Chaos - CCD sweep and behavioural improvements (for 1 CCD iteration)
- Custom SweepGeomCCD methods in Heightfield and TriMesh
- TOI calculated to first contact with TargetPenetration
- CCD does not apply impulses, it just moves bodies and leaves some penetration for the main solver to resolve
- Clean up some unused and duplicated functionality
- CVar to return to old functionality
- Fix CCD kinematics

diff --git a/Engine/Source/Runtime/Experimental/Chaos/Private/Chaos/CCDUtilities.cpp b/Engine/Source/Runtime/Experimental/Chaos/Private/Chaos/CCDUtilities.cpp
--- a/Engine/Source/Runtime/Experimental/Chaos/Private/Chaos/CCDUtilities.cpp
+++ b/Engine/Source/Runtime/Experimental/Chaos/Private/Chaos/CCDUtilities.cpp
@@ -28,9 +28,21 @@
 		int32 ChaosCollisionCCDConstraintMaxProcessCount = 1;
 		FAutoConsoleVariableRef CVarChaosCollisionCCDConstraintMaxProcessCount(TEXT("p.Chaos.Collision.CCD.ConstraintMaxProcessCount"), ChaosCollisionCCDConstraintMaxProcessCount, TEXT("The max number of times each constraint can be resolved when applying CCD constraints. Default is 2. The larger this number is, the more fully CCD constraints are resolved."));
 
+		// Determines when CCD is switched on
 		FRealSingle CCDEnableThresholdBoundsScale = 0.4f;
 		FAutoConsoleVariableRef  CVarCCDEnableThresholdBoundsScale(TEXT("p.Chaos.CCD.EnableThresholdBoundsScale"), CCDEnableThresholdBoundsScale , TEXT("CCD is used when object position is changing > smallest bound's extent * BoundsScale. 0 will always Use CCD. Values < 0 disables CCD."));
 
+		// Determines how much penetration CCD leaves after rewinding to TOI. Must be less than CCDEnableThresholdBoundsScale
+		Chaos::FRealSingle CCDAllowedDepthBoundsScale = 0.2f;
+		FAutoConsoleVariableRef CVarCCDAllowedDepthBoundsScale(TEXT("p.Chaos.CCD.AllowedDepthBoundsScale"), CCDAllowedDepthBoundsScale, TEXT("When rolling back to TOI, allow (smallest bound's extent) * AllowedDepthBoundsScale, instead of rolling back to exact TOI w/ penetration = 0."));
+
+		// When enabled, CCD TOI is set to the first contact which gets a depth of (CCDAllowedDepthBoundsScale * Object Size).
+		// When disabled, CCD will find the first contact, ignoreing the rest, and then fix the TOI to result in the specified depth.
+		// This fixes the issue where (when disabled) we would set TOI to some value that would leave us penetrating another (later) object. This
+		// could happen when colliding with the floor near the base of a ramp for example.
+		bool bCCDNewTargetDepthMode = true;
+		FAutoConsoleVariableRef  CVarCCDNewTargetDepthMode(TEXT("p.Chaos.CCD.NewTargetDepthMode"), bCCDNewTargetDepthMode, TEXT("Find the first contact with that results in a penetration of (CCDAllowedDepthBoundsScale*Size) as opposed to the first contact"));
+
 		int32 CCDAxisThresholdMode = 1;
 		FAutoConsoleVariableRef  CVarCCDAxisThresholdMode(TEXT("p.Chaos.CCD.AxisThresholdMode"), CCDAxisThresholdMode , TEXT("Change the mode used to generate CCD axis threshold bounds for particle geometries.\n0: Use object bounds\n1: Find the thinnest object bound on any axis and use it for all CCD axes\n2: On each axis, use the thinnest shape bound on that axis\n3: Find the thinnest shape bound on any axis and use this for all axes"));
 
@@ -51,22 +63,6 @@
 		AttachedCCDConstraints.Add(Constraint);
 	}
 
-	FReal GetParticleCCDThreshold(const FImplicitObject* Implicit)
-	{
-		if (Implicit)
-		{
-			// Trimesh/Heightfield are thin, cannot use bounds. We do not want them to contribute to CCD threshold.
-			if (Implicit->IsConvex())
-			{
-				const FReal MinExtent = Implicit->BoundingBox().Extents().Min();
-				return MinExtent * CVars::CCDEnableThresholdBoundsScale;
-			}
-
-			return 0;
-		}
-		return 0;
-	}
-
 	int32 FCCDConstraint::GetFastMovingKinematicIndex(const FPBDCollisionConstraint* Constraint, const FVec3 Displacements[]) const
 	{
 		for (int32 i = 0; i < 2; i++)
@@ -74,11 +70,7 @@
 			const TPBDRigidParticleHandle<FReal, 3>* Rigid = Constraint->GetParticle(i)->CastToRigidParticle();
 			if (Rigid && Rigid->ObjectState() == EObjectStateType::Kinematic)
 			{
-				// The same computation is carried out in UseCCDImpl function when constructing constraints. But we don't have access to FCCDConstraint at that point. This part could potentially be optimized away. 
-				const FVec3 D = Displacements[i];
-				const FReal DSizeSquared = D.SizeSquared();
-				const FReal CCDThreshold = GetParticleCCDThreshold(Constraint->GetImplicit(i));
-				if (DSizeSquared > CCDThreshold * CCDThreshold)
+				if (CCDHelpers::DeltaExceedsThreshold(Rigid->CCDAxisThreshold(), Displacements[i], Rigid->Q()))
 				{
 					return i;
 				}
@@ -106,12 +98,23 @@
 			return;
 		}
 
+		// Use simpler processing loop if we have only 1 CCD iteration
+		const bool bUseSimpleCCD = (CVars::ChaosCollisionCCDConstraintMaxProcessCount == 1);
+
 		AssignParticleIslandsAndGroupParticles();
 		AssignConstraintIslandsAndRecordConstraintNum();
 		GroupConstraintsWithIslands();
 		PhysicsParallelFor(IslandNum, [&](const int32 Island)
 		{
-			ApplyIslandSweptConstraints(Island, Dt);
+			if (bUseSimpleCCD)
+			{
+				ApplyIslandSweptConstraints2(Island, Dt);
+
+			}
+			else
+			{
+				ApplyIslandSweptConstraints(Island, Dt);
+			}
 		});
 	}
 
@@ -182,9 +185,7 @@
 			// make sure we ignore pairs that don't include any dynamics
 			if (CCDParticlePair[0] != nullptr || CCDParticlePair[1] != nullptr)
 			{
-				const FReal CCDConstraintThreshold = FMath::Min(Particle0->CCDAxisThreshold().GetMin(), Particle1->CCDAxisThreshold().GetMin());
-				const FReal PhiThreshold = -CCDConstraintThreshold;
-				CCDConstraints.Add(FCCDConstraint(Constraint, CCDParticlePair, Displacements, PhiThreshold));
+				CCDConstraints.Add(FCCDConstraint(Constraint, CCDParticlePair, Displacements));
 				for (int32 i = 0; i < 2; i++)
 				{
 					if (CCDParticlePair[i] != nullptr)
@@ -299,9 +300,163 @@
 
 	bool CCDConstraintSortPredicate(const FCCDConstraint* Constraint0, const FCCDConstraint* Constraint1)
 	{
-		return Constraint0->SweptConstraint->CCDTimeOfImpact < Constraint1->SweptConstraint->CCDTimeOfImpact;
+		return Constraint0->SweptConstraint->GetCCDTimeOfImpact() < Constraint1->SweptConstraint->GetCCDTimeOfImpact();
 	}
 
+	// Differences to the original ApplyIslandSweptConstraints
+	// - only 1 CCD iteration, and we roll back positions to the first contact on each body pair
+	// - if a body pair has multiple constraints (multiple shapes) only the first contact need to processed
+	// - no impulses are applied by CCD. NOTE: this only works if the TOI rollback leaves some penetration
+	//
+	void FCCDManager::ApplyIslandSweptConstraints2(const int32 Island, const FReal Dt)
+	{
+		const int32 ConstraintStart = IslandConstraintStart[Island];
+		const int32 ConstraintNum = IslandConstraintNum[Island];
+		const int32 ConstraintEnd = IslandConstraintEnd[Island];
+		
+		if (ConstraintNum == 0)
+		{
+			return;
+		}
+
+#if CHAOS_DEBUG_DRAW
+		if (CVars::ChaosSolverDrawCCDInteractions)
+		{
+			// Debugdraw the shape at the TOI=0 (black) and TOI=1 (white)
+			for (FCCDConstraint* CCDConstraint : SortedCCDConstraints)
+			{
+				DebugDraw::DrawCCDCollisionShape(FRigidTransform3::Identity, *CCDConstraint, true, FColor::Black, &CVars::ChaosSolverDebugDebugDrawSettings);
+				DebugDraw::DrawCCDCollisionShape(FRigidTransform3::Identity, *CCDConstraint, false, FColor::White, &CVars::ChaosSolverDebugDebugDrawSettings);
+			}
+		}
+#endif
+
+		// Sort constraints based on TOI
+		FReal IslandTOI = 0.f;
+		ResetIslandParticles(Island);
+		ResetIslandConstraints(Island);
+		bool bSortRequired = true;
+		int32 ConstraintIndex = ConstraintStart;
+		while (ConstraintIndex < ConstraintEnd)
+		{
+			// If we updated some constraints, we need to sort so that we handle the next TOI event
+			if (bSortRequired)
+			{
+				std::sort(SortedCCDConstraints.GetData() + ConstraintIndex, SortedCCDConstraints.GetData() + ConstraintStart + ConstraintNum, CCDConstraintSortPredicate);
+				bSortRequired = false;
+			}
+
+			FCCDConstraint* CCDConstraint = SortedCCDConstraints[ConstraintIndex];
+			FCCDParticle* CCDParticle0 = CCDConstraint->Particle[0];
+			FCCDParticle* CCDParticle1 = CCDConstraint->Particle[1];
+
+			IslandTOI = CCDConstraint->SweptConstraint->GetCCDTimeOfImpact();
+
+			// Constraints whose TOIs are in the range of [0, 1) are resolved for this frame. TOI = 1 means that the two 
+			// particles just start touching at the end of the frame and therefore cannot have tunneling this frame. 
+			// So this TOI = 1 can be left to normal collisions or CCD in next frame.
+			if (IslandTOI > 1)
+			{
+				break;
+			}
+
+			// If both particles are marked Done continue
+			const bool bParticle0Done = ((CCDParticle0 == nullptr) || CCDParticle0->Done);
+			const bool bParticle1Done = ((CCDParticle1 == nullptr) || CCDParticle1->Done);
+			if (bParticle0Done && bParticle1Done)
+			{
+				ConstraintIndex++;
+				continue;
+			}
+
+			CCDConstraint->ProcessedCount++;
+
+			// In UpdateConstraintSwept, InitManifoldPoint requires P, Q to be at TOI=1., but the input of 
+			// UpdateConstraintSwept requires transforms at current TOI. So instead of rewinding P, Q, we 
+			// advance X, R to current TOI and keep P, Q at TOI=1.
+			if (CCDParticle0 && !CCDParticle0->Done)
+			{
+				AdvanceParticleXToTOI(CCDParticle0, IslandTOI, Dt);
+			}
+			if (CCDParticle1 && !CCDParticle1->Done)
+			{
+				AdvanceParticleXToTOI(CCDParticle1, IslandTOI, Dt);
+			}
+
+#if CHAOS_DEBUG_DRAW
+			// Debugdraw the shape at the current TOI
+			if (CVars::ChaosSolverDrawCCDInteractions)
+			{
+				DebugDraw::DrawCCDCollisionShape(FRigidTransform3::Identity, *CCDConstraint, true, FColor::Magenta, &CVars::ChaosSolverDebugDebugDrawSettings);
+			}
+#endif
+
+			// We don't normally apply impulses in the CCD phase because we leave on overlap for the main solver to handle. However if the
+			// settings are such that this won't happen, we must apply the impulse here. CCD impulses assume point masses and no rotations or friction are applied.
+			const bool bApplyImpulse = (CVars::CCDAllowedDepthBoundsScale <= 0);
+			if (bApplyImpulse)
+			{
+				ApplyImpulse(CCDConstraint);
+			}
+
+			if (CCDParticle0)
+			{
+				if (CCDConstraint->FastMovingKinematicIndex != INDEX_NONE)
+				{
+					// @todo(chaos): can this just get the particle from the CCD constraint?
+					const FConstGenericParticleHandle KinematicParticle = FGenericParticleHandle(CCDConstraint->SweptConstraint->GetParticle(CCDConstraint->FastMovingKinematicIndex));
+					const FVec3 Normal = CCDConstraint->SweptConstraint->CalculateWorldContactNormal();
+					const FVec3 Offset = FVec3::DotProduct(KinematicParticle->V() * ((1.f - IslandTOI) * Dt), Normal) * Normal;
+					ClipParticleP(CCDParticle0, Offset);
+				}
+				else
+				{
+					ClipParticleP(CCDParticle0);
+				}
+				CCDParticle0->Done = true;
+			}
+			if (CCDParticle1)
+			{
+				if (CCDConstraint->FastMovingKinematicIndex != INDEX_NONE)
+				{
+					// @todo(chaos): can this just get the particle from the CCD constraint?
+					const FConstGenericParticleHandle KinematicParticle = FGenericParticleHandle(CCDConstraint->SweptConstraint->GetParticle(CCDConstraint->FastMovingKinematicIndex));
+					const FVec3 Normal = CCDConstraint->SweptConstraint->CalculateWorldContactNormal();
+					const FVec3 Offset = FVec3::DotProduct(KinematicParticle->V() * ((1.f - IslandTOI) * Dt), Normal) * Normal;
+					ClipParticleP(CCDParticle1, Offset);
+				}
+				else
+				{
+					ClipParticleP(CCDParticle1);
+				}
+				CCDParticle1->Done = true;
+			}
+
+			// We applied a CCD impulse and updated the particle positions, so we need to update all the constraints involving these particles
+			if (CCDParticle0)
+			{
+				bSortRequired |= UpdateParticleSweptConstraints(CCDParticle0, IslandTOI, Dt);
+			}
+			if (CCDParticle1)
+			{
+				bSortRequired |= UpdateParticleSweptConstraints(CCDParticle1, IslandTOI, Dt);
+			}
+
+			ConstraintIndex++;
+		}
+
+#if CHAOS_DEBUG_DRAW
+		// Debugdraw the shapes at the final position
+		if (CVars::ChaosSolverDrawCCDInteractions)
+		{
+			for (FCCDConstraint* CCDConstraint : SortedCCDConstraints)
+			{
+				DebugDraw::DrawCCDCollisionShape(FRigidTransform3::Identity, *CCDConstraint, false, FColor::Green, &CVars::ChaosSolverDebugDebugDrawSettings);
+			}
+		}
+#endif
+	}
+
 	void FCCDManager::ApplyIslandSweptConstraints(const int32 Island, const FReal Dt)
 	{
 		const int32 ConstraintStart = IslandConstraintStart[Island];
@@ -333,7 +488,7 @@
 			FCCDParticle* CCDParticle0 = CCDConstraint->Particle[0];
 			FCCDParticle* CCDParticle1 = CCDConstraint->Particle[1];
 
-			IslandTOI = CCDConstraint->SweptConstraint->CCDTimeOfImpact;
+			IslandTOI = CCDConstraint->SweptConstraint->GetCCDTimeOfImpact();
 
 			// Constraints whose TOIs are in the range of [0, 1) are resolved for this frame. TOI = 1 means that the two 
 			// particles just start touching at the end of the frame and therefore cannot have tunneling this frame. 
@@ -378,7 +533,7 @@
 			// After applying impulse, constraint TOI need be updated to reflect the new velocities. 
 			// Usually the new velocities are separating, and therefore TOI should be infinity.
 			// See resweep below which (optionally) updates TOI for all other contacts as a result of handling this one
-			CCDConstraint->SweptConstraint->CCDTimeOfImpact = TNumericLimits<FReal>::Max();
+			CCDConstraint->SweptConstraint->SetCCDTimeOfImpact(TNumericLimits<FReal>::Max());
 
 			bool bMovedParticle0 = false;
 			bool bMovedParticle1 = false;
@@ -512,6 +667,13 @@
 					continue;
 				}
 
+				const bool bParticle0Done = ((AttachedCCDConstraint->Particle[0] == nullptr) || AttachedCCDConstraint->Particle[0]->Done);
+				const bool bParticle1Done = ((AttachedCCDConstraint->Particle[1] == nullptr) || AttachedCCDConstraint->Particle[1]->Done);
+				if (bParticle0Done && bParticle1Done)
+				{
+					continue;
+				}
+
 				// Particle transforms at TOI
 				FRigidTransform3 ParticleStartWorldTransforms[2];
 				for (int32 j = 0; j < 2; j++)
@@ -567,21 +729,13 @@
 				// Update of swept constraint assumes that the constraint holds the end transforms for the sweep
 				SweptConstraint->SetShapeWorldTransforms(ShapeEndWorldTransform0, ShapeEndWorldTransform1);
 
-				if (CVars::bChaosCollisionCCDEnableResweep)
-				{
-					// Resweep the shape. This is the expensive option
-					Collisions::UpdateConstraintSwept(*SweptConstraint, ShapeStartWorldTransform0, ShapeStartWorldTransform1, RestDt);
-				}
-				else
-				{
-					// Keep the contact as-is but update the depth and TOI based on current transforms
-					SweptConstraint->UpdateSweptManifoldPoints(ShapeStartWorldTransform0.GetTranslation(), ShapeStartWorldTransform1.GetTranslation(), Dt);
-				}
+				// Calculate the contact point and TOI
+				Collisions::UpdateConstraintSwept(*SweptConstraint, ShapeStartWorldTransform0, ShapeStartWorldTransform1, RestDt);
 
-				const FReal RestDtTOI = AttachedCCDConstraint->SweptConstraint->CCDTimeOfImpact;
+				const FReal RestDtTOI = AttachedCCDConstraint->SweptConstraint->GetCCDTimeOfImpact();
 				if ((RestDtTOI >= 0) && (RestDtTOI < FReal(1)))
 				{
-					AttachedCCDConstraint->SweptConstraint->CCDTimeOfImpact = IslandTOI + (FReal(1) - IslandTOI) * RestDtTOI;
+					AttachedCCDConstraint->SweptConstraint->SetCCDTimeOfImpact(IslandTOI + (FReal(1) - IslandTOI) * RestDtTOI);
 				}
 
 				// When bUpdated==true, TOI was modified. When bUpdated==false, TOI was set to be TNumericLimits<FReal>::Max(). In either case, a re-sorting on the constraints is needed.
diff --git a/Engine/Source/Runtime/Experimental/Chaos/Private/Chaos/ChaosDebugDraw.cpp b/Engine/Source/Runtime/Experimental/Chaos/Private/Chaos/ChaosDebugDraw.cpp
--- a/Engine/Source/Runtime/Experimental/Chaos/Private/Chaos/ChaosDebugDraw.cpp
+++ b/Engine/Source/Runtime/Experimental/Chaos/Private/Chaos/ChaosDebugDraw.cpp
@@ -79,11 +79,11 @@
 		float ChaosDebugDrawConvexExplodeDistance = 0.0f;
 		FAutoConsoleVariableRef CVarChaosDebugDrawConvexExplodeDistance(TEXT("p.Chaos.DebugDraw.ConvexExplodeDistance"), ChaosDebugDrawConvexExplodeDistance, TEXT("Explode convex edges by this amount (useful for looking at convex integrity)"));
 
-		float ChaosDebugDrawCCDDuration = 2.0f;
-		FAutoConsoleVariableRef CVarChaosDebugDrawCCDDuration(TEXT("p.Chaos.DebugDraw.CCDDuration"), ChaosDebugDrawCCDDuration, TEXT("How long CCD debug draw should remain on screen"));
+		float ChaosDebugDrawCCDDuration = 0.0f;
+		FAutoConsoleVariableRef CVarChaosDebugDrawCCDDuration(TEXT("p.Chaos.DebugDraw.CCDDuration"), ChaosDebugDrawCCDDuration, TEXT("How long CCD debug draw should remain on screen in seconds. 0 for 1 frame."));
 
 		float ChaosDebugDrawCollisionDuration = 0.0f;
-		FAutoConsoleVariableRef CVarChaosDebugDrawCollisionDuration(TEXT("p.Chaos.DebugDraw.CollisionDuration"), ChaosDebugDrawCollisionDuration, TEXT("How long Collision debug draw should remain on screen"));
+		FAutoConsoleVariableRef CVarChaosDebugDrawCollisionDuration(TEXT("p.Chaos.DebugDraw.CollisionDuration"), ChaosDebugDrawCollisionDuration, TEXT("How long Collision debug draw should remain on screen in seconds. 0 for 1 frame."));
 
 
 		// NOTE: These settings should never really be used - they are the fallback defaults
@@ -1043,6 +1043,7 @@
 
 			DrawParticleCCDCollisionShapeImpl(SpaceTransform, CCDConstraint.Particle[0], bShowStartPos, ShapeColor, Duration, Settings);
 			DrawParticleCCDCollisionShapeImpl(SpaceTransform, CCDConstraint.Particle[1], bShowStartPos, ShapeColor, Duration, Settings);
+			DrawCollisionImpl(SpaceTransform, CCDConstraint.SweptConstraint, 1.0f, Duration, Settings);
 		}
 
 		void DrawCCDCollisionImpulseImpl(const FRigidTransform3& SpaceTransform, const FCCDConstraint& CCDConstraint, const int32 ManifoldPointIndex, const FVec3& Impulse, const FChaosDebugDrawSettings& Settings)
diff --git a/Engine/Source/Runtime/Experimental/Chaos/Private/Chaos/Collision/ContactPointsMiscShapes.cpp b/Engine/Source/Runtime/Experimental/Chaos/Private/Chaos/Collision/ContactPointsMiscShapes.cpp
--- a/Engine/Source/Runtime/Experimental/Chaos/Private/Chaos/Collision/ContactPointsMiscShapes.cpp
+++ b/Engine/Source/Runtime/Experimental/Chaos/Private/Chaos/Collision/ContactPointsMiscShapes.cpp
@@ -24,43 +24,45 @@
 
 namespace Chaos
 {
-
 	template <typename GeometryB>
-	FContactPoint GJKImplicitSweptContactPoint(const FImplicitObject& A, const FRigidTransform3& AStartTransform, const GeometryB& B, const FRigidTransform3& BTransform, const FVec3& Dir, const FReal Length, FReal& TOI)
+	FContactPoint GJKImplicitSweptContactPoint(const FImplicitObject& A, const FRigidTransform3& AStartTransform, const GeometryB& B, const FRigidTransform3& BTransform, const FVec3& Dir, const FReal Length, const FReal IgnorePenetration, const FReal TargetPenetration, FReal& OutTOI)
 	{
 		FContactPoint Contact;
 		const FRigidTransform3 AToBTM = AStartTransform.GetRelativeTransform(BTransform);
 		const FVec3 LocalDir = BTransform.InverseTransformVectorNoScale(Dir);
 
-		FReal OutTime = FLT_MAX;
-		int32 FaceIndex = -1;
+		FReal TOI = TNumericLimits<FReal>::Max();
+		FReal Phi = TNumericLimits<FReal>::Max();
+		int32 FaceIndex = INDEX_NONE;
 		FVec3 FaceNormal;
 		FVec3 Location, Normal;
 
 		Utilities::CastHelper(A, AStartTransform, [&](const auto& ADowncast, const FRigidTransform3& AFullTM)
 			{
-				// @todo(chaos): handle instances with margin
-				if (B.SweepGeom(ADowncast, AToBTM, LocalDir, Length, OutTime, Location, Normal, FaceIndex, FaceNormal, 0.0f, true))
+				if (B.SweepGeomCCD(ADowncast, AToBTM, LocalDir, Length, IgnorePenetration, TargetPenetration, TOI, Phi, Location, Normal, FaceIndex, FaceNormal))
 				{
-					// @todo(chaos): margin
 					Contact.ShapeContactPoints[0] = AToBTM.InverseTransformPosition(Location);
 					Contact.ShapeContactPoints[1] = Location;
 					Contact.ShapeContactNormal = Normal;
-					const FVec3 ContactNormal = BTransform.TransformVectorNoScale(Normal);
-					ComputeSweptContactPhiAndTOIHelper(ContactNormal, Dir, Length, OutTime, TOI, Contact.Phi);
+					Contact.Phi = Phi;
+					Contact.FaceIndex = FaceIndex;
+					OutTOI = TOI;
 				}
 			});
 
 		return Contact;
 	}
 
-	void ComputeSweptContactPhiAndTOIHelper(const FVec3& ContactNormal, const FVec3& Dir, const FReal& Length, const FReal& HitTime, FReal& OutTOI, FReal& OutPhi)
+	// GJKRaycast Time output depends on whether there was initial overlap.
+	// For no initial overlap, Time is the Distance along the sweep where the contact occurs
+	// For initial overlap, Time is the Phi value (separation, so negative)
+	void ComputeSweptContactPhiAndTOIHelper(const FReal DirDotNormal, const FReal& Length, const FReal& HitTime, FReal& OutTOI, FReal& OutPhi)
 	{
 		if (HitTime >= 0.0f)
 		{
 			// We subtract length to get the total penetration at at end of frame.
 			// Project penetration vector onto geometry normal for correct phi.
-			FReal Dot = FMath::Abs(FVec3::DotProduct(ContactNormal, -Dir));
+			FReal Dot = -DirDotNormal;
 			OutPhi = (HitTime - Length) * Dot;
 
 			// TOI is between [0,1], used to compute particle position
@@ -76,6 +78,120 @@
 		}
 	}
 
+	void ComputeSweptContactPhiAndTOIHelper(const FVec3& ContactNormal, const FVec3& Dir, const FReal& Length, const FReal& HitTime, FReal& OutTOI, FReal& OutPhi)
+	{
+		const FReal DirDotNormal = FMath::Abs(FVec3::DotProduct(ContactNormal, Dir));
+		ComputeSweptContactPhiAndTOIHelper(DirDotNormal, Length, HitTime, OutTOI, OutPhi);
+	}
+
+	// Calculate Phi at the start and end of a raycast, given the ray information and sweep result
+	bool ComputeSweptContactStartAndEndPhi(const FReal TOI, const FReal Phi, const FReal DirDotNormal, const FReal Length, FReal& StartPhi, FReal& EndPhi)
+	{
+		StartPhi = TNumericLimits<FReal>::Max();
+		EndPhi = TNumericLimits<FReal>::Max();
+		if (TOI <= FReal(1))
+		{
+			// Modify TOI so we ignore separating and shallow CCD contacts
+			// NOTE (See GJKRaycast2): If initially penetrating, Phi is the separation at the start of the sweep.
+			// But if not initially penetrating Phi is the separation at the end of the sweep
+			// @todo(chaos): The way Phi is set for sweeps is quite confusing - can we make this better?
+			if (TOI <= FReal(0))
+			{
+				StartPhi = Phi;
+				EndPhi = StartPhi + DirDotNormal * Length;
+			}
+			else
+			{
+				EndPhi = Phi;
+				StartPhi = EndPhi - DirDotNormal * Length;
+			}
+			return true;
+		}
+		return false;
+	}
+
+	// Use the ray start and end Phi to calculate the time when Phi was equal to TargetPhi. Also ignore separating rays and
+	// those where penatration is belocw some threshold at the end of the sweep.
+	FReal ComputeSweptContactTimeToTargetPhi(const FReal StartPhi, const FReal EndPhi, const FReal IgnorePhi, const FReal TargetPhi)
+	{
+		const FReal InfiniteTOI = TNumericLimits<FReal>::Max();
+		const FReal MovementTolerance = KINDA_SMALL_NUMBER;
+
+		// If we end up separated at TOI=1 ignore the contact
+		if (EndPhi > 0)
+		{
+			return InfiniteTOI;
+		}
+
+		// If contact is moving in the right direction or not moving ignore the contact
+		if (EndPhi > StartPhi - MovementTolerance)
+		{
+			return InfiniteTOI;
+		}
+
+		// If we penetrate by less than the IgnorePhi, treat it as TOI=1. This mean no CCD impulse and the non-CCD 
+		// solve is expected to handle it. This improves the behaviour when we are sliding along a surface at
+		// above CCD speeds - we don't want to handle TOI events with the floor
+		if (EndPhi > IgnorePhi)
+		{
+			return FReal(1);
+		}
+
+		// If we penetrate by more than the TargetPhi we roll back to the TOI-plus-a-bit leaving some penetration
+		// for the main solver to resolve. Note that is CCDRemainderPenetration is non-zero, secondary sweeps are
+		// more likely to encounter an bad edge collision.
+		const FReal TOI = (TargetPhi - StartPhi) / (EndPhi - StartPhi);
+
+		return FMath::Clamp(TOI, FReal(0), FReal(1));
+	}
+
+	// Modify the time of impact so that the contact depth is TargetPenetration. If penetration at T=1 is less than TargetPenetration, TOI will be "infinity".
+	// Returns true if we have a TOI less than 1 (i.e., a contact at TargetPenetration or more that needs CCD processing)
+	bool ComputeSweptContactTOIAndPhiAtTargetPenetration(const FReal DirDotNormal, const FReal SweepLength, const FReal HitDistance, const FReal IgnorePenetration, const FReal TargetPenetration, FReal& OutTOI, FReal& OutPhi)
+	{
+		// Parse the GJK result to get TOI and Phi (Interpretation depends on whether TOI=0 or not - see GJKRaycast2)
+		FReal HitTOI, HitPhi;
+		ComputeSweptContactPhiAndTOIHelper(DirDotNormal, SweepLength, HitDistance, HitTOI, HitPhi);
+
+		// Calculate the Phis at the start and end of the sweep
+		FReal StartPhi, EndPhi;
+		ComputeSweptContactStartAndEndPhi(HitTOI, HitPhi, DirDotNormal, SweepLength, StartPhi, EndPhi);
+
+		// Calculate the time to reach a depth of TargetPenetration (ignoring those penetrating by less than IgnorePenetration at the end)
+		const FReal IgnorePhi = -IgnorePenetration;
+		const FReal TargetPhi = -TargetPenetration;
+		const FReal TargetTOI = ComputeSweptContactTimeToTargetPhi(StartPhi, EndPhi, IgnorePhi, TargetPhi);
+
+		OutTOI = TargetTOI;
+		OutPhi = (TargetTOI == 0) ? StartPhi : TargetPhi;
+
+		return (OutTOI < FReal(1));
+	}
+
+	bool ComputeSweptContactTOIAndPhiAtTargetPenetration(const FVec3& ContactNormal, const FVec3& Dir, const FReal SweepLength, const FReal HitDistance, const FReal IgnorePenetration, const FReal TargetPenetration, FReal& OutTOI, FReal& OutPhi)
+	{
+		const FReal DirDotNormal = FVec3::DotProduct(ContactNormal, Dir);
+		return ComputeSweptContactTOIAndPhiAtTargetPenetration(DirDotNormal, SweepLength, HitDistance, IgnorePenetration, TargetPenetration, OutTOI, OutPhi);
+	}
+
+	// Calculate TOI to leave objects penetrating by TargetPhi. This is used by the legacy path that finds the first contacting shape pair and then ignores all others after that
+	// even if they would have a smaller TOI at TargetPenetration. It is only called if the low level sweep code has not already called ComputeSweptContactTOIAndPhiAtTargetPenetration. 
+	// See CVars::bCCDNewTargetDepthMode
+	void LegacyComputeSweptContactTOIAndPhiAtTargetPenetration(const FReal DirDotNormal, const FReal Length, const FReal IgnorePenetration, const FReal TargetPenetration, FReal& InOutTOI, FReal& InOutPhi)
+	{
+		if (InOutTOI <= 1)
+		{
+			if (InOutTOI != 0)
+			{
+				InOutPhi = (FReal(1) - InOutTOI) * Length * DirDotNormal;
+			}
+			FReal StartPhi, EndPhi;
+			ComputeSweptContactStartAndEndPhi(InOutTOI, InOutPhi, DirDotNormal, Length, StartPhi, EndPhi);
+			InOutTOI = ComputeSweptContactTimeToTargetPhi(StartPhi, EndPhi, -IgnorePenetration, -TargetPenetration);
+			InOutPhi = (InOutTOI == 0) ? StartPhi : -TargetPenetration;
+		}
+	}
+
 	template <typename GeometryA, typename GeometryB>
 	FContactPoint GJKImplicitContactPoint(const FImplicitObject& A, const FRigidTransform3& ATransform, const GeometryB& B, const FRigidTransform3& BTransform, const FReal CullDistance, const FReal ShapePadding)
 	{
@@ -145,7 +261,7 @@
 
 	// A is the implicit here, we want to return a contact point on B (trimesh)
 	template <typename GeometryA>
-	FContactPoint GJKImplicitScaledTriMeshSweptContactPoint(const FImplicitObject& A, const FRigidTransform3& AStartTransform, const TImplicitObjectScaled<FTriangleMeshImplicitObject>& B, const FRigidTransform3& BTransform, const FVec3& Dir, const FReal Length, FReal& TOI)
+	FContactPoint GJKImplicitScaledTriMeshSweptContactPoint(const FImplicitObject& A, const FRigidTransform3& AStartTransform, const TImplicitObjectScaled<FTriangleMeshImplicitObject>& B, const FRigidTransform3& BTransform, const FVec3& Dir, const FReal Length, const FReal IgnorePenetration, const FReal TargetPenetration, FReal& OutTOI)
 	{
 		FContactPoint Contact;
 		const FRigidTransform3 AToBTM = AStartTransform.GetRelativeTransform(BTransform);
@@ -156,22 +272,20 @@
 			return FContactPoint();
 		}
 
-		FReal OutTime = FLT_MAX;
-		FVec3 Location, Normal;
-		int32 FaceIndex = -1;
-		Chaos::FVec3 FaceNormal;
-
 		Utilities::CastHelper(A, AStartTransform, [&](const auto& ADowncast, const FRigidTransform3& AFullTM)
 			{
-				// @todo(chaos): handle Instanced with margin
-				if (B.LowLevelSweepGeom(ADowncast, AToBTM, LocalDir, Length, OutTime, Location, Normal, FaceIndex, FaceNormal, 0.0f, true))
+				FReal TOI, Phi;
+				FVec3 Location, Normal;
+				int32 FaceIndex = -1;
+				Chaos::FVec3 FaceNormal;
+				if (B.LowLevelSweepGeomCCD(ADowncast, AToBTM, LocalDir, Length, IgnorePenetration, TargetPenetration, TOI, Phi, Location, Normal, FaceIndex, FaceNormal))
 				{
 					Contact.ShapeContactPoints[0] = AToBTM.InverseTransformPositionNoScale(Location);
 					Contact.ShapeContactPoints[1] = Location;
 					Contact.ShapeContactNormal = Normal;
-
-					const FVec3& ContactNormal = BTransform.TransformVectorNoScale(Normal);
-					ComputeSweptContactPhiAndTOIHelper(ContactNormal, Dir, Length, OutTime, TOI, Contact.Phi);
+					Contact.Phi = Phi;
+					Contact.FaceIndex = FaceIndex;
+					OutTOI = TOI;
 				}
 			});
 
@@ -284,15 +398,15 @@
 	}
 
 	template<typename TriMeshType>
-	FContactPoint SphereTriangleMeshSweptContactPoint(const TSphere<FReal, 3>& A, const FRigidTransform3& ATransform, const TriMeshType& B, const FRigidTransform3& BStartTransform, const FVec3& Dir, const FReal Length, FReal& TOI)
+	FContactPoint SphereTriangleMeshSweptContactPoint(const TSphere<FReal, 3>& A, const FRigidTransform3& ATransform, const TriMeshType& B, const FRigidTransform3& BStartTransform, const FVec3& Dir, const FReal Length, const FReal IgnorePenetration, const FReal TargetPenetration, FReal& TOI)
 	{
 		if (const TImplicitObjectScaled<FTriangleMeshImplicitObject>* ScaledTriangleMesh = B.template GetObject<const TImplicitObjectScaled<FTriangleMeshImplicitObject>>())
 		{
-			return GJKImplicitScaledTriMeshSweptContactPoint<TSphere<FReal, 3>>(A, ATransform, *ScaledTriangleMesh, BStartTransform, Dir, Length, TOI);
+			return GJKImplicitScaledTriMeshSweptContactPoint<TSphere<FReal, 3>>(A, ATransform, *ScaledTriangleMesh, BStartTransform, Dir, Length, IgnorePenetration, TargetPenetration, TOI);
 		}
 		else if (const FTriangleMeshImplicitObject* TriangleMesh = B.template GetObject<const FTriangleMeshImplicitObject>())
 		{
-			return GJKImplicitSweptContactPoint(TSphere<FReal, 3>(A), ATransform, *TriangleMesh, BStartTransform, Dir, Length, TOI);
+			return GJKImplicitSweptContactPoint(TSphere<FReal, 3>(A), ATransform, *TriangleMesh, BStartTransform, Dir, Length, IgnorePenetration, TargetPenetration, TOI);
 		}
 
 		ensure(false);
@@ -329,16 +443,16 @@
 	}
 
 	template <typename TriMeshType>
-	FContactPoint CapsuleTriangleMeshSweptContactPoint(const FCapsule& A, const FRigidTransform3& ATransform, const TriMeshType& B, const FRigidTransform3& BStartTransform, const FVec3& Dir, const FReal Length, FReal& TOI)
+	FContactPoint CapsuleTriangleMeshSweptContactPoint(const FCapsule& A, const FRigidTransform3& ATransform, const TriMeshType& B, const FRigidTransform3& BStartTransform, const FVec3& Dir, const FReal Length, const FReal IgnorePenetration, const FReal TargetPenetration, FReal& TOI)
 	{
 		CONDITIONAL_SCOPE_CYCLE_COUNTER(STAT_Collisions_CapsuleTriangleMeshSweptContactPoint, ConstraintsDetailedStats);
 		if (const TImplicitObjectScaled<FTriangleMeshImplicitObject>* ScaledTriangleMesh = B.template GetObject<const TImplicitObjectScaled<FTriangleMeshImplicitObject>>())
 		{
-			return GJKImplicitScaledTriMeshSweptContactPoint<FCapsule>(A, ATransform, *ScaledTriangleMesh, BStartTransform, Dir, Length, TOI);
+			return GJKImplicitScaledTriMeshSweptContactPoint<FCapsule>(A, ATransform, *ScaledTriangleMesh, BStartTransform, Dir, Length, IgnorePenetration, TargetPenetration, TOI);
 		}
 		else if (const FTriangleMeshImplicitObject* TriangleMesh = B.template GetObject<const FTriangleMeshImplicitObject>())
 		{
-			return GJKImplicitSweptContactPoint(A, ATransform, *TriangleMesh, BStartTransform, Dir, Length, TOI);
+			return GJKImplicitSweptContactPoint(A, ATransform, *TriangleMesh, BStartTransform, Dir, Length, IgnorePenetration, TargetPenetration, TOI);
 		}
 
 		ensure(false);
@@ -365,16 +479,16 @@
 	}
 
 	template <typename TriMeshType>
-	FContactPoint ConvexTriangleMeshSweptContactPoint(const FImplicitObject& A, const FRigidTransform3& ATransform, const TriMeshType& B, const FRigidTransform3& BStartTransform, const FVec3& Dir, const FReal Length, FReal& TOI)
+	FContactPoint ConvexTriangleMeshSweptContactPoint(const FImplicitObject& A, const FRigidTransform3& ATransform, const TriMeshType& B, const FRigidTransform3& BStartTransform, const FVec3& Dir, const FReal Length, const FReal IgnorePenetration, const FReal TargetPenetration, FReal& TOI)
 	{
 		CONDITIONAL_SCOPE_CYCLE_COUNTER(STAT_Collisions_ConvexTriangleMeshSweptContactPoint, ConstraintsDetailedStats);
 		if (const TImplicitObjectScaled<FTriangleMeshImplicitObject>* ScaledTriangleMesh = B.template GetObject<const TImplicitObjectScaled<FTriangleMeshImplicitObject>>())
 		{
-			return GJKImplicitScaledTriMeshSweptContactPoint<FConvex>(A, ATransform, *ScaledTriangleMesh, BStartTransform, Dir, Length, TOI);
+			return GJKImplicitScaledTriMeshSweptContactPoint<FConvex>(A, ATransform, *ScaledTriangleMesh, BStartTransform, Dir, Length, IgnorePenetration, TargetPenetration, TOI);
 		}
 		else if (const FTriangleMeshImplicitObject* TriangleMesh = B.template GetObject<const FTriangleMeshImplicitObject>())
 		{
-			return GJKImplicitSweptContactPoint(A, ATransform, *TriangleMesh, BStartTransform, Dir, Length, TOI);
+			return GJKImplicitSweptContactPoint(A, ATransform, *TriangleMesh, BStartTransform, Dir, Length, IgnorePenetration, TargetPenetration, TOI);
 		}
 
 		ensure(false);
@@ -418,14 +532,14 @@
 
 
 	// Template  Instantiations
-	template FContactPoint GJKImplicitSweptContactPoint<FHeightField>(const FImplicitObject& A, const FRigidTransform3& AStartTransform, const FHeightField& B, const FRigidTransform3& BTransform, const FVec3& Dir, const FReal Length, FReal& TOI);
+	template FContactPoint GJKImplicitSweptContactPoint<FHeightField>(const FImplicitObject& A, const FRigidTransform3& AStartTransform, const FHeightField& B, const FRigidTransform3& BTransform, const FVec3& Dir, const FReal Length, const FReal IgnorePenetration, const FReal TargetPenetration, FReal& TOI);
 
-	template FContactPoint SphereTriangleMeshSweptContactPoint<TImplicitObjectScaled<class Chaos::FTriangleMeshImplicitObject, 1>>(const TSphere<FReal, 3>& A, const FRigidTransform3& ATransform, const TImplicitObjectScaled<class Chaos::FTriangleMeshImplicitObject, 1>& B, const FRigidTransform3& BStartTransform, const FVec3& Dir, const FReal Length, FReal& TOI);
-	template FContactPoint SphereTriangleMeshSweptContactPoint<FTriangleMeshImplicitObject>(const TSphere<FReal, 3>& A, const FRigidTransform3& ATransform, const FTriangleMeshImplicitObject& B, const FRigidTransform3& BStartTransform, const FVec3& Dir, const FReal Length, FReal& TOI);
-	template FContactPoint CapsuleTriangleMeshSweptContactPoint<TImplicitObjectScaled<class Chaos::FTriangleMeshImplicitObject, 1>>(const FCapsule& A, const FRigidTransform3& ATransform, const TImplicitObjectScaled<class Chaos::FTriangleMeshImplicitObject, 1>& B, const FRigidTransform3& BStartTransform, const FVec3& Dir, const FReal Length, FReal& TOI);
-	template FContactPoint CapsuleTriangleMeshSweptContactPoint<FTriangleMeshImplicitObject>(const FCapsule& A, const FRigidTransform3& ATransform, const FTriangleMeshImplicitObject& B, const FRigidTransform3& BStartTransform, const FVec3& Dir, const FReal Length, FReal& TOI);
-	template FContactPoint ConvexTriangleMeshSweptContactPoint<TImplicitObjectScaled<class Chaos::FTriangleMeshImplicitObject, 1>>(const FImplicitObject& A, const FRigidTransform3& ATransform, const TImplicitObjectScaled<class Chaos::FTriangleMeshImplicitObject, 1>& B, const FRigidTransform3& BStartTransform, const FVec3& Dir, const FReal Length, FReal& TOI);
-	template FContactPoint ConvexTriangleMeshSweptContactPoint<FTriangleMeshImplicitObject>(const FImplicitObject& A, const FRigidTransform3& ATransform, const FTriangleMeshImplicitObject& B, const FRigidTransform3& BStartTransform, const FVec3& Dir, const FReal Length, FReal& TOI);
+	template FContactPoint SphereTriangleMeshSweptContactPoint<TImplicitObjectScaled<class Chaos::FTriangleMeshImplicitObject, 1>>(const TSphere<FReal, 3>& A, const FRigidTransform3& ATransform, const TImplicitObjectScaled<class Chaos::FTriangleMeshImplicitObject, 1>& B, const FRigidTransform3& BStartTransform, const FVec3& Dir, const FReal Length, const FReal IgnorePenetration, const FReal TargetPenetration, FReal& TOI);
+	template FContactPoint SphereTriangleMeshSweptContactPoint<FTriangleMeshImplicitObject>(const TSphere<FReal, 3>& A, const FRigidTransform3& ATransform, const FTriangleMeshImplicitObject& B, const FRigidTransform3& BStartTransform, const FVec3& Dir, const FReal Length, const FReal IgnorePenetration, const FReal TargetPenetration, FReal& TOI);
+	template FContactPoint CapsuleTriangleMeshSweptContactPoint<TImplicitObjectScaled<class Chaos::FTriangleMeshImplicitObject, 1>>(const FCapsule& A, const FRigidTransform3& ATransform, const TImplicitObjectScaled<class Chaos::FTriangleMeshImplicitObject, 1>& B, const FRigidTransform3& BStartTransform, const FVec3& Dir, const FReal Length, const FReal IgnorePenetration, const FReal TargetPenetration, FReal& TOI);
+	template FContactPoint CapsuleTriangleMeshSweptContactPoint<FTriangleMeshImplicitObject>(const FCapsule& A, const FRigidTransform3& ATransform, const FTriangleMeshImplicitObject& B, const FRigidTransform3& BStartTransform, const FVec3& Dir, const FReal Length, const FReal IgnorePenetration, const FReal TargetPenetration, FReal& TOI);
+	template FContactPoint ConvexTriangleMeshSweptContactPoint<TImplicitObjectScaled<class Chaos::FTriangleMeshImplicitObject, 1>>(const FImplicitObject& A, const FRigidTransform3& ATransform, const TImplicitObjectScaled<class Chaos::FTriangleMeshImplicitObject, 1>& B, const FRigidTransform3& BStartTransform, const FVec3& Dir, const FReal Length, const FReal IgnorePenetration, const FReal TargetPenetration, FReal& TOI);
+	template FContactPoint ConvexTriangleMeshSweptContactPoint<FTriangleMeshImplicitObject>(const FImplicitObject& A, const FRigidTransform3& ATransform, const FTriangleMeshImplicitObject& B, const FRigidTransform3& BStartTransform, const FVec3& Dir, const FReal Length, const FReal IgnorePenetration, const FReal TargetPenetration, FReal& TOI);
 	template FContactPoint BoxTriangleMeshContactPoint<TImplicitObjectScaled<class Chaos::FTriangleMeshImplicitObject, 1>>(const FImplicitBox3& A, const FRigidTransform3& ATransform, const TImplicitObjectScaled<class Chaos::FTriangleMeshImplicitObject, 1>& B, const FRigidTransform3& BTransform, const FReal CullDistance, const FReal ShapePadding);
 	template FContactPoint BoxTriangleMeshContactPoint<FTriangleMeshImplicitObject>(const FImplicitBox3& A, const FRigidTransform3& ATransform, const FTriangleMeshImplicitObject& B, const FRigidTransform3& BTransform, const FReal CullDistance, const FReal ShapePadding);
 	template FContactPoint SphereTriangleMeshContactPoint<TImplicitObjectScaled<class Chaos::FTriangleMeshImplicitObject, 1>>(const TSphere<FReal, 3>& A, const FRigidTransform3& ATransform, const TImplicitObjectScaled<class Chaos::FTriangleMeshImplicitObject, 1>& B, const FRigidTransform3& BTransform, const FReal CullDistance, const FReal ShapePadding);
diff --git a/Engine/Source/Runtime/Experimental/Chaos/Private/Chaos/Collision/GJKContactPoint.h b/Engine/Source/Runtime/Experimental/Chaos/Private/Chaos/Collision/GJKContactPoint.h
--- a/Engine/Source/Runtime/Experimental/Chaos/Private/Chaos/Collision/GJKContactPoint.h
+++ b/Engine/Source/Runtime/Experimental/Chaos/Private/Chaos/Collision/GJKContactPoint.h
@@ -43,16 +43,4 @@
 		const FRigidTransform3 BToATM = BTM.GetRelativeTransform(ATM);
 		return GJKContactPoint2(A, B, ATM, BToATM, InitialDir, ShapePadding);
 	}
-
-	inline FContactPoint GenericConvexConvexContactPoint(const FImplicitObject& A, const FRigidTransform3& ATM, const FImplicitObject& B, const FRigidTransform3& BTM, const FReal ShapePadding)
-	{
-		// This expands to a switch of switches that calls the inner function with the appropriate concrete implicit types
-		return Utilities::CastHelperNoUnwrap(A, ATM, [&](const auto& ADowncast, const FRigidTransform3& AFullTM)
-			{
-				return Utilities::CastHelperNoUnwrap(B, BTM, [&](const auto& BDowncast, const FRigidTransform3& BFullTM)
-					{
-						return GJKContactPoint(ADowncast, AFullTM, BDowncast, BFullTM, FVec3(1, 0, 0), ShapePadding);
-					});
-			});
-	}
 }
diff --git a/Engine/Source/Runtime/Experimental/Chaos/Private/Chaos/Collision/GJKContactPointSwept.cpp b/Engine/Source/Runtime/Experimental/Chaos/Private/Chaos/Collision/GJKContactPointSwept.cpp
--- a/Engine/Source/Runtime/Experimental/Chaos/Private/Chaos/Collision/GJKContactPointSwept.cpp
+++ b/Engine/Source/Runtime/Experimental/Chaos/Private/Chaos/Collision/GJKContactPointSwept.cpp
@@ -1,76 +1,70 @@
-// Copyright Epic Games, Inc. All Rights Reserved.
-#include "Chaos/Collision/GJKContactPointSwept.h"
-
-#include "Chaos/CastingUtilities.h"
-#include "Chaos/Collision/ContactPoint.h"
-#include "Chaos/Collision/ContactPointsMiscShapes.h"
-#include "Chaos/GJK.h"
-
-namespace Chaos
-{
-	template <typename GeometryA, typename GeometryB>
-	FContactPoint GJKContactPointSwept(const GeometryA& A, const FRigidTransform3& AStartTM, const FRigidTransform3& AEndTM, const GeometryB& B, const FRigidTransform3& BStartTM, const FRigidTransform3& BEndTM, const FVec3& Dir, const FReal Length, FReal& TOI)
-	{
-		FContactPoint Contact;
-		/** B could be static or dynamic. In both cases, we compute contact point in B local space.  
-		* GJKRaycast2 assumes that A, B do not rotate. 
-		* If B is rotating, the trajectory of A in space B will become non-linear. Also if A or B is rotating, the minkowski difference of the sweep geometry will be very complex and the sweep geometry can be even non-convex, creating difficulties for the GJK algorithm.
-		* For computational efficiency, we could pick any rotation from the start frame rotation to the end frame rotation and lock that rotation for the sweeping process. No matter what rotation we pick, we could always get problems when rotation is large.
-		* If we pick the start frame rotation, we can think of this whole process as first sweeping the objects from X to P, then rotating them in place to end frame rotations. We might end up with A and B penetrating because of the end frame rotations. So here instead, we use the end frame rotations. This process becomes first rotating A and B in place to end frame rotations, then sweeping them from X to P. This produces less penetrations at the end of the frame.
-		*/
-		const FRigidTransform3 ATM(AStartTM.GetLocation(), AEndTM.GetRotation());
-		const FRigidTransform3 BTM(BStartTM.GetLocation(), BEndTM.GetRotation());
-		const FRigidTransform3 AToBTM = ATM.GetRelativeTransform(BTM);
-		const FVec3 LocalDir = BStartTM.InverseTransformVectorNoScale(Dir); 
-
-		FReal OutTime;
-		FVec3 Location, Normal;
-		if (GJKRaycast2(B, A, AToBTM, LocalDir, Length, OutTime, Location, Normal, (FReal)0, true))
-		{
-			// GJK output is all in the local space of B. We need to transform the B-relative position and the normal in to B-space
-			FRigidTransform3 ATOITM, BTOITM;
-			if (OutTime > 0.f)
-			{
-				TOI = OutTime / Length; // OutTime is a distance :|
-				ATOITM = FRigidTransform3(AStartTM.GetLocation() * (1 - TOI) + AEndTM.GetLocation() * TOI, AEndTM.GetRotation());
-				BTOITM = FRigidTransform3(BStartTM.GetLocation() * (1 - TOI) + BEndTM.GetLocation() * TOI, BEndTM.GetRotation());
-			}
-			else
-			{
-				TOI = 0.f;
-				ATOITM = ATM;
-				BTOITM = BTM;
-			}
-
-			const FVec3 WorldLocation = BTOITM.TransformPosition(Location);
-			const FVec3 WorldNormal = BTOITM.TransformVectorNoScale(Normal);
-			Contact.ShapeContactPoints[0] = ATOITM.InverseTransformPosition(WorldLocation);
-			Contact.ShapeContactPoints[1] = Location;
-			Contact.ShapeContactNormal = Normal;
-
-			if (OutTime > 0.f)
-			{
-				const FReal Dot = FVec3::DotProduct(WorldNormal, Dir);
-				Contact.Phi = (Length - OutTime) * Dot;
-			}
-			else // initial penetration case
-			{
-				Contact.Phi = OutTime;
-			}
-		}
-
-		return Contact;
-	}
-
-	FContactPoint GenericConvexConvexContactPointSwept(const FImplicitObject& A, const FRigidTransform3& AStartTM, const FRigidTransform3& AEndTM, const FImplicitObject& B, const FRigidTransform3& BStartTM, const FRigidTransform3& BEndTM, const FVec3& Dir, const FReal Length, FReal& TOI)
-	{
-		// This expands to a switch of switches that calls the inner function with the appropriate concrete implicit types
-		return Utilities::CastHelperNoUnwrap(A, AStartTM, AEndTM, [&](const auto& ADowncast, const FRigidTransform3& AFullStartTM, const FRigidTransform3& AFullEndTM)
-		{
-			return Utilities::CastHelperNoUnwrap(B, BStartTM, BEndTM, [&](const auto& BDowncast, const FRigidTransform3& BFullStartTM, const FRigidTransform3& BFullEndTM)
-			{
-				return GJKContactPointSwept(ADowncast, AFullStartTM, AFullEndTM, BDowncast, BFullStartTM, BFullEndTM, Dir, Length, TOI);
-			});
-		});
-	}
+// Copyright Epic Games, Inc. All Rights Reserved.
+#include "Chaos/Collision/GJKContactPointSwept.h"
+
+#include "Chaos/CastingUtilities.h"
+#include "Chaos/Collision/ContactPoint.h"
+#include "Chaos/Collision/ContactPointsMiscShapes.h"
+#include "Chaos/GJK.h"
+
+namespace Chaos
+{
+	template <typename GeometryA, typename GeometryB>
+	FContactPoint GJKContactPointSwept(const GeometryA& A, const FRigidTransform3& AStartTM, const FRigidTransform3& AEndTM, const GeometryB& B, const FRigidTransform3& BStartTM, const FRigidTransform3& BEndTM, const FVec3& Dir, const FReal Length, const FReal IgnorePenetration, const FReal TargetPenetration, FReal& OutTOI)
+	{
+		FContactPoint Contact;
+		/** B could be static or dynamic. In both cases, we compute contact point in B local space.  
+		* GJKRaycast2 assumes that A, B do not rotate. 
+		* If B is rotating, the trajectory of A in space B will become non-linear. Also if A or B is rotating, the minkowski difference of the sweep geometry will be very complex and the sweep geometry can be even non-convex, creating difficulties for the GJK algorithm.
+		* For computational efficiency, we could pick any rotation from the start frame rotation to the end frame rotation and lock that rotation for the sweeping process. No matter what rotation we pick, we could always get problems when rotation is large.
+		* If we pick the start frame rotation, we can think of this whole process as first sweeping the objects from X to P, then rotating them in place to end frame rotations. We might end up with A and B penetrating because of the end frame rotations. So here instead, we use the end frame rotations. This process becomes first rotating A and B in place to end frame rotations, then sweeping them from X to P. This produces less penetrations at the end of the frame.
+		*/
+		const FRigidTransform3 ATM(AStartTM.GetLocation(), AEndTM.GetRotation());
+		const FRigidTransform3 BTM(BStartTM.GetLocation(), BEndTM.GetRotation());
+		const FRigidTransform3 AToBTM = ATM.GetRelativeTransform(BTM);
+		const FVec3 LocalDir = BStartTM.InverseTransformVectorNoScale(Dir); 
+
+		FReal Distance;
+		FVec3 Location, Normal;
+		if (GJKRaycast2(B, A, AToBTM, LocalDir, Length, Distance, Location, Normal, (FReal)0, true))
+		{
+			const FReal DirDotNormal = FVec3::DotProduct(LocalDir, Normal);
+			FReal TargetTOI, TargetPhi;
+			if (ComputeSweptContactTOIAndPhiAtTargetPenetration(DirDotNormal, Length, Distance, IgnorePenetration, TargetPenetration, TargetTOI, TargetPhi))
+			{
+				// GJK output is all in the local space of B. We need to transform the B-relative position and the normal in to B-space
+				FRigidTransform3 ATOITM, BTOITM;
+				if (Distance > 0.f)
+				{
+					ATOITM = FRigidTransform3(AStartTM.GetLocation() * (1 - TargetTOI) + AEndTM.GetLocation() * TargetTOI, AEndTM.GetRotation());
+					BTOITM = FRigidTransform3(BStartTM.GetLocation() * (1 - TargetTOI) + BEndTM.GetLocation() * TargetTOI, BEndTM.GetRotation());
+				}
+				else
+				{
+					ATOITM = ATM;
+					BTOITM = BTM;
+				}
+
+				const FVec3 WorldLocation = BTOITM.TransformPosition(Location);
+				Contact.ShapeContactPoints[0] = ATOITM.InverseTransformPosition(WorldLocation);
+				Contact.ShapeContactPoints[1] = Location;
+				Contact.ShapeContactNormal = Normal;
+				Contact.Phi = TargetPhi;
+				OutTOI = TargetTOI;
+			}
+		}
+
+		return Contact;
+	}
+
+	FContactPoint GenericConvexConvexContactPointSwept(const FImplicitObject& A, const FRigidTransform3& AStartTM, const FRigidTransform3& AEndTM, const FImplicitObject& B, const FRigidTransform3& BStartTM, const FRigidTransform3& BEndTM, const FVec3& Dir, const FReal Length, const FReal IgnorePenetration, const FReal TargetPenetration, FReal& TOI)
+	{
+		// This expands to a switch of switches that calls the inner function with the appropriate concrete implicit types
+		return Utilities::CastHelperNoUnwrap(A, AStartTM, AEndTM, [&](const auto& ADowncast, const FRigidTransform3& AFullStartTM, const FRigidTransform3& AFullEndTM)
+		{
+			return Utilities::CastHelperNoUnwrap(B, BStartTM, BEndTM, [&](const auto& BDowncast, const FRigidTransform3& BFullStartTM, const FRigidTransform3& BFullEndTM)
+			{
+				return GJKContactPointSwept(ADowncast, AFullStartTM, AFullEndTM, BDowncast, BFullStartTM, BFullEndTM, Dir, Length, IgnorePenetration, TargetPenetration, TOI);
+			});
+		});
+	}
 }
\ No newline at end of file
diff --git a/Engine/Source/Runtime/Experimental/Chaos/Private/Chaos/Collision/GJKContactPointSwept.h b/Engine/Source/Runtime/Experimental/Chaos/Private/Chaos/Collision/GJKContactPointSwept.h
--- a/Engine/Source/Runtime/Experimental/Chaos/Private/Chaos/Collision/GJKContactPointSwept.h
+++ b/Engine/Source/Runtime/Experimental/Chaos/Private/Chaos/Collision/GJKContactPointSwept.h
@@ -1,16 +1,13 @@
-// Copyright Epic Games, Inc. All Rights Reserved.
-#pragma once
-
-#include "CoreMinimal.h"
-#include "Chaos/Core.h"
-#include "Chaos/ImplicitFwd.h"
-
-namespace Chaos
-{
-	class FContactPoint;
-
-	template <typename GeometryA, typename GeometryB>
-	FContactPoint GJKContactPointSwept(const GeometryA& A, const FRigidTransform3& ATM, const GeometryB& B, const FRigidTransform3& BTM, const FVec3& Dir, const FReal Length, FReal& TOI);
-
-	FContactPoint GenericConvexConvexContactPointSwept(const FImplicitObject& A, const FRigidTransform3& AStartTM, const FRigidTransform3& AEndTM, const FImplicitObject& B, const FRigidTransform3& BStartTM, const FRigidTransform3& BEndTM, const FVec3& Dir, const FReal Length, FReal& TOI);
-}
+// Copyright Epic Games, Inc. All Rights Reserved.
+#pragma once
+
+#include "CoreMinimal.h"
+#include "Chaos/Core.h"
+#include "Chaos/ImplicitFwd.h"
+
+namespace Chaos
+{
+	class FContactPoint;
+
+	FContactPoint GenericConvexConvexContactPointSwept(const FImplicitObject& A, const FRigidTransform3& AStartTM, const FRigidTransform3& AEndTM, const FImplicitObject& B, const FRigidTransform3& BStartTM, const FRigidTransform3& BEndTM, const FVec3& Dir, const FReal Length, const FReal IgnorePenetration, const FReal TargetPenetration, FReal& TOI);
+}
diff --git a/Engine/Source/Runtime/Experimental/Chaos/Private/Chaos/Collision/PBDCollisionConstraint.cpp b/Engine/Source/Runtime/Experimental/Chaos/Private/Chaos/Collision/PBDCollisionConstraint.cpp
--- a/Engine/Source/Runtime/Experimental/Chaos/Private/Chaos/Collision/PBDCollisionConstraint.cpp
+++ b/Engine/Source/Runtime/Experimental/Chaos/Private/Chaos/Collision/PBDCollisionConstraint.cpp
@@ -16,10 +16,14 @@
 
 //PRAGMA_DISABLE_OPTIMIZATION
 
-extern Chaos::FRealSingle CCDAllowedDepthBoundsScale;
-
 namespace Chaos
 {
+	namespace CVars
+	{
+		extern FRealSingle CCDEnableThresholdBoundsScale;
+		extern FRealSingle CCDAllowedDepthBoundsScale;
+	}
+
 	extern bool bChaos_Collision_EnableManifoldGJKInject;
 	extern bool bChaos_Collision_EnableManifoldGJKReplace;
 
@@ -185,7 +189,8 @@
 		, Stiffness(1)
 		, AccumulatedImpulse(0)
 		, CCDTimeOfImpact(0)
-		, CCDPenetrationThreshold(0)
+		, CCDEnablePenetration(0)
+		, CCDTargetPenetration(0)
 		, ContainerCookie()
 		, ShapesType(EContactShapesType::Unknown)
 		, CCDType(ECollisionCCDType::Disabled)
@@ -225,7 +230,8 @@
 		, Stiffness(1)
 		, AccumulatedImpulse(0)
 		, CCDTimeOfImpact(0)
-		, CCDPenetrationThreshold(0)
+		, CCDEnablePenetration(0)
+		, CCDTargetPenetration(0)
 		, ContainerCookie()
 		, ShapesType(EContactShapesType::Unknown)
 		, CCDType(ECollisionCCDType::Disabled)
@@ -360,10 +366,12 @@
 
 	void FPBDCollisionConstraint::InitCCDThreshold()
 	{
-		// Calculate the max penetration that we ignore with CCD contacts
-		const FReal Threshold0 = FConstGenericParticleHandle(Particle[0])->CCDEnabled() ? Implicit[0]->BoundingBox().Extents().GetAbsMin() : FReal(0);
-		const FReal Threshold1 = FConstGenericParticleHandle(Particle[1])->CCDEnabled() ? Implicit[1]->BoundingBox().Extents().GetAbsMin() : FReal(0);
-		CCDPenetrationThreshold = FMath::Max(Threshold0, Threshold1) * CCDAllowedDepthBoundsScale;
+		// Calculate the depth at which we enforce CCD constraints, and the amount of penetration to leave when resolving a CCD constraint.
+		const FReal MinBounds0 = FConstGenericParticleHandle(Particle[0])->CCDEnabled() ? Implicit[0]->BoundingBox().Extents().GetAbsMin() : FReal(0);
+		const FReal MinBounds1 = FConstGenericParticleHandle(Particle[1])->CCDEnabled() ? Implicit[1]->BoundingBox().Extents().GetAbsMin() : FReal(0);
+		const FReal MinBounds = FMath::Max(MinBounds0, MinBounds1);
+		CCDEnablePenetration = MinBounds * CVars::CCDEnableThresholdBoundsScale;
+		CCDTargetPenetration = FMath::Min(MinBounds * CVars::CCDAllowedDepthBoundsScale, CCDEnablePenetration);
 	}
 
 	void FPBDCollisionConstraint::SetIsSleeping(const bool bInIsSleeping)
@@ -1046,91 +1054,4 @@
 			return NoRestingDependency;
 		}
 	}
-
-	// NOTE: This only works well if we do not update the rotation during CCD interations (which we currently do not)
-	// and we start off with a good selection for the contact plane. The latter is true if we have large objects but not so true
-	// when colliding against non-smooth triangle meshes
-	void FPBDCollisionConstraint::UpdateSweptManifoldPoints(const FVec3& ShapeStartWorldPosition0, const FVec3& ShapeStartWorldPosition1, const FReal Dt)
-	{
-		FReal MinTOI = TNumericLimits<FReal>::Max();
-		int32 MinTOIManifoldPointIndex = INDEX_NONE;
-
-		for (int32 ManifoldPointIndex = 0; ManifoldPointIndex < ManifoldPoints.Num(); ++ManifoldPointIndex)
-		{
-			FManifoldPoint& ManifoldPoint = ManifoldPoints[ManifoldPointIndex];
-			if (ManifoldPoint.Flags.bDisabled)
-			{
-				continue;
-			}
-
-			// Start and End contact positions in world space
-			const FVec3 ContactOffset0 = ShapeWorldTransform0.TransformVectorNoScale(ManifoldPoint.ContactPoint.ShapeContactPoints[0]);
-			const FVec3 ContactStartPos0 = ShapeStartWorldPosition0 + ContactOffset0;
-			const FVec3 ContactEndPos0 = ShapeWorldTransform0.GetTranslation() + ContactOffset0;
-			const FVec3 ContactOffset1 = ShapeWorldTransform1.TransformVectorNoScale(ManifoldPoint.ContactPoint.ShapeContactPoints[1]);
-			const FVec3 ContactStartPos1 = ShapeStartWorldPosition1 + ContactOffset1;
-			const FVec3 ContactEndPos1 = ShapeWorldTransform1.GetTranslation() + ContactOffset1;
-
-			// Contact normal in world space
-			const FVec3 ContactNormal = ShapeWorldTransform1.TransformVectorNoScale(ManifoldPoint.ContactPoint.ShapeContactNormal);
-
-			// Start and end separation
-			const FReal ContactStartPhi = FVec3::DotProduct(ContactStartPos0 - ContactStartPos1, ContactNormal);
-			const FReal ContactEndPhi = FVec3::DotProduct(ContactEndPos0 - ContactEndPos1, ContactNormal);
-
-			// Update the TOI and track closest manifold point
-			const FReal ContactTOI = CalculateModifiedSweptTOI(ContactStartPhi, ContactEndPhi);
-			if (ContactTOI < MinTOI)
-			{
-				MinTOI = ContactTOI;
-				MinTOIManifoldPointIndex = ManifoldPointIndex;
-			}
-
-			// Update the contact separation
-			ManifoldPoint.ContactPoint.Phi = ContactEndPhi;
-		}
-
-		CCDTimeOfImpact = MinTOI;
-		ClosestManifoldPointIndex = MinTOIManifoldPointIndex;
-	}
-
-	// Calculate a Time of Impact (TOI) to use with CCD. This is based on the initial and final contact separation from the sweep test but modified so that
-	// - we ignore separating contacts (increasing Phi)
-	// - we ignore contacts that are separated at TOI=1 (EndPhi > 0)
-	// - we ignore contacts if the penetration is less than the CCD tolerance
-	FReal FPBDCollisionConstraint::CalculateModifiedSweptTOI(const FReal StartPhi, const FReal EndPhi)
-	{
-		const FReal InfiniteTOI = TNumericLimits<FReal>::Max();
-		const FReal MovementTolerance = KINDA_SMALL_NUMBER;
-
-		// If we end up separated at TOI=1 ignore the contact
-		if (EndPhi > 0)
-		{
-			return InfiniteTOI;
-		}
-
-		// If contact is moving in the right direction or not moving ignore the contact
-		if (EndPhi > StartPhi - MovementTolerance)
-		{
-			return InfiniteTOI;
-		}
-
-		// If we penetrate by less than the CCD tolerance, treat it as TOI=1. This mean no CCD impulse and the non-CCD 
-		// solve is expected to handle it. E.g., this improves the behaviour when we are sliding along a surface at
-		// above CCD speeds - we don't want to handle TOI events with the floor
-		const FReal PhiThreshold = -CCDPenetrationThreshold;
-		if (EndPhi > PhiThreshold)
-		{
-			return FReal(1);
-		}
-
-		// If we penetrate by more than the CCD threshold we roll back all the way to the TOI leaving no penetration
-		// It would be nice to leave to penetration for the non-CCD solve to handle in this case, but then we will have
-		// initiial-overlap problems when we have secondary CCD collisions that can result in missed collisions
-		const FReal DesiredPhi = FReal(0);
-		const FReal TOI = (DesiredPhi - StartPhi) / (EndPhi - StartPhi);
-
-		return FMath::Clamp(TOI, FReal(0), FReal(1));
-	}
-
 }
diff --git a/Engine/Source/Runtime/Experimental/Chaos/Private/Chaos/CollisionOneShotManifolds.cpp b/Engine/Source/Runtime/Experimental/Chaos/Private/Chaos/CollisionOneShotManifolds.cpp
--- a/Engine/Source/Runtime/Experimental/Chaos/Private/Chaos/CollisionOneShotManifolds.cpp
+++ b/Engine/Source/Runtime/Experimental/Chaos/Private/Chaos/CollisionOneShotManifolds.cpp
@@ -82,8 +82,8 @@
 	bool bChaos_Collision_UseGJK2 = false;
 	FAutoConsoleVariableRef CVarChaos_Collision_UseGJK2(TEXT("p.Chaos.Collision.UseGJK2"), bChaos_Collision_UseGJK2, TEXT(""));
 
-	bool bChaos_Collision_OneSidedTriangleMesh = false;
-	bool bChaos_Collision_OneSidedHeightfield = false;
+	bool bChaos_Collision_OneSidedTriangleMesh = true;
+	bool bChaos_Collision_OneSidedHeightfield = true;
 	FAutoConsoleVariableRef CVarChaos_Collision_OneSidedTriangleMesh(TEXT("p.Chaos.Collision.OneSidedTriangleMesh"), bChaos_Collision_OneSidedTriangleMesh, TEXT(""));
 	FAutoConsoleVariableRef CVarChaos_Collision_OneSidedHeightfield(TEXT("p.Chaos.Collision.OneSidedHeightfield"), bChaos_Collision_OneSidedHeightfield, TEXT(""));
 
@@ -92,7 +92,6 @@
 		// Forward delarations we need from CollisionRestitution.cpp
 
 		FContactPoint BoxBoxContactPoint(const FImplicitBox3& Box1, const FImplicitBox3& Box2, const FRigidTransform3& Box1TM, const FRigidTransform3& Box2TM, const FReal ShapePadding);
-		FContactPoint GenericConvexConvexContactPoint(const FImplicitObject& A, const FRigidTransform3& ATM, const FImplicitObject& B, const FRigidTransform3& BTM, const FReal ShapePadding);
 
 		//////////////////////////
 		// Box Box
diff --git a/Engine/Source/Runtime/Experimental/Chaos/Private/Chaos/CollisionResolution.cpp b/Engine/Source/Runtime/Experimental/Chaos/Private/Chaos/CollisionResolution.cpp
--- a/Engine/Source/Runtime/Experimental/Chaos/Private/Chaos/CollisionResolution.cpp
+++ b/Engine/Source/Runtime/Experimental/Chaos/Private/Chaos/CollisionResolution.cpp
@@ -45,7 +45,6 @@
 DECLARE_CYCLE_STAT(TEXT("Collisions::FindAllIntersectingClusteredObjects"), STAT_Collisions_FindAllIntersectingClusteredObjects, STATGROUP_ChaosCollision);
 DECLARE_CYCLE_STAT(TEXT("Collisions::ConstructGenericConvexConvexConstraints"), STAT_Collisions_ConstructGenericConvexConvexConstraints, STATGROUP_ChaosCollision);
 DECLARE_CYCLE_STAT(TEXT("Collisions::ConstructGenericConvexConvexConstraintsSwept"), STAT_Collisions_ConstructGenericConvexConvexConstraintsSwept, STATGROUP_ChaosCollision);
-DECLARE_CYCLE_STAT(TEXT("Collisions::SetSweptConstraintTOI"), STAT_Collisions_SetSweptConstraintTOI, STATGROUP_ChaosCollision);
 DECLARE_CYCLE_STAT(TEXT("Collisions::UpdateConstraintFromGeometryInternal"), STAT_Collisions_UpdateConstraintFromGeometryInternal, STATGROUP_ChaosCollision);
 DECLARE_CYCLE_STAT(TEXT("Collisions::UpdateGenericConvexConvexConstraint"), STAT_Collisions_UpdateGenericConvexConvexConstraint, STATGROUP_ChaosCollision);
 DECLARE_CYCLE_STAT(TEXT("Collisions::UpdateConvexTriangleMeshConstraint"), STAT_Collisions_UpdateConvexTriangleMeshConstraint, STATGROUP_ChaosCollision);
@@ -73,9 +72,6 @@
 DECLARE_DWORD_ACCUMULATOR_STAT(TEXT("NumContactsCreated"), STAT_ChaosCollisionCounter_NumContactsCreated, STATGROUP_ChaosCollisionCounters);
 DECLARE_DWORD_ACCUMULATOR_STAT(TEXT("NumContactUpdates"), STAT_ChaosCollisionCounter_NumContactUpdates, STATGROUP_ChaosCollisionCounters);
 
-Chaos::FRealSingle CCDAllowedDepthBoundsScale = 0.1f;
-FAutoConsoleVariableRef CVarCCDAllowedDepthBoundsScale(TEXT("p.Chaos.CCD.AllowedDepthBoundsScale"), CCDAllowedDepthBoundsScale, TEXT("When rolling back to TOI, allow (smallest bound's extent) * AllowedDepthBoundsScale, instead of rolling back to exact TOI w/ penetration = 0."));
-
 int32 CCDUseGenericSweptConvexConstraints = 1;
 FAutoConsoleVariableRef CVarUseGenericSweptConvexConstraints(TEXT("p.Chaos.CCD.UseGenericSweptConvexConstraints"), CCDUseGenericSweptConvexConstraints, TEXT("Use generic convex convex swept constraint generation for convex shape pairs which don't have specialized implementations."));
 
@@ -174,7 +170,9 @@
 				return false;
 			}
 
-			bool bUseCCD = CCDHelpers::DeltaExceedsThreshold(*Particle0, *Particle1);
+			bool bUseCCD = CCDHelpers::DeltaExceedsThreshold(
+				Particle0->CCDAxisThreshold(), DeltaX0, FConstGenericParticleHandle(Particle0)->Q(),
+				Particle1->CCDAxisThreshold(), DeltaX1, FConstGenericParticleHandle(Particle1)->Q());
 
 			if (bUseCCD)
 			{
@@ -196,37 +194,6 @@
 			return bUseCCD;
 		}
 
-		// Initializes TimeOfImpact and adjusts Phi for TOI.
-		// @todo(chaos): this repeats normal calculation and dot product - move to ComputeSweptContactPhiAndTOIHelper?
-		void SetSweptConstraintTOI(TGeometryParticleHandle<FReal, 3>* Particle, const FReal TOI, const FReal Length, const FVec3& Dir, FPBDCollisionConstraint& SweptConstraint)
-		{
-			CONDITIONAL_SCOPE_CYCLE_COUNTER(STAT_Collisions_SetSweptConstraintTOI, ConstraintsDetailedStats);
-
-			if (TOI <= FReal(1))
-			{
-				const FVec3 ContactNormal = SweptConstraint.CalculateWorldContactNormal();
-				const FReal DirDotNormal = FVec3::DotProduct(Dir, ContactNormal);
-
-				// Modify TOI so we ignore separating and shallow CCD contacts
-				// NOTE (See GJKRaycast2): If initially penetrating, Phi is the separation at the start of the sweep.
-				// But if not initially penetrating Phi is the separation at the end of the sweep
-				// @todo(chaos): The way Phi is set for sweeps is quite confusing - can we make this better?
-				FReal StartPhi, EndPhi;
-				if (TOI <= FReal(0))
-				{
-					StartPhi = SweptConstraint.GetPhi();
-					EndPhi = StartPhi + DirDotNormal * Length;
-				}
-				else
-				{
-					EndPhi = SweptConstraint.GetPhi();
-					StartPhi = EndPhi - DirDotNormal * Length;
-				}
-
-				SweptConstraint.CCDTimeOfImpact = SweptConstraint.CalculateModifiedSweptTOI(StartPhi, EndPhi);
-			}
-		}
-
 		// Add the contact to the manifold (or update the existing point). Disable the contact if the contact distance is greater than the Cull Distance
 		void UpdateContactPoint(FPBDCollisionConstraint& Constraint, const FContactPoint& ContactPoint, const FReal Dt)
 		{
@@ -724,8 +691,8 @@
 		void UpdateSphereHeightFieldConstraintSwept(TGeometryParticleHandle<FReal, 3>* Particle0, const TSphere<FReal, 3>& A, const FRigidTransform3& ATransform, const FHeightField& B, const FRigidTransform3& BTransform, const FVec3& Dir, const FReal Length, const FReal Dt, FPBDCollisionConstraint& Constraint)
 		{
 			FReal TOI = 1.0f;
-			UpdateContactPointNoCull(Constraint, GJKImplicitSweptContactPoint(A, ATransform, B, BTransform, Dir, Length, TOI), Dt, CCDNoCullAllShapePairs);
-			SetSweptConstraintTOI(Particle0, TOI, Length, Dir, Constraint);
+			UpdateContactPointNoCull(Constraint, GJKImplicitSweptContactPoint(A, ATransform, B, BTransform, Dir, Length, Constraint.GetCCDEnablePenetration(), Constraint.GetCCDTargetPenetration(), TOI), Dt, CCDNoCullAllShapePairs);
+			Constraint.SetCCDTimeOfImpact(TOI);
 		}
 
 		template<typename T_TRAITS>
@@ -918,8 +885,8 @@
 		void UpdateSphereTriangleMeshConstraintSwept(TGeometryParticleHandle<FReal, 3>* Particle0, const TSphere<FReal, 3>& Sphere0, const FRigidTransform3& WorldTransform0, const TriMeshType& TriangleMesh1, const FRigidTransform3& WorldTransform1, const FVec3& Dir, const FReal Length, const FReal Dt, FPBDCollisionConstraint& Constraint)
 		{
 			FReal TOI = 1.0f;
-			UpdateContactPointNoCull(Constraint, SphereTriangleMeshSweptContactPoint(Sphere0, WorldTransform0, TriangleMesh1, WorldTransform1, Dir, Length, TOI), Dt, CCDNoCullAllShapePairs);
-			SetSweptConstraintTOI(Particle0, TOI, Length, Dir, Constraint);
+			UpdateContactPointNoCull(Constraint, SphereTriangleMeshSweptContactPoint(Sphere0, WorldTransform0, TriangleMesh1, WorldTransform1, Dir, Length, Constraint.GetCCDEnablePenetration(), Constraint.GetCCDTargetPenetration(), TOI), Dt, CCDNoCullAllShapePairs);
+			Constraint.SetCCDTimeOfImpact(TOI);
 		}
 
 		template<typename T_TRAITS>
@@ -1165,8 +1132,8 @@
 		{
 			CONDITIONAL_SCOPE_CYCLE_COUNTER(STAT_Collisions_UpdateCapsuleHeightFieldConstraintSwept, ConstraintsDetailedStats);
 			FReal TOI = 1.0f;
-			UpdateContactPointNoCull(Constraint, GJKImplicitSweptContactPoint(A, ATransform, B, BTransform, Dir, Length, TOI), Dt, CCDNoCullAllShapePairs);
-			SetSweptConstraintTOI(Particle0, TOI, Length, Dir, Constraint);
+			UpdateContactPointNoCull(Constraint, GJKImplicitSweptContactPoint(A, ATransform, B, BTransform, Dir, Length, Constraint.GetCCDEnablePenetration(), Constraint.GetCCDTargetPenetration(), TOI), Dt, CCDNoCullAllShapePairs);
+			Constraint.SetCCDTimeOfImpact(TOI);
 		}
 
 		template<typename T_TRAITS>
@@ -1228,8 +1195,8 @@
 		{
 			CONDITIONAL_SCOPE_CYCLE_COUNTER(STAT_Collisions_UpdateCapsuleTriangleMeshConstraint, ConstraintsDetailedStats);
 			FReal TOI = 1.0f;
-			UpdateContactPointNoCull(Constraint, CapsuleTriangleMeshSweptContactPoint(Capsule0, WorldTransform0, TriangleMesh1, WorldTransform1, Dir, Length, TOI), Dt, CCDNoCullAllShapePairs);
-			SetSweptConstraintTOI(Particle0, TOI, Length, Dir, Constraint);
+			UpdateContactPointNoCull(Constraint, CapsuleTriangleMeshSweptContactPoint(Capsule0, WorldTransform0, TriangleMesh1, WorldTransform1, Dir, Length, Constraint.GetCCDEnablePenetration(), Constraint.GetCCDTargetPenetration(), TOI), Dt, CCDNoCullAllShapePairs);
+			Constraint.SetCCDTimeOfImpact(TOI);
 		}
 
 		template<typename T_TRAITS>
@@ -1296,7 +1263,7 @@
 		}
 
 		//
-		// Generic Convex - Convex (actualy concrete type could be anything)
+		// Generic Convex - Convex (actual concrete type could be anything)
 		//
 
 		void UpdateGenericConvexConvexConstraint(const FImplicitObject& Implicit0, const FRigidTransform3& WorldTransform0, const FImplicitObject& Implicit1, const FRigidTransform3& WorldTransform1, const FReal Dt, FPBDCollisionConstraint& Constraint)
@@ -1312,8 +1279,8 @@
 			FReal TOI = TNumericLimits<FReal>::Max();
 			const FRigidTransform3& EndWorldTransform0 = Constraint.GetShapeWorldTransform0();
 			const FRigidTransform3& EndWorldTransform1 = Constraint.GetShapeWorldTransform1();
-			UpdateContactPointNoCull(Constraint, GenericConvexConvexContactPointSwept(Implicit0, StartWorldTransform0, EndWorldTransform0, Implicit1, StartWorldTransform1, EndWorldTransform1, Dir, Length, TOI), Dt);
-			SetSweptConstraintTOI(Particle0, TOI, Length, Dir, Constraint);
+			UpdateContactPointNoCull(Constraint, GenericConvexConvexContactPointSwept(Implicit0, StartWorldTransform0, EndWorldTransform0, Implicit1, StartWorldTransform1, EndWorldTransform1, Dir, Length, Constraint.GetCCDEnablePenetration(), Constraint.GetCCDTargetPenetration(), TOI), Dt);
+			Constraint.SetCCDTimeOfImpact(TOI);
 		}
 
 
@@ -1375,8 +1342,8 @@
 		{
 			CONDITIONAL_SCOPE_CYCLE_COUNTER(STAT_Collisions_UpdateConvexHeightFieldConstraintSwept, ConstraintsDetailedStats);
 			FReal TOI = 1.0f;
-			UpdateContactPointNoCull(Constraint, GJKImplicitSweptContactPoint(A, ATransform, B, BTransform, Dir, Length, TOI), Dt, CCDNoCullAllShapePairs);
-			SetSweptConstraintTOI(Particle0, TOI, Length, Dir, Constraint);
+			UpdateContactPointNoCull(Constraint, GJKImplicitSweptContactPoint(A, ATransform, B, BTransform, Dir, Length, Constraint.GetCCDEnablePenetration(), Constraint.GetCCDTargetPenetration(), TOI), Dt, CCDNoCullAllShapePairs);
+			Constraint.SetCCDTimeOfImpact(TOI);
 		}
 
 		template<typename T_TRAITS>
@@ -1432,8 +1399,8 @@
 		{
 			CONDITIONAL_SCOPE_CYCLE_COUNTER(STAT_Collisions_UpdateConvexTriangleMeshConstraintSwept, ConstraintsDetailedStats);
 			FReal TOI = 1.0f;
-			UpdateContactPointNoCull(Constraint, ConvexTriangleMeshSweptContactPoint(Convex0, WorldTransform0, TriangleMesh1, WorldTransform1, Dir, Length, TOI), Dt, CCDNoCullAllShapePairs);
-			SetSweptConstraintTOI(Particle0, TOI, Length, Dir, Constraint);
+			UpdateContactPointNoCull(Constraint, ConvexTriangleMeshSweptContactPoint(Convex0, WorldTransform0, TriangleMesh1, WorldTransform1, Dir, Length, Constraint.GetCCDEnablePenetration(), Constraint.GetCCDTargetPenetration(), TOI), Dt, CCDNoCullAllShapePairs);
+			Constraint.SetCCDTimeOfImpact(TOI);
 		}
 
 		template<typename T_TRAITS>
@@ -1722,6 +1689,8 @@
 			const FVec3 DeltaX1 = Constraint.GetShapeWorldTransform1().GetTranslation() - WorldTransform1.GetLocation();
 			bool bUseCCD = ShouldUseCCD(Particle0, DeltaX0, Particle1, DeltaX1, DirCCD, LengthCCD, false);
 
+			Constraint.SetCCDTimeOfImpact(TNumericLimits<FReal>::Max());
+
 			if (bUseCCD)
 			{
 				const EContactShapesType ShapePairType = Constraint.GetShapesType();
@@ -1817,9 +1786,7 @@
 				}
 			}
 
-			Constraint.CCDTimeOfImpact = TNumericLimits<FReal>::Max(); // CCD will not be used for this constraint
-			// Do a normal non-swept update if we reach this point - Not required for now, since this will be done in solver
-			// UpdateConstraintFromGeometryImpl<UpdateType>(*(Constraint.As<FPBDCollisionConstraint>()), WorldTransform0, WorldTransform1, Dt);
+			// CCD will not be used for this constraint this tick
 			return false;
 		}
 
diff --git a/Engine/Source/Runtime/Experimental/Chaos/Private/Chaos/HeightField.cpp b/Engine/Source/Runtime/Experimental/Chaos/Private/Chaos/HeightField.cpp
--- a/Engine/Source/Runtime/Experimental/Chaos/Private/Chaos/HeightField.cpp
+++ b/Engine/Source/Runtime/Experimental/Chaos/Private/Chaos/HeightField.cpp
@@ -2,6 +2,7 @@
 
 #include "Chaos/HeightField.h"
 #include "Chaos/Collision/ContactPoint.h"
+#include "Chaos/Collision/ContactPointsMiscShapes.h"
 #include "Chaos/Collision/PBDCollisionConstraint.h"
 #include "Chaos/CollisionOneShotManifolds.h"
 #include "Chaos/Core.h"
@@ -20,6 +21,11 @@
 
 namespace Chaos
 {
+	namespace CVars
+	{
+		extern bool bCCDNewTargetDepthMode;
+	}
+
 	extern FRealSingle Chaos_Collision_EdgePrunePlaneDistance;
 
 	extern bool bChaos_Collision_OneSidedHeightfield;
@@ -457,6 +463,197 @@
 
 	};
 
+	template<typename GeomQueryType>
+	class THeightfieldSweepVisitorCCD
+	{
+	public:
+
+		THeightfieldSweepVisitorCCD(const typename FHeightField::FDataType* InData, const GeomQueryType& InQueryGeom, const FRigidTransform3& InStartTM, const FVec3& InDir, const FReal Length, const FReal InIgnorePenetration, const FReal InTargetPenetration)
+			: OutDistance(TNumericLimits<FReal>::Max())
+			, OutPhi(TNumericLimits<FReal>::Max())
+			, OutFaceIndex(INDEX_NONE)
+			, HfData(InData)
+			, StartTM(InStartTM)
+			, OtherGeom(InQueryGeom)
+			, IgnorePenetration(InIgnorePenetration)
+			, TargetPenetration(InTargetPenetration)
+		{
+			const FAABB3 QueryBounds = InQueryGeom.BoundingBox();
+			const FVec3 StartPoint = StartTM.TransformPositionNoScale(QueryBounds.Center());
+			const FVec3 Inflation3D = QueryBounds.Extents() * FReal(0.5);
+
+			const VectorRegister4Float LengthSimd = MakeVectorRegisterFloatFromDouble(VectorSetFloat1(Length));
+			StartPointSimd = MakeVectorRegisterFloatFromDouble(MakeVectorRegister(StartPoint.X, StartPoint.Y, StartPoint.Z, 0.0));
+			DirSimd = MakeVectorRegisterFloatFromDouble(MakeVectorRegister(InDir.X, InDir.Y, InDir.Z, 0.0));
+			EndPointSimd = VectorAdd(StartPointSimd, VectorMultiply(DirSimd, LengthSimd));
+			RotationSimd = MakeVectorRegisterFloatFromDouble(StartTM.GetRotationRegister());
+			ParallelSimd = VectorCompareLT(VectorAbs(DirSimd), GlobalVectorConstants::SmallNumber);
+			InvDirSimd = VectorBitwiseNotAnd(ParallelSimd, VectorDivide(VectorOne(), DirSimd));
+			Inflation3DSimd = MakeVectorRegisterFloatFromDouble(VectorAbs(MakeVectorRegister(Inflation3D.X, Inflation3D.Y, Inflation3D.Z, 0.0)));
+
+			// @todo(chaos): do we really need this?
+			RotationSimd = VectorNormalizeSafe(RotationSimd, GlobalVectorConstants::Float0001);
+		}
+
+		bool VisitSweep(int32 CellIndex, FReal& CurrentLength)
+		{
+			// Skip holes
+			if (HfData->MaterialIndices.IsValidIndex(CellIndex))
+			{
+				const bool bHole = (HfData->MaterialIndices[CellIndex] == TNumericLimits<uint8>::Max());
+				if (bHole)
+				{
+					return true;
+				}
+			}
+
+			auto TestTriangle = [&](int32 FaceIndex, const VectorRegister4Float& A, const VectorRegister4Float& B, const VectorRegister4Float& C) -> bool
+			{
+				// Convert into A-relative positions to get better precision
+				FTriangleRegister Triangle(VectorZero(), VectorSubtract(B, A), VectorSubtract(C, A));
+				VectorRegister4Float TranslationSimd = MakeVectorRegisterFloatFromDouble(VectorSubtract(StartTM.GetTranslationRegister(), A));
+
+				// If we are definitely penetrating by less than IgnorePenetration at T=1, we can ignore this triangle
+				// Checks the query geometry support point along the normal at the sweep end point
+				VectorRegister4Float IgnorePenetrationSimd = MakeVectorRegisterFloatFromDouble(VectorSetFloat1(IgnorePenetration));
+				VectorRegister4Float TriangleNormalSimd = VectorNormalize(VectorCross(VectorSubtract(B, A), VectorSubtract(C, A)));
+				VectorRegister4Float OtherTriangleNormalSimd = VectorQuaternionInverseRotateVector(RotationSimd, TriangleNormalSimd);
+				VectorRegister4Float OtherTrianglePositionSimd = VectorQuaternionInverseRotateVector(RotationSimd, VectorSubtract(A, EndPointSimd));
+				VectorRegister4Float OtherSupportSimd = OtherGeom.SupportCoreSimd(VectorNegate(OtherTriangleNormalSimd), 0);
+				VectorRegister4Float MaxDepthSimd = VectorDot3(VectorSubtract(OtherTrianglePositionSimd, OtherSupportSimd), OtherTriangleNormalSimd);
+				if (VectorMaskBits(VectorCompareLT(MaxDepthSimd, IgnorePenetrationSimd)))
+				{
+					return true;
+				}
+
+				FRealSingle Distance;
+				VectorRegister4Float PositionSimd, NormalSimd;
+				if (GJKRaycast2ImplSimd(Triangle, OtherGeom, RotationSimd, TranslationSimd, DirSimd, FRealSingle(CurrentLength), Distance, PositionSimd, NormalSimd, true, GlobalVectorConstants::Float1000))
+				{
+					// Convert back from A-relative positions
+					PositionSimd = VectorAdd(PositionSimd, A);
+
+					const VectorRegister4Float DirDotNormalSimd = VectorDot3(DirSimd, NormalSimd);
+					FReal DirDotNormal;
+					VectorStoreFloat1(DirDotNormalSimd, &DirDotNormal);
+
+					// Calculate the time to reach a depth of TargetPenetration
+					FReal TargetTOI, TargetPhi;
+					ComputeSweptContactTOIAndPhiAtTargetPenetration(DirDotNormal, CurrentLength, FReal(Distance), IgnorePenetration, TargetPenetration, TargetTOI, TargetPhi);
+
+					// If this hit is closest, store the result
+					const FReal TargetDistance = TargetTOI * CurrentLength;
+					if (TargetDistance < CurrentLength)
+					{
+						CurrentLength = TargetDistance;
+
+						VectorStoreFloat3(MakeVectorRegisterDouble(NormalSimd), &OutNormal);
+						VectorStoreFloat3(MakeVectorRegisterDouble(PositionSimd), &OutPosition);
+						OutDistance = TargetDistance;
+						OutPhi = TargetPhi;
+						OutFaceIndex = FaceIndex;
+					}
+				}
+
+				return (CurrentLength > 0);
+			};
+
+
+			VectorRegister4Float Points[4];
+			FAABBVectorized CellBounds;
+			const int32 VertexIndex = HfData->CellIndexToVertexIndex(CellIndex);
+			HfData->GetPointsAndBoundsScaledSimd(VertexIndex, Points, CellBounds);
+			CellBounds.Thicken(Inflation3DSimd);
+
+			if (CellBounds.RaycastFast(StartPointSimd, InvDirSimd, ParallelSimd, VectorSetFloat1(static_cast<FRealSingle>(CurrentLength))))
+			{
+				const int32 TriIndex0 = 2 * CellIndex + 0;
+				const int32 TriIndex1 = 2 * CellIndex + 1;
+				if (!TestTriangle(TriIndex0, Points[0], Points[1], Points[3]))
+				{
+					return false;
+				}
+				if (!TestTriangle(TriIndex1, Points[0], Points[3], Points[2]))
+				{
+					return false;
+				}
+			}
+
+			return true;
+		}
+
+		FReal OutDistance;
+		FReal OutPhi;
+		FVec3 OutPosition;
+		FVec3 OutNormal;
+		int32 OutFaceIndex;
+
+	private:
+		const typename FHeightField::FDataType* HfData;
+		const FRigidTransform3 StartTM;
+		const GeomQueryType& OtherGeom;
+		FReal IgnorePenetration;
+		FReal TargetPenetration;
+
+		VectorRegister4Float StartPointSimd;
+		VectorRegister4Float EndPointSimd;
+		VectorRegister4Float InvDirSimd;
+		VectorRegister4Float DirSimd;
+		VectorRegister4Float RotationSimd;
+		VectorRegister4Float ParallelSimd;
+		VectorRegister4Float Inflation3DSimd;
+	};
+
+	/**
+	 * Call the triangle visitor on all the triangles that are swept through by the InQueryBounds along the vector InDir. Bounds and direction
+	 * are in heightfield space.
+	 * 
+	 * The visitor signature must be:
+	 * 
+	 *		bool(int32 FaceIndex, const VectorRegister4Float& A, const VectorRegister4Float& B, const VectorRegister4Float& C, FRealSingle& CurrentDistance)
+	 * 
+	 * The visitor must return true if we should continue processing other triangles. It must update the current hit distance if the
+	 * passed in triangle is accepted and closer. The vertex positions passed to the triangle visitor are in heightfield space.
+	*/
+	template<typename TriangleVisitor>
+	class THeightfieldTriangleGenerator
+	{
+	public:
+		THeightfieldTriangleGenerator(const typename FHeightField::FDataType* InData, const TriangleVisitor& InTriangleVisitor)
+			: HfData(InData)
+			, TriangleVisitor(InTriangleVisitor)
+		{
+		}
+
+		bool VisitSweep(int32 CellIndex, FReal& CurrentLength)
+		{
+			const int32 VertexIndex = HfData->CellIndexToVertexIndex(CellIndex);
+			VectorRegister4Float Points[4];
+			FAABBVectorized CellBounds;
+			HfData->GetPointsAndBoundsScaledSimd(VertexIndex, Points, CellBounds);
+
+			const int32 TriIndex0 = CellIndex * 2 + 0;
+			const int32 TriIndex1 = CellIndex * 2 + 1;
+
+			if (!InTriangleVisitor(TriIndex0, Points[0], Points[1], Points[3], CurrentLength))
+			{
+				return false;
+			}
+
+			if (!InTriangleVisitor(TriIndex1, Points[0], Points[3], Points[2], CurrentLength))
+			{
+				return false;
+			}
+
+			return true;
+		}
+
+	private:
+		const typename FHeightField::FDataType* HfData;
+		TriangleVisitor InTriangleVisitor;
+	};
+
+
 	template<typename BufferType>
 	void BuildGeomData(TArrayView<BufferType> BufferView, TArrayView<uint8> MaterialIndexView, int32 NumRows, int32 NumCols, const FVec3& InScale, TUniqueFunction<FReal(const BufferType)> ToRealFunc, typename FHeightField::FDataType& OutData, FAABB3& OutBounds)
 	{
@@ -2185,6 +2382,83 @@
 		return SweepGeomImp(QueryGeom, StartTM, Dir, Length, OutTime, OutPosition, OutNormal, OutFaceIndex, Thickness, bComputeMTD);
 	}
 
+	template <typename QueryGeomType>
+	bool FHeightField::SweepGeomCCDImp(const QueryGeomType& QueryGeom, const FRigidTransform3& StartTM, const FVec3& Dir, const FReal Length, const FReal InIgnorePenetration, const FReal InTargetPenetration, FReal& OutTOI, FReal& OutPhi, FVec3& OutPosition, FVec3& OutNormal, int32& OutFaceIndex) const
+	{
+		// Emulate the old TargetPenetration mode which always uses the first contact encountered, as opposed to the first contact which reaches a depth of TargetPenetration
+		const FReal IgnorePenetration = CVars::bCCDNewTargetDepthMode ? InIgnorePenetration : 0;
+		const FReal TargetPenetration = CVars::bCCDNewTargetDepthMode ? InTargetPenetration : 0;
+
+		bool bHit = false;
+		THeightfieldSweepVisitorCCD<QueryGeomType> SQVisitor(&GeomData, QueryGeom, StartTM, Dir, Length, IgnorePenetration, TargetPenetration);
+		const FAABB3 QueryBounds = QueryGeom.BoundingBox();
+		const FVec3 StartPoint = StartTM.TransformPositionNoScale(QueryBounds.Center());
+
+		const FVec3 Inflation3D = QueryBounds.Extents() * FReal(0.5);
+		GridSweep(StartPoint, Dir, Length, Inflation3D, SQVisitor);
+
+		FReal TOI = (Length > 0) ? SQVisitor.OutDistance / Length : FReal(0);
+		FReal Phi = SQVisitor.OutPhi;
+
+		// @todo(chaos): Legacy path to be removed when fully tested. See ComputeSweptContactTOIAndPhiAtTargetPenetration
+		if (!CVars::bCCDNewTargetDepthMode)
+		{
+			LegacyComputeSweptContactTOIAndPhiAtTargetPenetration(FVec3::DotProduct(Dir, SQVisitor.OutNormal), Length, InIgnorePenetration, InTargetPenetration, TOI, Phi);
+		}
+
+		if (TOI <= 1)
+		{
+			OutTOI = TOI;
+			OutPhi = Phi;
+			OutPosition = SQVisitor.OutPosition;
+			OutNormal = SQVisitor.OutNormal;
+			OutFaceIndex = SQVisitor.OutFaceIndex;
+			bHit = true;
+		}
+
+		return bHit;
+	}
+
+	bool FHeightField::SweepGeomCCD(const TSphere<FReal, 3>& QueryGeom, const FRigidTransform3& StartTM, const FVec3& Dir, const FReal Length, const FReal IgnorePenetration, const FReal TargetPenetration, FReal& OutTOI, FReal& OutPhi, FVec3& OutPosition, FVec3& OutNormal, int32& OutFaceIndex, FVec3& OutFaceNormal) const
+	{
+		return SweepGeomCCDImp(QueryGeom, StartTM, Dir, Length, IgnorePenetration, TargetPenetration, OutTOI, OutPhi, OutPosition, OutNormal, OutFaceIndex);
+	}
+
+	bool FHeightField::SweepGeomCCD(const TBox<FReal, 3>& QueryGeom, const FRigidTransform3& StartTM, const FVec3& Dir, const FReal Length, const FReal IgnorePenetration, const FReal TargetPenetration, FReal& OutTOI, FReal& OutPhi, FVec3& OutPosition, FVec3& OutNormal, int32& OutFaceIndex, FVec3& OutFaceNormal) const
+	{
+		return SweepGeomCCDImp(QueryGeom, StartTM, Dir, Length, IgnorePenetration, TargetPenetration, OutTOI, OutPhi, OutPosition, OutNormal, OutFaceIndex);
+	}
+
+	bool FHeightField::SweepGeomCCD(const FCapsule& QueryGeom, const FRigidTransform3& StartTM, const FVec3& Dir, const FReal Length, const FReal IgnorePenetration, const FReal TargetPenetration, FReal& OutTOI, FReal& OutPhi, FVec3& OutPosition, FVec3& OutNormal, int32& OutFaceIndex, FVec3& OutFaceNormal) const
+	{
+		return SweepGeomCCDImp(QueryGeom, StartTM, Dir, Length, IgnorePenetration, TargetPenetration, OutTOI, OutPhi, OutPosition, OutNormal, OutFaceIndex);
+	}
+
+	bool FHeightField::SweepGeomCCD(const FConvex& QueryGeom, const FRigidTransform3& StartTM, const FVec3& Dir, const FReal Length, const FReal IgnorePenetration, const FReal TargetPenetration, FReal& OutTOI, FReal& OutPhi, FVec3& OutPosition, FVec3& OutNormal, int32& OutFaceIndex, FVec3& OutFaceNormal) const
+	{
+		return SweepGeomCCDImp(QueryGeom, StartTM, Dir, Length, IgnorePenetration, TargetPenetration, OutTOI, OutPhi, OutPosition, OutNormal, OutFaceIndex);
+	}
+
+	bool FHeightField::SweepGeomCCD(const TImplicitObjectScaled<TSphere<FReal, 3>>& QueryGeom, const FRigidTransform3& StartTM, const FVec3& Dir, const FReal Length, const FReal IgnorePenetration, const FReal TargetPenetration, FReal& OutTOI, FReal& OutPhi, FVec3& OutPosition, FVec3& OutNormal, int32& OutFaceIndex, FVec3& OutFaceNormal) const
+	{
+		return SweepGeomCCDImp(QueryGeom, StartTM, Dir, Length, IgnorePenetration, TargetPenetration, OutTOI, OutPhi, OutPosition, OutNormal, OutFaceIndex);
+	}
+
+	bool FHeightField::SweepGeomCCD(const TImplicitObjectScaled<TBox<FReal, 3>>& QueryGeom, const FRigidTransform3& StartTM, const FVec3& Dir, const FReal Length, const FReal IgnorePenetration, const FReal TargetPenetration, FReal& OutTOI, FReal& OutPhi, FVec3& OutPosition, FVec3& OutNormal, int32& OutFaceIndex, FVec3& OutFaceNormal) const
+	{
+		return SweepGeomCCDImp(QueryGeom, StartTM, Dir, Length, IgnorePenetration, TargetPenetration, OutTOI, OutPhi, OutPosition, OutNormal, OutFaceIndex);
+	}
+
+	bool FHeightField::SweepGeomCCD(const TImplicitObjectScaled<FCapsule>& QueryGeom, const FRigidTransform3& StartTM, const FVec3& Dir, const FReal Length, const FReal IgnorePenetration, const FReal TargetPenetration, FReal& OutTOI, FReal& OutPhi, FVec3& OutPosition, FVec3& OutNormal, int32& OutFaceIndex, FVec3& OutFaceNormal) const
+	{
+		return SweepGeomCCDImp(QueryGeom, StartTM, Dir, Length, IgnorePenetration, TargetPenetration, OutTOI, OutPhi, OutPosition, OutNormal, OutFaceIndex);
+	}
+
+	bool FHeightField::SweepGeomCCD(const TImplicitObjectScaled<FConvex>& QueryGeom, const FRigidTransform3& StartTM, const FVec3& Dir, const FReal Length, const FReal IgnorePenetration, const FReal TargetPenetration, FReal& OutTOI, FReal& OutPhi, FVec3& OutPosition, FVec3& OutNormal, int32& OutFaceIndex, FVec3& OutFaceNormal) const
+	{
+		return SweepGeomCCDImp(QueryGeom, StartTM, Dir, Length, IgnorePenetration, TargetPenetration, OutTOI, OutPhi, OutPosition, OutNormal, OutFaceIndex);
+	}
+
 	void FHeightField::VisitTriangles(const FAABB3& QueryBounds, const TFunction<void(const FTriangle& Triangle)>& Visitor) const
 	{
 		FBounds2D FlatQueryBounds;
diff --git a/Engine/Source/Runtime/Experimental/Chaos/Private/Chaos/TriangleMeshImplicitObject.cpp b/Engine/Source/Runtime/Experimental/Chaos/Private/Chaos/TriangleMeshImplicitObject.cpp
--- a/Engine/Source/Runtime/Experimental/Chaos/Private/Chaos/TriangleMeshImplicitObject.cpp
+++ b/Engine/Source/Runtime/Experimental/Chaos/Private/Chaos/TriangleMeshImplicitObject.cpp
@@ -1,6 +1,7 @@
 // Copyright Epic Games, Inc. All Rights Reserved.
 #include "Chaos/TriangleMeshImplicitObject.h"
 #include "Chaos/Collision/ContactPoint.h"
+#include "Chaos/Collision/ContactPointsMiscShapes.h"
 #include "Chaos/Collision/PBDCollisionConstraint.h"
 #include "Chaos/CollisionOneShotManifolds.h"
 #include "Chaos/Capsule.h"
@@ -15,6 +16,11 @@
 
 namespace Chaos
 {
+	namespace CVars
+	{
+		extern bool bCCDNewTargetDepthMode;
+	}
+
 	extern FRealSingle Chaos_Collision_EdgePrunePlaneDistance;
 
 	// Note that if this is re-enabled when previously off, the cooked trimeshes won't have the vertex map serialized, so the change will not take effect until re-cooked.
@@ -1150,6 +1156,164 @@
 
 };
 
+template <typename QueryGeomType, typename IdxType>
+struct FTriangleMeshSweepVisitorCCD
+{
+	FTriangleMeshSweepVisitorCCD(const FTriangleMeshImplicitObject& InTriMesh, const TArray<TVec3<IdxType>>& InElements, const QueryGeomType& InQueryGeom,
+		const FVec3& InScaledDirNormalized, const FReal InLengthScale, const FReal InLength, const FRigidTransform3& InScaledStartTM, FVec3 InTriMeshScale, FReal InCullsBackFaceSweepsCode, const FReal InIgnorePenetration, const FReal InTargetPenetration)
+		: OutDistance(TNumericLimits<FRealSingle>::Max())
+		, OutPhi(TNumericLimits<FRealSingle>::Max())
+		, OutFaceIndex(INDEX_NONE)
+		, TriMesh(InTriMesh)
+		, Elements(InElements)
+		, QueryGeom(InQueryGeom)
+		, IgnorePenetration(InIgnorePenetration)
+		, TargetPenetration(InTargetPenetration)
+		, CullsBackFaceSweepsCode(InCullsBackFaceSweepsCode)
+		, LengthScale(FRealSingle(InLengthScale))
+		, Length(FRealSingle(InLength))
+		, TriMeshScale(InTriMeshScale)
+	{
+		VectorScaledDirNormalized = MakeVectorRegisterFloatFromDouble(MakeVectorRegister(InScaledDirNormalized.X, InScaledDirNormalized.Y, InScaledDirNormalized.Z, 0.0));
+		VectorCullsBackFaceSweepsCode = MakeVectorRegisterFloatFromDouble(VectorLoadFloat1(&InCullsBackFaceSweepsCode));
+
+		const UE::Math::TQuat<FReal>& RotationDouble = InScaledStartTM.GetRotation();
+		RotationSimd = MakeVectorRegisterFloatFromDouble(MakeVectorRegister(RotationDouble.X, RotationDouble.Y, RotationDouble.Z, RotationDouble.W));
+
+		const UE::Math::TVector<FReal>& TranslationDouble = InScaledStartTM.GetTranslation();
+		TranslationSimd = MakeVectorRegisterFloatFromDouble(MakeVectorRegister(TranslationDouble.X, TranslationDouble.Y, TranslationDouble.Z, 0.0));
+
+		// Normalize rotation
+		RotationSimd = VectorNormalizeSafe(RotationSimd, GlobalVectorConstants::Float0001);
+		RayDirSimd = MakeVectorRegisterFloatFromDouble(MakeVectorRegister(InScaledDirNormalized[0], InScaledDirNormalized[1], InScaledDirNormalized[2], 0.0));
+	}
+
+	const void* GetQueryData() const { return nullptr; }
+	const void* GetSimData() const { return nullptr; }
+
+	/** Return a pointer to the payload on which we are querying the acceleration structure */
+	const void* GetQueryPayload() const
+	{
+		return nullptr;
+	}
+
+	bool VisitOverlap(const TSpatialVisitorData<int32>& VisitData)
+	{
+		check(false);
+		return true;
+	}
+
+	bool VisitRaycast(const TSpatialVisitorData<int32>& VisitData, FRealSingle& CurDataLength)
+	{
+		check(false);
+		return true;
+	}
+
+	bool VisitSweep(const TSpatialVisitorData<int32>& VisitData, FRealSingle& CurDataLength)
+	{
+		const int32 TriIdx = VisitData.Payload;
+
+		const VectorRegister4Float TriMeshScaleVector = MakeVectorRegisterFloatFromDouble(MakeVectorRegister(TriMeshScale.X, TriMeshScale.Y, TriMeshScale.Z, 0.0));
+
+		const TParticles<FRealSingle, 3>& Particles = TriMesh.MParticles;
+
+		const TVector<FRealSingle, 3>& AVec = Particles.X(Elements[TriIdx][0]);
+		const TVector<FRealSingle, 3>& BVec = Particles.X(Elements[TriIdx][1]);
+		const TVector<FRealSingle, 3>& CVec = Particles.X(Elements[TriIdx][2]);
+
+		VectorRegister4Float A = MakeVectorRegister(AVec.X, AVec.Y, AVec.Z, 0.0f);
+		VectorRegister4Float B = MakeVectorRegister(BVec.X, BVec.Y, BVec.Z, 0.0f);
+		VectorRegister4Float C = MakeVectorRegister(CVec.X, CVec.Y, CVec.Z, 0.0f);
+
+		A = VectorMultiply(A, TriMeshScaleVector);
+		B = VectorMultiply(B, TriMeshScaleVector);
+		C = VectorMultiply(C, TriMeshScaleVector);
+
+		FTriangleRegister Tri(A, B, C);
+		const VectorRegister4Float TriNormal = VectorCross(VectorSubtract(B, A), VectorSubtract(C, A));
+
+		if (CullsBackFaceSweepsCode != 0)
+		{
+			const VectorRegister4Float ReturnTrue = VectorCompareGT(VectorMultiply(VectorDot3(TriNormal, VectorScaledDirNormalized), VectorCullsBackFaceSweepsCode), VectorZero());
+			if (VectorMaskBits(ReturnTrue))
+			{
+				return true;
+			}
+		}
+
+		// If we are definitely penetrating by less than IgnorePenetration at T=1, we can ignore this triangle
+		// Checks the query geometry support point along the normal at the sweep end point
+		VectorRegister4Float IgnorePenetrationSimd = MakeVectorRegisterFloatFromDouble(VectorSetFloat1(IgnorePenetration));
+		VectorRegister4Float LengthSimd = MakeVectorRegisterFloatFromDouble(VectorSetFloat1(LengthScale * Length));
+		VectorRegister4Float EndPointSimd = VectorAdd(TranslationSimd, VectorMultiply(RayDirSimd, LengthSimd));
+		VectorRegister4Float TriangleNormalSimd = VectorNormalize(VectorCross(VectorSubtract(B, A), VectorSubtract(C, A)));
+		VectorRegister4Float OtherTriangleNormalSimd = VectorQuaternionInverseRotateVector(RotationSimd, TriangleNormalSimd);
+		VectorRegister4Float OtherTrianglePositionSimd = VectorQuaternionInverseRotateVector(RotationSimd, VectorSubtract(A, EndPointSimd));
+		VectorRegister4Float OtherSupportSimd = QueryGeom.SupportCoreSimd(VectorNegate(OtherTriangleNormalSimd), 0);
+		VectorRegister4Float MaxDepthSimd = VectorDot3(VectorSubtract(OtherTrianglePositionSimd, OtherSupportSimd), OtherTriangleNormalSimd);
+		if (VectorMaskBits(VectorCompareLT(MaxDepthSimd, IgnorePenetrationSimd)))
+		{
+			return true;
+		}
+
+		FRealSingle Distance;
+		VectorRegister4Float PositionSimd, NormalSimd;
+		if (GJKRaycast2ImplSimd(Tri, QueryGeom, RotationSimd, TranslationSimd, RayDirSimd, LengthScale * CurDataLength, Distance, PositionSimd, NormalSimd, true, GlobalVectorConstants::Float1000))
+		{
+			const VectorRegister4Float DirDotNormalSimd = VectorDot3(RayDirSimd, NormalSimd);
+			FReal DirDotNormal;
+			VectorStoreFloat1(DirDotNormalSimd, &DirDotNormal);
+
+			// Calculate the time to reach a depth of TargetPenetration
+			FReal TargetTOI, TargetPhi;
+			ComputeSweptContactTOIAndPhiAtTargetPenetration(DirDotNormal, LengthScale * CurDataLength, FReal(Distance), IgnorePenetration, TargetPenetration, TargetTOI, TargetPhi);
+
+			// TargetDistance is local scale, OutDistance is local scale.
+			const FRealSingle TargetDistance = FRealSingle(TargetTOI) * CurDataLength;
+			if (TargetDistance < CurDataLength)
+			{
+				CurDataLength = TargetDistance;
+
+				// Transform results back into world scale and save
+				FVec3 HitPosition;
+				FVec3 HitNormal;
+				FRealSingle TOI;	// Unused
+				VectorStoreFloat3(MakeVectorRegisterDouble(PositionSimd), &HitPosition);
+				VectorStoreFloat3(MakeVectorRegisterDouble(NormalSimd), &HitNormal);
+				TransformSweepOutputsHelper(TriMeshScale, HitNormal, HitPosition, LengthScale, TargetDistance, OutNormal, OutPosition, TOI);
+				OutDistance = TargetDistance;
+				OutPhi = TargetPhi;
+				OutFaceIndex = TriIdx;
+			}
+		}
+
+		return (CurDataLength > 0);
+	}
+
+	FReal OutDistance;
+	FReal OutPhi;
+	FVec3 OutPosition;
+	FVec3 OutNormal;
+	int32 OutFaceIndex;
+
+private:
+	const FTriangleMeshImplicitObject& TriMesh;
+	const TArray<TVec3<IdxType>>& Elements;
+	const QueryGeomType& QueryGeom;
+	const FReal IgnorePenetration;
+	const FReal TargetPenetration;
+	const FReal CullsBackFaceSweepsCode; // 0: no culling, 1/-1: winding order
+	const FRealSingle LengthScale;
+	const FRealSingle Length;
+	const FVec3 TriMeshScale;
+
+	VectorRegister4Float RotationSimd;
+	VectorRegister4Float TranslationSimd;
+	VectorRegister4Float RayDirSimd;
+	VectorRegister4Float VectorScaledDirNormalized;
+	VectorRegister4Float VectorCullsBackFaceSweepsCode; // 0: no culling, 1/-1: winding order
+};
+
 void ComputeScaledSweepInputs(FVec3 TriMeshScale, const FRigidTransform3& StartTM, const FVec3& Dir, const FReal Length,
 	FVec3& OutScaledDirNormalized, FReal& OutLengthScale, FRigidTransform3& OutScaledStartTM)
 {
@@ -1271,6 +1435,106 @@
 	return SweepGeomImp(QueryGeom, StartTM, Dir, Length, OutTime, OutPosition, OutNormal, OutFaceIndex, OutFaceNormal, Thickness, bComputeMTD, TriMeshScale);
 }
 
+
+template<typename QueryGeomType>
+bool FTriangleMeshImplicitObject::SweepGeomCCDImp(const QueryGeomType& QueryGeom, const FRigidTransform3& StartTM, const FVec3& Dir, const FReal Length, const FReal InIgnorePenetration, const FReal InTargetPenetration, FReal& OutTOI, FReal& OutPhi, FVec3& OutPosition, FVec3& OutNormal, int32& OutFaceIndex, FVec3& OutFaceNormal, const FVec3& TriMeshScale) const
+{
+	// Compute scaled sweep inputs to cache in visitor.
+	FVec3 ScaledDirNormalized;
+	FReal LengthScale;
+	FRigidTransform3 ScaledStartTM;
+	ComputeScaledSweepInputs(TriMeshScale, StartTM, Dir, Length, ScaledDirNormalized, LengthScale, ScaledStartTM);
+
+	bool bHit = false;
+	auto LambdaHelper = [&](const auto& Elements)
+	{
+		// Emulate the old TargetPenetration mode which always uses the first contact encountered, as opposed to the first contact which reaches a depth of TargetPenetration
+		const FReal IgnorePenetration = CVars::bCCDNewTargetDepthMode ? InIgnorePenetration : 0;
+		const FReal TargetPenetration = CVars::bCCDNewTargetDepthMode ? InTargetPenetration : 0;
+
+		const FReal CullsBackFaceRaycastCode = bCullsBackFaceRaycast ? GetWindingOrder(TriMeshScale) : 0.f;
+		using VisitorType = FTriangleMeshSweepVisitorCCD<QueryGeomType, decltype(Elements[0][0])>;
+		VisitorType SQVisitor(*this, Elements, QueryGeom, ScaledDirNormalized, LengthScale, Length, ScaledStartTM, TriMeshScale, CullsBackFaceRaycastCode, IgnorePenetration, TargetPenetration);
+
+		const FAABB3 QueryBounds = QueryGeom.CalculateTransformedBounds(FRigidTransform3(FVec3::ZeroVector, StartTM.GetRotation()));
+		const FVec3 InvTriMeshScale = SafeInvScale(TriMeshScale);
+		const FVec3 StartPoint = QueryBounds.Center() * InvTriMeshScale + StartTM.GetLocation();
+		const FVec3 Inflation = QueryBounds.Extents() * InvTriMeshScale.GetAbs() * 0.5;
+		FastBVH.template Sweep<VisitorType>(StartPoint, Dir, Length, Inflation, SQVisitor);
+
+		FReal TOI = (Length > 0) ? SQVisitor.OutDistance / Length : FReal(0);
+		FReal Phi = SQVisitor.OutPhi;
+
+		// @todo(chaos): Legacy path to be removed when fully tested. See ComputeSweptContactTOIAndPhiAtTargetPenetration
+		if (!CVars::bCCDNewTargetDepthMode)
+		{
+			LegacyComputeSweptContactTOIAndPhiAtTargetPenetration(FVec3::DotProduct(Dir, SQVisitor.OutNormal), LengthScale * Length, InIgnorePenetration, InTargetPenetration, TOI, Phi);
+		}
+
+		if (TOI <= 1)
+		{
+			OutTOI = TOI;
+			OutPhi = Phi;
+			OutPosition = SQVisitor.OutPosition;
+			OutNormal = SQVisitor.OutNormal;
+			OutFaceIndex = SQVisitor.OutFaceIndex;
+			OutFaceNormal = GetFaceNormal(OutFaceIndex);
+			bHit = true;
+		}
+	};
+
+	if (MElements.RequiresLargeIndices())
+	{
+		LambdaHelper(MElements.GetLargeIndexBuffer());
+	}
+	else
+	{
+		LambdaHelper(MElements.GetSmallIndexBuffer());
+	}
+	return bHit;
+}
+
+bool FTriangleMeshImplicitObject::SweepGeomCCD(const TSphere<FReal, 3>& QueryGeom, const FRigidTransform3& StartTM, const FVec3& Dir, const FReal Length, const FReal IgnorePenetration, const FReal TargetPenetration, FReal& OutTOI, FReal& OutPhi, FVec3& OutPosition, FVec3& OutNormal, int32& OutFaceIndex, FVec3& OutFaceNormal, const FVec3& TriMeshScale) const
+{
+	return SweepGeomCCDImp(QueryGeom, StartTM, Dir, Length, IgnorePenetration, TargetPenetration, OutTOI, OutPhi, OutPosition, OutNormal, OutFaceIndex, OutFaceNormal, TriMeshScale);
+}
+
+bool FTriangleMeshImplicitObject::SweepGeomCCD(const TBox<FReal, 3>& QueryGeom, const FRigidTransform3& StartTM, const FVec3& Dir, const FReal Length, const FReal IgnorePenetration, const FReal TargetPenetration, FReal& OutTOI, FReal& OutPhi, FVec3& OutPosition, FVec3& OutNormal, int32& OutFaceIndex, FVec3& OutFaceNormal, const FVec3& TriMeshScale) const
+{
+	return SweepGeomCCDImp(QueryGeom, StartTM, Dir, Length, IgnorePenetration, TargetPenetration, OutTOI, OutPhi, OutPosition, OutNormal, OutFaceIndex, OutFaceNormal, TriMeshScale);
+}
+
+bool FTriangleMeshImplicitObject::SweepGeomCCD(const FCapsule& QueryGeom, const FRigidTransform3& StartTM, const FVec3& Dir, const FReal Length, const FReal IgnorePenetration, const FReal TargetPenetration, FReal& OutTOI, FReal& OutPhi, FVec3& OutPosition, FVec3& OutNormal, int32& OutFaceIndex, FVec3& OutFaceNormal, const FVec3& TriMeshScale) const
+{
+	return SweepGeomCCDImp(QueryGeom, StartTM, Dir, Length, IgnorePenetration, TargetPenetration, OutTOI, OutPhi, OutPosition, OutNormal, OutFaceIndex, OutFaceNormal, TriMeshScale);
+}
+
+bool FTriangleMeshImplicitObject::SweepGeomCCD(const FConvex& QueryGeom, const FRigidTransform3& StartTM, const FVec3& Dir, const FReal Length, const FReal IgnorePenetration, const FReal TargetPenetration, FReal& OutTOI, FReal& OutPhi, FVec3& OutPosition, FVec3& OutNormal, int32& OutFaceIndex, FVec3& OutFaceNormal, const FVec3& TriMeshScale) const
+{
+	return SweepGeomCCDImp(QueryGeom, StartTM, Dir, Length, IgnorePenetration, TargetPenetration, OutTOI, OutPhi, OutPosition, OutNormal, OutFaceIndex, OutFaceNormal, TriMeshScale);
+}
+
+bool FTriangleMeshImplicitObject::SweepGeomCCD(const TImplicitObjectScaled<TSphere<FReal, 3>>& QueryGeom, const FRigidTransform3& StartTM, const FVec3& Dir, const FReal Length, const FReal IgnorePenetration, const FReal TargetPenetration, FReal& OutTOI, FReal& OutPhi, FVec3& OutPosition, FVec3& OutNormal, int32& OutFaceIndex, FVec3& OutFaceNormal, const FVec3& TriMeshScale) const
+{
+	return SweepGeomCCDImp(QueryGeom, StartTM, Dir, Length, IgnorePenetration, TargetPenetration, OutTOI, OutPhi, OutPosition, OutNormal, OutFaceIndex, OutFaceNormal, TriMeshScale);
+}
+
+bool FTriangleMeshImplicitObject::SweepGeomCCD(const TImplicitObjectScaled<TBox<FReal, 3>>& QueryGeom, const FRigidTransform3& StartTM, const FVec3& Dir, const FReal Length, const FReal IgnorePenetration, const FReal TargetPenetration, FReal& OutTOI, FReal& OutPhi, FVec3& OutPosition, FVec3& OutNormal, int32& OutFaceIndex, FVec3& OutFaceNormal, const FVec3& TriMeshScale) const
+{
+	return SweepGeomCCDImp(QueryGeom, StartTM, Dir, Length, IgnorePenetration, TargetPenetration, OutTOI, OutPhi, OutPosition, OutNormal, OutFaceIndex, OutFaceNormal, TriMeshScale);
+}
+
+bool FTriangleMeshImplicitObject::SweepGeomCCD(const TImplicitObjectScaled<FCapsule>& QueryGeom, const FRigidTransform3& StartTM, const FVec3& Dir, const FReal Length, const FReal IgnorePenetration, const FReal TargetPenetration, FReal& OutTOI, FReal& OutPhi, FVec3& OutPosition, FVec3& OutNormal, int32& OutFaceIndex, FVec3& OutFaceNormal, const FVec3& TriMeshScale) const
+{
+	return SweepGeomCCDImp(QueryGeom, StartTM, Dir, Length, IgnorePenetration, TargetPenetration, OutTOI, OutPhi, OutPosition, OutNormal, OutFaceIndex, OutFaceNormal, TriMeshScale);
+}
+
+bool FTriangleMeshImplicitObject::SweepGeomCCD(const TImplicitObjectScaled<FConvex>& QueryGeom, const FRigidTransform3& StartTM, const FVec3& Dir, const FReal Length, const FReal IgnorePenetration, const FReal TargetPenetration, FReal& OutTOI, FReal& OutPhi, FVec3& OutPosition, FVec3& OutNormal, int32& OutFaceIndex, FVec3& OutFaceNormal, const FVec3& TriMeshScale) const
+{
+	return SweepGeomCCDImp(QueryGeom, StartTM, Dir, Length, IgnorePenetration, TargetPenetration, OutTOI, OutPhi, OutPosition, OutNormal, OutFaceIndex, OutFaceNormal, TriMeshScale);
+}
+
+
 template <typename IdxType>
 int32 FTriangleMeshImplicitObject::FindMostOpposingFace(const TArray<TVec3<IdxType>>& Elements, const FVec3& Position, const FVec3& UnitDir, int32 HintFaceIndex, FReal SearchDist, const FVec3& Scale) const
 {
diff --git a/Engine/Source/Runtime/Experimental/Chaos/Public/Chaos/CCDUtilities.h b/Engine/Source/Runtime/Experimental/Chaos/Public/Chaos/CCDUtilities.h
--- a/Engine/Source/Runtime/Experimental/Chaos/Public/Chaos/CCDUtilities.h
+++ b/Engine/Source/Runtime/Experimental/Chaos/Public/Chaos/CCDUtilities.h
@@ -54,7 +54,6 @@
 	{
 		FPBDCollisionConstraint* SweptConstraint;
 		FCCDParticle* Particle[2];
-		FReal PhiThreshold;
 		int32 Island;
 		int32 ProcessedCount; // The number of times this constraint is processed in the current frame.
 		int32 FastMovingKinematicIndex;
@@ -63,17 +62,15 @@
 		FCCDConstraint()
 			: SweptConstraint(nullptr)
 			, Particle{ nullptr, nullptr }
-			, PhiThreshold(0)
 			, Island(INDEX_NONE)
 			, ProcessedCount(0)
 			, FastMovingKinematicIndex(INDEX_NONE)
 			, NetImpulse(0)
 		{}
 
-		FCCDConstraint(FPBDCollisionConstraint* const InConstraint, FCCDParticle* InParticle[], const FVec3 Displacements[], const FReal InPhiThreshold)
+		FCCDConstraint(FPBDCollisionConstraint* const InConstraint, FCCDParticle* InParticle[], const FVec3 Displacements[])
 			: SweptConstraint(InConstraint)
 			, Particle{ InParticle[0], InParticle[1] }
-			, PhiThreshold(InPhiThreshold)
 			, Island(INDEX_NONE)
 			, ProcessedCount(0)
 			, FastMovingKinematicIndex(GetFastMovingKinematicIndex(InConstraint, Displacements))
@@ -100,6 +97,7 @@
 		void AssignConstraintIslandsAndRecordConstraintNum();
 		void GroupConstraintsWithIslands();
 		void ApplyIslandSweptConstraints(const int32 Island, const FReal Dt);
+		void ApplyIslandSweptConstraints2(const int32 Island, const FReal Dt);
 		void ResetIslandParticles(const int32 Island);
 		void ResetIslandConstraints(const int32 Island);
 		void AdvanceParticleXToTOI(FCCDParticle *CCDParticle, const FReal TOI, const FReal Dt) const;
diff --git a/Engine/Source/Runtime/Experimental/Chaos/Public/Chaos/Collision/ContactPointsMiscShapes.h b/Engine/Source/Runtime/Experimental/Chaos/Public/Chaos/Collision/ContactPointsMiscShapes.h
--- a/Engine/Source/Runtime/Experimental/Chaos/Public/Chaos/Collision/ContactPointsMiscShapes.h
+++ b/Engine/Source/Runtime/Experimental/Chaos/Public/Chaos/Collision/ContactPointsMiscShapes.h
@@ -12,17 +12,34 @@
 	class FContactPoint;
 	class FHeightField;
 	
-	void ComputeSweptContactPhiAndTOIHelper(const FVec3& ContactNormal, const FVec3& Dir, const FReal& Length, const FReal& HitTime, FReal& OutTOI, FReal& OutPhi);
+	/**
+	 * @brief Given a sweep result, calculate the sweep time at which the penetration depth will be TargetPenetration
+	 * This is based on the initial and final contact separation from the sweep test but modified so that
+	 * - we ignore separating contacts (increasing Phi)
+	 * - we ignore contacts that are separated at T=1 (EndPhi > 0)
+	 * - we ignore contacts if the penetration is less than the IgnorePenetration
+	 * - we calculate the TOI that leaves a penetration depth of TargetPenetration (except when initially overlapping by more than TargetPenetration, in which case TOI=0)
+	 * 
+	 * @param SweepLength The length of the sweep test
+	 * @param HitDistance The distance along the sweep where the shapes first touch
+	 * @param IgnorePenetration Ignore contacts if the depth at T=1 is less than this
+	 * @param TargetPenetration Calculate the TOI for this penetration depth
+	 * @param OutTOI The sweep time at which we hit the TargetPenetration, or "infinity" (numeric max)
+	 * @param OutPhi The depth at OutTOI, which is usually TargetPenetration unless OutTOI is zero
+	 * @return true if we have a TOI less than 1
+	*/
+	bool ComputeSweptContactTOIAndPhiAtTargetPenetration(const FReal DirDotNormal, const FReal SweepLength, const FReal HitDistance, const FReal IgnorePenetration, const FReal TargetPenetration, FReal& OutTOI, FReal& OutPhi);
+	bool ComputeSweptContactTOIAndPhiAtTargetPenetration(const FVec3& ContactNormal, const FVec3& Dir, const FReal SweepLength, const FReal HitDistance, const FReal IgnorePenetration, const FReal TargetPenetration, FReal& OutTOI, FReal& OutPhi);
+	void LegacyComputeSweptContactTOIAndPhiAtTargetPenetration(const FReal DirDotNormal, const FReal Length, const FReal IgnorePenetration, const FReal TargetPenetration, FReal& InOutTOI, FReal& InOutPhi);
 
+
 	template <typename GeometryB>
-	FContactPoint GJKImplicitSweptContactPoint(const FImplicitObject& A, const FRigidTransform3& AStartTransform, const GeometryB& B, const FRigidTransform3& BTransform, const FVec3& Dir, const FReal Length, FReal& TOI);
+	FContactPoint GJKImplicitSweptContactPoint(const FImplicitObject& A, const FRigidTransform3& AStartTransform, const GeometryB& B, const FRigidTransform3& BTransform, const FVec3& Dir, const FReal Length, const FReal IgnorePenetration, const FReal TargetPenetration, FReal& TOI);
 
 	
 	template <typename GeometryA, typename GeometryB>
 	FContactPoint GJKImplicitContactPoint(const FImplicitObject& A, const FRigidTransform3& ATransform, const GeometryB& B, const FRigidTransform3& BTransform, const FReal CullDistance, const FReal ShapePadding);
 
-	FContactPoint GJKImplicitScaledTriMeshSweptContactPoint(const FImplicitObject& A, const FRigidTransform3& AStartTransform, const TImplicitObjectScaled<FTriangleMeshImplicitObject>& B, const FRigidTransform3& BTransform, const FVec3& Dir, const FReal Length, FReal& TOI);
-
 	template <typename GeometryA, typename GeometryB>
 	void GJKImplicitManifold(const FImplicitObject& A, const FRigidTransform3& ATransform, const GeometryB& B, const FRigidTransform3& BTransform, const FReal CullDistance, const FReal ShapePadding, TArray<FContactPoint>& ContactPoints);
 
@@ -35,7 +52,7 @@
 	FContactPoint SphereTriangleMeshContactPoint(const TSphere<FReal, 3>& A, const FRigidTransform3& ATransform, const TriMeshType& B, const FRigidTransform3& BTransform, const FReal CullDistance, const FReal ShapePadding);
 
 	template<typename TriMeshType>
-	FContactPoint SphereTriangleMeshSweptContactPoint(const TSphere<FReal, 3>& A, const FRigidTransform3& ATransform, const TriMeshType& B, const FRigidTransform3& BStartTransform, const FVec3& Dir, const FReal Length, FReal& TOI);
+	FContactPoint SphereTriangleMeshSweptContactPoint(const TSphere<FReal, 3>& A, const FRigidTransform3& ATransform, const TriMeshType& B, const FRigidTransform3& BStartTransform, const FVec3& Dir, const FReal Length, const FReal IgnorePenetration, const FReal TargetPenetration, FReal& TOI);
 
 	FContactPoint BoxHeightFieldContactPoint(const FImplicitBox3& A, const FRigidTransform3& ATransform, const FHeightField& B, const FRigidTransform3& BTransform, const FReal CullDistance, const FReal ShapePadding);
 
@@ -47,12 +64,12 @@
 	template <typename TriMeshType>
 	FContactPoint CapsuleTriangleMeshContactPoint(const FCapsule& A, const FRigidTransform3& ATransform, const TriMeshType& B, const FRigidTransform3& BTransform, const FReal CullDistance, const FReal ShapePadding);
 	template <typename TriMeshType>
-	FContactPoint CapsuleTriangleMeshSweptContactPoint(const FCapsule& A, const FRigidTransform3& ATransform, const TriMeshType& B, const FRigidTransform3& BStartTransform, const FVec3& Dir, const FReal Length, FReal& TOI);
+	FContactPoint CapsuleTriangleMeshSweptContactPoint(const FCapsule& A, const FRigidTransform3& ATransform, const TriMeshType& B, const FRigidTransform3& BStartTransform, const FVec3& Dir, const FReal Length, const FReal IgnorePenetration, const FReal TargetPenetration, FReal& TOI);
 
 	FContactPoint ConvexHeightFieldContactPoint(const FImplicitObject& A, const FRigidTransform3& ATransform, const FHeightField& B, const FRigidTransform3& BTransform, const FReal CullDistance, const FReal ShapePadding);
 	FContactPoint ConvexTriangleMeshContactPoint(const FImplicitObject& A, const FRigidTransform3& ATransform, const FImplicitObject& B, const FRigidTransform3& BTransform, const FReal CullDistance, const FReal ShapePadding);
 	template <typename TriMeshType>
-	FContactPoint ConvexTriangleMeshSweptContactPoint(const FImplicitObject& A, const FRigidTransform3& ATransform, const TriMeshType& B, const FRigidTransform3& BStartTransform, const FVec3& Dir, const FReal Length, FReal& TOI);
+	FContactPoint ConvexTriangleMeshSweptContactPoint(const FImplicitObject& A, const FRigidTransform3& ATransform, const TriMeshType& B, const FRigidTransform3& BStartTransform, const FVec3& Dir, const FReal Length, const FReal IgnorePenetration, const FReal TargetPenetration, FReal& TOI);
 
 	FContactPoint CapsuleCapsuleContactPoint(const FCapsule& A, const FRigidTransform3& ATransform, const FCapsule& B, const FRigidTransform3& BTransform, const FReal ShapePadding);
 	FContactPoint CapsuleBoxContactPoint(const FCapsule& A, const FRigidTransform3& ATransform, const FImplicitBox3& B, const FRigidTransform3& BTransform, const FVec3& InitialDir, const FReal ShapePadding);
diff --git a/Engine/Source/Runtime/Experimental/Chaos/Public/Chaos/Collision/PBDCollisionConstraint.h b/Engine/Source/Runtime/Experimental/Chaos/Public/Chaos/Collision/PBDCollisionConstraint.h
--- a/Engine/Source/Runtime/Experimental/Chaos/Public/Chaos/Collision/PBDCollisionConstraint.h
+++ b/Engine/Source/Runtime/Experimental/Chaos/Public/Chaos/Collision/PBDCollisionConstraint.h
@@ -307,12 +307,30 @@
 		void SetCCDEnabled(const bool bCCDEnabled)
 		{ 
 			CCDType = bCCDEnabled ? ECollisionCCDType::Enabled : ECollisionCCDType::Disabled;
-			if (bCCDEnabled && (CCDPenetrationThreshold == FReal(0)))
+
+			// Initialize the CCD thresholds (one time only)
+			if (bCCDEnabled && (CCDEnablePenetration == FReal(0)))
 			{
 				InitCCDThreshold();
 			}
 		}
 
+		/**
+		 * @brief If CCD is enabled, contacts deeper than this will be handled by CCD
+		*/
+		FReal GetCCDEnablePenetration() const
+		{
+			return CCDEnablePenetration;
+		}
+
+		/**
+		 * @brief If CCD is enabled and processed the contact, CCD resolution leaves up to this much penetration
+		*/
+		FReal GetCCDTargetPenetration() const
+		{
+			return CCDTargetPenetration;
+		}
+
 		//
 		// API
 		//
@@ -524,17 +542,6 @@
 			}
 		}
 
-		/**
-		 * @brief Update the contact depth and TOI based on the current shape world transforms and the given start transforms
-		*/
-		void UpdateSweptManifoldPoints(const FVec3& ShapeStartWorldPosition0, const FVec3& ShapeStartWorldPosition1, const FReal Dt);
-
-		/**
-		 * @brief return the TOI to use for CCD given the start and end separations
-		 * We ignore separated and separating contacts, and modify TOI to leave some minimim overlap.
-		*/
-		FReal CalculateModifiedSweptTOI(const FReal StartPhi, const FReal EndPhi);
-
 		// The GJK warm-start data. This is updated directly in the narrow phase
 		FGJKSimplexData& GetGJKWarmStartData() { return GJKWarmStartData; }
 
@@ -575,6 +582,16 @@
 		}
 
 		/**
+		 * @brief Time of impact from CCD sweep test if CCD is activate.Otherwise undefined.
+		*/
+		FReal GetCCDTimeOfImpact() const { return CCDTimeOfImpact; }
+
+		/**
+		 * @brief Set the CCD TOI from the collision detection sweep
+		*/
+		void SetCCDTimeOfImpact(const FReal TOI) { CCDTimeOfImpact = TOI; }
+
+		/**
 		 * \brief Store the results of CCD contact resolution, if active
 		 */
 		void SetCCDResults(const FVec3& InNetImpulse)
@@ -707,11 +724,19 @@
 	public:
 		FVec3 AccumulatedImpulse;					// @todo(chaos): we need to accumulate angular impulse separately
 
+	private:
 		// Value in range [0,1] used to interpolate P between [X,P] that we will rollback to when solving at time of impact.
 		FReal CCDTimeOfImpact;
-		FReal CCDPenetrationThreshold;
 
-	private:
+		// The penetration at which CCD contacts get processed following a CCD sweep test.
+		// NOTE: also see GeometryParticle CCDAxisThreshold, which is used to determine when to enable sweeping during collision detection.
+		// Calculated from particle pair properties when constraint is created.
+		FReal CCDEnablePenetration;
+
+		// The penetration we leave behind when rolling back to a CCD time of impact. Should be less than or equal to CCDEnablePenetration.
+		// Calculated from particle pair properties when constraint is created.
+		FReal CCDTargetPenetration;
+
 		FPBDCollisionConstraintContainerCookie ContainerCookie;
 		EContactShapesType ShapesType;
 		ECollisionCCDType CCDType;
diff --git a/Engine/Source/Runtime/Experimental/Chaos/Public/Chaos/HeightField.h b/Engine/Source/Runtime/Experimental/Chaos/Public/Chaos/HeightField.h
--- a/Engine/Source/Runtime/Experimental/Chaos/Public/Chaos/HeightField.h
+++ b/Engine/Source/Runtime/Experimental/Chaos/Public/Chaos/HeightField.h
@@ -74,6 +74,16 @@
 		bool SweepGeom(const TImplicitObjectScaled<FCapsule>& QueryGeom, const FRigidTransform3& StartTM, const FVec3& Dir, const FReal Length, FReal& OutTime, FVec3& OutPosition, FVec3& OutNormal, int32& OutFaceIndex, FVec3& OutFaceNormal, const FReal Thickness = 0, bool bComputeMTD = false) const;
 		bool SweepGeom(const TImplicitObjectScaled<FConvex>& QueryGeom, const FRigidTransform3& StartTM, const FVec3& Dir, const FReal Length, FReal& OutTime, FVec3& OutPosition, FVec3& OutNormal, int32& OutFaceIndex, FVec3& OutFaceNormal, const FReal Thickness = 0, bool bComputeMTD = false) const;
 
+		// Sweep used for CCD. Ignores triangles we penetrate by less than IgnorePenetration, and calculate the TOI for a depth of TargetPenetration. If both are zero, this is equivalent to SweepGeom
+		bool SweepGeomCCD(const TSphere<FReal, 3>& QueryGeom, const FRigidTransform3& StartTM, const FVec3& Dir, const FReal Length, const FReal IgnorePenetration, const FReal TargetPenetration, FReal& OutTOI, FReal& OutPhi, FVec3& OutPosition, FVec3& OutNormal, int32& OutFaceIndex, FVec3& OutFaceNormal) const;
+		bool SweepGeomCCD(const TBox<FReal, 3>& QueryGeom, const FRigidTransform3& StartTM, const FVec3& Dir, const FReal Length, const FReal IgnorePenetration, const FReal TargetPenetration, FReal& OutTOI, FReal& OutPhi, FVec3& OutPosition, FVec3& OutNormal, int32& OutFaceIndex, FVec3& OutFaceNormal) const;
+		bool SweepGeomCCD(const FCapsule& QueryGeom, const FRigidTransform3& StartTM, const FVec3& Dir, const FReal Length, const FReal IgnorePenetration, const FReal TargetPenetration, FReal& OutTOI, FReal& OutPhi, FVec3& OutPosition, FVec3& OutNormal, int32& OutFaceIndex, FVec3& OutFaceNormal) const;
+		bool SweepGeomCCD(const FConvex& QueryGeom, const FRigidTransform3& StartTM, const FVec3& Dir, const FReal Length, const FReal IgnorePenetration, const FReal TargetPenetration, FReal& OutTOI, FReal& OutPhi, FVec3& OutPosition, FVec3& OutNormal, int32& OutFaceIndex, FVec3& OutFaceNormal) const;
+		bool SweepGeomCCD(const TImplicitObjectScaled<TSphere<FReal, 3>>& QueryGeom, const FRigidTransform3& StartTM, const FVec3& Dir, const FReal Length, const FReal IgnorePenetration, const FReal TargetPenetration, FReal& OutTOI, FReal& OutPhi, FVec3& OutPosition, FVec3& OutNormal, int32& OutFaceIndex, FVec3& OutFaceNormal) const;
+		bool SweepGeomCCD(const TImplicitObjectScaled<TBox<FReal, 3>>& QueryGeom, const FRigidTransform3& StartTM, const FVec3& Dir, const FReal Length, const FReal IgnorePenetration, const FReal TargetPenetration, FReal& OutTOI, FReal& OutPhi, FVec3& OutPosition, FVec3& OutNormal, int32& OutFaceIndex, FVec3& OutFaceNormal) const;
+		bool SweepGeomCCD(const TImplicitObjectScaled<FCapsule>& QueryGeom, const FRigidTransform3& StartTM, const FVec3& Dir, const FReal Length, const FReal IgnorePenetration, const FReal TargetPenetration, FReal& OutTOI, FReal& OutPhi, FVec3& OutPosition, FVec3& OutNormal, int32& OutFaceIndex, FVec3& OutFaceNormal) const;
+		bool SweepGeomCCD(const TImplicitObjectScaled<FConvex>& QueryGeom, const FRigidTransform3& StartTM, const FVec3& Dir, const FReal Length, const FReal IgnorePenetration, const FReal TargetPenetration, FReal& OutTOI, FReal& OutPhi, FVec3& OutPosition, FVec3& OutNormal, int32& OutFaceIndex, FVec3& OutFaceNormal) const;
+
 		bool GJKContactPoint(const TBox<FReal, 3>& QueryGeom, const FRigidTransform3& QueryTM, const FReal Thickness, FVec3& ContactLocation, FVec3& ContactNormal, FReal& ContactPhi) const;
 		bool GJKContactPoint(const TSphere<FReal, 3>& QueryGeom, const FRigidTransform3& QueryTM, const FReal Thickness, FVec3& ContactLocation, FVec3& ContactNormal, FReal& ContactPhi) const;
 		bool GJKContactPoint(const FCapsule& QueryGeom, const FRigidTransform3& QueryTM, const FReal Thickness, FVec3& ContactLocation, FVec3& ContactNormal, FReal& ContactPhi) const;
@@ -234,6 +244,28 @@
 				return Scale[1];
 			}
 
+			/**
+			 * @brief Convert a Cell Index to a Vertex Index. 
+			 * Returns the vertex index of the first corner in the cell. The other corners will be 
+			 * (VertexIndex + 1), (VertexIndex + NumCols), (VertexIndex + NumCols + 1).
+			 * Most functions in this class take a vertex index.
+			 *
+			 *	00 --- 01 --- 02 --- 03
+			 *	|  0   |  1   |  2   |
+			 *	04 --- 05 --- 06 --- 07
+			 *	|  3   |  4   |  5   |
+			 *	08 --- 09 --- 10 --- 11
+			 *	|  6   |  7   |  8   |
+			 *	12 --- 13 --- 14 --- 15
+			 *
+			*/
+			FORCEINLINE int32 CellIndexToVertexIndex(const int32 CellIndex) const
+			{
+				const int32 CellY = CellIndex / (NumCols - 1);
+				const int32 HeightIndex = CellIndex + CellY;
+				return HeightIndex;
+			}
+
 			FORCEINLINE FVec3 GetPoint(int32 Index) const
 			{
 				const FReal Height = MinValue + Heights[Index] * HeightPerUnit;
@@ -736,6 +768,9 @@
 		template <typename QueryGeomType>
 		bool SweepGeomImp(const QueryGeomType& QueryGeom, const FRigidTransform3& StartTM, const FVec3& Dir, const FReal Length, FReal& OutTime, FVec3& OutPosition, FVec3& OutNormal, int32& OutFaceIndex, const FReal Thickness, bool bComputeMTD) const;
 
+		template <typename QueryGeomType>
+		bool SweepGeomCCDImp(const QueryGeomType& QueryGeom, const FRigidTransform3& StartTM, const FVec3& Dir, const FReal Length, const FReal IgnorePenetration, const FReal TargetPenetration, FReal& OutTOI, FReal& OutPhi, FVec3& OutPosition, FVec3& OutNormal, int32& OutFaceIndex) const;
+
 		template <typename GeomType>
 		bool GJKContactPointImp(const GeomType& QueryGeom, const FRigidTransform3& QueryTM, const FReal Thickness, FVec3& ContactLocation, FVec3& ContactNormal, FReal& ContactPhi) const;
 
diff --git a/Engine/Source/Runtime/Experimental/Chaos/Public/Chaos/ImplicitObjectScaled.h b/Engine/Source/Runtime/Experimental/Chaos/Public/Chaos/ImplicitObjectScaled.h
--- a/Engine/Source/Runtime/Experimental/Chaos/Public/Chaos/ImplicitObjectScaled.h
+++ b/Engine/Source/Runtime/Experimental/Chaos/Public/Chaos/ImplicitObjectScaled.h
@@ -671,6 +671,44 @@
 		return false;
 	}
 
+	/** This is a low level function and assumes the internal object has a SweepGeom function. Should not be called directly. See GeometryQueries.h : SweepQuery */
+	template <typename QueryGeomType>
+	bool LowLevelSweepGeomCCD(const QueryGeomType& B, const TRigidTransform<T, d>& BToATM, const TVector<T, d>& LocalDir, const T Length, const FReal IgnorePenetration, const FReal TargetPenetration, T& OutTOI, T& OutPhi, TVector<T, d>& LocalPosition, TVector<T, d>& LocalNormal, int32& OutFaceIndex, TVector<T, d>& OutFaceNormal) const
+	{
+		ensure(Length > 0);
+		ensure(FMath::IsNearlyEqual(LocalDir.SizeSquared(), (T)1, (T)UE_KINDA_SMALL_NUMBER));
+
+		const TVector<T, d> UnscaledDirDenorm = MInvScale * LocalDir;
+		const T LengthScale = UnscaledDirDenorm.Size();
+		if (ensure(LengthScale > TNumericLimits<T>::Min()))
+		{
+			const T LengthScaleInv = 1.f / LengthScale;
+			const T UnscaledLength = Length * LengthScale;
+			const TVector<T, d> UnscaledDir = UnscaledDirDenorm * LengthScaleInv;
+
+			TVector<T, d> UnscaledPosition;
+			TVector<T, d> UnscaledNormal;
+			T TOI;
+			T Phi;
+
+			TRigidTransform<T, d> BToATMNoScale(BToATM.GetLocation() * MInvScale, BToATM.GetRotation());
+
+			if (MObject->SweepGeomCCD(B, BToATMNoScale, UnscaledDir, UnscaledLength, IgnorePenetration, TargetPenetration, TOI, Phi, UnscaledPosition, UnscaledNormal, OutFaceIndex, OutFaceNormal, MScale))
+			{
+				if (TOI < FReal(1))
+				{
+					OutTOI = TOI;
+					OutPhi = Phi;
+					LocalPosition = MScale * UnscaledPosition;
+					LocalNormal = (MInvScale * UnscaledNormal).GetSafeNormal();
+					return true;
+				}
+			}
+		}
+
+		return false;
+	}
+
 	template <typename QueryGeomType>
 	bool GJKContactPoint(const QueryGeomType& A, const FRigidTransform3& AToBTM, const FReal Thickness, FVec3& Location, FVec3& Normal, FReal& Penetration) const
 	{
diff --git a/Engine/Source/Runtime/Experimental/Chaos/Public/Chaos/TriangleMeshImplicitObject.h b/Engine/Source/Runtime/Experimental/Chaos/Public/Chaos/TriangleMeshImplicitObject.h
--- a/Engine/Source/Runtime/Experimental/Chaos/Public/Chaos/TriangleMeshImplicitObject.h
+++ b/Engine/Source/Runtime/Experimental/Chaos/Public/Chaos/TriangleMeshImplicitObject.h
@@ -527,6 +527,16 @@
 		bool SweepGeom(const TImplicitObjectScaled<FCapsule>& QueryGeom, const FRigidTransform3& StartTM, const FVec3& Dir, const FReal Length, FReal& OutTime, FVec3& OutPosition, FVec3& OutNormal, int32& OutFaceIndex, FVec3& OutFaceNormal, const FReal Thickness = 0, const bool bComputeMTD = false, FVec3 TriMeshScale = FVec3(1.0f)) const;
 		bool SweepGeom(const TImplicitObjectScaled<FConvex>& QueryGeom, const FRigidTransform3& StartTM, const FVec3& Dir, const FReal Length, FReal& OutTime, FVec3& OutPosition, FVec3& OutNormal, int32& OutFaceIndex, FVec3& OutFaceNormal, const FReal Thickness = 0, const bool bComputeMTD = false, FVec3 TriMeshScale = FVec3(1.0f)) const;
 
+		// Sweep used for CCD. Ignores triangles we penetrate by less than IgnorePenetration, and calculate the TOI for a depth of TargetPenetration. If both are zero, this is equivalent to SweepGeom
+		bool SweepGeomCCD(const TSphere<FReal, 3>& QueryGeom, const FRigidTransform3& StartTM, const FVec3& Dir, const FReal Length, const FReal IgnorePenetration, const FReal TargetPenetration, FReal& OutTOI, FReal& OutPhi, FVec3& OutPosition, FVec3& OutNormal, int32& OutFaceIndex, FVec3& OutFaceNormal, const FVec3& TriMeshScale = FVec3(1.0f)) const;
+		bool SweepGeomCCD(const TBox<FReal, 3>& QueryGeom, const FRigidTransform3& StartTM, const FVec3& Dir, const FReal Length, const FReal IgnorePenetration, const FReal TargetPenetration, FReal& OutTOI, FReal& OutPhi, FVec3& OutPosition, FVec3& OutNormal, int32& OutFaceIndex, FVec3& OutFaceNormal, const FVec3& TriMeshScale = FVec3(1.0f)) const;
+		bool SweepGeomCCD(const FCapsule& QueryGeom, const FRigidTransform3& StartTM, const FVec3& Dir, const FReal Length, const FReal IgnorePenetration, const FReal TargetPenetration, FReal& OutTOI, FReal& OutPhi, FVec3& OutPosition, FVec3& OutNormal, int32& OutFaceIndex, FVec3& OutFaceNormal, const FVec3& TriMeshScale = FVec3(1.0f)) const;
+		bool SweepGeomCCD(const FConvex& QueryGeom, const FRigidTransform3& StartTM, const FVec3& Dir, const FReal Length, const FReal IgnorePenetration, const FReal TargetPenetration, FReal& OutTOI, FReal& OutPhi, FVec3& OutPosition, FVec3& OutNormal, int32& OutFaceIndex, FVec3& OutFaceNormal, const FVec3& TriMeshScale = FVec3(1.0f)) const;
+		bool SweepGeomCCD(const TImplicitObjectScaled<TSphere<FReal, 3>>& QueryGeom, const FRigidTransform3& StartTM, const FVec3& Dir, const FReal Length, const FReal IgnorePenetration, const FReal TargetPenetration, FReal& OutTOI, FReal& OutPhi, FVec3& OutPosition, FVec3& OutNormal, int32& OutFaceIndex, FVec3& OutFaceNormal, const FVec3& TriMeshScale = FVec3(1.0f)) const;
+		bool SweepGeomCCD(const TImplicitObjectScaled<TBox<FReal, 3>>& QueryGeom, const FRigidTransform3& StartTM, const FVec3& Dir, const FReal Length, const FReal IgnorePenetration, const FReal TargetPenetration, FReal& OutTOI, FReal& OutPhi, FVec3& OutPosition, FVec3& OutNormal, int32& OutFaceIndex, FVec3& OutFaceNormal, const FVec3& TriMeshScale = FVec3(1.0f)) const;
+		bool SweepGeomCCD(const TImplicitObjectScaled<FCapsule>& QueryGeom, const FRigidTransform3& StartTM, const FVec3& Dir, const FReal Length, const FReal IgnorePenetration, const FReal TargetPenetration, FReal& OutTOI, FReal& OutPhi, FVec3& OutPosition, FVec3& OutNormal, int32& OutFaceIndex, FVec3& OutFaceNormal, const FVec3& TriMeshScale = FVec3(1.0f)) const;
+		bool SweepGeomCCD(const TImplicitObjectScaled<FConvex>& QueryGeom, const FRigidTransform3& StartTM, const FVec3& Dir, const FReal Length, const FReal IgnorePenetration, const FReal TargetPenetration, FReal& OutTOI, FReal& OutPhi, FVec3& OutPosition, FVec3& OutNormal, int32& OutFaceIndex, FVec3& OutFaceNormal, const FVec3& TriMeshScale = FVec3(1.0f)) const;
+
 		bool GJKContactPoint(const TSphere<FReal, 3>& QueryGeom, const FRigidTransform3& QueryTM, const FReal Thickness, FVec3& Location, FVec3& Normal, FReal& Penetration) const;
 		bool GJKContactPoint(const TBox<FReal, 3>& QueryGeom, const FRigidTransform3& QueryTM, const FReal Thickness, FVec3& Location, FVec3& Normal, FReal& Penetration) const;
 		bool GJKContactPoint(const FCapsule& QueryGeom, const FRigidTransform3& QueryTM, const FReal Thickness, FVec3& Location, FVec3& Normal, FReal& Penetration) const;
@@ -817,6 +827,9 @@
 		template<typename Geom, typename IdxType>
 		friend struct FTriangleMeshSweepVisitor;
 
+		template<typename Geom, typename IdxType>
+		friend struct FTriangleMeshSweepVisitorCCD;
+
 		// Required by implicit object serialization, disabled for general use.
 		friend class FImplicitObject;
 
@@ -835,6 +848,9 @@
 		template<typename QueryGeomType>
 		bool SweepGeomImp(const QueryGeomType& QueryGeom, const FRigidTransform3& StartTM, const FVec3& Dir, const FReal Length, FReal& OutTime, FVec3& OutPosition, FVec3& OutNormal, int32& OutFaceIndex, FVec3& OutFaceNormal, const FReal Thickness, const bool bComputeMTD, FVec3 TriMeshScale = FVec3(1.0f)) const;
 
+		template <typename QueryGeomType>
+		bool SweepGeomCCDImp(const QueryGeomType& QueryGeom, const FRigidTransform3& StartTM, const FVec3& Dir, const FReal Length, const FReal IgnorePenetration, const FReal TargetPenetration, FReal& OutTOI, FReal& OutPhi, FVec3& OutPosition, FVec3& OutNormal, int32& OutFaceIndex, FVec3& OutFaceNormal, const FVec3& TriMeshScale) const;
+
 		template <typename IdxType>
 		bool RaycastImp(const TArray<TVec3<IdxType>>& Elements, const FVec3& StartPoint, const FVec3& Dir, const FReal Length, const FReal Thickness, FReal& OutTime, FVec3& OutPosition, FVec3& OutNormal, int32& OutFaceIndex) const;
 
