Fix - NoPCH CIS

diff --git a/Engine/Source/Runtime/Experimental/Dataflow/Core/Public/Dataflow/DataflowProperty.h b/Engine/Source/Runtime/Experimental/Dataflow/Core/Public/Dataflow/DataflowProperty.h
--- a/Engine/Source/Runtime/Experimental/Dataflow/Core/Public/Dataflow/DataflowProperty.h
+++ b/Engine/Source/Runtime/Experimental/Dataflow/Core/Public/Dataflow/DataflowProperty.h
@@ -10,16 +10,6 @@ namespace Dataflow
 {
 	class FNode;
 
-	struct FFilename
-	{
-		FString Filename;
-		FMD5Hash FileHash;
-
-		FFilename() {}
-		FFilename(const FString& File) : Filename(File) {}
-	};
-
-
 	template<class T>
 	struct DATAFLOWCORE_API TPropertyParameters {
 		FName Name = FName("");
