Chaos destruction: optimization and experimental collision handling logic
- ignore impulse below close to 0
- cache contact accumulated impulse to avoid recomputing the size for each tested particle
- skip seraching the map for null parent
- add an experimental graph propagation based collision handling ( behind a cvar )

diff --git a/Engine/Source/Runtime/Experimental/Chaos/Private/Chaos/PBDRigidClustering.cpp b/Engine/Source/Runtime/Experimental/Chaos/Private/Chaos/PBDRigidClustering.cpp
--- a/Engine/Source/Runtime/Experimental/Chaos/Private/Chaos/PBDRigidClustering.cpp
+++ b/Engine/Source/Runtime/Experimental/Chaos/Private/Chaos/PBDRigidClustering.cpp
@@ -48,6 +48,9 @@
 	float BoundingBoxMarginForConnectionGraphFiltering = 0;
 	FAutoConsoleVariableRef CVarBoundingBoxMarginForConnectionGraphFiltering(TEXT("p.BoundingBoxMarginForConnectionGraphFiltering"), BoundingBoxMarginForConnectionGraphFiltering, TEXT("when UseBoundingBoxForConnectionGraphFiltering is on, the margin to use for the oevrlap test [def: 0]"));
 
+	int32 GraphPropagationBasedCollisionImpulseProcessing = 0;
+	FAutoConsoleVariableRef CVarGraphPropagationBasedCollisionImpulseProcessing(TEXT("p.GraphPropagationBasedCollisionImpulseProcessing"), GraphPropagationBasedCollisionImpulseProcessing, TEXT("when processing collision impulse toc ompute strain, pick the closest child from the impact point and propagate using the connection graph [def: 0]"));
+	
 	//==========================================================================
 	// TPBDRigidClustering
 	//==========================================================================
@@ -602,6 +605,28 @@
 			{
 				//UE_LOG(LogTemp, Warning, TEXT("Releasing child %d from parent %p due to strain %.5f Exceeding internal strain %.5f (Source: %s)"), ChildIdx, ClusteredParticle, ChildStrain, Child->Strain(), bForceRelease ? TEXT("Forced by caller") : ExternalStrainMap ? TEXT("External") : TEXT("Collision"));
 
+				// if necessary propagate through the graph
+				if (GraphPropagationBasedCollisionImpulseProcessing && !bForceRelease)
+				{
+					const FReal RemainingStrain = Child->Strain() - ChildStrain;
+					if (RemainingStrain > 0)
+					{
+						// todo(chaos) : could do better and have something weighted on distance with a falloff maybe ?  
+						const FReal RemainingStrainPerSibling = RemainingStrain / Child->ConnectivityEdges().Num();
+						for (const TConnectivityEdge<FReal>& Edge: Child->ConnectivityEdges())
+						{
+							if (Edge.Sibling)
+							{
+								if (FPBDRigidClusteredParticleHandle* ClusteredSibling = Edge.Sibling->CastToClustered())
+								{
+									// todo(chaos) this may currently be non optimal as we are in the apply loop and this may be cleareed right after 
+									ClusteredSibling->CollisionImpulses() +=  RemainingStrainPerSibling;
+								}
+							}
+						}
+					}
+				}
+				
 				// The piece that hits just breaks off - we may want more control 
 				// by looking at the edges of this piece which would give us cleaner 
 				// breaks (this approach produces more rubble)
@@ -1009,31 +1034,69 @@
 				const FPBDRigidClusteredParticleHandle* Cluster,
 				const TArray<FPBDRigidParticleHandle*>& ParentToChildren)
 			{
-				const FRigidTransform3 WorldToClusterTM = FRigidTransform3(Cluster->P(), Cluster->Q());
-				const FVec3 ContactLocationClusterLocal = WorldToClusterTM.InverseTransformPosition(ContactHandle->GetContact().CalculateWorldContactLocation());
-				FAABB3 ContactBox(ContactLocationClusterLocal, ContactLocationClusterLocal);
-				ContactBox.Thicken(ClusterDistanceThreshold);
-				if (Cluster->ChildrenSpatial())
+				const FVec3 ContactWorldLocation = ContactHandle->GetContact().CalculateWorldContactLocation();
+				
+				const FReal AccumulatedImpulse = ContactHandle->GetAccumulatedImpulse().Size();
+				if (AccumulatedImpulse > UE_SMALL_NUMBER)
 				{
-					const TArray<FPBDRigidParticleHandle*> Intersections = Cluster->ChildrenSpatial()->FindAllIntersectingChildren(ContactBox);
-					for (FPBDRigidParticleHandle* Child : Intersections)
+					if (GraphPropagationBasedCollisionImpulseProcessing)
 					{
-						if (TPBDRigidClusteredParticleHandle<FReal, 3>*ClusteredChild = Child->CastToClustered())
+						// first find the closest child
+						FPBDRigidParticleHandle* ClosestChild = nullptr;
+						FReal ClosestDistanceSquared = TNumericLimits<FReal>::Max();
+						for (FPBDRigidParticleHandle* Child:ParentToChildren)
 						{
-							ClusteredChild->CollisionImpulses() += ContactHandle->GetAccumulatedImpulse().Size();
+							const FReal DistanceSquared = (Child->X() - ContactWorldLocation).SizeSquared();
+							if (DistanceSquared < ClosestDistanceSquared)
+							{
+								ClosestDistanceSquared = DistanceSquared;
+								ClosestChild = Child;
+							}
+						}
+						if (ClosestChild)
+						{
+							if (TPBDRigidClusteredParticleHandle<FReal, 3>* ClusteredChild = ClosestChild->CastToClustered())
+							{
+								ClusteredChild->CollisionImpulses() += AccumulatedImpulse;
+							}
+						}
+					}
+					else
+					{
+						const FRigidTransform3 WorldToClusterTM = FRigidTransform3(Cluster->P(), Cluster->Q());
+						const FVec3 ContactLocationClusterLocal = WorldToClusterTM.InverseTransformPosition(ContactWorldLocation);
+						FAABB3 ContactBox(ContactLocationClusterLocal, ContactLocationClusterLocal);
+						ContactBox.Thicken(ClusterDistanceThreshold);
+						if (Cluster->ChildrenSpatial())
+						{
+							// todo(chaos): FindAllIntersectingChildren may return an unfiltered list of children ( when num children is under a certain threshold )   
+							const TArray<FPBDRigidParticleHandle*> Intersections = Cluster->ChildrenSpatial()->FindAllIntersectingChildren(ContactBox);
+							for (FPBDRigidParticleHandle* Child : Intersections)
+							{
+								if (TPBDRigidClusteredParticleHandle<FReal, 3>*ClusteredChild = Child->CastToClustered())
+								{
+									ClusteredChild->CollisionImpulses() += AccumulatedImpulse;
+								}
+							}
 						}
 					}
 				}
 			};
 
-			if (const TArray<FPBDRigidParticleHandle*>* ChildrenPtr = MParentToChildren.Find(ClusteredConstrainedParticles0))
+			if (ClusteredConstrainedParticles0)
 			{
-				ComputeStrainLambda(ClusteredConstrainedParticles0, *ChildrenPtr);
+				if (const TArray<FPBDRigidParticleHandle*>* ChildrenPtr = MParentToChildren.Find(ClusteredConstrainedParticles0))
+				{
+					ComputeStrainLambda(ClusteredConstrainedParticles0, *ChildrenPtr);
+				}
 			}
 
-			if (const TArray<FPBDRigidParticleHandle*>* ChildrenPtr = MParentToChildren.Find(ClusteredConstrainedParticles1))
+			if (ClusteredConstrainedParticles1)
 			{
-				ComputeStrainLambda(ClusteredConstrainedParticles1, *ChildrenPtr);
+				if (const TArray<FPBDRigidParticleHandle*>* ChildrenPtr = MParentToChildren.Find(ClusteredConstrainedParticles1))
+				{
+					ComputeStrainLambda(ClusteredConstrainedParticles1, *ChildrenPtr);
+				}
 			}
 
 			MCollisionImpulseArrayDirty = true;
