Chaos : ManagedArrayCollection
- To get the deep copy working in the dataflow on a ManagedArrayCollection. I need to lift the restriction that ALL unique pointers can not be deep copied. This change allows for the deep copy of the unique pointer types that support the Copy(). The current attributes that do not support ::Copy() are: TGeometryParticle, BVHParticles, TArray<TVector<float>>.

diff --git a/Engine/Source/Runtime/Experimental/Chaos/Public/GeometryCollection/ManagedArray.h b/Engine/Source/Runtime/Experimental/Chaos/Public/GeometryCollection/ManagedArray.h
--- a/Engine/Source/Runtime/Experimental/Chaos/Public/GeometryCollection/ManagedArray.h
+++ b/Engine/Source/Runtime/Experimental/Chaos/Public/GeometryCollection/ManagedArray.h
@@ -528,9 +528,35 @@ void InitHelper(TArray<T>& Array, const TManagedArrayBase<T>& NewTypedArray, int
 template <typename T>
 void InitHelper(TArray<TUniquePtr<T>>& Array, const TManagedArrayBase<TUniquePtr<T>>& NewTypedArray, int32 Size)
 {
-	check(false);	//Cannot make copies of a managed array with unique pointers. Typically used for shared data
+	for (int32 Index = 0; Index < Size; Index++)
+	{
+		if (NewTypedArray[Index])
+		{
+			Array[Index].Reset((T*)NewTypedArray[Index]->Copy().Release());
+		}
+	}
 }
 
+//
+//
+//
+#include "ChaosLog.h"
+#define UNSUPPORTED_UNIQUE_ARRAY_COPIES(TYPE, NAME) \
+template<> inline void InitHelper(TArray<TYPE>& Array, const TManagedArrayBase<TYPE>& NewTypedArray, int32 Size) { \
+	UE_LOG(LogChaos,Warning, TEXT("Cannot make a copy of unique array of type (%s) within the managed array collection. Regenerate unique pointer attributes if needed."), NAME); \
+}
+
+#include "Chaos/ParticleHandle.h"
+typedef TUniquePtr<Chaos::TGeometryParticle<Chaos::FReal, 3>> LOCAL_MA_UniqueTGeometryParticle;
+UNSUPPORTED_UNIQUE_ARRAY_COPIES(LOCAL_MA_UniqueTGeometryParticle, "Chaos::TGeometryParticle");
+
+#include "Chaos/BVHParticles.h"
+typedef TUniquePtr<Chaos::FBVHParticles, TDefaultDelete<Chaos::FBVHParticles>> LOCAL_MA_UniqueTBVHParticles;
+UNSUPPORTED_UNIQUE_ARRAY_COPIES(LOCAL_MA_UniqueTBVHParticles, "Chaos::FBVHParticles");
+
+#include "Math/Vector.h"
+typedef TUniquePtr<TArray<UE::Math::TVector<float>>> LOCAL_MA_UniqueTArrayTVector;
+UNSUPPORTED_UNIQUE_ARRAY_COPIES(LOCAL_MA_UniqueTArrayTVector, "TArray<UE::Math::TVector<float>>");
 
 template<class InElementType>
 class TManagedArray : public TManagedArrayBase<InElementType>
