Fix for CIS Warning

diff --git a/Engine/Source/Runtime/Experimental/Chaos/Public/GeometryCollection/ManagedArray.h b/Engine/Source/Runtime/Experimental/Chaos/Public/GeometryCollection/ManagedArray.h
--- a/Engine/Source/Runtime/Experimental/Chaos/Public/GeometryCollection/ManagedArray.h
+++ b/Engine/Source/Runtime/Experimental/Chaos/Public/GeometryCollection/ManagedArray.h
@@ -6,6 +6,11 @@
 #include "Templates/UnrealTemplate.h"
 #include "Chaos/ChaosArchive.h"
 #include "UObject/DestructionObjectVersion.h"
+//
+#include "ChaosLog.h"
+#include "Chaos/ParticleHandle.h"
+#include "Chaos/BVHParticles.h"
+#include "Math/Vector.h"
 
 class FManagedArrayCollection;
 DEFINE_LOG_CATEGORY_STATIC(UManagedArrayLogging, NoLogging, All);
@@ -540,21 +545,17 @@ void InitHelper(TArray<TUniquePtr<T>>& Array, const TManagedArrayBase<TUniquePtr
 //
 //
 //
-#include "ChaosLog.h"
 #define UNSUPPORTED_UNIQUE_ARRAY_COPIES(TYPE, NAME) \
 template<> inline void InitHelper(TArray<TYPE>& Array, const TManagedArrayBase<TYPE>& NewTypedArray, int32 Size) { \
 	UE_LOG(LogChaos,Warning, TEXT("Cannot make a copy of unique array of type (%s) within the managed array collection. Regenerate unique pointer attributes if needed."), NAME); \
 }
 
-#include "Chaos/ParticleHandle.h"
 typedef TUniquePtr<Chaos::TGeometryParticle<Chaos::FReal, 3>> LOCAL_MA_UniqueTGeometryParticle;
 UNSUPPORTED_UNIQUE_ARRAY_COPIES(LOCAL_MA_UniqueTGeometryParticle, "Chaos::TGeometryParticle");
 
-#include "Chaos/BVHParticles.h"
 typedef TUniquePtr<Chaos::FBVHParticles, TDefaultDelete<Chaos::FBVHParticles>> LOCAL_MA_UniqueTBVHParticles;
 UNSUPPORTED_UNIQUE_ARRAY_COPIES(LOCAL_MA_UniqueTBVHParticles, "Chaos::FBVHParticles");
 
-#include "Math/Vector.h"
 typedef TUniquePtr<TArray<UE::Math::TVector<float>>> LOCAL_MA_UniqueTArrayTVector;
 UNSUPPORTED_UNIQUE_ARRAY_COPIES(LOCAL_MA_UniqueTArrayTVector, "TArray<UE::Math::TVector<float>>");
 
