Chaos destruction : allow import of collision shapes from static meshes when creating a geometry collection

diff --git a/Engine/Plugins/Experimental/GeometryCollectionPlugin/Source/GeometryCollectionEditor/Private/GeometryCollection/GeometryCollectionConversion.cpp b/Engine/Plugins/Experimental/GeometryCollectionPlugin/Source/GeometryCollectionEditor/Private/GeometryCollection/GeometryCollectionConversion.cpp
--- a/Engine/Plugins/Experimental/GeometryCollectionPlugin/Source/GeometryCollectionEditor/Private/GeometryCollection/GeometryCollectionConversion.cpp
+++ b/Engine/Plugins/Experimental/GeometryCollectionPlugin/Source/GeometryCollectionEditor/Private/GeometryCollection/GeometryCollectionConversion.cpp
@@ -18,6 +18,7 @@
 #include "GeometryCollection/GeometryCollectionFactory.h"
 #include "GeometryCollection/GeometryCollectionComponent.h"
 #include "GeometryCollection/GeometryCollectionClusteringUtility.h"
+#include "GeometryCollectionProxyData.h"
 #include "GeometryCollection/GeometryCollectionUtility.h"
 #include "Logging/LogMacros.h"
 #include "Materials/Material.h"
@@ -25,6 +26,7 @@
 #include "MeshDescription.h"
 #include "StaticMeshAttributes.h"
 #include "StaticMeshOperations.h"
+#include "Physics/Experimental/ChaosInterfaceUtils.h"
 
 
 DEFINE_LOG_CATEGORY_STATIC(UGeometryCollectionConversionLogging, Log, All);
@@ -247,6 +249,48 @@ void FGeometryCollectionConversion::AppendStaticMesh(const UStaticMesh* StaticMe
 		Transform[TransformIndex1] = StaticMeshTransform;
 		Transform[TransformIndex1].SetScale3D(FVector(1.f, 1.f, 1.f));
 
+		// collisions
+		if (UBodySetup* BodySetup = StaticMesh->GetBodySetup())
+		{
+			TArray<TUniquePtr<Chaos::FImplicitObject>> Geoms;
+			Chaos::FShapesArray Shapes;
+
+			FGeometryAddParams CreateGeometryParams;
+			CreateGeometryParams.bDoubleSided = false;
+			CreateGeometryParams.CollisionData.CollisionFlags.bEnableQueryCollision = true;
+			CreateGeometryParams.CollisionData.CollisionFlags.bEnableSimCollisionComplex = false; // no support for trimesh in destruction
+			CreateGeometryParams.CollisionData.CollisionFlags.bEnableSimCollisionSimple = true;
+			CreateGeometryParams.CollisionTraceType = ECollisionTraceFlag::CTF_UseSimpleAsComplex;
+			CreateGeometryParams.Scale = Scale;
+			CreateGeometryParams.LocalTransform = Chaos::FRigidTransform3::Identity;
+			CreateGeometryParams.WorldTransform = Chaos::FRigidTransform3::Identity;
+			CreateGeometryParams.Geometry = &BodySetup->AggGeom;
+#if PHYSICS_INTERFACE_PHYSX
+			AddParams.TriMeshes = TArrayView<PxTriangleMesh*>(BodySetup->TriMeshes);
+#endif
+#if WITH_CHAOS
+			CreateGeometryParams.ChaosTriMeshes = MakeArrayView(BodySetup->ChaosTriMeshes);
+#endif
+
+			// todo(chaos) : this currently also create the shape array which is unnecessary ,this could be optimized by having a common function to create only the implicits 
+			ChaosInterface::CreateGeometry(CreateGeometryParams, Geoms, Shapes);
+
+			using FCollisionType = FGeometryDynamicCollection::FSharedImplicit;
+			TManagedArray<FCollisionType>& ExternaCollisions = GeometryCollection->AddAttribute<FCollisionType>("ExternalCollisions", FGeometryCollection::TransformGroup);
+			if (Geoms.Num() == 0)
+			{
+				ExternaCollisions[TransformIndex1] = nullptr;
+			}
+			else if (Geoms.Num() == 1)
+			{
+				ExternaCollisions[TransformIndex1] = TSharedPtr<Chaos::FImplicitObject>(Geoms[0].Release());  
+			}
+			else
+			{
+				ExternaCollisions[TransformIndex1] = MakeShared<Chaos::FImplicitObjectUnion>(MoveTemp(Geoms)); 
+			}
+		}
+		
 		// Bone Hierarchy - Added at root with no common parent
 		TManagedArray<int32>& Parent = GeometryCollection->Parent;
 		TManagedArray<int32>& SimulationType = GeometryCollection->SimulationType;
@@ -598,7 +642,7 @@ void FGeometryCollectionConversion::AppendStaticMesh(const UStaticMesh* StaticMe
 	UMaterialInterface* BoneSelectedMaterial = LoadObject<UMaterialInterface>(nullptr, UGeometryCollection::GetSelectedMaterialPath(), nullptr, LOAD_None, nullptr);
 	GeometryCollectionObject->Materials.Remove(BoneSelectedMaterial);
 	Materials.Remove(BoneSelectedMaterial);
-
+	
 	AppendStaticMesh(StaticMesh, Materials, StaticMeshTransform, GeometryCollectionObject, ReindexMaterials);
 }
 
@@ -843,15 +887,49 @@ void FGeometryCollectionConversion::AppendSkeletalMesh(const USkeletalMesh* Skel
 	}
 }
 
+struct FGeometryCollectionAsset
+{
+	UPackage* Package = nullptr;
+	UGeometryCollection* GeometryCollection = nullptr;
+
+	bool IsCreated() const
+	{
+		return (Package != nullptr && GeometryCollection != nullptr);
+	}
+	
+	void Create()
+	{
+		if (ensure(!IsCreated()))
+		{
+			FGeometryCollectionAsset NewAsset;
+			Package = CreatePackage(TEXT("/Game/GeometryCollectionAsset"));
+			const auto GeometryCollectionFactory = NewObject<UGeometryCollectionFactory>();
+			GeometryCollection = static_cast<UGeometryCollection*>(
+				GeometryCollectionFactory->FactoryCreateNew(UGeometryCollection::StaticClass(),
+					NewAsset.Package,
+					FName("GeometryCollectionAsset"),
+					RF_Standalone | RF_Public,
+					nullptr,
+					GWarn));
+		}
+	}
+
+	void Finalize() const 
+	{
+		if (ensure(IsCreated()))
+		{
+			FAssetRegistryModule::AssetCreated(GeometryCollection);
+			GeometryCollection->MarkPackageDirty();
+			Package->SetDirtyFlag(true);
+		}
+	}
+};
+
 void FGeometryCollectionConversion::CreateGeometryCollectionCommand(UWorld * World)
 {
-	UPackage* Package = CreatePackage(TEXT("/Game/GeometryCollectionAsset"));
-	auto GeometryCollectionFactory = NewObject<UGeometryCollectionFactory>();
-	UGeometryCollection* GeometryCollection = static_cast<UGeometryCollection*>(
-		GeometryCollectionFactory->FactoryCreateNew(UGeometryCollection::StaticClass(), Package,
-			FName("GeometryCollectionAsset"), RF_Standalone | RF_Public, NULL, GWarn));
-	FAssetRegistryModule::AssetCreated(GeometryCollection);
-	Package->SetDirtyFlag(true);
+	FGeometryCollectionAsset GeometryCollectionAsset;
+	GeometryCollectionAsset.Create();
+	GeometryCollectionAsset.Finalize();
 }
 
 void FGeometryCollectionConversion::CreateFromSelectedActorsCommand(UWorld * World)
@@ -863,8 +941,7 @@ void FGeometryCollectionConversion::CreateFromSelectedActorsCommand(UWorld * Wor
 		return;
 	}
 
-	UPackage* Package = nullptr;
-	UGeometryCollection* GeometryCollection = nullptr;
+	FGeometryCollectionAsset GeometryCollectionAsset; 
 	for (FSelectionIterator Iter(*SelectedActors); Iter; ++Iter)
 	{
 		if (AActor* Actor = Cast<AActor>(*Iter))
@@ -877,24 +954,15 @@ void FGeometryCollectionConversion::CreateFromSelectedActorsCommand(UWorld * Wor
 				{
 					if (StaticMeshComponents[Index]->GetStaticMesh())
 					{
-						if (!Package || !GeometryCollection)
+						if (!GeometryCollectionAsset.IsCreated())
 						{
-							Package = CreatePackage(TEXT("/Game/GeometryCollectionAsset"));
-							auto GeometryCollectionFactory = NewObject<UGeometryCollectionFactory>();
-							GeometryCollection = static_cast<UGeometryCollection*>(
-								GeometryCollectionFactory->FactoryCreateNew(
-									UGeometryCollection::StaticClass(),
-									Package,
-									FName("GeometryCollectionAsset"),
-									RF_Standalone | RF_Public,
-									NULL,
-									GWarn));
+							GeometryCollectionAsset.Create();
 						}
 						FGeometryCollectionConversion::AppendStaticMesh(
 							StaticMeshComponents[Index]->GetStaticMesh(),
 							StaticMeshComponents[Index],
 							Actor->GetTransform(),
-							GeometryCollection);
+							GeometryCollectionAsset.GeometryCollection);
 					}
 				}
 			}
@@ -907,46 +975,33 @@ void FGeometryCollectionConversion::CreateFromSelectedActorsCommand(UWorld * Wor
 				{
 					if (SkeletalMeshComponents[Index]->SkeletalMesh)
 					{
-						if (!Package || !GeometryCollection)
+						if (!GeometryCollectionAsset.IsCreated())
 						{
-							Package = CreatePackage(TEXT("/Game/GeometryCollectionAsset"));
-							auto GeometryCollectionFactory = NewObject<UGeometryCollectionFactory>();
-							GeometryCollection = static_cast<UGeometryCollection*>(
-								GeometryCollectionFactory->FactoryCreateNew(
-									UGeometryCollection::StaticClass(),
-									Package,
-									FName("GeometryCollectionAsset"),
-									RF_Standalone | RF_Public,
-									NULL,
-									GWarn));
+							GeometryCollectionAsset.Create();
 						}
 						FGeometryCollectionConversion::AppendSkeletalMesh(
 							SkeletalMeshComponents[Index]->SkeletalMesh,
 							SkeletalMeshComponents[Index],
 							Actor->GetTransform(),
-							GeometryCollection);
+							GeometryCollectionAsset.GeometryCollection);
 					}
 				}
 			}
 		}
 	}
 
-	if(GeometryCollection != nullptr)
+	if(GeometryCollectionAsset.IsCreated())
 	{
-		GeometryCollectionAlgo::PrepareForSimulation(GeometryCollection->GetGeometryCollection().Get());
-
-		FAssetRegistryModule::AssetCreated(GeometryCollection);
-		GeometryCollection->MarkPackageDirty();
-		Package->SetDirtyFlag(true);
+		GeometryCollectionAlgo::PrepareForSimulation(GeometryCollectionAsset.GeometryCollection->GetGeometryCollection().Get());
+		GeometryCollectionAsset.Finalize();
 	}
 }
 
 void FGeometryCollectionConversion::CreateFromSelectedAssetsCommand(UWorld * World)
 {
 	//UE_LOG(UGeometryCollectionConversionLogging, Log, TEXT("FGeometryCollectionConversion::CreateCommand()"));
-	UPackage* Package = nullptr;
-	UGeometryCollection* GeometryCollection = nullptr;
-
+	FGeometryCollectionAsset GeometryCollectionAsset;
+	
 	TArray<FAssetData> SelectedAssets;
 	GEditor->GetContentBrowserSelections(SelectedAssets);
 	for (const FAssetData& AssetData : SelectedAssets)
@@ -954,55 +1009,34 @@ void FGeometryCollectionConversion::CreateFromSelectedAssetsCommand(UWorld * Wor
 		if (AssetData.GetAsset()->IsA<UStaticMesh>())
 		{
 			UE_LOG(UGeometryCollectionConversionLogging, Log, TEXT("Static Mesh Content Browser : %s"), *AssetData.GetClass()->GetName());
-			if (!Package || !GeometryCollection)
+			if (!GeometryCollectionAsset.IsCreated())
 			{
-				Package = CreatePackage(TEXT("/Game/GeometryCollectionAsset"));
-				auto GeometryCollectionFactory = NewObject<UGeometryCollectionFactory>();
-				GeometryCollection = static_cast<UGeometryCollection*>(
-					GeometryCollectionFactory->FactoryCreateNew(
-						UGeometryCollection::StaticClass(), 
-						Package,
-						FName("GeometryCollectionAsset"), 
-						RF_Standalone | RF_Public, 
-						NULL, 
-						GWarn));
+				GeometryCollectionAsset.Create();
 			}
 			FGeometryCollectionConversion::AppendStaticMesh(
 				static_cast<const UStaticMesh*>(AssetData.GetAsset()),
 				nullptr,
 				FTransform(), 
-				GeometryCollection);
+				GeometryCollectionAsset.GeometryCollection);
 		}
 		else if (AssetData.GetAsset()->IsA<USkeletalMesh>())
 		{
 			UE_LOG(UGeometryCollectionConversionLogging, Log, TEXT("Skeletal Mesh Content Browser : %s"), *AssetData.GetClass()->GetName());
-			if (!Package || !GeometryCollection)
+			if (!GeometryCollectionAsset.IsCreated())
 			{
-				Package = CreatePackage(TEXT("/Game/GeometryCollectionAsset"));
-				auto GeometryCollectionFactory = NewObject<UGeometryCollectionFactory>();
-				GeometryCollection = static_cast<UGeometryCollection*>(
-					GeometryCollectionFactory->FactoryCreateNew(
-						UGeometryCollection::StaticClass(), 
-						Package,
-						FName("GeometryCollectionAsset"), 
-						RF_Standalone | RF_Public, 
-						NULL, 
-						GWarn));
+				GeometryCollectionAsset.Create();
 			}
 			FGeometryCollectionConversion::AppendSkeletalMesh(
 				static_cast<const USkeletalMesh *>(AssetData.GetAsset()), 
 				nullptr,
 				FTransform(),
-				GeometryCollection);
+				GeometryCollectionAsset.GeometryCollection);
 		}
 	}
 
-	if(GeometryCollection != nullptr)
+	if(GeometryCollectionAsset.IsCreated())
 	{
-		GeometryCollectionAlgo::PrepareForSimulation(GeometryCollection->GetGeometryCollection().Get());
-
-		FAssetRegistryModule::AssetCreated(GeometryCollection);
-		GeometryCollection->MarkPackageDirty();
-		Package->SetDirtyFlag(true);
+		GeometryCollectionAlgo::PrepareForSimulation(GeometryCollectionAsset.GeometryCollection->GetGeometryCollection().Get());
+		GeometryCollectionAsset.Finalize();
 	}
 }
diff --git a/Engine/Source/Runtime/Engine/Private/PhysicsEngine/Experimental/ChaosInterfaceUtils.cpp b/Engine/Source/Runtime/Engine/Private/PhysicsEngine/Experimental/ChaosInterfaceUtils.cpp
--- a/Engine/Source/Runtime/Engine/Private/PhysicsEngine/Experimental/ChaosInterfaceUtils.cpp
+++ b/Engine/Source/Runtime/Engine/Private/PhysicsEngine/Experimental/ChaosInterfaceUtils.cpp
@@ -424,115 +424,9 @@ namespace ChaosInterface
 	}
 
 #if WITH_CHAOS
-	bool CalculateMassPropertiesOfImplicitType(
-		Chaos::FMassProperties& OutMassProperties,
-		const Chaos::FRigidTransform3& WorldTransform,
-		const Chaos::FImplicitObject* ImplicitObject,
-		Chaos::FReal InDensityKGPerCM)
-	{
-		using namespace Chaos;
-
-		if (ImplicitObject)
-		{
-			// Hack to handle Transformed and Scaled<ImplicitObjectTriangleMesh> until CastHelper can properly support transformed
-			// Commenting this out temporarily as it breaks vehicles
-			/*	if (Chaos::IsScaled(ImplicitObject->GetType(true)) && Chaos::GetInnerType(ImplicitObject->GetType(true)) & Chaos::ImplicitObjectType::TriangleMesh)
-				{
-					OutMassProperties.Volume = 0.f;
-					OutMassProperties.Mass = FLT_MAX;
-					OutMassProperties.InertiaTensor = FMatrix33(0, 0, 0);
-					OutMassProperties.CenterOfMass = FVector(0);
-					OutMassProperties.RotationOfMass = Chaos::FRotation3::FromIdentity();
-					return false;
-				}
-				else if (ImplicitObject->GetType(true) & Chaos::ImplicitObjectType::TriangleMesh)
-				{
-					OutMassProperties.Volume = 0.f;
-					OutMassProperties.Mass = FLT_MAX;
-					OutMassProperties.InertiaTensor = FMatrix33(0, 0, 0);
-					OutMassProperties.CenterOfMass = FVector(0);
-					OutMassProperties.RotationOfMass = Chaos::FRotation3::FromIdentity();
-					return false;
-				}
-			else*/
-
-			Chaos::Utilities::CastHelper(*ImplicitObject, FTransform::Identity, [&OutMassProperties, InDensityKGPerCM](const auto& Object, const auto& LocalTM)
-				{
-					OutMassProperties.Volume = Object.GetVolume();
-					OutMassProperties.Mass = OutMassProperties.Volume * InDensityKGPerCM;
-					OutMassProperties.InertiaTensor = Object.GetInertiaTensor(OutMassProperties.Mass);
-					OutMassProperties.CenterOfMass = LocalTM.TransformPosition(Object.GetCenterOfMass());
-					OutMassProperties.RotationOfMass = LocalTM.GetRotation() * Object.GetRotationOfMass();
-				});
-		}
-
-		// If the implicit is null, or it is scaled to zero it will have zero volume, mass or inertia
-		return (OutMassProperties.Mass > 0);
-	}
-
-	void CalculateMassPropertiesFromShapeCollectionImp(
-		Chaos::FMassProperties& OutProperties, 
-		int32 InNumShapes, 
-		Chaos::FReal InDensityKGPerCM,
-		const TArray<bool>& bContributesToMass,
-		TFunction<Chaos::FPerShapeData* (int32 ShapeIndex)> GetShapeDelegate)
-	{
-		Chaos::FReal TotalMass = 0;
-		Chaos::FReal TotalVolume = 0;
-		Chaos::FVec3 TotalCenterOfMass(0);
-		TArray< Chaos::FMassProperties > MassPropertiesList;
-		for (int32 ShapeIndex = 0; ShapeIndex < InNumShapes; ++ShapeIndex)
-		{
-			const Chaos::FPerShapeData* Shape = GetShapeDelegate(ShapeIndex);
-
-			const bool bHassMass = (ShapeIndex < bContributesToMass.Num()) ? bContributesToMass[ShapeIndex] : true;
-			if (bHassMass)
-			{
-				if (const Chaos::FImplicitObject* ImplicitObject = Shape->GetGeometry().Get())
-				{
-					Chaos::FMassProperties MassProperties;
-					if (CalculateMassPropertiesOfImplicitType(MassProperties, FTransform::Identity, ImplicitObject, InDensityKGPerCM))
-					{
-						MassPropertiesList.Add(MassProperties);
-						TotalMass += MassProperties.Mass;
-						TotalVolume += MassProperties.Volume;
-						TotalCenterOfMass += MassProperties.CenterOfMass * MassProperties.Mass;
-					}
-				}
-			}
-		}
-
-		Chaos::FMatrix33 Tensor;
-
-		// If no shapes contribute to mass, or they are scaled to zero, we may end up with zero mass here
-		if ((TotalMass > 0.f) && (MassPropertiesList.Num() > 0))
-		{
-			TotalCenterOfMass /= TotalMass;
-
-			Chaos::FMassProperties CombinedMassProperties = Chaos::CombineWorldSpace(MassPropertiesList);
-			ensure(CombinedMassProperties.RotationOfMass.IsIdentity());
-			Tensor = CombinedMassProperties.InertiaTensor;
-		}
-		else
-		{
-			// @todo(chaos): We should support shape-less particles as long as their mass an inertia are set directly
-			// For now hard-code a 50cm sphere with density 1g/cc
-			Tensor = Chaos::FMatrix33(5.24e5f, 5.24e5f, 5.24e5f);
-			TotalMass = 523.0f;
-			TotalVolume = 523000;
-		}
-
-		OutProperties.InertiaTensor = Tensor;
-		OutProperties.Mass = TotalMass;
-		OutProperties.Volume = TotalVolume;
-		OutProperties.CenterOfMass = TotalCenterOfMass;
-		OutProperties.RotationOfMass = Chaos::FRotation3::Identity;
-	}
-
-
 	void CalculateMassPropertiesFromShapeCollection(Chaos::FMassProperties& OutProperties, const TArray<FPhysicsShapeHandle>& InShapes, float InDensityKGPerCM)
 	{
-		CalculateMassPropertiesFromShapeCollectionImp(
+		Chaos::CalculateMassPropertiesFromShapeCollection(
 			OutProperties,
 			InShapes.Num(),
 			InDensityKGPerCM,
@@ -543,7 +437,7 @@ namespace ChaosInterface
 
 	void CalculateMassPropertiesFromShapeCollection(Chaos::FMassProperties& OutProperties, const Chaos::FShapesArray& InShapes, const TArray<bool>& bContributesToMass, float InDensityKGPerCM)
 	{
-		CalculateMassPropertiesFromShapeCollectionImp(
+		Chaos::CalculateMassPropertiesFromShapeCollection(
 			OutProperties,
 			InShapes.Num(),
 			InDensityKGPerCM,
diff --git a/Engine/Source/Runtime/Engine/Public/Physics/Experimental/ChaosInterfaceUtils.h b/Engine/Source/Runtime/Engine/Public/Physics/Experimental/ChaosInterfaceUtils.h
--- a/Engine/Source/Runtime/Engine/Public/Physics/Experimental/ChaosInterfaceUtils.h
+++ b/Engine/Source/Runtime/Engine/Public/Physics/Experimental/ChaosInterfaceUtils.h
@@ -18,7 +18,7 @@ namespace ChaosInterface
 	/**
 	 * Create the Chaos Geometry based on the geometry parameters.
 	 */
-	void CreateGeometry(const FGeometryAddParams& InParams, TArray<TUniquePtr<Chaos::FImplicitObject>>& OutGeoms, Chaos::FShapesArray& OutShapes);
+	void ENGINE_API CreateGeometry(const FGeometryAddParams& InParams, TArray<TUniquePtr<Chaos::FImplicitObject>>& OutGeoms, Chaos::FShapesArray& OutShapes);
 
 #if WITH_CHAOS
 	/**
diff --git a/Engine/Source/Runtime/Experimental/Chaos/Private/Chaos/ImplicitObjectUnion.cpp b/Engine/Source/Runtime/Experimental/Chaos/Private/Chaos/ImplicitObjectUnion.cpp
--- a/Engine/Source/Runtime/Experimental/Chaos/Private/Chaos/ImplicitObjectUnion.cpp
+++ b/Engine/Source/Runtime/Experimental/Chaos/Private/Chaos/ImplicitObjectUnion.cpp
@@ -152,7 +152,50 @@ void FImplicitObjectUnion::CacheAllImplicitObjects()
 		}
 	}
 }
+TUniquePtr<FImplicitObject> FImplicitObjectUnion::Copy() const
+{
+	TArray<TUniquePtr<FImplicitObject>> CopyOfObjects;
+	CopyOfObjects.Reserve(MObjects.Num());
+	for (const TUniquePtr<FImplicitObject>& Object: MObjects)
+	{
+		CopyOfObjects.Emplace(Object->Copy());
+	}
+	return MakeUnique<FImplicitObjectUnion>(MoveTemp(CopyOfObjects));
+}
+
+TUniquePtr<FImplicitObject> FImplicitObjectUnion::CopyWithScale(const FVec3& Scale) const
+{
+	TArray<TUniquePtr<FImplicitObject>> CopyOfObjects;
+	CopyOfObjects.Reserve(MObjects.Num());
+	for (const TUniquePtr<FImplicitObject>& Object: MObjects)
+	{
+		CopyOfObjects.Emplace(Object->CopyWithScale(Scale));
+	}
+	return MakeUnique<FImplicitObjectUnion>(MoveTemp(CopyOfObjects));
+}
 
+TUniquePtr<FImplicitObject> FImplicitObjectUnion::DeepCopy() const
+{
+	TArray<TUniquePtr<FImplicitObject>> CopyOfObjects;
+	CopyOfObjects.Reserve(MObjects.Num());
+	for (const TUniquePtr<FImplicitObject>& Object: MObjects)
+	{
+		CopyOfObjects.Emplace(Object->DeepCopy());
+	}
+	return MakeUnique<FImplicitObjectUnion>(MoveTemp(CopyOfObjects));
+}
+
+TUniquePtr<FImplicitObject> FImplicitObjectUnion::DeepCopyWithScale(const FVec3& Scale) const
+{
+	TArray<TUniquePtr<FImplicitObject>> CopyOfObjects;
+	CopyOfObjects.Reserve(MObjects.Num());
+	for (const TUniquePtr<FImplicitObject>& Object: MObjects)
+	{
+		CopyOfObjects.Emplace(Object->DeepCopyWithScale(Scale));
+	}
+	return MakeUnique<FImplicitObjectUnion>(MoveTemp(CopyOfObjects));
+}
+	
 void FImplicitObjectUnion::Serialize(FChaosArchive& Ar)
 {
 	Ar.UsingCustomVersion(FExternalPhysicsCustomObjectVersion::GUID);
diff --git a/Engine/Source/Runtime/Experimental/Chaos/Private/Chaos/MassProperties.cpp b/Engine/Source/Runtime/Experimental/Chaos/Private/Chaos/MassProperties.cpp
--- a/Engine/Source/Runtime/Experimental/Chaos/Private/Chaos/MassProperties.cpp
+++ b/Engine/Source/Runtime/Experimental/Chaos/Private/Chaos/MassProperties.cpp
@@ -7,7 +7,7 @@
 #include "Chaos/TriangleMesh.h"
 #include "Chaos/Utilities.h"
 #include "ChaosCheck.h"
-
+#include "Chaos/CastingUtilities.h"
 
 //PRAGMA_DISABLE_OPTIMIZATION
 
@@ -121,15 +121,15 @@ namespace Chaos
 		OutRotationOfMass = FRotation3::Identity;
 	}
 	
-   	FMassProperties CalculateMassProperties(const FBox& BoundingBox, const FVector::FReal Density)
+	FMassProperties CalculateMassProperties(const FBox& BoundingBox, const FVector::FReal Density)
 	{
-    	FMassProperties MassProperties;
-    	CalculateVolumeAndCenterOfMass(BoundingBox, MassProperties.Volume, MassProperties.CenterOfMass);
-    	check(Density > 0);
+		FMassProperties MassProperties;
+		CalculateVolumeAndCenterOfMass(BoundingBox, MassProperties.Volume, MassProperties.CenterOfMass);
+		check(Density > 0);
 		MassProperties.Mass = MassProperties.Volume * Density;
 		CalculateInertiaAndRotationOfMass(BoundingBox, Density, MassProperties.InertiaTensor, MassProperties.RotationOfMass);
-   		return MassProperties;
-   	}
+		return MassProperties;
+	}
 	
 	template<typename T, typename TVec, typename TSurfaces>
 	void CalculateVolumeAndCenterOfMassImpl(const TArray<TVec>& Vertices, const TSurfaces& Surfaces, T& OutVolume, TVec& OutCenterOfMass)
@@ -168,10 +168,10 @@ namespace Chaos
 		}
 		// @todo(mlentine): Should add suppoert for thin shell mass properties
 		if (Volume < UE_KINDA_SMALL_NUMBER)	//handle negative volume using fallback for now. Need to investigate cases where this happens
-		{
+			{
 			OutVolume = 0;
 			return;
-		}
+			}
 		OutCenterOfMass = Center + VolumeTimesSum / (4 * Volume);
 		OutVolume = Volume / 6;
 	}
@@ -313,6 +313,172 @@ namespace Chaos
 		return NewMP;
 	}
 
+	bool CalculateMassPropertiesOfImplicitUnion(
+		Chaos::FMassProperties& OutMassProperties,
+		const Chaos::FRigidTransform3& WorldTransform,
+		const Chaos::FImplicitObjectUnion& ImplicitUnion,
+		Chaos::FReal InDensityKGPerCM)
+	{
+		Chaos::FReal TotalMass = 0;
+		Chaos::FReal TotalVolume = 0;
+		Chaos::FVec3 TotalCenterOfMass(0);
+		TArray< Chaos::FMassProperties > MassPropertiesList;
+		for (const TUniquePtr<FImplicitObject>& ImplicitObjectUniquePtr: ImplicitUnion.GetObjects())
+		{
+			if (const Chaos::FImplicitObject* ImplicitObject = ImplicitObjectUniquePtr.Get())
+			{
+				Chaos::FMassProperties MassProperties;
+				if (CalculateMassPropertiesOfImplicitType(MassProperties, FTransform::Identity, ImplicitObject, InDensityKGPerCM))
+				{
+					MassPropertiesList.Add(MassProperties);
+					TotalMass += MassProperties.Mass;
+					TotalVolume += MassProperties.Volume;
+					TotalCenterOfMass += MassProperties.CenterOfMass * MassProperties.Mass;
+				}
+			}
+		}
+
+		Chaos::FMatrix33 Tensor;
+
+		// If no shapes contribute to mass, or they are scaled to zero, we may end up with zero mass here
+		if ((TotalMass > 0.f) && (MassPropertiesList.Num() > 0))
+		{
+			TotalCenterOfMass /= TotalMass;
+
+			const Chaos::FMassProperties CombinedMassProperties = Chaos::CombineWorldSpace(MassPropertiesList);
+			ensure(CombinedMassProperties.RotationOfMass.IsIdentity());
+			Tensor = CombinedMassProperties.InertiaTensor;
+		}
+		else
+		{
+			// @todo(chaos): We should support shape-less particles as long as their mass an inertia are set directly
+			// For now hard-code a 50cm sphere with density 1g/cc
+			Tensor = Chaos::FMatrix33(5.24e5f, 5.24e5f, 5.24e5f);
+			TotalMass = 523.0f;
+			TotalVolume = 523000;
+		}
+
+		OutMassProperties.InertiaTensor = Tensor;
+		OutMassProperties.Mass = TotalMass;
+		OutMassProperties.Volume = TotalVolume;
+		OutMassProperties.CenterOfMass = TotalCenterOfMass;
+		OutMassProperties.RotationOfMass = Chaos::FRotation3::Identity;
+
+		return (OutMassProperties.Mass > 0);
+	}
+	
+	bool CalculateMassPropertiesOfImplicitType(
+		Chaos::FMassProperties& OutMassProperties,
+		const Chaos::FRigidTransform3& WorldTransform,
+		const Chaos::FImplicitObject* ImplicitObject,
+		Chaos::FReal InDensityKGPerCM)
+	{
+		using namespace Chaos;
+
+		if (ImplicitObject)
+		{
+			// Hack to handle Transformed and Scaled<ImplicitObjectTriangleMesh> until CastHelper can properly support transformed
+			// Commenting this out temporarily as it breaks vehicles
+			/*	if (Chaos::IsScaled(ImplicitObject->GetType(true)) && Chaos::GetInnerType(ImplicitObject->GetType(true)) & Chaos::ImplicitObjectType::TriangleMesh)
+				{
+					OutMassProperties.Volume = 0.f;
+					OutMassProperties.Mass = FLT_MAX;
+					OutMassProperties.InertiaTensor = FMatrix33(0, 0, 0);
+					OutMassProperties.CenterOfMass = FVector(0);
+					OutMassProperties.RotationOfMass = Chaos::FRotation3::FromIdentity();
+					return false;
+				}
+				else if (ImplicitObject->GetType(true) & Chaos::ImplicitObjectType::TriangleMesh)
+				{
+					OutMassProperties.Volume = 0.f;
+					OutMassProperties.Mass = FLT_MAX;
+					OutMassProperties.InertiaTensor = FMatrix33(0, 0, 0);
+					OutMassProperties.CenterOfMass = FVector(0);
+					OutMassProperties.RotationOfMass = Chaos::FRotation3::FromIdentity();
+					return false;
+				}
+			else*/
+			if (ImplicitObject->IsUnderlyingUnion())
+			{
+				const FImplicitObjectUnion* Union = static_cast<const FImplicitObjectUnion*>(ImplicitObject);
+				CalculateMassPropertiesOfImplicitUnion(OutMassProperties, WorldTransform, *Union, InDensityKGPerCM);
+			}
+			else
+			{
+				Chaos::Utilities::CastHelper(*ImplicitObject, FTransform::Identity, [&OutMassProperties, InDensityKGPerCM](const auto& Object, const auto& LocalTM)
+					{
+						OutMassProperties.Volume = Object.GetVolume();
+						OutMassProperties.Mass = OutMassProperties.Volume * InDensityKGPerCM;
+						OutMassProperties.InertiaTensor = Object.GetInertiaTensor(OutMassProperties.Mass);
+						OutMassProperties.CenterOfMass = LocalTM.TransformPosition(Object.GetCenterOfMass());
+						OutMassProperties.RotationOfMass = LocalTM.GetRotation() * Object.GetRotationOfMass();
+					});
+			}
+		}
+
+		// If the implicit is null, or it is scaled to zero it will have zero volume, mass or inertia
+		return (OutMassProperties.Mass > 0);
+	}
+
+	void CalculateMassPropertiesFromShapeCollection(
+		Chaos::FMassProperties& OutProperties, 
+		int32 InNumShapes, 
+		Chaos::FReal InDensityKGPerCM,
+		const TArray<bool>& bContributesToMass,
+		TFunction<Chaos::FPerShapeData* (int32 ShapeIndex)> GetShapeDelegate)
+	{
+		Chaos::FReal TotalMass = 0;
+		Chaos::FReal TotalVolume = 0;
+		Chaos::FVec3 TotalCenterOfMass(0);
+		TArray< Chaos::FMassProperties > MassPropertiesList;
+		for (int32 ShapeIndex = 0; ShapeIndex < InNumShapes; ++ShapeIndex)
+		{
+			const Chaos::FPerShapeData* Shape = GetShapeDelegate(ShapeIndex);
+
+			const bool bHassMass = (ShapeIndex < bContributesToMass.Num()) ? bContributesToMass[ShapeIndex] : true;
+			if (bHassMass)
+			{
+				if (const Chaos::FImplicitObject* ImplicitObject = Shape->GetGeometry().Get())
+				{
+					Chaos::FMassProperties MassProperties;
+					if (CalculateMassPropertiesOfImplicitType(MassProperties, FTransform::Identity, ImplicitObject, InDensityKGPerCM))
+					{
+						MassPropertiesList.Add(MassProperties);
+						TotalMass += MassProperties.Mass;
+						TotalVolume += MassProperties.Volume;
+						TotalCenterOfMass += MassProperties.CenterOfMass * MassProperties.Mass;
+					}
+				}
+			}
+		}
+
+		Chaos::FMatrix33 Tensor;
+
+		// If no shapes contribute to mass, or they are scaled to zero, we may end up with zero mass here
+		if ((TotalMass > 0.f) && (MassPropertiesList.Num() > 0))
+		{
+			TotalCenterOfMass /= TotalMass;
+
+			Chaos::FMassProperties CombinedMassProperties = Chaos::CombineWorldSpace(MassPropertiesList);
+			ensure(CombinedMassProperties.RotationOfMass.IsIdentity());
+			Tensor = CombinedMassProperties.InertiaTensor;
+		}
+		else
+		{
+			// @todo(chaos): We should support shape-less particles as long as their mass an inertia are set directly
+			// For now hard-code a 50cm sphere with density 1g/cc
+			Tensor = Chaos::FMatrix33(5.24e5f, 5.24e5f, 5.24e5f);
+			TotalMass = 523.0f;
+			TotalVolume = 523000;
+		}
+
+		OutProperties.InertiaTensor = Tensor;
+		OutProperties.Mass = TotalMass;
+		OutProperties.Volume = TotalVolume;
+		OutProperties.CenterOfMass = TotalCenterOfMass;
+		OutProperties.RotationOfMass = Chaos::FRotation3::Identity;
+	}
+	
 	template CHAOS_API TRotation<FRealSingle,3> TransformToLocalSpace(PMatrix<FRealSingle,3,3>& Inertia);
 	template CHAOS_API TRotation<FRealDouble,3> TransformToLocalSpace(PMatrix<FRealDouble,3,3>& Inertia);
 	
diff --git a/Engine/Source/Runtime/Experimental/Chaos/Private/Chaos/PBDRigidClusteringAlgo.cpp b/Engine/Source/Runtime/Experimental/Chaos/Private/Chaos/PBDRigidClusteringAlgo.cpp
--- a/Engine/Source/Runtime/Experimental/Chaos/Private/Chaos/PBDRigidClusteringAlgo.cpp
+++ b/Engine/Source/Runtime/Experimental/Chaos/Private/Chaos/PBDRigidClusteringAlgo.cpp
@@ -402,7 +402,7 @@ namespace Chaos
 		// ignore unions for now as we don't yet support deep copy of it
 		// on the GT they are only used by clusters that aggregate their children shapes ( see GeometryCollectionPhysicsProxy.cpp )
 		// by failing artificially this condition thmake sure we create a FImplicitObjectUnionClustered for this particle 
-		if (ProxyGeometry && ProxyGeometry->GetType() != ImplicitObjectType::Union)
+		if (ProxyGeometry)
 		{
 			const FVector Scale = Parameters.Scale;
 			auto DeepCopyImplicit = [&Scale](const TSharedPtr<Chaos::FImplicitObject, ESPMode::ThreadSafe>& ImplicitToCopy) -> TUniquePtr<Chaos::FImplicitObject>
diff --git a/Engine/Source/Runtime/Experimental/Chaos/Private/PhysicsProxy/GeometryCollectionPhysicsProxy.cpp b/Engine/Source/Runtime/Experimental/Chaos/Private/PhysicsProxy/GeometryCollectionPhysicsProxy.cpp
--- a/Engine/Source/Runtime/Experimental/Chaos/Private/PhysicsProxy/GeometryCollectionPhysicsProxy.cpp
+++ b/Engine/Source/Runtime/Experimental/Chaos/Private/PhysicsProxy/GeometryCollectionPhysicsProxy.cpp
@@ -2143,17 +2143,24 @@ static const FSharedSimulationSizeSpecificData& GetSizeSpecificData(const TArray
 	return SizeSpecificData[SizeSpecificIdx];
 }
 
-static void GenerateInnerAndOuterRadiiIfNeeded(FGeometryCollection& RestCollection, const int32 GeometryIndex)
+// compute inner and outer radii from geometry in parallel if the attribute value as not been set
+static void GenerateInnerAndOuterRadiiIfNeeded(FGeometryCollection& RestCollection)
 {
-	if (RestCollection.InnerRadius[GeometryIndex] == 0.0f || RestCollection.OuterRadius[GeometryIndex] == 0.0f)
-	{
-		GeometryCollection::ComputeInnerAndOuterRadiiFromGeometryVertices(
-			RestCollection.Vertex,
-			RestCollection.VertexStart[GeometryIndex],
-			RestCollection.VertexCount[GeometryIndex],
-			RestCollection.InnerRadius[GeometryIndex],
-			RestCollection.OuterRadius[GeometryIndex]);
-	}	
+	const int32 NumGeometries = RestCollection.NumElements(FGeometryCollection::GeometryGroup);
+	const TManagedArray<FVector3f>& Vertices = RestCollection.Vertex;
+
+	ParallelFor(NumGeometries, [&](int32 GeometryIndex)
+	{
+		if (RestCollection.InnerRadius[GeometryIndex] == 0.0f || RestCollection.OuterRadius[GeometryIndex] == 0.0f)
+		{
+			GeometryCollection::ComputeInnerAndOuterRadiiFromGeometryVertices(
+				Vertices,
+				RestCollection.VertexStart[GeometryIndex],
+				RestCollection.VertexCount[GeometryIndex],
+				RestCollection.InnerRadius[GeometryIndex],
+				RestCollection.OuterRadius[GeometryIndex]);
+		}
+	});
 }
 
 static FGeometryDynamicCollection::FSharedImplicit CreateImplicitGeometry(
@@ -2375,22 +2382,20 @@ static void ComputeMassPropertyFromMesh(
 		bOutIsTooSmallGeometry
 		);
 
-	// now compute unit inertia ( assuming mass of 1 it will be scaled later by the real mass
-	constexpr Chaos::FReal UnitMass = 1000;
-	ensure (OutMassProperties.Volume > UE_SMALL_NUMBER);
-	const Chaos::FReal Density = UnitMass / OutMassProperties.Volume;
+	// now compute unit inertia ( assuming density of 1 )
+	constexpr Chaos::FReal UnitDensityKGPerCM = 1;
 	const bool bUseBoundingBoxInertia = (UncorrectedVolume == 0);  
 	if (bUseBoundingBoxInertia)
 	{
 		OutMassProperties.CenterOfMass = BoundingBox[GeometryIndex].GetCenter();
-		Chaos::CalculateInertiaAndRotationOfMass(BoundingBox[GeometryIndex], Density, OutMassProperties.InertiaTensor, OutMassProperties.RotationOfMass);
+		Chaos::CalculateInertiaAndRotationOfMass(BoundingBox[GeometryIndex], UnitDensityKGPerCM, OutMassProperties.InertiaTensor, OutMassProperties.RotationOfMass);
 	}
 	else
 	{
 		const FVector3f CenterOfMass{ OutMassProperties.CenterOfMass };
 		Chaos::PMatrix<float,3,3> InertiaTensor;
 		Chaos::TRotation<float,3> RotationOfMass;
-		CalculateInertiaAndRotationOfMass(Vertex.GetConstArray(), TriMesh->GetSurfaceElements(), (float)Density, CenterOfMass, InertiaTensor, RotationOfMass);
+		CalculateInertiaAndRotationOfMass(Vertex.GetConstArray(), TriMesh->GetSurfaceElements(), (float)UnitDensityKGPerCM, CenterOfMass, InertiaTensor, RotationOfMass);
 		OutMassProperties.CenterOfMass = CenterOfMass;
 		OutMassProperties.InertiaTensor = InertiaTensor;
 		OutMassProperties.RotationOfMass = RotationOfMass;
@@ -2404,6 +2409,74 @@ static void ComputeMassPropertyFromMesh(
 	}
 }
 
+// compute mass properties and trimesh in parallel
+// ( non-simulatable particles will not compute the it  ) 
+// also return if the geometry mesh is too small 
+static void ComputeMassPropertiesAndTriMeshes(
+	const FGeometryCollection& RestCollection,
+	const FSharedSimulationParameters& SharedParams,
+	TArray<TUniquePtr<Chaos::FTriangleMesh>>& OutTriangleMeshesArray,
+	TArray<Chaos::FMassProperties>& MassPropertiesArray,
+	TArray<bool>& IsTooSmallGeometryArray
+	)
+{
+	using FImplicitGeom = FGeometryDynamicCollection::FSharedImplicit;
+	
+	const int32 NumGeometries = RestCollection.NumElements(FGeometryCollection::GeometryGroup);
+	const TManagedArray<int32>& TransformIndex = RestCollection.TransformIndex;
+	const TManagedArray<bool>& CollectionSimulatableParticles = RestCollection.GetAttribute<bool>(FGeometryCollection::SimulatableParticlesAttribute, FTransformCollection::TransformGroup);
+	const TManagedArray<FImplicitGeom>* ExternaCollisions = RestCollection.FindAttribute<FImplicitGeom>("ExternalCollisions", FGeometryCollection::TransformGroup);
+	const TManagedArray<int32>& SimulationType = RestCollection.SimulationType;
+
+	ParallelFor(NumGeometries, [&](int32 GeometryIndex)
+	{
+		const int32 TransformGroupIndex = TransformIndex[GeometryIndex];
+		if (CollectionSimulatableParticles[TransformGroupIndex])
+		{
+			// todo compute mesh here and pass it on 
+			ComputeMassPropertyFromMesh(GeometryIndex, RestCollection, SharedParams, OutTriangleMeshesArray[GeometryIndex], MassPropertiesArray[GeometryIndex], IsTooSmallGeometryArray[GeometryIndex]);
+			// todo : needs to be a better decision as we are still computing the inertia above to replace it ?( see comment above )
+			if (SharedParams.bUseImportedCollisionImplicits && ExternaCollisions && (*ExternaCollisions)[TransformGroupIndex])
+			{
+				constexpr Chaos::FReal UnitDensityKGPerCM = 1;
+				Chaos::CalculateMassPropertiesOfImplicitType(MassPropertiesArray[GeometryIndex], Chaos::FRigidTransform3::Identity, (*ExternaCollisions)[TransformGroupIndex].Get(), UnitDensityKGPerCM);
+				IsTooSmallGeometryArray[GeometryIndex] = false;
+			}
+		}
+
+	});
+}
+
+static TUniquePtr<Chaos::FImplicitObject> MakeTransformImplicitObject(const Chaos::FImplicitObject& ImplicitObject, const Chaos::FRigidTransform3& Transform)
+{
+	TUniquePtr<Chaos::FImplicitObject> ResultObject;
+	
+	// we cannot really put a transform on top a union, so we need to transform each member
+	if (ImplicitObject.IsUnderlyingUnion())
+	{
+		 TArray<TUniquePtr<Chaos::FImplicitObject>> TransformedObjects;
+		const Chaos::FImplicitObjectUnion& Union = static_cast<const Chaos::FImplicitObjectUnion&>(ImplicitObject);
+		for (const TUniquePtr<Chaos::FImplicitObject>& Object: Union.GetObjects())
+		{
+			TransformedObjects.Add(MakeTransformImplicitObject(*Object, Transform));
+		}
+		ResultObject = MakeUnique<Chaos::FImplicitObjectUnion>(MoveTemp(TransformedObjects));
+	}
+	else if (ImplicitObject.GetType() == Chaos::ImplicitObjectType::Transformed)
+	{
+		const Chaos::TImplicitObjectTransformed<Chaos::FReal,3>& TransformObject = static_cast<const Chaos::TImplicitObjectTransformed<Chaos::FReal,3>&>(ImplicitObject);
+		// we deep copy at this point as the transform is going to be handled at this level
+		TUniquePtr<Chaos::FImplicitObject> TransformedObjectCopy =  TransformObject.GetTransformedObject()->DeepCopy();
+		ResultObject = MakeUnique<Chaos::TImplicitObjectTransformed<Chaos::FReal,3>>(MoveTemp(TransformedObjectCopy),  TransformObject.GetTransform() * Transform);
+	}
+	else
+	{
+		// we deep copy at this point as the transform is going to be handled at this level
+		TUniquePtr<Chaos::FImplicitObject> TransformedObjectCopy =  ImplicitObject.DeepCopy();
+		ResultObject = MakeUnique<Chaos::TImplicitObjectTransformed<Chaos::FReal,3>>(MoveTemp(TransformedObjectCopy), Transform);
+	}
+	return ResultObject;
+}
 
 /** 
 	NOTE - Making any changes to data stored on the rest collection below MUST be accompanied
@@ -2424,11 +2497,6 @@ void FGeometryCollectionPhysicsProxy::InitializeSharedCollisionStructures(
 		return;
 	}
 
-	//TArray<TArray<TArray<int32>>> BoundaryVertexIndices;
-	//GeometryCollectionAlgo::FindOpenBoundaries(&RestCollection, 1e-2, BoundaryVertexIndices);
-	//GeometryCollectionAlgo::TriangulateBoundaries(&RestCollection, BoundaryVertexIndices);
-	//RestCollection.ReindexMaterials();
-
 	using namespace Chaos;
 
 	// TransformGroup
@@ -2444,9 +2512,6 @@ void FGeometryCollectionPhysicsProxy::InitializeSharedCollisionStructures(
 	RestCollection.RemoveAttribute(FGeometryDynamicCollection::ImplicitsAttribute, FTransformCollection::TransformGroup);
 	TManagedArray<FGeometryDynamicCollection::FSharedImplicit>& CollectionImplicits = RestCollection.AddAttribute<FGeometryDynamicCollection::FSharedImplicit>(FGeometryDynamicCollection::ImplicitsAttribute, FTransformCollection::TransformGroup);
 
-	//TManagedArray<TSet<int32>>* TransformToConvexIndices = RestCollection.FindAttribute<TSet<int32>>("TransformToConvexIndices", FTransformCollection::TransformGroup);
-	//TManagedArray<TUniquePtr<Chaos::FConvex>>* ConvexGeometry = RestCollection.FindAttribute<TUniquePtr<Chaos::FConvex>>("ConvexHull", "Convex");
-
 	bool bUseRelativeSize = RestCollection.HasAttribute(TEXT("Size"), FTransformCollection::TransformGroup);
 	if (!bUseRelativeSize)
 	{
@@ -2489,19 +2554,9 @@ void FGeometryCollectionPhysicsProxy::InitializeSharedCollisionStructures(
 	TArray<TUniquePtr<FTriangleMesh>> TriangleMeshesArray;	//use to union trimeshes in cluster case
 	TriangleMeshesArray.AddDefaulted(NumTransforms);
 
-	// FParticles MassSpaceParticles;
-	// MassSpaceParticles.AddParticles(Vertex.Num());
-	// for (int32 Idx = 0; Idx < Vertex.Num(); ++Idx)
-	// {
-	// 	MassSpaceParticles.X(Idx) = Vertex[Idx];	//mass space computation done later down
-	// }
-
 	TArray<FMassProperties> MassPropertiesArray;
 	MassPropertiesArray.AddUninitialized(NumGeometries);
 
-	//TArray<FVector::FReal> UncorrectedVolume;
-	//UncorrectedVolume.Init(0, NumGeometries);
-
 	// The geometry group has a set of transform indices that maps a geometry index
 	// to a transform index, but only in the case where there is a 1-to-1 mapping 
 	// between the two.  In the event where a geometry is instanced for multiple
@@ -2513,20 +2568,15 @@ void FGeometryCollectionPhysicsProxy::InitializeSharedCollisionStructures(
 	TArray<bool> IsTooSmallGeometryArray;
 	IsTooSmallGeometryArray.Init(false, NumGeometries);
 
-	// compute mass properties in parallel
-	ParallelFor(NumGeometries, [&](int32 GeometryIndex)
-	{
-		// @chaos(todo): is that really needed? 
-		GenerateInnerAndOuterRadiiIfNeeded(RestCollection, GeometryIndex);
+	GenerateInnerAndOuterRadiiIfNeeded(RestCollection);
 
-		const int32 TransformGroupIndex = TransformIndex[GeometryIndex];
+	for (int32 TransformGroupIndex = 0; TransformGroupIndex < NumTransforms; TransformGroupIndex++)
+	{
 		CollectionSimulatableParticles[TransformGroupIndex] = (SimulationType[TransformGroupIndex] > FGeometryCollection::ESimulationTypes::FST_None);
-		
-		if (CollectionSimulatableParticles[TransformGroupIndex])
-		{
-			ComputeMassPropertyFromMesh(GeometryIndex, RestCollection, SharedParams, TriangleMeshesArray[GeometryIndex], MassPropertiesArray[GeometryIndex], IsTooSmallGeometryArray[GeometryIndex]);
-		}
-	});
+	}
+	
+	// compute mass properties in parallel
+	ComputeMassPropertiesAndTriMeshes(RestCollection, SharedParams, TriangleMeshesArray, MassPropertiesArray, IsTooSmallGeometryArray);
 	
 	// compute total volume and whether we need to skip small geometry ( can't use parallelfors )
 	bool bSkippedSmallGeometry = false;
@@ -2552,6 +2602,9 @@ void FGeometryCollectionPhysicsProxy::InitializeSharedCollisionStructures(
 	{
 		UE_LOG(LogChaos, Warning, TEXT("Some geometry is too small to be simulated and has been skipped."));
 	}
+
+	using FImplicitGeomSharePtr = FGeometryDynamicCollection::FSharedImplicit;
+	const TManagedArray<FImplicitGeomSharePtr>* ExternaCollisions = RestCollection.FindAttribute<FImplicitGeomSharePtr>("ExternalCollisions", FGeometryCollection::TransformGroup);
 	
 	// User provides us with total mass or density.
 	// Density must be the same for individual parts and the total. Density_i = Density = Mass_i / Volume_i
@@ -2591,8 +2644,8 @@ void FGeometryCollectionPhysicsProxy::InitializeSharedCollisionStructures(
 			const FReal Density_i = Mass_i / Volume_i;
 			CollectionMass[TransformGroupIndex] = (FRealSingle)Mass_i;
 
-			// scale the mass properties inertia tensor by the mass
-			MassPropertiesArray[GeometryIndex].InertiaTensor *= Mass_i / 1000.0;
+			// scale the mass properties inertia tensor by the Density_i ( as we computedit earlier with a density of one )
+			MassPropertiesArray[GeometryIndex].InertiaTensor *= Density_i;
 			const Chaos::FMatrix33& InertiaTensor = MassPropertiesArray[GeometryIndex].InertiaTensor;
 			CollectionInertiaTensor[TransformGroupIndex] = FVector3f((float)InertiaTensor.M[0][0], (float)InertiaTensor.M[1][1], (float)InertiaTensor.M[2][2]);
 
@@ -2619,9 +2672,9 @@ void FGeometryCollectionPhysicsProxy::InitializeSharedCollisionStructures(
 			const TUniquePtr<Chaos::FTriangleMesh>& TriMesh = TriangleMeshesArray[GeometryIndex];
 			const FTransform& MassToLocalTransform = CollectionMassToLocal[TransformGroupIndex];
 
-			// Orient particle in the local mass space 
-			if (!MassPropertiesArray[GeometryIndex].RotationOfMass.Equals(FQuat::Identity) ||
-				!FMath::IsNearlyZero(MassPropertiesArray[GeometryIndex].CenterOfMass.SizeSquared()))
+			// Orient particle in the local mass space
+			const bool bIdentityMassTransform = MassToLocalTransform.Equals(FTransform::Identity);
+			if (!bIdentityMassTransform)
 			{
 				const int32 IdxStart = VertexStart[GeometryIndex];
 				const int32 IdxEnd = IdxStart + VertexCount[GeometryIndex];
@@ -2682,16 +2735,36 @@ void FGeometryCollectionPhysicsProxy::InitializeSharedCollisionStructures(
 						ensureMsgf(false, TEXT("Simplicial is empty."));
 					}
 
-					LocalErrorReporter.SetPrefix(BaseErrorPrefix + " | Transform Index: " + FString::FromInt(TransformGroupIndex) + " of " + FString::FromInt(TransformIndex.Num()));
-					CollectionImplicits[TransformGroupIndex] = CreateImplicitGeometry(
-						SizeSpecificData,
-						TransformGroupIndex,
-						RestCollection,
-						MassSpaceParticles,
-						*TriMesh,
-						InstanceBoundingBox,
-						InnerRadius[GeometryIndex],
-						LocalErrorReporter);
+					// tdo(chaos) : if imported is selected but no shape exists shoudl we simply ignore ? 
+					if (SharedParams.bUseImportedCollisionImplicits && ExternaCollisions && (*ExternaCollisions)[TransformGroupIndex])
+					{
+						// for now  simply copy the shared pointer as the imported geometry should not change under the hood without being recreated
+						const FImplicitGeomSharePtr ExternalCollisionImplicit = (*ExternaCollisions)[TransformGroupIndex];
+						if (!bIdentityMassTransform)
+						{
+							// since we do not set the rotation of mass and center of mass properties on the particle and have a MasstoLocal managed property on the collection instead
+							// we need to reverse transform the external shapes 
+							TUniquePtr<FImplicitObject> TransformedCollisionImplicit = MakeTransformImplicitObject(*ExternalCollisionImplicit, MassToLocalTransform.Inverse());
+							CollectionImplicits[TransformGroupIndex] = TSharedPtr<Chaos::FImplicitObject>(TransformedCollisionImplicit.Release());
+						}
+						else
+						{
+							CollectionImplicits[TransformGroupIndex] = ExternalCollisionImplicit;
+						}
+					}
+					else
+					{
+						LocalErrorReporter.SetPrefix(BaseErrorPrefix + " | Transform Index: " + FString::FromInt(TransformGroupIndex) + " of " + FString::FromInt(TransformIndex.Num()));
+						CollectionImplicits[TransformGroupIndex] = CreateImplicitGeometry(
+							SizeSpecificData,
+							TransformGroupIndex,
+							RestCollection,
+							MassSpaceParticles,
+							*TriMesh,
+							InstanceBoundingBox,
+							InnerRadius[GeometryIndex],
+							LocalErrorReporter);
+					}
 					
 					if (CollectionImplicits[TransformGroupIndex] && CollectionImplicits[TransformGroupIndex]->HasBoundingBox())
 					{
@@ -2876,18 +2949,39 @@ void FGeometryCollectionPhysicsProxy::InitializeSharedCollisionStructures(
 
 				ErrorReporter.SetPrefix(BaseErrorPrefix + " | Cluster Transform Index: " + FString::FromInt(ClusterTransformIdx));
 
-				CollectionImplicits[ClusterTransformIdx] = CreateImplicitGeometry(
-					SizeSpecificData,
-					ClusterTransformIdx,
-					RestCollection,
-					MassSpaceParticles,
-					*UnionMesh,
-					InstanceBoundingBox,
-					InstanceBoundingBox.GetExtent().GetAbsMin(), // InnerRadius
-					ErrorReporter,
-					&MaxChildBounds
-					);
+				if (SharedParams.bUseImportedCollisionImplicits && ExternaCollisions && (*ExternaCollisions)[TransformGroupIndex])
+				{
+					const FTransform& ClusterMassToLocal = CollectionMassToLocal[ClusterTransformIdx];
+					const bool bIdentityMassTransform = ClusterMassToLocal.Equals(FTransform::Identity);
 
+					// for now  simply copy the shared pointer as the imported geometry should not change under the hood without being recreated
+					const FImplicitGeomSharePtr ExternalCollisionImplicit = (*ExternaCollisions)[TransformGroupIndex];
+					if (!bIdentityMassTransform)
+					{
+						// since we do not set the rotation of mass and center of mass properties on the particle and have a MasstoLocal managed property on the collection instead
+						// we need to reverse transform the external shapes 
+						TUniquePtr<FImplicitObject> TransformedCollisionImplicit = MakeTransformImplicitObject(*ExternalCollisionImplicit, ClusterMassToLocal.Inverse());
+						CollectionImplicits[TransformGroupIndex] = TSharedPtr<Chaos::FImplicitObject>(TransformedCollisionImplicit.Release());
+					}
+					else
+					{
+						CollectionImplicits[ClusterTransformIdx] = (*ExternaCollisions)[ClusterTransformIdx];
+					}
+				}
+				else
+				{
+					CollectionImplicits[ClusterTransformIdx] = CreateImplicitGeometry(
+						SizeSpecificData,
+						ClusterTransformIdx,
+						RestCollection,
+						MassSpaceParticles,
+						*UnionMesh,
+						InstanceBoundingBox,
+						InstanceBoundingBox.GetExtent().GetAbsMin(), // InnerRadius
+						ErrorReporter,
+						&MaxChildBounds
+						);
+				}
 				// create simplicial from the implicit geometry
 				CollectionSimplicials[ClusterTransformIdx] = TUniquePtr<FSimplicial>(
 							FCollisionStructureManager::NewSimplicial(MassSpaceParticles, *UnionMesh, CollectionImplicits[ClusterTransformIdx].Get(),
diff --git a/Engine/Source/Runtime/Experimental/Chaos/Public/Chaos/ImplicitObjectTransformed.h b/Engine/Source/Runtime/Experimental/Chaos/Public/Chaos/ImplicitObjectTransformed.h
--- a/Engine/Source/Runtime/Experimental/Chaos/Public/Chaos/ImplicitObjectTransformed.h
+++ b/Engine/Source/Runtime/Experimental/Chaos/Public/Chaos/ImplicitObjectTransformed.h
@@ -100,19 +100,80 @@ public:
 		}
 	}
 
-	~TImplicitObjectTransformed() {}
+	virtual TUniquePtr<FImplicitObject> Copy() const
+	{
+		if (MObject)
+		{
+			TUniquePtr<FImplicitObject> ObjectCopy(MObject->Copy());
+			return MakeUnique<TImplicitObjectTransformed<T,d>>(MoveTemp(ObjectCopy), MTransform);
+		}
+		else
+		{
+			check(false);
+			return nullptr;
+		}
+	}
+	
+	virtual TUniquePtr<FImplicitObject> CopyWithScale(const FVec3& Scale) const override
+	{
+		if(MObject)
+		{
+			//return MakeCopyWithScaleTransformed(MObjectOwner, MTransform, Scale);
+			// since we cannot have a { Scaled -- Transformed -- Shape } ( scaled can only directly reference concrete shapes )
+			// we need to scale the transform translation and set the Scaled on the shape itself like { (Adjusted)Transformed -- Scaled -- Shape }  
+			FRigidTransform3 AdjustedTransform{ MTransform };
+			AdjustedTransform.ScaleTranslation(Scale);
+
+			TUniquePtr<FImplicitObject> ScaledObject(MObject->CopyWithScale(Scale));
+			return MakeUnique<TImplicitObjectTransformed<FReal,3>>(MoveTemp(ScaledObject), AdjustedTransform);
+		}
+		else
+		{
+			check(false);
+			return nullptr;
+		}
+	}
 
-	static constexpr EImplicitObjectType StaticType()
+	virtual TUniquePtr<FImplicitObject> DeepCopy() const
 	{
-		return ImplicitObjectType::Transformed;
+		if(MObject)
+		{
+			TUniquePtr<FImplicitObject> ObjectCopy(MObject->DeepCopy());
+			return MakeUnique<TImplicitObjectTransformed<T,d>>(MoveTemp(ObjectCopy), MTransform);
+		}
+		else
+		{
+			check(false);
+			return nullptr;
+		}
 	}
 
-	virtual TUniquePtr<FImplicitObject> Copy() const
+	virtual TUniquePtr<FImplicitObject> DeepCopyWithScale(const FVec3& Scale) const override
 	{
-		TUniquePtr<FImplicitObject> ImplicitCopy = MObject->Copy();
-		return TUniquePtr<FImplicitObject>(new TImplicitObjectTransformed<T, d>(MoveTemp(ImplicitCopy), MTransform));
+		if(MObject)
+		{
+			//return MakeCopyWithScaleTransformed(MObjectOwner, MTransform, Scale);
+			// since we cannot have a { Scaled -- Transformed -- Shape } ( scaled can only directly reference concrete shapes )
+			// we need to scale the transform translation and set the Scaled on the shape itself like { (Adjusted)Transformed -- Scaled -- Shape }  
+			FRigidTransform3 AdjustedTransform{ MTransform };
+			AdjustedTransform.ScaleTranslation(Scale);
+
+			TUniquePtr<FImplicitObject> ScaledObject(MObject->DeepCopyWithScale(Scale));
+			return MakeUnique<TImplicitObjectTransformed<FReal,3>>(MoveTemp(ScaledObject), AdjustedTransform);
+		}
+		else
+		{
+			check(false);
+			return nullptr;
+		}
 	}
+	
+	~TImplicitObjectTransformed() {}
 
+	static constexpr EImplicitObjectType StaticType()
+	{
+		return ImplicitObjectType::Transformed;
+	}
 
 	const FImplicitObject* GetTransformedObject() const
 	{
@@ -334,4 +395,4 @@ namespace Utilities
 template <typename T, int d>
 using TImplicitObjectTransformedNonSerializable = TImplicitObjectTransformed<T, d, false>;
 
-}
\ No newline at end of file
+}
diff --git a/Engine/Source/Runtime/Experimental/Chaos/Public/Chaos/ImplicitObjectUnion.h b/Engine/Source/Runtime/Experimental/Chaos/Public/Chaos/ImplicitObjectUnion.h
--- a/Engine/Source/Runtime/Experimental/Chaos/Public/Chaos/ImplicitObjectUnion.h
+++ b/Engine/Source/Runtime/Experimental/Chaos/Public/Chaos/ImplicitObjectUnion.h
@@ -33,6 +33,11 @@ class CHAOS_API FImplicitObjectUnion : public FImplicitObject
 		return ImplicitObjectType::Union;
 	}
 
+	virtual TUniquePtr<FImplicitObject> Copy() const;
+	virtual TUniquePtr<FImplicitObject> CopyWithScale(const FVec3& Scale) const override;
+	virtual TUniquePtr<FImplicitObject> DeepCopy() const;
+	virtual TUniquePtr<FImplicitObject> DeepCopyWithScale(const FVec3& Scale) const override;
+	
 	virtual FReal PhiWithNormal(const FVec3& x, FVec3& Normal) const override
 	{
 		FReal Phi = TNumericLimits<FReal>::Max();
diff --git a/Engine/Source/Runtime/Experimental/Chaos/Public/Chaos/MassProperties.h b/Engine/Source/Runtime/Experimental/Chaos/Public/Chaos/MassProperties.h
--- a/Engine/Source/Runtime/Experimental/Chaos/Public/Chaos/MassProperties.h
+++ b/Engine/Source/Runtime/Experimental/Chaos/Public/Chaos/MassProperties.h
@@ -12,6 +12,9 @@
 namespace Chaos
 {
 	class FTriangleMesh;
+	class FPerShapeData;
+	class FImplicitObject;
+	class FImplicitObjectUnion;
 
 	template<class T, int d>
 	class TParticles;
@@ -70,6 +73,27 @@ namespace Chaos
 	// @see Combine()
 	FMassProperties CHAOS_API CombineWorldSpace(const TArray<FMassProperties>& MPArray);
 
+	// Calculate the mass properties from a union
+	bool CHAOS_API CalculateMassPropertiesOfImplicitUnion(
+	Chaos::FMassProperties& OutMassProperties,
+	const Chaos::FRigidTransform3& WorldTransform,
+	const Chaos::FImplicitObjectUnion& ImplicitUnion,
+	Chaos::FReal InDensityKGPerCM);
+	
+	// Calculate the mass properties from a specific implicit
+	bool CHAOS_API CalculateMassPropertiesOfImplicitType(
+		FMassProperties& OutMassProperties,
+		const FRigidTransform3& WorldTransform,
+		const FImplicitObject* ImplicitObject,
+		Chaos::FReal InDensityKGPerCM);
+	
+	// Calculate the mass properties from a list of shapes
+	void CHAOS_API CalculateMassPropertiesFromShapeCollection(
+		FMassProperties& OutProperties, 
+		int32 InNumShapes, 
+		FReal InDensityKGPerCM,
+		const TArray<bool>& bContributesToMass,
+		TFunction<FPerShapeData* (int32 ShapeIndex)> GetShapeDelegate);
 
 	template <typename T, int d>
 	using TMassProperties UE_DEPRECATED(4.27, "Deprecated. this class is to be deleted, use FMassProperties instead") = FMassProperties;
diff --git a/Engine/Source/Runtime/Experimental/Chaos/Public/GeometryCollection/GeometryCollectionSimulationCoreTypes.h b/Engine/Source/Runtime/Experimental/Chaos/Public/GeometryCollection/GeometryCollectionSimulationCoreTypes.h
--- a/Engine/Source/Runtime/Experimental/Chaos/Public/GeometryCollection/GeometryCollectionSimulationCoreTypes.h
+++ b/Engine/Source/Runtime/Experimental/Chaos/Public/GeometryCollection/GeometryCollectionSimulationCoreTypes.h
@@ -99,6 +99,7 @@ struct FSharedSimulationParameters
 	, MinimumInertiaTensorDiagonalClamp(UE_SMALL_NUMBER)	// todo : Expose to users with better initial values
 	, MaximumInertiaTensorDiagonalClamp(1e20f)				// todo : Expose to users with better initial values
 	, MaximumCollisionParticleCount(60)
+	, bUseImportedCollisionImplicits(false)
 	{
 		SizeSpecificData.AddDefaulted();
 	}
@@ -119,7 +120,8 @@ struct FSharedSimulationParameters
 		, float InMaximumInertiaTensorDiagonalClamp
 		, float InCollisionParticlesFraction
 		, int32 InMaximumCollisionParticleCount
-		, float InCollisionMarginFraction)
+		, float InCollisionMarginFraction
+		, bool InUseImportedCollisionImplicits )
 	: bMassAsDensity(InMassAsDensity)
 	, Mass(InMass)
 	, MinimumMassClamp(InMinimumMassClamp)
@@ -129,6 +131,7 @@ struct FSharedSimulationParameters
 	, MinimumInertiaTensorDiagonalClamp(InMinimumInertiaTensorDiagonalClamp)
 	, MaximumInertiaTensorDiagonalClamp(InMaximumInertiaTensorDiagonalClamp)
 	, MaximumCollisionParticleCount(InMaximumCollisionParticleCount)
+	, bUseImportedCollisionImplicits(InUseImportedCollisionImplicits)
 	{
 		SizeSpecificData.AddDefaulted();
 		if (ensure(SizeSpecificData.Num() && SizeSpecificData[0].CollisionShapesData.Num()))
@@ -160,6 +163,7 @@ struct FSharedSimulationParameters
 	TArray<FSharedSimulationSizeSpecificData> SizeSpecificData;
 	TArray<int32> RemoveOnFractureIndices;
 	int32 MaximumCollisionParticleCount;
+	bool bUseImportedCollisionImplicits;
 };
 
 struct FSimulationParameters
diff --git a/Engine/Source/Runtime/Experimental/GeometryCollectionEngine/Private/GeometryCollection/GeometryCollectionObject.cpp b/Engine/Source/Runtime/Experimental/GeometryCollectionEngine/Private/GeometryCollection/GeometryCollectionObject.cpp
--- a/Engine/Source/Runtime/Experimental/GeometryCollectionEngine/Private/GeometryCollection/GeometryCollectionObject.cpp
+++ b/Engine/Source/Runtime/Experimental/GeometryCollectionEngine/Private/GeometryCollection/GeometryCollectionObject.cpp
@@ -91,6 +91,7 @@ UGeometryCollection::UGeometryCollection(const FObjectInitializer& ObjectInitial
 	, bMassAsDensity(true)
 	, Mass(2500.0f)
 	, MinimumMassClamp(0.1f)
+	, bImportCollisionFromSource(false)
 	, bRemoveOnMaxSleep(false)
 	, MaximumSleepTime(5.0, 10.0)
 	, RemovalDuration(2.5, 5.0)
@@ -347,6 +348,8 @@ void UGeometryCollection::GetSharedSimulationParams(FSharedSimulationParameters&
 		FixupRemoveOnFractureMaterials(OutParams);
 	}
 
+	OutParams.bUseImportedCollisionImplicits = bImportCollisionFromSource;
+	
 	OutParams.SizeSpecificData.Sort();	//can we do this at editor time on post edit change?
 }
 
diff --git a/Engine/Source/Runtime/Experimental/GeometryCollectionEngine/Public/GeometryCollection/GeometryCollectionObject.h b/Engine/Source/Runtime/Experimental/GeometryCollectionEngine/Public/GeometryCollection/GeometryCollectionObject.h
--- a/Engine/Source/Runtime/Experimental/GeometryCollectionEngine/Public/GeometryCollection/GeometryCollectionObject.h
+++ b/Engine/Source/Runtime/Experimental/GeometryCollectionEngine/Public/GeometryCollection/GeometryCollectionObject.h
@@ -533,6 +533,12 @@ public:
 	UPROPERTY(EditAnywhere, BlueprintReadOnly, Category = "Collisions")
 	float MinimumMassClamp;
 
+	/**
+	* whether to import collision from the source asset
+	*/
+	UPROPERTY(EditAnywhere, BlueprintReadOnly, Category = "Collisions")
+	bool bImportCollisionFromSource;
+	
 #if WITH_EDITORONLY_DATA
 	/**
 	 * Number of particles on the triangulated surface to use for collisions.
