DataflowCore - Updates

diff --git a/Engine/Source/Runtime/Experimental/Dataflow/Core/Private/Dataflow/DataflowNode.cpp b/Engine/Source/Runtime/Experimental/Dataflow/Core/Private/Dataflow/DataflowNode.cpp
--- a/Engine/Source/Runtime/Experimental/Dataflow/Core/Private/Dataflow/DataflowNode.cpp
+++ b/Engine/Source/Runtime/Experimental/Dataflow/Core/Private/Dataflow/DataflowNode.cpp
@@ -3,6 +3,7 @@
 #include "Dataflow/DataflowNode.h"
 
 #include "Dataflow/DataflowNode.h"
+#include "Dataflow/DataflowProperty.h"
 
 namespace Dataflow
 {
diff --git a/Engine/Source/Runtime/Experimental/Dataflow/Core/Public/Dataflow/DataflowCore.h b/Engine/Source/Runtime/Experimental/Dataflow/Core/Public/Dataflow/DataflowCore.h
--- a/Engine/Source/Runtime/Experimental/Dataflow/Core/Public/Dataflow/DataflowCore.h
+++ b/Engine/Source/Runtime/Experimental/Dataflow/Core/Public/Dataflow/DataflowCore.h
@@ -7,4 +7,5 @@
 #include "Dataflow/DataflowInputOutput.h"
 #include "Dataflow/DataflowNode.h"
 #include "Dataflow/DataflowNodeFactory.h"
-#include "Dataflow/DataflowNodeParameters.h"
\ No newline at end of file
+#include "Dataflow/DataflowNodeParameters.h"
+#include "Dataflow/DataflowProperty.h"
diff --git a/Engine/Source/Runtime/Experimental/Dataflow/Core/Public/Dataflow/DataflowNode.h b/Engine/Source/Runtime/Experimental/Dataflow/Core/Public/Dataflow/DataflowNode.h
--- a/Engine/Source/Runtime/Experimental/Dataflow/Core/Public/Dataflow/DataflowNode.h
+++ b/Engine/Source/Runtime/Experimental/Dataflow/Core/Public/Dataflow/DataflowNode.h
@@ -4,11 +4,11 @@
 
 #include "CoreMinimal.h"
 #include "Dataflow/DataflowConnection.h"
-#include "Dataflow/DataflowProperty.h"
 #include "Dataflow/DataflowNodeParameters.h"
 
 namespace Dataflow
 {
+	class FProperty;
 
 	struct DATAFLOWCORE_API FNodeParameters {
 		FName Name;
