Chaos: Fixed MTD path for heightfield queries not checking cell bounds before entering GJK

diff --git a/Engine/Source/Runtime/Experimental/Chaos/Private/Chaos/HeightField.cpp b/Engine/Source/Runtime/Experimental/Chaos/Private/Chaos/HeightField.cpp
--- a/Engine/Source/Runtime/Experimental/Chaos/Private/Chaos/HeightField.cpp
+++ b/Engine/Source/Runtime/Experimental/Chaos/Private/Chaos/HeightField.cpp
@@ -2227,13 +2227,17 @@
 				OutMTD->Normal = FVec3(0);
 				OutMTD->Penetration = TNumericLimits<FReal>::Lowest();
 				FVec3 Points[4];
+				FAABB3 CellBounds;
 				for (const TVec2<int32>& Cell : Intersections)
 				{
 					const int32 SingleIndex = Cell[1] * (GeomData.NumCols) + Cell[0];
-					GeomData.GetPointsScaled(SingleIndex, Points);
+					GeomData.GetPointsAndBoundsScaled(SingleIndex, Points, CellBounds);
 
-					bOverlaps |= OverlapTriangleMTD(Points[0], Points[1], Points[3], OutMTD);
-					bOverlaps |= OverlapTriangleMTD(Points[0], Points[3], Points[2], OutMTD);
+					if(CellBounds.Intersects(QueryBounds))
+					{
+						bOverlaps |= OverlapTriangleMTD(Points[0], Points[1], Points[3], OutMTD);
+						bOverlaps |= OverlapTriangleMTD(Points[0], Points[3], Points[2], OutMTD);
+					}
 				}
 				return bOverlaps;
 			}
@@ -2248,7 +2252,6 @@
 
 				const UE::Math::TVector<FReal>& TranslationDouble = QueryTM.GetTranslation();
 				const VectorRegister4Float Translation = MakeVectorRegisterFloatFromDouble(MakeVectorRegister(TranslationDouble.X, TranslationDouble.Y, TranslationDouble.Z, 0.0));
-
 				FAABBVectorized CellBounds;
 				for (const TVec2<int32>& Cell : Intersections)
 				{
