GeometryCollection : Fix CIS

diff --git a/Engine/Source/Runtime/Experimental/GeometryCollectionEngine/Public/GeometryCollection/GeometryCollectionObject.h b/Engine/Source/Runtime/Experimental/GeometryCollectionEngine/Public/GeometryCollection/GeometryCollectionObject.h
--- a/Engine/Source/Runtime/Experimental/GeometryCollectionEngine/Public/GeometryCollection/GeometryCollectionObject.h
+++ b/Engine/Source/Runtime/Experimental/GeometryCollectionEngine/Public/GeometryCollection/GeometryCollectionObject.h
@@ -7,7 +7,6 @@
 #include "Rendering/NaniteResources.h"
 #include "InstanceUniformShaderParameters.h"
 #include "GeometryCollection/ManagedArray.h"
-#include "GeometryCollection/GeometryCollectionObject.h"
 #include "GeometryCollection/GeometryCollectionSimulationTypes.h"
 #include "Chaos/ChaosSolverActor.h"
 
