Chaos - Collision auto-stiffness experiment (disabled on cvar, WIP)

diff --git a/Engine/Source/Runtime/Experimental/Chaos/Private/Chaos/Collision/SolverCollisionContainer.cpp b/Engine/Source/Runtime/Experimental/Chaos/Private/Chaos/Collision/SolverCollisionContainer.cpp
--- a/Engine/Source/Runtime/Experimental/Chaos/Private/Chaos/Collision/SolverCollisionContainer.cpp
+++ b/Engine/Source/Runtime/Experimental/Chaos/Private/Chaos/Collision/SolverCollisionContainer.cpp
@@ -14,6 +14,7 @@
 
 #include "ChaosLog.h"
 
+//PRAGMA_DISABLE_OPTIMIZATION
 
 namespace Chaos
 {
@@ -23,6 +24,15 @@ namespace Chaos
 
 		extern bool bChaos_PBDCollisionSolver_Position_SolveEnabled;
 		extern bool bChaos_PBDCollisionSolver_Velocity_SolveEnabled;
+
+		// If one body is more than MassRatio1 times the mass of the other, adjust the solver stiffness when the lighter body is underneath.
+		// Solver stiffness will be equal to 1 when the mass ratio is MassRatio1.
+		// Solver stiffness will be equal to 0 when the mass ratio is MassRatio2.
+		FRealSingle Chaos_PBDCollisionSolver_AutoStiffness_MassRatio1 = 0;
+		FRealSingle Chaos_PBDCollisionSolver_AutoStiffness_MassRatio2 = 0;
+		FAutoConsoleVariableRef CVarChaosPBDCollisionSolverAutoStiffnessMassRatio1(TEXT("p.Chaos.PBDCollision.AutoStiffness.MassRatio1"), Chaos_PBDCollisionSolver_AutoStiffness_MassRatio1, TEXT(""));
+		FAutoConsoleVariableRef CVarChaosPBDCollisionSolverAutoStiffnessMassRatio2(TEXT("p.Chaos.PBDCollision.AutoStiffness.MassRatio2"), Chaos_PBDCollisionSolver_AutoStiffness_MassRatio2, TEXT(""));
+
 	}
 	using namespace CVars;
 
@@ -75,6 +85,59 @@ namespace Chaos
 			SolverBodyContainer.FindOrAdd(Constraint->GetParticle1());
 		}
 
+		/**
+		 * @brief Modify solver stiffness when we have bodies with large mass differences
+		*/
+		FReal CalculateSolverStiffness(
+			const FSolverBody* Body0,
+			const FSolverBody* Body1,
+			const FReal MassRatio1,
+			const FReal MassRatio2)
+		{
+			// Adjust the solver stiffness if one body is more than MassRatio1 times the mass of the other and the heavier one is on top.
+			// Solver stiffness will be equal to 1 when the mass ratio is MassRatio1.
+			// Solver stiffness will be equal to 0 when the mass ratio is MassRatio2.
+			if (Body0->IsDynamic() && Body1->IsDynamic() && (MassRatio1 > 0) && (MassRatio2 > MassRatio1))
+			{
+				// Find heavy body and the mass ratio
+				const FSolverBody* HeavyBody;
+				FReal MassRatio;
+				if (Body0->InvM() < Body1->InvM())
+				{
+					HeavyBody = Body0;
+					MassRatio = Body1->InvM() / Body0->InvM();
+				}
+				else
+				{
+					HeavyBody = Body1;
+					MassRatio = Body0->InvM() / Body1->InvM();
+				}
+
+				if (MassRatio > MassRatio1)
+				{
+					// Is this a load-bearing contact (normal is significantly along gravity direction)?
+					// @todo(chaos): should use gravity direction. Currently assumes -Z
+					// @todo(chaos): maybe gradually introduce stiffness scaling based on normal rather than on/off
+					// @todo(chaos): could use solver manifold data which is already in world space rather than CalculateWorldContactNormal
+					const FVec3 WorldNormal = Constraint->CalculateWorldContactNormal();
+					if (FMath::Abs(WorldNormal.Z) > FReal(0.3))
+					{
+						// Which body is on the top? (Normal always points away from second body - see FContactPoint)
+						const FSolverBody* TopBody = (WorldNormal.Z > 0) ? Body0 : Body1;
+						if (TopBody == HeavyBody)
+						{
+							// The heavy body is on top - reduce the solver stiffness
+							const FReal StiffnessScale = FMath::Clamp((MassRatio2 - MassRatio) / (MassRatio2 - MassRatio1), FReal(0), FReal(1));
+							return StiffnessScale * Constraint->GetStiffness();
+						}
+					}
+				}
+
+			}
+
+			return Constraint->GetStiffness();
+		}
+
 		/**
 		 * @brief Initialize the data required for the solver and bind to the bodies
 		*/
@@ -124,7 +187,8 @@ namespace Chaos
 
 			Solver.SetFriction(PositionStaticFriction, PositionDynamicFriction, VelocityDynamicFriction);
 
-			Solver.SetStiffness(FSolverReal(Constraint->GetStiffness()));
+			const FReal SolverStiffness = CalculateSolverStiffness(Body0, Body1, Chaos_PBDCollisionSolver_AutoStiffness_MassRatio1, Chaos_PBDCollisionSolver_AutoStiffness_MassRatio2);
+			Solver.SetStiffness(FSolverReal(SolverStiffness));
 
 			Solver.SetSolverBodies(*Body0, *Body1);
 			Solver.SolverBody0().SetInvMScale(Constraint->GetInvMassScale0());
