Add support for  Geometry Collection collision filter changes at runtime

diff --git a/Engine/Source/Runtime/Experimental/Chaos/Private/Chaos/PBDRigidClustering.cpp b/Engine/Source/Runtime/Experimental/Chaos/Private/Chaos/PBDRigidClustering.cpp
--- a/Engine/Source/Runtime/Experimental/Chaos/Private/Chaos/PBDRigidClustering.cpp
+++ b/Engine/Source/Runtime/Experimental/Chaos/Private/Chaos/PBDRigidClustering.cpp
@@ -149,7 +149,7 @@
 		FMatrix33 ClusterInertia(0);
 		UpdateClusterMassProperties(NewParticle, ChildrenSet, ClusterInertia, ForceMassOrientation);
 		UpdateKinematicProperties(NewParticle, MChildren, MEvolution);
-		UpdateGeometry(NewParticle, ChildrenSet, ProxyGeometry, Parameters);
+		UpdateGeometry(NewParticle, ChildrenSet, MChildren, ProxyGeometry, Parameters);
 		GenerateConnectionGraph(NewParticle, Parameters);
 
 		NewParticle->SetSleeping(bClusterIsAsleep);
@@ -256,7 +256,7 @@
 		UpdateClusterMassProperties(NewParticle, ChildrenSet, ClusterInertia, nullptr);
 		UpdateKinematicProperties(NewParticle, MChildren, MEvolution);
 
-		UpdateGeometry(NewParticle, ChildrenSet, nullptr, NoCleanParams);
+		UpdateGeometry(NewParticle, ChildrenSet, MChildren, nullptr, NoCleanParams);
 
 		return NewParticle;
 	}
@@ -870,7 +870,6 @@
 		return AllActivatedChildren;
 	}
 
-
 	DECLARE_CYCLE_STAT(TEXT("FRigidClustering::Visitor"), STAT_ClusterVisitor, STATGROUP_Chaos);
 	void FRigidClustering::Visitor(FClusterHandle Cluster, FVisitorFunction Function)
 	{
diff --git a/Engine/Source/Runtime/Experimental/Chaos/Private/Chaos/PBDRigidClusteringAlgo.cpp b/Engine/Source/Runtime/Experimental/Chaos/Private/Chaos/PBDRigidClusteringAlgo.cpp
--- a/Engine/Source/Runtime/Experimental/Chaos/Private/Chaos/PBDRigidClusteringAlgo.cpp
+++ b/Engine/Source/Runtime/Experimental/Chaos/Private/Chaos/PBDRigidClusteringAlgo.cpp
@@ -282,6 +282,7 @@ namespace Chaos
 	UpdateGeometry(
 		Chaos::FPBDRigidClusteredParticleHandle* Parent,
 		const TSet<FPBDRigidParticleHandle*>& Children,
+		const FRigidClustering::FClusterMap& ChildrenMap,
 		TSharedPtr<Chaos::FImplicitObject, ESPMode::ThreadSafe> ProxyGeometry,
 		const FClusterCreationParameters& Parameters)
 	{
@@ -305,34 +306,6 @@ namespace Chaos
 		const bool bUseCollisionPoints = (ProxyGeometry || Parameters.bCopyCollisionParticles) && !Parameters.CollisionParticles;
 		bool bUseParticleImplicit = false;
 
-		// Need to extract a filter off one of the cluster children 
-		FCollisionFilterData SimFilter, QueryFilter;
-		for (FPBDRigidParticleHandle* Child : Children)
-		{
-			bool bFilterValid = false;
-			for (const TUniquePtr<FPerShapeData>& Shape : Child->ShapesArray())
-			{
-				if (Shape)
-				{
-					SimFilter = Shape->GetSimData();
-					bFilterValid = SimFilter.Word0 != 0 || SimFilter.Word1 != 0 || SimFilter.Word2 != 0 || SimFilter.Word3 != 0;
-
-					QueryFilter = Shape->GetQueryData();
-				}
-
-				if (bFilterValid)
-				{
-					break;
-				}
-			}
-
-			// Bail once we've found one
-			if (bFilterValid)
-			{
-				break;
-			}
-		}
-
 		{ // STAT_UpdateGeometry_GatherObjects
 			SCOPE_CYCLE_COUNTER(STAT_UpdateGeometry_GatherObjects);
 
@@ -525,12 +498,52 @@ namespace Chaos
 			Parent->UpdateWorldSpaceState(Xf, FVec3(0));
 		}
 
-		// Set the captured filter to our new shapes
-		for (const TUniquePtr<FPerShapeData>& Shape : Parent->ShapesArray())
+		
+		// Update filter data on new shapes
+		const FRigidClustering::FRigidHandleArray& ChildrenArray = ChildrenMap[Parent];
+		UpdateClusterFilterDataFromChildren(Parent, ChildrenArray);
+	}
+	
+	void UpdateClusterFilterDataFromChildren(FPBDRigidClusteredParticleHandle* ClusterParent, const TArray<FPBDRigidParticleHandle*>& Children)
+	{
+		SCOPE_CYCLE_COUNTER(STAT_GCUpdateFilterData);
+
+		FCollisionFilterData SelectedSimFilter, SelectedQueryFilter;
+		bool bFilterValid = false;
+		for (FPBDRigidParticleHandle* ChildHandle : Children)
+		{
+			for (const TUniquePtr<FPerShapeData>& Shape : ChildHandle->ShapesArray())
+			{
+				if (Shape)
+				{
+					SelectedSimFilter = Shape->GetSimData();
+					bFilterValid = SelectedSimFilter.Word0 != 0 || SelectedSimFilter.Word1 != 0 || SelectedSimFilter.Word2 != 0 || SelectedSimFilter.Word3 != 0;
+					SelectedQueryFilter = Shape->GetQueryData();
+
+					if (bFilterValid)
+					{
+						break;
+					}
+				}
+			}
+
+			if (bFilterValid)
+			{
+				break;
+			}
+		}
+
+		// Apply selected filters to shapes
+		if (bFilterValid)
 		{
-			Shape->SetSimData(SimFilter);
-			Shape->SetQueryData(QueryFilter);
+			const FShapesArray& ShapesArray = ClusterParent->ShapesArray();
+			for (const TUniquePtr<FPerShapeData>& Shape : ShapesArray)
+			{
+				Shape->SetQueryData(SelectedQueryFilter);
+				Shape->SetSimData(SelectedSimFilter);
+			}
 		}
 	}
 
+
 } // namespace Chaos
diff --git a/Engine/Source/Runtime/Experimental/Chaos/Private/ChaosStats.cpp b/Engine/Source/Runtime/Experimental/Chaos/Private/ChaosStats.cpp
--- a/Engine/Source/Runtime/Experimental/Chaos/Private/ChaosStats.cpp
+++ b/Engine/Source/Runtime/Experimental/Chaos/Private/ChaosStats.cpp
@@ -22,6 +22,7 @@ DEFINE_STAT(STAT_GCRaycast);
 DEFINE_STAT(STAT_GCOverlap);
 DEFINE_STAT(STAT_GCSweep);
 DEFINE_STAT(STAT_GCCUpdateBounds);
+DEFINE_STAT(STAT_GCUpdateFilterData);
 DEFINE_STAT(STAT_GCCUGlobalMatrices);
 DEFINE_STAT(STAT_GCInitDynamicData);
 DEFINE_STAT(STAT_GCTotalTransforms);
diff --git a/Engine/Source/Runtime/Experimental/Chaos/Private/PhysicsProxy/GeometryCollectionPhysicsProxy.cpp b/Engine/Source/Runtime/Experimental/Chaos/Private/PhysicsProxy/GeometryCollectionPhysicsProxy.cpp
--- a/Engine/Source/Runtime/Experimental/Chaos/Private/PhysicsProxy/GeometryCollectionPhysicsProxy.cpp
+++ b/Engine/Source/Runtime/Experimental/Chaos/Private/PhysicsProxy/GeometryCollectionPhysicsProxy.cpp
@@ -10,6 +10,7 @@
 #include "Chaos/BVHParticles.h"
 #include "Chaos/Transform.h"
 #include "Chaos/ParallelFor.h"
+#include "Chaos/GeometryParticles.h"
 #include "Chaos/Particles.h"
 #include "Chaos/TriangleMesh.h"
 #include "Chaos/MassProperties.h"
@@ -371,6 +372,7 @@ FGeometryCollectionPhysicsProxy::FGeometryCollectionPhysicsProxy(
 
 	, GameThreadCollection(GameThreadCollectionIn)
 	, bIsPhysicsThreadWorldTransformDirty(false)
+	, bIsCollisionFilterDataDirty(false)
 {
 	// We rely on a guarded buffer.
 	check(BufferMode == Chaos::EMultiBufferMode::TripleGuarded);
@@ -1660,13 +1662,22 @@ void FGeometryCollectionPhysicsProxy::PushStateOnGameThread(Chaos::FPBDRigidsSol
 		Parameters.WorldTransform = GameThreadPerFrameData.GetWorldTransform();
 		GameThreadPerFrameData.ResetIsWorldTransformDirty();
 	}
+
+	bIsCollisionFilterDataDirty = GameThreadPerFrameData.GetIsCollisionFilterDataDirty();
+	if (bIsCollisionFilterDataDirty)
+	{
+		Parameters.QueryFilterData = GameThreadPerFrameData.GetQueryFilter();
+		Parameters.SimulationFilterData = GameThreadPerFrameData.GetSimFilter();
+		GameThreadPerFrameData.ResetIsCollisionFilterDataDirty();
+	}
 }
 
 void FGeometryCollectionPhysicsProxy::PushToPhysicsState()
 {
+	using namespace Chaos;
 	// CONTEXT: PHYSICSTHREAD
 	// because the attached actor can be dynamic, we need to update the kinematic particles properly
-	if (bIsPhysicsThreadWorldTransformDirty)
+	if (bIsPhysicsThreadWorldTransformDirty || bIsCollisionFilterDataDirty)
 	{
 		const FTransform& ActorToWorld = Parameters.WorldTransform;
 
@@ -1679,29 +1690,56 @@ void FGeometryCollectionPhysicsProxy::PushToPhysicsState()
 			Chaos::FPBDRigidClusteredParticleHandle* Handle = SolverParticleHandles[TransformGroupIndex];
 			if (Handle)
 			{
-				if (Handle->ObjectState() == Chaos::EObjectStateType::Kinematic)
+				// Must update our filters before updating an internal cluster parent
+				if (bIsCollisionFilterDataDirty)
+				{
+					const Chaos::FShapesArray& ShapesArray = Handle->ShapesArray();
+					for (const TUniquePtr<Chaos::FPerShapeData>& Shape : ShapesArray)
+					{
+						Shape->SetQueryData(Parameters.QueryFilterData);
+						Shape->SetSimData(Parameters.SimulationFilterData);
+					}
+				}
+
+				// in the case of cluster union we need to find our Internal Cluster parent and update it
+				if (!InternalClusterParentUpdated)
 				{
-					// in the case of cluster union we need to find our Internal Cluster parent and update it
-					if (!InternalClusterParentUpdated)
+					FClusterHandle* ParentHandle = Handle->Parent();
+					if (ParentHandle && ParentHandle->InternalCluster())
 					{
-						FClusterHandle* ParentHandle = Handle->Parent();
-						if (ParentHandle && ParentHandle->InternalCluster() && !ParentHandle->Disabled() && ParentHandle->ObjectState() == Chaos::EObjectStateType::Kinematic)
+						if (bIsPhysicsThreadWorldTransformDirty && Handle->ObjectState() == Chaos::EObjectStateType::Kinematic && ParentHandle->ObjectState() == Chaos::EObjectStateType::Kinematic && !ParentHandle->Disabled())
 						{
 							FTransform NewChildWorldTransform = PhysicsThreadCollection.MassToLocal[TransformGroupIndex] * PhysicsThreadCollection.Transform[TransformGroupIndex] * ActorToWorld;
 							Chaos::FRigidTransform3 ParentToChildTransform = Handle->ChildToParent().Inverse();
 							FTransform NewParentWorldTRansform = ParentToChildTransform * NewChildWorldTransform;
 							SetClusteredParticleKinematicTarget_Internal(ParentHandle, NewParentWorldTRansform);
+						}
+
+						if (bIsCollisionFilterDataDirty)
+						{
+							if(Chaos::FPhysicsSolver* RigidSolver = GetSolver<Chaos::FPhysicsSolver>())
+							{
+								FRigidClustering& RigidClustering = RigidSolver->GetEvolution()->GetRigidClustering();
+								FRigidClustering::FRigidHandleArray* ChildArray = RigidClustering.GetChildrenMap().Find(ParentHandle);
+								if (ensure(ChildArray))
+								{
+									// If our filter changed, internal cluster parent's filter may be stale.
+									UpdateClusterFilterDataFromChildren(ParentHandle, *ChildArray);
+								}
 
-							InternalClusterParentUpdated = true;
+							}
 						}
-					}
 
-					if (!Handle->Disabled())
-					{
-						FTransform WorldTransform = PhysicsThreadCollection.MassToLocal[TransformGroupIndex] * PhysicsThreadCollection.Transform[TransformGroupIndex] * ActorToWorld;
-						SetClusteredParticleKinematicTarget_Internal(Handle, WorldTransform);
+
+						InternalClusterParentUpdated = true;
 					}
 				}
+
+				if (bIsPhysicsThreadWorldTransformDirty && !Handle->Disabled() && Handle->ObjectState() == Chaos::EObjectStateType::Kinematic)
+				{
+					FTransform WorldTransform = PhysicsThreadCollection.MassToLocal[TransformGroupIndex] * PhysicsThreadCollection.Transform[TransformGroupIndex] * ActorToWorld;
+					SetClusteredParticleKinematicTarget_Internal(Handle, WorldTransform);
+				}
 			}
 		}
 	}
@@ -2124,6 +2162,37 @@ bool FGeometryCollectionPhysicsProxy::PullFromPhysicsState(const Chaos::FDirtyGe
 	return true;
 }
 
+
+void FGeometryCollectionPhysicsProxy::UpdateFilterData_External(const FCollisionFilterData& NewSimFilter, const FCollisionFilterData& NewQueryFilter)
+{
+	SCOPE_CYCLE_COUNTER(STAT_GCUpdateFilterData);
+	check(IsInGameThread());
+
+	// SimFilter/QueryFilter members are read on both threads, these are const after initialization and are not updated here.
+
+	int32 NumTransforms = GameThreadCollection.Transform.Num();
+	for (int32 Index = 0; Index < NumTransforms; ++Index)
+	{
+		Chaos::FGeometryParticle* P = GTParticles[Index].Get();
+		const Chaos::FShapesArray& Shapes = P->ShapesArray();
+		const int32 NumShapes = Shapes.Num();
+		for (int32 ShapeIndex = 0; ShapeIndex < NumShapes; ++ShapeIndex)
+		{
+			Chaos::FPerShapeData* Shape = Shapes[ShapeIndex].Get();
+			Shape->SetSimData(NewSimFilter);
+			Shape->SetQueryData(NewQueryFilter);
+		}
+	}
+
+	GameThreadPerFrameData.SetSimFilter(NewSimFilter);
+	GameThreadPerFrameData.SetQueryFilter(NewQueryFilter);
+
+	if (Chaos::FPhysicsSolver* RBDSolver = GetSolver<Chaos::FPhysicsSolver>())
+	{
+		RBDSolver->AddDirtyProxy(this);
+	}
+}
+
 //==============================================================================
 // STATIC SETUP FUNCTIONS
 //==============================================================================
diff --git a/Engine/Source/Runtime/Experimental/Chaos/Public/Chaos/PBDRigidClusteringAlgo.h b/Engine/Source/Runtime/Experimental/Chaos/Public/Chaos/PBDRigidClusteringAlgo.h
--- a/Engine/Source/Runtime/Experimental/Chaos/Public/Chaos/PBDRigidClusteringAlgo.h
+++ b/Engine/Source/Runtime/Experimental/Chaos/Public/Chaos/PBDRigidClusteringAlgo.h
@@ -19,7 +19,14 @@ namespace Chaos
 	void CHAOS_API UpdateGeometry(
 		Chaos::FPBDRigidClusteredParticleHandle* Parent,
 		const TSet<FPBDRigidParticleHandle*>& Children,
+		const FRigidClustering::FClusterMap& ChildrenMap,
 		TSharedPtr<Chaos::FImplicitObject, ESPMode::ThreadSafe> ProxyGeometry,
 		const FClusterCreationParameters& Parameters);
 
+
+	/**
+	* Update shapes with first valid filter data from a child. Should be called after a child's filter
+	* is modified to ensure cluster's filters match at least one of the children and isn't stale.
+	*/
+	void UpdateClusterFilterDataFromChildren(FPBDRigidClusteredParticleHandle* ClusterParent, const TArray<FPBDRigidParticleHandle*>& Children);
 } // namespace Chaos
diff --git a/Engine/Source/Runtime/Experimental/Chaos/Public/ChaosStats.h b/Engine/Source/Runtime/Experimental/Chaos/Public/ChaosStats.h
--- a/Engine/Source/Runtime/Experimental/Chaos/Public/ChaosStats.h
+++ b/Engine/Source/Runtime/Experimental/Chaos/Public/ChaosStats.h
@@ -36,6 +36,7 @@ DECLARE_CYCLE_STAT_EXTERN(TEXT("Trailing Data Generation"), STAT_TrailingCallbac
 DECLARE_CYCLE_STAT_EXTERN(TEXT("Geometry Collection Raycast"), STAT_GCRaycast, STATGROUP_Chaos, CHAOS_API);
 DECLARE_CYCLE_STAT_EXTERN(TEXT("Geometry Collection Overlap"), STAT_GCOverlap, STATGROUP_Chaos, CHAOS_API);
 DECLARE_CYCLE_STAT_EXTERN(TEXT("Geometry Collection Sweep"), STAT_GCSweep, STATGROUP_Chaos, CHAOS_API);
+DECLARE_CYCLE_STAT_EXTERN(TEXT("Geoemtry Collection Update Filter Data"), STAT_GCUpdateFilterData, STATGROUP_Chaos, CHAOS_API)
 DECLARE_CYCLE_STAT_EXTERN(TEXT("Geometry Collection Component UpdateBounds"), STAT_GCCUpdateBounds, STATGROUP_Chaos, CHAOS_API);
 DECLARE_CYCLE_STAT_EXTERN(TEXT("Geometry Collection Component CalculateGlobalMatrices"), STAT_GCCUGlobalMatrices, STATGROUP_Chaos, CHAOS_API);
 DECLARE_CYCLE_STAT_EXTERN(TEXT("Geometry Collection InitDynamicData"), STAT_GCInitDynamicData, STATGROUP_Chaos, CHAOS_API);
diff --git a/Engine/Source/Runtime/Experimental/Chaos/Public/GeometryCollectionProxyData.h b/Engine/Source/Runtime/Experimental/Chaos/Public/GeometryCollectionProxyData.h
--- a/Engine/Source/Runtime/Experimental/Chaos/Public/GeometryCollectionProxyData.h
+++ b/Engine/Source/Runtime/Experimental/Chaos/Public/GeometryCollectionProxyData.h
@@ -93,7 +93,8 @@ class FGeometryCollectioPerFrameData
 {
 public:
 	FGeometryCollectioPerFrameData()
-		: IsWorldTransformDirty(false) {}
+		: IsWorldTransformDirty(false)
+		, bIsCollisionFilterDataDirty(false) {}
 
 	const FTransform& GetWorldTransform() const { return WorldTransform; }
 
@@ -109,9 +110,30 @@ public:
 	bool GetIsWorldTransformDirty() const { return IsWorldTransformDirty; }
 	void ResetIsWorldTransformDirty() { IsWorldTransformDirty = false; }
 
+	const FCollisionFilterData& GetSimFilter() const { return SimFilter; }
+	void SetSimFilter(const FCollisionFilterData& NewSimFilter)
+	{
+		SimFilter = NewSimFilter;
+		bIsCollisionFilterDataDirty = true;
+	}
+
+	const FCollisionFilterData& GetQueryFilter() const { return QueryFilter; }
+	void SetQueryFilter(const FCollisionFilterData& NewQueryFilter)
+	{
+		QueryFilter = NewQueryFilter;
+		bIsCollisionFilterDataDirty = true;
+	}
+
+	bool GetIsCollisionFilterDataDirty() const { return bIsCollisionFilterDataDirty; }
+	void ResetIsCollisionFilterDataDirty() { bIsCollisionFilterDataDirty = false; }
+
 private:
 	FTransform WorldTransform;
 	bool IsWorldTransformDirty;
+
+	FCollisionFilterData SimFilter;
+	FCollisionFilterData QueryFilter;
+	bool bIsCollisionFilterDataDirty;
 };
 
 /**
diff --git a/Engine/Source/Runtime/Experimental/Chaos/Public/PhysicsProxy/GeometryCollectionPhysicsProxy.h b/Engine/Source/Runtime/Experimental/Chaos/Public/PhysicsProxy/GeometryCollectionPhysicsProxy.h
--- a/Engine/Source/Runtime/Experimental/Chaos/Public/PhysicsProxy/GeometryCollectionPhysicsProxy.h
+++ b/Engine/Source/Runtime/Experimental/Chaos/Public/PhysicsProxy/GeometryCollectionPhysicsProxy.h
@@ -248,6 +248,8 @@ public:
 	void SetInitializedForReplication() { bInitializedForReplication = true; }
 	bool GetInitializedForReplication() const { return bInitializedForReplication; }
 
+	void UpdateFilterData_External(const FCollisionFilterData& NewSimFilter, const FCollisionFilterData& NewQueryFilter);
+
 protected:
 	/**
 	* Compute damage threshold for a specific transform
@@ -319,8 +321,10 @@ private:
 	bool bInitializedForReplication = false;	//Indicates that initialization for replication has finished
 
 	TManagedArray<TUniquePtr<Chaos::FGeometryParticle>> GTParticles;
-	FCollisionFilterData SimFilter;
-	FCollisionFilterData QueryFilter;
+
+	// These are read on both threads and should not be changed
+	const FCollisionFilterData SimFilter;
+	const FCollisionFilterData QueryFilter;
 
 	// This is a subset of the geometry group that are used in the transform hierarchy to represent geometry
 	TArray<FBox> ValidGeometryBoundingBoxes;
@@ -358,6 +362,7 @@ private:
 	// this data flows from Game thread to physics thread
 	FGeometryCollectioPerFrameData GameThreadPerFrameData;
 	bool bIsPhysicsThreadWorldTransformDirty;
+	bool bIsCollisionFilterDataDirty;
 
 	// Currently this is using triple buffers for game-physics and 
 	// physics-game thread communication, but not for any reason other than this 
diff --git a/Engine/Source/Runtime/Experimental/GeometryCollectionEngine/Private/GeometryCollection/GeometryCollectionComponent.cpp b/Engine/Source/Runtime/Experimental/GeometryCollectionEngine/Private/GeometryCollection/GeometryCollectionComponent.cpp
--- a/Engine/Source/Runtime/Experimental/GeometryCollectionEngine/Private/GeometryCollection/GeometryCollectionComponent.cpp
+++ b/Engine/Source/Runtime/Experimental/GeometryCollectionEngine/Private/GeometryCollection/GeometryCollectionComponent.cpp
@@ -2176,22 +2176,7 @@ void UGeometryCollectionComponent::OnCreatePhysicsState()
 			BodyInstance.BodySetup = DummyBodySetup;
 			BodyInstance.OwnerComponent = this; // Required to make filter data include component/actor ID for ignored actors/components
 
-			FBodyCollisionFilterData FilterData;
-			FMaskFilter FilterMask = BodyInstance.GetMaskFilter();
-			BodyInstance.BuildBodyFilterData(FilterData);
-
-			InitialSimFilter = FilterData.SimFilter;
-			InitialQueryFilter = FilterData.QuerySimpleFilter;
-
-			// Enable for complex and simple (no dual representation currently like other meshes)
-			InitialQueryFilter.Word3 |= (EPDF_SimpleCollision | EPDF_ComplexCollision);
-			InitialSimFilter.Word3 |= (EPDF_SimpleCollision | EPDF_ComplexCollision);
-			
-			if (bNotifyCollisions)
-			{
-				InitialQueryFilter.Word3 |= EPDF_ContactNotify;
-				InitialSimFilter.Word3 |= EPDF_ContactNotify;
-			}
+			BuildInitialFilterData();
 
 			if (BodyInstance.bSimulatePhysics)
 			{
@@ -2361,6 +2346,42 @@ void UGeometryCollectionComponent::SendRenderDynamicData_Concurrent()
 	}
 }
 
+void UGeometryCollectionComponent::OnActorEnableCollisionChanged()
+{
+	// Update filters on BI
+	BodyInstance.UpdatePhysicsFilterData();
+
+	// Update InitialSimFilter and InitialQueryFilter
+	BuildInitialFilterData();
+
+	// Update filters stored on proxy
+	if (PhysicsProxy)
+	{
+		PhysicsProxy->UpdateFilterData_External(InitialSimFilter, InitialQueryFilter);
+	}
+
+}
+
+void UGeometryCollectionComponent::BuildInitialFilterData()
+{
+	FBodyCollisionFilterData FilterData;
+	FMaskFilter FilterMask = BodyInstance.GetMaskFilter();
+	BodyInstance.BuildBodyFilterData(FilterData);
+
+	InitialSimFilter = FilterData.SimFilter;
+	InitialQueryFilter = FilterData.QuerySimpleFilter;
+
+	// Enable for complex and simple (no dual representation currently like other meshes)
+	InitialQueryFilter.Word3 |= (EPDF_SimpleCollision | EPDF_ComplexCollision);
+	InitialSimFilter.Word3 |= (EPDF_SimpleCollision | EPDF_ComplexCollision);
+
+	if (bNotifyCollisions)
+	{
+		InitialQueryFilter.Word3 |= EPDF_ContactNotify;
+		InitialSimFilter.Word3 |= EPDF_ContactNotify;
+	}
+}
+
 void UGeometryCollectionComponent::SetRestCollection(const UGeometryCollection* RestCollectionIn)
 {
 	//UE_LOG(UGCC_LOG, Log, TEXT("GeometryCollectionComponent[%p]::SetRestCollection()"), this);
diff --git a/Engine/Source/Runtime/Experimental/GeometryCollectionEngine/Public/GeometryCollection/GeometryCollectionComponent.h b/Engine/Source/Runtime/Experimental/GeometryCollectionEngine/Public/GeometryCollection/GeometryCollectionComponent.h
--- a/Engine/Source/Runtime/Experimental/GeometryCollectionEngine/Public/GeometryCollection/GeometryCollectionComponent.h
+++ b/Engine/Source/Runtime/Experimental/GeometryCollectionEngine/Public/GeometryCollection/GeometryCollectionComponent.h
@@ -365,6 +365,7 @@ public:
 	virtual void CreateRenderState_Concurrent(FRegisterComponentContext* Context) override;
 	virtual void SendRenderDynamicData_Concurrent() override;
 	FORCEINLINE void SetRenderStateDirty() { bRenderStateDirty = true; }
+	virtual void OnActorEnableCollisionChanged() override;
 	virtual void BeginPlay() override;
 	virtual void EndPlay(const EEndPlayReason::Type ReasonEnd) override;
 	virtual void GetLifetimeReplicatedProps(TArray<FLifetimeProperty>& OutLifetimeProps) const override;
@@ -885,6 +886,8 @@ private:
 	bool CalculateInnerSphere(int32 TransformIndex, FSphere& SphereOut) const;
 	void ProcessRepData();
 
+	void BuildInitialFilterData();
+
 	/** The clusters we need to replicate */
 	TUniquePtr<TSet<Chaos::FPBDRigidClusteredParticleHandle*>> ClustersToRep;
 
