Fix crash in Triangle Mesh BVH when loading Reverb project

diff --git a/Engine/Source/Runtime/Experimental/Chaos/Private/Chaos/TriangleMeshImplicitObject.cpp b/Engine/Source/Runtime/Experimental/Chaos/Private/Chaos/TriangleMeshImplicitObject.cpp
--- a/Engine/Source/Runtime/Experimental/Chaos/Private/Chaos/TriangleMeshImplicitObject.cpp
+++ b/Engine/Source/Runtime/Experimental/Chaos/Private/Chaos/TriangleMeshImplicitObject.cpp
@@ -1869,7 +1869,7 @@
 			// Reorder triangle indices, triangles will be in the same order as the bounding volume in the BVH structure.
 			// And all triangles in a leaf will be contiguous in memory.
 			AddTriangles(LargeIndices, SmallIndices, MElements, LeafPayload.Payload);
-			if (ExternalFaceIndexMap.IsValid())
+			if (ExternalFaceIndexMap.IsValid() && !OldFaceIndexMap.IsEmpty())
 			{
 				ExternalFaceIndexMap->Add(OldFaceIndexMap[LeafPayload.Payload]);
 			}
@@ -1950,7 +1950,7 @@
 							// Reorder triangle indices, triangles will be in the same order as the bounding volume in the BVH structure.
 							// And all triangles in a leaf will be contiguous in memory.
 							AddTriangles(LargeIndices, SmallIndices, MElements, LeafPayload.Payload);
-							if (ExternalFaceIndexMap.IsValid())
+							if (ExternalFaceIndexMap.IsValid() && !OldFaceIndexMap.IsEmpty())
 							{
 								ExternalFaceIndexMap->Add(OldFaceIndexMap[LeafPayload.Payload]);
 							}
diff --git a/Engine/Source/Runtime/Experimental/ChaosCore/Public/Chaos/Core.h b/Engine/Source/Runtime/Experimental/ChaosCore/Public/Chaos/Core.h
--- a/Engine/Source/Runtime/Experimental/ChaosCore/Public/Chaos/Core.h
+++ b/Engine/Source/Runtime/Experimental/ChaosCore/Public/Chaos/Core.h
@@ -51,7 +51,7 @@ namespace Chaos
 	// NOTE: if you get a merge conflict on the GUID, you must replace it with a new GUID - do not accept the source or target
 	// or you will likely get DDC version conflicts resulting in crashes during load.
 	// Core version string for Chaos data. Any DDC builder dependent on Chaos for serialization should depend on this version
-	static const TCHAR* const ChaosVersionGUID = TEXT("5F283595-0103-4FE6-9684-E8B820158BFC");
+	static const TCHAR* const ChaosVersionGUID = TEXT("5453DC22-C449-42C0-A74C-4662A08B13E8");
 
 	inline FString GetChaosVersionStringInner()
 	{
