don't exclude childless roots from GeometryCollectionAlgo::ComputeRecursiveOrder

diff --git a/Engine/Source/Runtime/Experimental/Chaos/Private/GeometryCollection/GeometryCollectionAlgo.cpp b/Engine/Source/Runtime/Experimental/Chaos/Private/GeometryCollection/GeometryCollectionAlgo.cpp
--- a/Engine/Source/Runtime/Experimental/Chaos/Private/GeometryCollection/GeometryCollectionAlgo.cpp
+++ b/Engine/Source/Runtime/Experimental/Chaos/Private/GeometryCollection/GeometryCollectionAlgo.cpp
@@ -1440,7 +1440,7 @@ namespace GeometryCollectionAlgo
 		//enqueue all roots
 		for (int32 TransformGroupIndex = 0; TransformGroupIndex < NumTransforms; TransformGroupIndex++)
 		{
-			if (Parent[TransformGroupIndex] == FGeometryCollection::Invalid && Children[TransformGroupIndex].Num() > 0)
+			if (Parent[TransformGroupIndex] == FGeometryCollection::Invalid)
 			{
 				ClustersToProcess.Emplace(TransformGroupIndex);
 			}
