Chaos destruction : imporvement on the remove on sleep feature
- Now also remove of the piece does not move for a certain amount of time ( velocity based ) ( behind a Cvar )
- Fix bugs with shrinking not going all the way down to zero
- Simplify math for computing off-center scale matrix ( removed most of the Transform conversion and multiplication )
- Internal clusters can now break upon removal
- Passing more info about internal clusters from PT to GT
- More memory efficient structure passed between PT and GT

diff --git a/Engine/Source/Runtime/Experimental/Chaos/Private/PhysicsProxy/GeometryCollectionPhysicsProxy.cpp b/Engine/Source/Runtime/Experimental/Chaos/Private/PhysicsProxy/GeometryCollectionPhysicsProxy.cpp
--- a/Engine/Source/Runtime/Experimental/Chaos/Private/PhysicsProxy/GeometryCollectionPhysicsProxy.cpp
+++ b/Engine/Source/Runtime/Experimental/Chaos/Private/PhysicsProxy/GeometryCollectionPhysicsProxy.cpp
@@ -97,7 +97,7 @@ FGeometryCollectionResults::FGeometryCollectionResults()
 void FGeometryCollectionResults::Reset()
 {
 	SolverDt = 0.0f;
-	DisabledStates.SetNum(0);
+	States.SetNum(0);
 	GlobalTransforms.SetNum(0);
 	ParticleXs.SetNum(0);
 	ParticleRs.SetNum(0);
@@ -1488,15 +1488,15 @@ void FGeometryCollectionPhysicsProxy::GetRelevantParticleHandles(
 #endif
 }
 
-void FGeometryCollectionPhysicsProxy::DisableParticles(TArray<int32>& TransformGroupIndices)
+void FGeometryCollectionPhysicsProxy::DisableParticles(TArray<int32>&& TransformGroupIndices)
 {
 	check(IsInGameThread());
 
 	if (Chaos::FPhysicsSolver* RBDSolver = GetSolver<Chaos::FPhysicsSolver>())
 	{
-		RBDSolver->EnqueueCommandImmediate([this, RBDSolver, TransformGroupIndices]()
+		RBDSolver->EnqueueCommandImmediate([this, RBDSolver, IndicesToDisable = MoveTemp(TransformGroupIndices)]()
 			{
-				for (int32 TransformIdx : TransformGroupIndices)
+				for (int32 TransformIdx : IndicesToDisable)
 				{
 					RBDSolver->GetEvolution()->DisableParticleWithRemovalEvent(SolverParticleHandles[TransformIdx]);
 				}
@@ -1504,6 +1504,42 @@ void FGeometryCollectionPhysicsProxy::DisableParticles(TArray<int32>& TransformG
 	}
 }
 
+void FGeometryCollectionPhysicsProxy::BreakInternalClusterParent(TArray<int32>&& TransformGroupIndices)
+{
+	check(IsInGameThread());
+
+	if (Chaos::FPhysicsSolver* RBDSolver = GetSolver<Chaos::FPhysicsSolver>())
+	{
+		RBDSolver->EnqueueCommandImmediate([this, RBDSolver, IndicesToBreakParent = MoveTemp(TransformGroupIndices)]()
+		{
+			// could use a set but the normal usage will always have a very small amount of unique parents and the array is more cache friendly
+			TArray<Chaos::FPBDRigidClusteredParticleHandle*> ParentToBreak;
+			for (int32 TransformIdx : IndicesToBreakParent)
+			{
+				if (Chaos::FPBDRigidClusteredParticleHandle* ChildHandle = SolverParticleHandles[TransformIdx])
+				{
+					if (Chaos::FPBDRigidClusteredParticleHandle* ClusteredParent =  ChildHandle->Parent())
+					{
+						if (ClusteredParent->InternalCluster())
+						{
+							ParentToBreak.AddUnique(ClusteredParent);
+						}
+					}
+				}
+				for (Chaos::FPBDRigidClusteredParticleHandle* ClusteredParent: ParentToBreak)
+				{
+					Chaos::FRigidClustering& Clustering = RBDSolver->GetEvolution()->GetRigidClustering();
+					if (!ClusteredParent->Disabled() && Clustering.GetChildrenMap().Contains(ClusteredParent))
+					{
+						Clustering.ReleaseClusterParticles(ClusteredParent, nullptr, true /* bForceRelease */);
+					}
+				}
+			}
+		});
+	}
+	
+}
+
 int32 FGeometryCollectionPhysicsProxy::CalculateHierarchyLevel(const FGeometryDynamicCollection& DynamicCollection, int32 TransformIndex)
 {
 	int32 Level = 0;
@@ -1905,16 +1941,18 @@ void FGeometryCollectionPhysicsProxy::BufferPhysicsResults_Internal(Chaos::FPBDR
 			TargetResults.Transforms[TransformGroupIndex] = PhysicsThreadCollection.Transform[TransformGroupIndex];
 			TargetResults.Parent[TransformGroupIndex] = PhysicsThreadCollection.Parent[TransformGroupIndex];
 
-			TargetResults.DisabledStates[TransformGroupIndex] = true;
+			TargetResults.States[TransformGroupIndex].DisabledState = true;
+			TargetResults.States[TransformGroupIndex].HasInternalClusterParent = false;
+			TargetResults.States[TransformGroupIndex].DynamicInternalClusterParent = false;
 			Chaos::FPBDRigidClusteredParticleHandle* Handle = SolverParticleHandles[TransformGroupIndex];
 			if (!Handle)
 			{
-				PhysicsThreadCollection.Active[TransformGroupIndex] = !TargetResults.DisabledStates[TransformGroupIndex];
+				PhysicsThreadCollection.Active[TransformGroupIndex] = !TargetResults.States[TransformGroupIndex].DisabledState;
 				continue;
 			}
 
 			// Dynamic state is also updated by the solver during field interaction.
-			TargetResults.DynamicState[TransformGroupIndex] = static_cast<int>(GetObjectStateFromHandle(Handle));
+			TargetResults.States[TransformGroupIndex].DynamicState = static_cast<int8>(GetObjectStateFromHandle(Handle));
 
 			// Update the transform and parent hierarchy of the active rigid bodies. Active bodies can be either
 			// rigid geometry defined from the leaf nodes of the collection, or cluster bodies that drive an entire
@@ -1948,7 +1986,9 @@ void FGeometryCollectionPhysicsProxy::BufferPhysicsResults_Internal(Chaos::FPBDR
 				PhysicsThreadCollection.Transform[TransformGroupIndex] = TargetResults.Transforms[TransformGroupIndex];
 
 				// Indicate that this object needs to be updated and the proxy is active.
-				TargetResults.DisabledStates[TransformGroupIndex] = false;
+				TargetResults.States[TransformGroupIndex].DisabledState = false;
+				TargetResults.States[TransformGroupIndex].HasInternalClusterParent = false;
+				TargetResults.States[TransformGroupIndex].DynamicInternalClusterParent = false;
 				IsObjectDynamic = true;
 
 				// If the parent of this NON DISABLED body is set to anything other than INDEX_NONE,
@@ -1994,8 +2034,8 @@ void FGeometryCollectionPhysicsProxy::BufferPhysicsResults_Internal(Chaos::FPBDR
 								PhysicsThreadCollection.Parent[TransformGroupIndex] = INDEX_NONE;
 
 								// Indicate that this object needs to be updated and the proxy is active.
-								TargetResults.DisabledStates[TransformGroupIndex] = false;
-								IsObjectDynamic                                   = true;
+								TargetResults.States[TransformGroupIndex].DisabledState = false;
+								IsObjectDynamic = true;
 							}
 							SolverClusterID[TransformGroupIndex] = Handle->ClusterIds().Id;
 						}
@@ -2020,8 +2060,10 @@ void FGeometryCollectionPhysicsProxy::BufferPhysicsResults_Internal(Chaos::FPBDR
 							PhysicsThreadCollection.Transform[TransformGroupIndex] = TargetResults.Transforms[TransformGroupIndex];
 
 							// Indicate that this object needs to be updated and the proxy is active.
-							TargetResults.DisabledStates[TransformGroupIndex] = false;
-							IsObjectDynamic                                   = true;
+							TargetResults.States[TransformGroupIndex].DisabledState = false;
+							TargetResults.States[TransformGroupIndex].HasInternalClusterParent = true;
+							TargetResults.States[TransformGroupIndex].DynamicInternalClusterParent = (ClusterParent->IsDynamic());
+							IsObjectDynamic = true;
 
 							// as we just transitioned from disabled to non disabled the update is unconditional 
 							UpdateParticleHandleTransform(*CurrentSolver, *Handle, ParticleToWorld);
@@ -2039,7 +2081,7 @@ void FGeometryCollectionPhysicsProxy::BufferPhysicsResults_Internal(Chaos::FPBDR
 				}
 			}    // end if
 
-			PhysicsThreadCollection.Active[TransformGroupIndex] = !TargetResults.DisabledStates[TransformGroupIndex];
+			PhysicsThreadCollection.Active[TransformGroupIndex] = !TargetResults.States[TransformGroupIndex].DisabledState;
 		}    // end for
 	}        // STAT_CalcParticleToWorld scope
 	
@@ -2089,6 +2131,7 @@ bool FGeometryCollectionPhysicsProxy::PullFromPhysicsState(const Chaos::FDirtyGe
 
 	TManagedArray<FVector3f>* LinearVelocities = GameThreadCollection.FindAttributeTyped<FVector3f>("LinearVelocity", FTransformCollection::TransformGroup);
 	TManagedArray<FVector3f>* AngularVelocities = GameThreadCollection.FindAttributeTyped<FVector3f>("AngularVelocity", FTransformCollection::TransformGroup);
+	TManagedArray<uint8>* InternalClusterParentTypeArray = GameThreadCollection.FindAttributeTyped<uint8>("InternalClusterParentTypeArray", FTransformCollection::TransformGroup);
 
 	// We should never be changing the number of entries, this would break other 
 	// attributes in the transform group.
@@ -2098,12 +2141,22 @@ bool FGeometryCollectionPhysicsProxy::PullFromPhysicsState(const Chaos::FDirtyGe
 		// first : copy the non interpolate-able values
 		for (int32 TransformGroupIndex = 0; TransformGroupIndex < NumTransforms; ++TransformGroupIndex)
 		{
-			if (!TargetResults.DisabledStates[TransformGroupIndex])
+			if (!TargetResults.States[TransformGroupIndex].DisabledState)
 			{
 				GameThreadCollection.Parent[TransformGroupIndex] = TargetResults.Parent[TransformGroupIndex];
 			}
-			GameThreadCollection.DynamicState[TransformGroupIndex] = TargetResults.DynamicState[TransformGroupIndex];
-			GameThreadCollection.Active[TransformGroupIndex] = !TargetResults.DisabledStates[TransformGroupIndex];
+			GameThreadCollection.DynamicState[TransformGroupIndex] = TargetResults.States[TransformGroupIndex].DynamicState;
+			GameThreadCollection.Active[TransformGroupIndex] = !TargetResults.States[TransformGroupIndex].DisabledState;
+
+			if (InternalClusterParentTypeArray)
+			{
+				Chaos::EInternalClusterType Parenttype = Chaos::EInternalClusterType::None;
+				if (TargetResults.States[TransformGroupIndex].HasInternalClusterParent != 0)
+				{
+					Parenttype = (TargetResults.States[TransformGroupIndex].DynamicInternalClusterParent != 0)? Chaos::EInternalClusterType::Dynamic: Chaos::EInternalClusterType::KinematicOrStatic; 
+				}
+				(*InternalClusterParentTypeArray)[TransformGroupIndex] =static_cast<uint8>(Parenttype); 
+			}
 		}
 
 		// second : interpolate-able ones
@@ -2114,7 +2167,7 @@ bool FGeometryCollectionPhysicsProxy::PullFromPhysicsState(const Chaos::FDirtyGe
 			const FGeometryCollectionResults& Next = NextPullData->Results;
 			for (int32 TransformGroupIndex = 0; TransformGroupIndex < NumTransforms; ++TransformGroupIndex)
 			{
-				if (!TargetResults.DisabledStates[TransformGroupIndex])
+				if (!TargetResults.States[TransformGroupIndex].DisabledState)
 				{
 					Chaos::FGeometryParticle& GTParticle = *GTParticles[TransformGroupIndex];
 					GTParticle.SetX(FMath::Lerp(Prev.ParticleXs[TransformGroupIndex], Next.ParticleXs[TransformGroupIndex], *Alpha), false);
@@ -2136,7 +2189,7 @@ bool FGeometryCollectionPhysicsProxy::PullFromPhysicsState(const Chaos::FDirtyGe
 		{
 			for (int32 TransformGroupIndex = 0; TransformGroupIndex < NumTransforms; ++TransformGroupIndex)
 			{
-				if (!TargetResults.DisabledStates[TransformGroupIndex])
+				if (!TargetResults.States[TransformGroupIndex].DisabledState)
 				{
 					Chaos::FGeometryParticle& GTParticle = *GTParticles[TransformGroupIndex];
 					GTParticle.SetX(TargetResults.ParticleXs[TransformGroupIndex], false);
diff --git a/Engine/Source/Runtime/Experimental/Chaos/Public/Chaos/GeometryParticlesfwd.h b/Engine/Source/Runtime/Experimental/Chaos/Public/Chaos/GeometryParticlesfwd.h
--- a/Engine/Source/Runtime/Experimental/Chaos/Public/Chaos/GeometryParticlesfwd.h
+++ b/Engine/Source/Runtime/Experimental/Chaos/Public/Chaos/GeometryParticlesfwd.h
@@ -98,4 +98,12 @@ namespace Chaos
 	{
 		return ::GetTypeHash(Unique.Idx);
 	}
+
+	//Used for down casting when iterating over multiple SOAs.
+	enum class EInternalClusterType : uint8
+	{
+		None,
+		KinematicOrStatic,
+		Dynamic,
+	};
 }
\ No newline at end of file
diff --git a/Engine/Source/Runtime/Experimental/Chaos/Public/GeometryCollectionProxyData.h b/Engine/Source/Runtime/Experimental/Chaos/Public/GeometryCollectionProxyData.h
--- a/Engine/Source/Runtime/Experimental/Chaos/Public/GeometryCollectionProxyData.h
+++ b/Engine/Source/Runtime/Experimental/Chaos/Public/GeometryCollectionProxyData.h
@@ -152,7 +152,7 @@ public:
 	void InitArrays(const FGeometryDynamicCollection& Other)
 	{
 		const int32 NumTransforms = Other.NumElements(FGeometryCollection::TransformGroup);
-		DisabledStates.SetNumUninitialized(NumTransforms);
+		States.SetNumUninitialized(NumTransforms);
 		GlobalTransforms.SetNumUninitialized(NumTransforms);
 		ParticleXs.SetNumUninitialized(NumTransforms);
 		ParticleRs.SetNumUninitialized(NumTransforms);
@@ -161,11 +161,19 @@ public:
 		
 		Transforms.SetNumUninitialized(NumTransforms);
 		Parent.SetNumUninitialized(NumTransforms);
-		DynamicState.SetNumUninitialized(NumTransforms);
 	}
 
+	struct FState
+	{
+		int16 DynamicState: 8; // need to fit EObjectStateTypeEnum
+		int16 DisabledState: 1;
+		int16 HasInternalClusterParent: 1;
+		int16 DynamicInternalClusterParent: 1;
+		// 6 bits left
+	};
+	
 	Chaos::FReal SolverDt;
-	TArray<bool> DisabledStates;
+	TArray<FState> States;
 	TArray<FMatrix> GlobalTransforms;
 	TArray<Chaos::FVec3> ParticleXs;
 	TArray<Chaos::FRotation3> ParticleRs;
@@ -174,7 +182,6 @@ public:
 
 	TArray<FTransform> Transforms;
 	TArray<int32> Parent;
-	TArray<int32> DynamicState;
 	
 	bool IsObjectDynamic;
 	bool IsObjectLoading;
diff --git a/Engine/Source/Runtime/Experimental/Chaos/Public/PhysicsProxy/GeometryCollectionPhysicsProxy.h b/Engine/Source/Runtime/Experimental/Chaos/Public/PhysicsProxy/GeometryCollectionPhysicsProxy.h
--- a/Engine/Source/Runtime/Experimental/Chaos/Public/PhysicsProxy/GeometryCollectionPhysicsProxy.h
+++ b/Engine/Source/Runtime/Experimental/Chaos/Public/PhysicsProxy/GeometryCollectionPhysicsProxy.h
@@ -240,7 +240,8 @@ public:
 
 	bool GetIsObjectDynamic() const { return IsObjectDynamic; }
 
-	void DisableParticles(TArray<int32>& TransformGroupIndices);
+	void DisableParticles(TArray<int32>&& TransformGroupIndices);
+	void BreakInternalClusterParent(TArray<int32>&& TransformGroupIndices);
 
 	FProxyInterpolationData& GetInterpolationData() { return InterpolationData; }
 	const FProxyInterpolationData& GetInterpolationData() const { return InterpolationData; }
diff --git a/Engine/Source/Runtime/Experimental/GeometryCollectionEngine/Private/GeometryCollection/GeometryCollectionComponent.cpp b/Engine/Source/Runtime/Experimental/GeometryCollectionEngine/Private/GeometryCollection/GeometryCollectionComponent.cpp
--- a/Engine/Source/Runtime/Experimental/GeometryCollectionEngine/Private/GeometryCollection/GeometryCollectionComponent.cpp
+++ b/Engine/Source/Runtime/Experimental/GeometryCollectionEngine/Private/GeometryCollection/GeometryCollectionComponent.cpp
@@ -77,6 +77,12 @@ static bool bChaos_BoxCalcBounds_ISPC_Enabled = CHAOS_BOX_CALC_BOUNDS_ISPC_ENABL
 static FAutoConsoleVariableRef CVarChaosBoxCalcBoundsISPCEnabled(TEXT("p.Chaos.BoxCalcBounds.ISPC"), bChaos_BoxCalcBounds_ISPC_Enabled, TEXT("Whether to use ISPC optimizations in calculating box bounds in geometry collections"));
 #endif
 
+static bool ChaosCheckIfStillForRemoval  = false;
+static FAutoConsoleVariableRef CVarChaosCheckIfStillForRemoval(TEXT("p.Chaos.CheckIfStillForRemoval"), ChaosCheckIfStillForRemoval, TEXT("if true, non moving GC ( even non sleeping ) will be considered for removal if the feature is enabled"));
+
+static float ChaosStillCheckVelocityThreshold  = 1; // 1 cm/s
+static FAutoConsoleVariableRef CVarChaosStillCheckDistanceThreshold(TEXT("p.Chaos.ChaosStillCheckVelocityThreshold"), ChaosStillCheckVelocityThreshold, TEXT("When using ChaosStillCheckVelocityThreshold, this set the velocity threshold in cm/s "));
+
 DEFINE_LOG_CATEGORY_STATIC(UGCC_LOG, Error, All);
 
 extern FGeometryCollectionDynamicDataPool GDynamicDataPool;
@@ -1122,6 +1128,7 @@ void UGeometryCollectionComponent::InitializeComponent()
 				DynamicCollection->AddAttribute<FVector3f>("AngularVelocity", FTransformCollection::TransformGroup);
 			}
 		}
+		DynamicCollection->AddAttribute<uint8>("InternalClusterParentTypeArray", FTransformCollection::TransformGroup);
 	}
 
 	AActor* Owner = GetOwner();
@@ -2034,6 +2041,18 @@ void UGeometryCollectionComponent::ResetDynamicCollection()
 				SleepTimer.Fill(0.f);
 			}
 
+			if (!DynamicCollection->HasAttribute("Decay", FGeometryCollection::TransformGroup))
+			{
+				TManagedArray<float>& Decay = DynamicCollection->AddAttribute<float>("Decay", FGeometryCollection::TransformGroup);
+				Decay.Fill(0);
+			}
+
+			if (!DynamicCollection->HasAttribute("LastPosition", FGeometryCollection::TransformGroup))
+			{
+				TManagedArray<FVector>& LastPosition = DynamicCollection->AddAttribute<FVector>("LastPosition", FGeometryCollection::TransformGroup);
+				LastPosition.Fill(FVector::Zero());
+			}
+			
 			if (!DynamicCollection->HasAttribute("UniformScale", FGeometryCollection::TransformGroup))
 			{
 				TManagedArray<FTransform>& UniformScale = DynamicCollection->AddAttribute<FTransform>("UniformScale", FGeometryCollection::TransformGroup);
@@ -3206,37 +3225,46 @@ void UGeometryCollectionComponent::CalculateGlobalMatrices()
 			// Have to fully rebuild
 			if (DynamicCollection 
 				&& RestCollection->bRemoveOnMaxSleep 
-				&& DynamicCollection->HasAttribute("SleepTimer", FGeometryCollection::TransformGroup) 
 				&& DynamicCollection->HasAttribute("UniformScale", FGeometryCollection::TransformGroup)
-				&& DynamicCollection->HasAttribute("MaxSleepTime", FGeometryCollection::TransformGroup)
-				&& DynamicCollection->HasAttribute("RemovalDuration", FGeometryCollection::TransformGroup))
+				&& DynamicCollection->HasAttribute("Decay", FGeometryCollection::TransformGroup))
 			{
-				const TManagedArray<float>& SleepTimer = DynamicCollection->GetAttribute<float>("SleepTimer", FGeometryCollection::TransformGroup);
-				const TManagedArray<float>& MaxSleepTime = DynamicCollection->GetAttribute<float>("MaxSleepTime", FGeometryCollection::TransformGroup);
-				const TManagedArray<float>& RemovalDuration = DynamicCollection->GetAttribute<float>("RemovalDuration", FGeometryCollection::TransformGroup);
+				const TManagedArray<float>& Decay = DynamicCollection->GetAttribute<float>("Decay", FGeometryCollection::TransformGroup);
 				TManagedArray<FTransform>& UniformScale = DynamicCollection->ModifyAttribute<FTransform>("UniformScale", FGeometryCollection::TransformGroup);
+				const TManagedArray<FBox>& BoundingBoxes = RestCollection->GetGeometryCollection()->BoundingBox;
+				const TManagedArray<int32>& TransformToGeom = RestCollection->GetGeometryCollection()->TransformToGeometryIndex;
 
+				const FTransform InverseComponentTransform = GetComponentTransform().Inverse();
+				const FTransform ZeroScaleTransform(FQuat::Identity, FVector::Zero(), FVector(0, 0, 0));
 				for (int32 Idx = 0; Idx < GetTransformArray().Num(); ++Idx)
 				{
-					if (SleepTimer[Idx] > MaxSleepTime[Idx])
+					// only update values if the decay has changed 
+					if (Decay[Idx] > 0.f && Decay[Idx] <= 1.f)
 					{
-						float Scale = 1.0 - FMath::Min(1.0, (SleepTimer[Idx] - MaxSleepTime[Idx]) / RemovalDuration[Idx]);
-						
-						if ((Scale < 1.0) && (Scale > 0.0))
+						const float Scale = 1.0 - Decay[Idx];
+						if (Scale < UE_SMALL_NUMBER)
+						{
+							UniformScale[Idx] = ZeroScaleTransform;
+						}
+						else
 						{
 							float ShrinkRadius = 0.0f;
 							FSphere AccumulatedSphere;
+							// todo(chaos) : find a faster way to do that ( precompute the data ? )
 							if (CalculateInnerSphere(Idx, AccumulatedSphere))
 							{
 								ShrinkRadius = -AccumulatedSphere.W;
 							}
-							
-							FQuat LocalRotation = (GetComponentTransform().Inverse() * FTransform(GlobalMatrices[Idx]).Inverse()).GetRotation();
-							FTransform LocalDown(LocalRotation.RotateVector(FVector(0.f, 0.f, ShrinkRadius)));
-							FTransform ToCOM(DynamicCollection->MassToLocal[Idx].GetTranslation());
-							UniformScale[Idx] = ToCOM.Inverse() * LocalDown.Inverse() * FTransform(FQuat::Identity, FVector(0.f, 0.f, 0.f), FVector(Scale)) * LocalDown * ToCOM;
+						
+							const FQuat LocalRotation = (InverseComponentTransform * FTransform(GlobalMatrices[Idx]).Inverse()).GetRotation();
+							const FVector LocalDown = LocalRotation.RotateVector(FVector(0.f, 0.f, ShrinkRadius));
+							const FVector CenterOfMass = DynamicCollection->MassToLocal[Idx].GetTranslation();
+							const FVector ScaleCenter = LocalDown + CenterOfMass;
+							//const FTransform ToScaleCenter(ScaleCenter);
+							//const FTransform ToScaleCenterInverse(-ScaleCenter);
+							//UniformScale[Idx] = ToCOM.Inverse() * LocalDown.Inverse() * FTransform(FQuat::Identity, FVector(0.f, 0.f, 0.f), FVector(Scale)) * LocalDown * ToCOM;
+							//UniformScale[Idx] = ToScaleCenterInverse * FTransform(FQuat::Identity, FVector(0.f, 0.f, 0.f), FVector(Scale)) * ToScaleCenter;
+							UniformScale[Idx] = FTransform(FQuat::Identity, ScaleCenter * (FVector::FReal)(1.f - Scale), FVector(Scale));
 						}
-	
 					}
 				}
 				
@@ -3459,38 +3487,85 @@ void UGeometryCollectionComponent::InitializeEmbeddedGeometry()
 
 void UGeometryCollectionComponent::IncrementSleepTimer(float DeltaTime)
 {
+	if (DeltaTime <= 0)
+	{
+		return;
+	}
+	
 	// If a particle is sleeping, increment its sleep timer, otherwise reset it.
 	if (DynamicCollection && PhysicsProxy 
 		&& DynamicCollection->HasAttribute("SleepTimer", FGeometryCollection::TransformGroup) 
 		&& DynamicCollection->HasAttribute("MaxSleepTime", FGeometryCollection::TransformGroup)
-		&& DynamicCollection->HasAttribute("RemovalDuration", FGeometryCollection::TransformGroup))
+		&& DynamicCollection->HasAttribute("RemovalDuration", FGeometryCollection::TransformGroup)
+		&& DynamicCollection->HasAttribute("Decay", FGeometryCollection::TransformGroup)
+		&& DynamicCollection->HasAttribute("LastPosition", FGeometryCollection::TransformGroup))
 	{
 		TManagedArray<float>& SleepTimer = DynamicCollection->ModifyAttribute<float>("SleepTimer", FGeometryCollection::TransformGroup);
 		const TManagedArray<float>& RemovalDuration = DynamicCollection->GetAttribute<float>("RemovalDuration", FGeometryCollection::TransformGroup);
 		const TManagedArray<float>& MaxSleepTime = DynamicCollection->GetAttribute<float>("MaxSleepTime", FGeometryCollection::TransformGroup);
+		TManagedArray<float>& Decay = DynamicCollection->ModifyAttribute<float>("Decay", FGeometryCollection::TransformGroup);
+		TManagedArray<FVector>& LastPosition = DynamicCollection->ModifyAttribute<FVector>("LastPosition", FGeometryCollection::TransformGroup);
+		const TManagedArray<uint8>& InternalClusterParentType = DynamicCollection->GetAttribute<uint8>("InternalClusterParentTypeArray", FGeometryCollection::TransformGroup);
+		const TManagedArray<bool>& Active = DynamicCollection->Active;
+		const TManagedArray<int32>& Parents = DynamicCollection->Parent;
+		
 		TArray<int32> ToDisable;
+		TArray<int32> ToBreakParent;
 		for (int32 TransformIdx = 0; TransformIdx < SleepTimer.Num(); ++TransformIdx)
 		{
-			bool PreviouslyAwake = SleepTimer[TransformIdx] < MaxSleepTime[TransformIdx];
-			if (SleepTimer[TransformIdx] < (MaxSleepTime[TransformIdx] + RemovalDuration[TransformIdx]))
-			{
-				SleepTimer[TransformIdx] = (DynamicCollection->DynamicState[TransformIdx] == (int)EObjectStateTypeEnum::Chaos_Object_Sleeping) ? SleepTimer[TransformIdx] + DeltaTime : 0.0f;
+			// update the sleeping timer accordingly
+			const int32 State = DynamicCollection->DynamicState[TransformIdx];
+			const bool bIsDynamicOrSleeping = (State == (int)EObjectStateTypeEnum::Chaos_Object_Sleeping) || (State == (int)EObjectStateTypeEnum::Chaos_Object_Dynamic);
+			const bool bNoParent = Parents[TransformIdx] == INDEX_NONE;
+			const bool bHasKinematicOrStaticInternalClusterParent = (InternalClusterParentType[TransformIdx] == (uint8)Chaos::EInternalClusterType::KinematicOrStatic);
+			const bool IsRemovalActive = (MaxSleepTime[TransformIdx] >= 0) && bIsDynamicOrSleeping && bNoParent && !bHasKinematicOrStaticInternalClusterParent;
+			if (IsRemovalActive)
+			{
+				const FVector CurrentPosition = DynamicCollection->Transform[TransformIdx].GetTranslation();
+				const FVector::FReal InstantVelocity = (CurrentPosition-LastPosition[TransformIdx]).Size() / (FVector::FReal)DeltaTime; 
+				const bool bIsStill = ChaosCheckIfStillForRemoval && (InstantVelocity < ChaosStillCheckVelocityThreshold);
+				LastPosition[TransformIdx] = CurrentPosition;
+				
+				const bool IsSleeping = (DynamicCollection->DynamicState[TransformIdx] == (int)EObjectStateTypeEnum::Chaos_Object_Sleeping);
+				const bool HasStartedDecaying = (Decay[TransformIdx] > 0); 
+				if (IsSleeping || bIsStill || HasStartedDecaying)
+				{
+					SleepTimer[TransformIdx] += DeltaTime;
+				}
+			
+				// update the decay and disable the particle when decay has completed 
+				const bool bZeroRemovalDuration = (RemovalDuration[TransformIdx] < UE_SMALL_NUMBER);
+				const float UpdatedDecay = bZeroRemovalDuration? 1.f: FMath::Clamp<float>((SleepTimer[TransformIdx] - MaxSleepTime[TransformIdx]) / RemovalDuration[TransformIdx], 0.f, 1.f);
+				if (UpdatedDecay != Decay[TransformIdx])
+				{
+					ensure(UpdatedDecay > Decay[TransformIdx]);
+
+					// if parent is internal and dynamic we need to break it before decaying 
+					if (InternalClusterParentType[TransformIdx] == (uint8)Chaos::EInternalClusterType::Dynamic)
+					{
+						ToBreakParent.AddUnique(TransformIdx);
+					}
 
-				if (SleepTimer[TransformIdx] > MaxSleepTime[TransformIdx])
-				{ 
 					DynamicCollection->MakeDirty();
-					if (PreviouslyAwake)
+					Decay[TransformIdx] = UpdatedDecay;
+
+					if (Decay[TransformIdx] >= 1.0f)
 					{
 						// Disable the particle if it has been asleep for the requisite time
+						Decay[TransformIdx] = 1.0f;
 						ToDisable.Add(TransformIdx);
 					}
 				}
 			}
 		}
 
+		if (ToBreakParent.Num())
+		{
+			PhysicsProxy->BreakInternalClusterParent(MoveTemp(ToBreakParent));
+		}
 		if (ToDisable.Num())
 		{
-			PhysicsProxy->DisableParticles(ToDisable);	
+			PhysicsProxy->DisableParticles(MoveTemp(ToDisable));
 		}
 	}
 }
