Chaos : fix inertia scaling logic for scaled geometry collection

diff --git a/Engine/Source/Runtime/Experimental/Chaos/Private/PhysicsProxy/GeometryCollectionPhysicsProxy.cpp b/Engine/Source/Runtime/Experimental/Chaos/Private/PhysicsProxy/GeometryCollectionPhysicsProxy.cpp
--- a/Engine/Source/Runtime/Experimental/Chaos/Private/PhysicsProxy/GeometryCollectionPhysicsProxy.cpp
+++ b/Engine/Source/Runtime/Experimental/Chaos/Private/PhysicsProxy/GeometryCollectionPhysicsProxy.cpp
@@ -221,6 +221,10 @@ void PopulateSimulatedParticle(
 
 		Handle->SetM(MassIn);
 		Handle->SetI(InertiaTensorVec);
+		const Chaos::FReal MassInv = (MassIn > 0.0f) ? 1.0f / MassIn : 0.0f;
+		const Chaos::FVec3 InertiaInv = (MassIn > 0.0f) ? Chaos::FVec3(InertiaTensorVec).Reciprocal() : Chaos::FVec3::ZeroVector;
+		Handle->SetInvM(MassInv);
+		Handle->SetInvI(InertiaInv);
 		Handle->SetObjectStateLowLevel(Chaos::EObjectStateType::Dynamic); // this step sets InvM, InvInertia, P, Q
 	}
 
@@ -827,6 +831,8 @@ void FGeometryCollectionPhysicsProxy::InitializeBodiesPT(Chaos::FPBDRigidsSolver
 				// Mass space -> Composed parent space -> world
 				const FTransform WorldTransform = 
 					MassToLocal[TransformGroupIndex] * Transform[TransformGroupIndex] * Parameters.WorldTransform;
+
+				const Chaos::FVec3f ScaledInertia = Chaos::Utilities::ScaleInertia<float>((Chaos::FVec3f)InertiaTensor[TransformGroupIndex], (Chaos::FVec3f)(WorldScale), true);
 				
 				PopulateSimulatedParticle(
 					Handle,
@@ -836,7 +842,7 @@ void FGeometryCollectionPhysicsProxy::InitializeBodiesPT(Chaos::FPBDRigidsSolver
 					SimFilter,
 					QueryFilter,
 					Mass[TransformGroupIndex] * MassScale,
-					InertiaTensor[TransformGroupIndex] * MassScale,
+					ScaledInertia,
 					WorldTransform,
 					static_cast<uint8>(DynamicState[TransformGroupIndex]),
 					static_cast<int16>(CollisionGroup[TransformGroupIndex]),
@@ -1237,6 +1243,7 @@ FGeometryCollectionPhysicsProxy::BuildClusters_Internal(
 
 	const FVector WorldScale = Parameters.WorldTransform.GetScale3D();
 	const FVector::FReal MassScale = WorldScale.X * WorldScale.Y * WorldScale.Z;
+	const Chaos::FVec3f ScaledInertia = Chaos::Utilities::ScaleInertia<float>((Chaos::FVec3f)InertiaTensor[CollectionClusterIndex], FVector3f(WorldScale), true);  
 	
 	PopulateSimulatedParticle(
 		Parent,
@@ -1246,7 +1253,7 @@ FGeometryCollectionPhysicsProxy::BuildClusters_Internal(
 		SimFilter,
 		QueryFilter,
 		Parent->M() > 0.0 ? Parent->M() : Mass[CollectionClusterIndex] * MassScale, 
-		Parent->I() != Chaos::TVec3<Chaos::FRealSingle>(0.0) ? Parent->I() : Chaos::TVec3<Chaos::FRealSingle>(InertiaTensor[CollectionClusterIndex] * MassScale),
+		Parent->I() != Chaos::FVec3f(0.0) ? Parent->I() : ScaledInertia,
 		ParticleTM, 
 		(uint8)DynamicState[CollectionClusterIndex], 
 		0,
diff --git a/Engine/Source/Runtime/Experimental/Chaos/Public/Chaos/Utilities.h b/Engine/Source/Runtime/Experimental/Chaos/Public/Chaos/Utilities.h
--- a/Engine/Source/Runtime/Experimental/Chaos/Public/Chaos/Utilities.h
+++ b/Engine/Source/Runtime/Experimental/Chaos/Public/Chaos/Utilities.h
@@ -449,14 +449,15 @@ namespace Chaos
 		 * Given the local-space diagonal inertia for an unscaled object, return an inertia as if generated from a non-uniformly scaled shape with the specified scale.
 		 * If bScaleMass is true, it also takes into account the fact that the mass would have changed by the increase in volume.
 		 */
-		inline FVec3 ScaleInertia(const FVec3& Inertia, const FVec3& Scale, const bool bScaleMass)
-		{
-			FVec3 XYZSq = (FVec3(0.5f * (Inertia.X + Inertia.Y + Inertia.Z)) - Inertia) * Scale * Scale;
-			FReal XX = XYZSq.Y + XYZSq.Z;
-			FReal YY = XYZSq.X + XYZSq.Z;
-			FReal ZZ = XYZSq.X + XYZSq.Y;
-			FVec3 ScaledInertia = FVec3(XX, YY, ZZ);
-			FReal MassScale = (bScaleMass) ? Scale.X * Scale.Y * Scale.Z : 1.0f;
+		template<typename T>
+		inline TVec3<T> ScaleInertia(const TVec3<T>& Inertia, const TVec3<T>& Scale, const bool bScaleMass)
+		{
+			TVec3<T> XYZSq = (TVec3<T>(0.5f * (Inertia.X + Inertia.Y + Inertia.Z)) - Inertia) * Scale * Scale;
+			T XX = XYZSq.Y + XYZSq.Z;
+			T YY = XYZSq.X + XYZSq.Z;
+			T ZZ = XYZSq.X + XYZSq.Y;
+			TVec3<T> ScaledInertia = TVec3<T>(XX, YY, ZZ);
+			T MassScale = (bScaleMass) ? Scale.X * Scale.Y * Scale.Z : 1.0f;
 			return MassScale * ScaledInertia;
 		}
 
