Constraints in Sequencer : Improved constraints baking and properties changes

diff --git a/Engine/Plugins/Animation/ControlRig/Source/ControlRigEditor/Private/Tools/BakingHelper.cpp b/Engine/Plugins/Animation/ControlRig/Source/ControlRigEditor/Private/Tools/BakingHelper.cpp
--- a/Engine/Plugins/Animation/ControlRig/Source/ControlRigEditor/Private/Tools/BakingHelper.cpp
+++ b/Engine/Plugins/Animation/ControlRig/Source/ControlRigEditor/Private/Tools/BakingHelper.cpp
@@ -56,7 +56,10 @@
 	}
 }
 
-UMovieScene3DTransformSection* FBakingHelper::GetTransformSection(const ISequencer* InSequencer, const FGuid& InGuid)
+UMovieScene3DTransformSection* FBakingHelper::GetTransformSection(
+	const ISequencer* InSequencer,
+	const FGuid& InGuid,
+	const FTransform& InDefaultTransform)
 {
 	if (!InSequencer || !InSequencer->GetFocusedMovieSceneSequence())
 	{
@@ -90,6 +93,22 @@
 	if (bSectionAdded)
 	{
 		TransformSection->SetRange(TRange<FFrameNumber>::All());
+
+		const FVector Location0 = InDefaultTransform.GetLocation();
+		const FRotator Rotation0 = InDefaultTransform.GetRotation().Rotator();
+		const FVector Scale3D0 = InDefaultTransform.GetScale3D();
+
+		const TArrayView<FMovieSceneDoubleChannel*> Channels =
+			TransformSection->GetChannelProxy().GetChannels<FMovieSceneDoubleChannel>();
+		Channels[0]->SetDefault(Location0.X);
+		Channels[1]->SetDefault(Location0.Y);
+		Channels[2]->SetDefault(Location0.Z);
+		Channels[3]->SetDefault(Rotation0.Roll);
+		Channels[4]->SetDefault(Rotation0.Pitch);
+		Channels[5]->SetDefault(Rotation0.Yaw);
+		Channels[6]->SetDefault(Scale3D0.X);
+		Channels[7]->SetDefault(Scale3D0.Y);
+		Channels[8]->SetDefault(Scale3D0.Z);
 	}
 	
 	return TransformSection;
@@ -161,6 +180,22 @@
 	const TArrayView<FMovieSceneDoubleChannel*> Channels =
 		InTransformSection->GetChannelProxy().GetChannels<FMovieSceneDoubleChannel>();
 
+	// set default
+	const FTransform& LocalTransform0 = InLocalTransforms[0];
+	const FVector Location0 = LocalTransform0.GetLocation();
+	const FRotator Rotation0 = LocalTransform0.GetRotation().Rotator();
+	const FVector Scale3D0 = LocalTransform0.GetScale3D();
+	
+	for (int32 ChannelIndex = 0; ChannelIndex < 9; ChannelIndex++)
+	{
+		if (!Channels[ChannelIndex]->GetDefault().IsSet())
+		{
+			const double Value = GetValue(ChannelIndex, Location0, Rotation0, Scale3D0);
+			Channels[ChannelIndex]->SetDefault(Value);
+		}
+	}
+
+	// add keys
 	for (int32 Index = 0; Index < Frames.Num(); ++Index)
 	{
 		const FFrameNumber& Frame = Frames[Index];
@@ -172,16 +207,7 @@
 
 		for (const int32 ChannelIndex: ChannelsIndexToKey)
 		{
-			// get value
 			const double Value = GetValue(ChannelIndex, Location, Rotation, Scale3D);
-
-			// set default
-			if (Index == 0 && !Channels[ChannelIndex]->GetDefault().IsSet())
-			{
-				Channels[ChannelIndex]->SetDefault(Value);
-			}
-			
-			// add key
 			TMovieSceneChannelData<FMovieSceneDoubleValue> ChannelData = Channels[ChannelIndex]->GetData();
 			MovieSceneToolHelpers::SetOrAddKey(ChannelData, Frame, Value);
 		}
diff --git a/Engine/Plugins/Animation/ControlRig/Source/ControlRigEditor/Private/Tools/ConstraintBaker.cpp b/Engine/Plugins/Animation/ControlRig/Source/ControlRigEditor/Private/Tools/ConstraintBaker.cpp
--- a/Engine/Plugins/Animation/ControlRig/Source/ControlRigEditor/Private/Tools/ConstraintBaker.cpp
+++ b/Engine/Plugins/Animation/ControlRig/Source/ControlRigEditor/Private/Tools/ConstraintBaker.cpp
@@ -21,7 +21,10 @@
 namespace
 {
 	
-UMovieScene3DTransformSection* GetTransformSection(const TSharedPtr<ISequencer>& InSequencer, AActor* InActor)
+UMovieScene3DTransformSection* GetTransformSection(
+	const TSharedPtr<ISequencer>& InSequencer,
+	AActor* InActor,
+	const FTransform& InTransform0)
 {
 	if (!InSequencer || !InSequencer->GetFocusedMovieSceneSequence())
 	{
@@ -34,7 +37,7 @@
 		return nullptr;
 	}
 	
-	return FBakingHelper::GetTransformSection(InSequencer.Get(), Guid);
+	return FBakingHelper::GetTransformSection(InSequencer.Get(), Guid, InTransform0);
 }
 
 void BakeComponent(
@@ -44,6 +47,8 @@
 	const TArray<FTransform>& InTransforms,
 	const EMovieSceneTransformChannel& InChannels)
 {
+	ensure(InTransforms.Num());
+	
 	if (!InComponentHandle->IsValid())
 	{
 		return;
@@ -53,7 +58,7 @@
 	{
 		return;
 	}
-	const UMovieScene3DTransformSection* TransformSection = GetTransformSection(InSequencer, Actor);
+	const UMovieScene3DTransformSection* TransformSection = GetTransformSection(InSequencer, Actor, InTransforms[0]);
 	if (!TransformSection)
 	{
 		return;
@@ -126,9 +131,19 @@
 	TArray<FFrameNumber> Frames;
 	FBakingHelper::CalculateFramesBetween(MovieScene, StartFrame, EndFrame, Frames);
 
+	if (Frames.IsEmpty())
+	{
+		return;
+	}
+
 	// compute transforms
 	TArray<FTransform> Transforms;
 	GetChildTransforms(Sequencer, *InConstraint, Frames, true, Transforms);
+	
+	if (Frames.Num() != Transforms.Num())
+	{
+		return;
+	}
 
 	// bake to channel curves
 	const EMovieSceneTransformChannel Channels = GetChannelsToKey(InConstraint);
diff --git a/Engine/Plugins/Animation/ControlRig/Source/ControlRigEditor/Public/Tools/BakingHelper.h b/Engine/Plugins/Animation/ControlRig/Source/ControlRigEditor/Public/Tools/BakingHelper.h
--- a/Engine/Plugins/Animation/ControlRig/Source/ControlRigEditor/Public/Tools/BakingHelper.h
+++ b/Engine/Plugins/Animation/ControlRig/Source/ControlRigEditor/Public/Tools/BakingHelper.h
@@ -33,7 +33,8 @@
 	/** Returns the transform section for that guid. */
 	static UMovieScene3DTransformSection* GetTransformSection(
 		const ISequencer* InSequencer,
-		const FGuid& InGuid);
+		const FGuid& InGuid,
+		const FTransform& InDefaultTransform = FTransform::Identity);
 
 	/** Adds transform keys to the section based on the channels filters. */
 	static bool AddTransformKeys(
diff --git a/Engine/Source/Runtime/Experimental/Animation/Constraints/Private/TransformConstraint.cpp b/Engine/Source/Runtime/Experimental/Animation/Constraints/Private/TransformConstraint.cpp
--- a/Engine/Source/Runtime/Experimental/Animation/Constraints/Private/TransformConstraint.cpp
+++ b/Engine/Source/Runtime/Experimental/Animation/Constraints/Private/TransformConstraint.cpp
@@ -22,6 +22,28 @@
 {
 	return ParentTRSHandle->IsValid() ? ParentTRSHandle->GetName() : NAME_None;
 }
+
+void UTickableTransformConstraint::PostEditChangeProperty(FPropertyChangedEvent& PropertyChangedEvent)
+{
+	Super::PostEditChangeProperty(PropertyChangedEvent);
+
+	const FName PropertyName = PropertyChangedEvent.GetPropertyName();
+	if (PropertyName == GET_MEMBER_NAME_CHECKED(UTickableTransformConstraint, bMaintainOffset))
+	{
+		Evaluate();
+		return;
+	}
+
+	if (const FProperty* MemberProperty = PropertyChangedEvent.MemberProperty)
+	{
+		static const FString OffsetStr("Offset");
+		if (MemberProperty->GetFName().ToString().Contains(OffsetStr) )
+		{
+			Evaluate();
+		}
+	}
+}
+
 #endif
 
 void UTickableTransformConstraint::Setup()
@@ -137,15 +159,20 @@
 {
 	return [this]()
 	{
+		if (!Active)
+		{
+			return;
+		}
+		
 		const float ClampedWeight = FMath::Clamp<float>(Weight, 0.f, 1.f);
 		if (ClampedWeight < KINDA_SMALL_NUMBER)
 		{
 			return;
 		}
 
-		const FTransform ParentTransform = GetParentGlobalTransform();
+		const FVector ParentTranslation = GetParentGlobalTransform().GetLocation();
 		FTransform Transform = GetChildGlobalTransform();
-		FVector NewTranslation = ParentTransform.GetLocation() + OffsetTranslation;
+		FVector NewTranslation = bMaintainOffset ? ParentTranslation + OffsetTranslation : ParentTranslation;
 		if (ClampedWeight < 1.0f - KINDA_SMALL_NUMBER)
 		{
 			NewTranslation = FMath::Lerp<FVector>(Transform.GetLocation(), NewTranslation, ClampedWeight);
@@ -182,16 +209,21 @@
 {
 	return [this]()
 	{
+		if (!Active)
+		{
+			return;
+		}
+		
 		const float ClampedWeight = FMath::Clamp<float>(Weight, 0.f, 1.f);
 		if (ClampedWeight < KINDA_SMALL_NUMBER)
 		{
 			return;
 		}
 		
-		const FTransform ParentTransform = GetParentGlobalTransform();
+		const FQuat ParentRotation = GetParentGlobalTransform().GetRotation();
 		FTransform Transform = GetChildGlobalTransform();
 
-		FQuat NewRotation = ParentTransform.GetRotation() * OffsetRotation;
+		FQuat NewRotation = bMaintainOffset ? ParentRotation * OffsetRotation : ParentRotation;
 		if (ClampedWeight < 1.0f - KINDA_SMALL_NUMBER)
 		{
 			NewRotation = FQuat::Slerp(Transform.GetRotation(), NewRotation, ClampedWeight);
@@ -231,15 +263,20 @@
 {
 	return [this]()
 	{
+		if (!Active)
+		{
+			return;
+		}
+		
 		const float ClampedWeight = FMath::Clamp<float>(Weight, 0.f, 1.f);
 		if (ClampedWeight < KINDA_SMALL_NUMBER)
 		{
 			return;
 		}
 		
-		const FTransform ParentTransform = GetParentGlobalTransform();
+		const FVector ParentScale = GetParentGlobalTransform().GetScale3D();
 		FTransform Transform = GetChildGlobalTransform();
-		FVector NewScale = ParentTransform.GetScale3D() * OffsetScale;
+		FVector NewScale = bMaintainOffset ? ParentScale * OffsetScale : ParentScale;
 		if (ClampedWeight < 1.0f - KINDA_SMALL_NUMBER)
 		{
 			NewScale = FMath::Lerp<FVector>(Transform.GetScale3D(), NewScale, ClampedWeight);
@@ -275,6 +312,11 @@
 {
 	return [this]()
 	{
+		if (!Active)
+		{
+			return;
+		}
+		
 		const float ClampedWeight = FMath::Clamp<float>(Weight, 0.f, 1.f);
 		if (ClampedWeight < KINDA_SMALL_NUMBER)
 		{
@@ -283,7 +325,7 @@
 		
 		const FTransform ParentTransform = ParentTRSHandle->GetGlobalTransform();
 
-		FTransform TargetTransform = OffsetTransform * ParentTransform;
+		FTransform TargetTransform = bMaintainOffset ? OffsetTransform * ParentTransform : ParentTransform;
 		if (ClampedWeight < 1.0f - KINDA_SMALL_NUMBER)
 		{
 			const  FTransform Transform = GetChildGlobalTransform();
@@ -318,6 +360,11 @@
 	// @todo handle weight here
 	return [this]()
 	{
+		if (!Active)
+		{
+			return;
+		}
+		
 		const FTransform ParentTransform = GetParentGlobalTransform();
 		const FTransform ChildTransform = GetChildGlobalTransform();
 		
diff --git a/Engine/Source/Runtime/Experimental/Animation/Constraints/Public/TransformConstraint.h b/Engine/Source/Runtime/Experimental/Animation/Constraints/Public/TransformConstraint.h
--- a/Engine/Source/Runtime/Experimental/Animation/Constraints/Public/TransformConstraint.h
+++ b/Engine/Source/Runtime/Experimental/Animation/Constraints/Public/TransformConstraint.h
@@ -82,6 +82,10 @@
 public:
 	/** @todo document */
 	virtual FName GetLabel() const override;
+
+	// UObject interface
+	virtual void PostEditChangeProperty(FPropertyChangedEvent& PropertyChangedEvent) override;
+	// End of UObject interface
 #endif
 };
 
