fix crash in Geometry Collection Proxy when we try to use GetPrimitiveSceneInfo() to get the general hitproxy color for the collection, but it returns null

diff --git a/Engine/Source/Runtime/Experimental/GeometryCollectionEngine/Private/GeometryCollection/GeometryCollectionSceneProxy.cpp b/Engine/Source/Runtime/Experimental/GeometryCollectionEngine/Private/GeometryCollection/GeometryCollectionSceneProxy.cpp
--- a/Engine/Source/Runtime/Experimental/GeometryCollectionEngine/Private/GeometryCollection/GeometryCollectionSceneProxy.cpp
+++ b/Engine/Source/Runtime/Experimental/GeometryCollectionEngine/Private/GeometryCollection/GeometryCollectionSceneProxy.cpp
@@ -508,7 +508,6 @@ void FGeometryCollectionSceneProxy::SetConstantData_RenderThread(FGeometryCollec
 
 	if (GetRequiredVertexCount())
 	{
-		FColor WholeObjectHitProxyColor = GetPrimitiveSceneInfo()->DefaultDynamicHitProxyId.GetColor();
 		ParallelFor(Vertices.Num(), [&](int32 i)
 		{
 			const FDynamicMeshVertex& Vertex = Vertices[i];
@@ -1191,6 +1190,7 @@ HHitProxy* FGeometryCollectionSceneProxy::CreateHitProxies(UPrimitiveComponent*
 	// In order to be able to click on static meshes when they're batched up, we need to have catch all default
 	// hit proxy to return.
 	HHitProxy* DefaultHitProxy = FPrimitiveSceneProxy::CreateHitProxies(Component, OutHitProxies);
+	WholeObjectHitProxyColor = DefaultHitProxy->Id.GetColor();
 
 	// @todo FractureTools - Reconcile with subsection hit proxies.  Subsection is a draw call per hit proxy but is not suitable per-vertex as written
 	if (bEnableBoneSelection)
diff --git a/Engine/Source/Runtime/Experimental/GeometryCollectionEngine/Private/GeometryCollection/GeometryCollectionSceneProxy.h b/Engine/Source/Runtime/Experimental/GeometryCollectionEngine/Private/GeometryCollection/GeometryCollectionSceneProxy.h
--- a/Engine/Source/Runtime/Experimental/GeometryCollectionEngine/Private/GeometryCollection/GeometryCollectionSceneProxy.h
+++ b/Engine/Source/Runtime/Experimental/GeometryCollectionEngine/Private/GeometryCollection/GeometryCollectionSceneProxy.h
@@ -249,6 +249,7 @@ class FGeometryCollectionSceneProxy final : public FPrimitiveSceneProxy
 	TMap<int32, int32> SubSectionHitProxyIndexMap;
 	// @todo FractureTools - Reconcile with SubSectionHitProxies.  Currently subsection hit proxies dont work for per-vertex submission
 	TArray<TRefCountPtr<HGeometryCollectionBone>> PerBoneHitProxies;
+	FColor WholeObjectHitProxyColor = FColor(EForceInit::ForceInit);
 	bool bUsesSubSections;
 #endif
 
