Chaos destruction : make geometry collection use the inertia conditioning option

diff --git a/Engine/Source/Runtime/Experimental/Chaos/Private/PhysicsProxy/GeometryCollectionPhysicsProxy.cpp b/Engine/Source/Runtime/Experimental/Chaos/Private/PhysicsProxy/GeometryCollectionPhysicsProxy.cpp
--- a/Engine/Source/Runtime/Experimental/Chaos/Private/PhysicsProxy/GeometryCollectionPhysicsProxy.cpp
+++ b/Engine/Source/Runtime/Experimental/Chaos/Private/PhysicsProxy/GeometryCollectionPhysicsProxy.cpp
@@ -756,6 +756,18 @@ void FGeometryCollectionPhysicsProxy::SetGravityEnabled(const Chaos::FPBDRigidsS
 	}
 }
 
+void FGeometryCollectionPhysicsProxy::EnableInertiaConditioning(const Chaos::FPBDRigidsSolver& RigidsSolver, bool bEnabled)
+{
+	Chaos::FPerParticleGravity& GravityForces = RigidsSolver.GetEvolution()->GetGravityForces();
+	for (int32 HandleIdx = 0; HandleIdx < SolverParticleHandles.Num(); ++HandleIdx)
+	{
+		if (Chaos::TPBDRigidParticleHandle<Chaos::FReal, 3>* Handle = SolverParticleHandles[HandleIdx])
+		{
+			Handle->SetInertiaConditioningEnabled(bEnabled);
+		}
+	}
+}
+
 void FGeometryCollectionPhysicsProxy::SetSleepingState(const Chaos::FPBDRigidsSolver& RigidsSolver)
 {
 	for (FClusterHandle* Handle: SolverParticleHandles)
@@ -1112,6 +1124,8 @@ void FGeometryCollectionPhysicsProxy::InitializeBodiesPT(Chaos::FPBDRigidsSolver
 			SetSleepingState(*RigidsSolver);
 		}
 
+		EnableInertiaConditioning(*RigidsSolver, Parameters.UseInertiaConditioning);
+
 		// call DirtyParticle to make sure the acceleration structure is up to date with all the changes happening here
 		DirtyAllParticles(*RigidsSolver);
 
diff --git a/Engine/Source/Runtime/Experimental/Chaos/Public/GeometryCollection/GeometryCollectionSimulationCoreTypes.h b/Engine/Source/Runtime/Experimental/Chaos/Public/GeometryCollection/GeometryCollectionSimulationCoreTypes.h
--- a/Engine/Source/Runtime/Experimental/Chaos/Public/GeometryCollection/GeometryCollectionSimulationCoreTypes.h
+++ b/Engine/Source/Runtime/Experimental/Chaos/Public/GeometryCollection/GeometryCollectionSimulationCoreTypes.h
@@ -200,6 +200,7 @@ struct FSimulationParameters
 		, bGenerateRemovalsData(false)
 		, RemoveOnFractureEnabled(false)
 		, EnableGravity(true)
+		, UseInertiaConditioning(true)
 		, SimulationFilterData()
 		, QueryFilterData()
 		, UserData(nullptr)
@@ -240,6 +241,7 @@ struct FSimulationParameters
 		, Shared(Other.Shared)
 		, RemoveOnFractureEnabled(Other.RemoveOnFractureEnabled)
 		, EnableGravity(Other.EnableGravity)
+		, UseInertiaConditioning(Other.UseInertiaConditioning)
 		, SimulationFilterData(Other.SimulationFilterData)
 		, QueryFilterData(Other.QueryFilterData)
 		, UserData(Other.UserData)
@@ -303,6 +305,7 @@ struct FSimulationParameters
 	bool RemoveOnFractureEnabled;
 	
 	bool EnableGravity;
+	bool UseInertiaConditioning;
 
 	FCollisionFilterData SimulationFilterData;
 	FCollisionFilterData QueryFilterData;
diff --git a/Engine/Source/Runtime/Experimental/Chaos/Public/PhysicsProxy/GeometryCollectionPhysicsProxy.h b/Engine/Source/Runtime/Experimental/Chaos/Public/PhysicsProxy/GeometryCollectionPhysicsProxy.h
--- a/Engine/Source/Runtime/Experimental/Chaos/Public/PhysicsProxy/GeometryCollectionPhysicsProxy.h
+++ b/Engine/Source/Runtime/Experimental/Chaos/Public/PhysicsProxy/GeometryCollectionPhysicsProxy.h
@@ -277,6 +277,7 @@ protected:
 		const Chaos::FUniqueIdx* ExistingIndex);
 
 	void SetGravityEnabled(const Chaos::FPBDRigidsSolver& RigidsSolver, bool bEnabled);
+	void EnableInertiaConditioning(const Chaos::FPBDRigidsSolver& RigidsSolver, bool bEnabled);
 	void SetSleepingState(const Chaos::FPBDRigidsSolver& RigidsSolver);
 	void DirtyAllParticles(const Chaos::FPBDRigidsSolver& RigidsSolver);
 	
diff --git a/Engine/Source/Runtime/Experimental/GeometryCollectionEngine/Private/GeometryCollection/GeometryCollectionComponent.cpp b/Engine/Source/Runtime/Experimental/GeometryCollectionEngine/Private/GeometryCollection/GeometryCollectionComponent.cpp
--- a/Engine/Source/Runtime/Experimental/GeometryCollectionEngine/Private/GeometryCollection/GeometryCollectionComponent.cpp
+++ b/Engine/Source/Runtime/Experimental/GeometryCollectionEngine/Private/GeometryCollection/GeometryCollectionComponent.cpp
@@ -2256,6 +2256,7 @@ void UGeometryCollectionComponent::RegisterAndInitializePhysicsProxy()
 		SimulationParameters.bGenerateRemovalsData = bNotifyRemovals;
 		SimulationParameters.RemoveOnFractureEnabled = SimulationParameters.Shared.RemoveOnFractureIndices.Num() > 0;
 		SimulationParameters.EnableGravity = BodyInstance.bEnableGravity;
+		SimulationParameters.UseInertiaConditioning = BodyInstance.IsInertiaConditioningEnabled();
 		SimulationParameters.WorldTransform = GetComponentToWorld();
 		SimulationParameters.UserData = static_cast<void*>(&PhysicsUserData);
 
