Physics interface cleanup.

diff --git a/Engine/Source/Runtime/Engine/Public/PhysicsReplication.h b/Engine/Source/Runtime/Engine/Public/PhysicsReplication.h
--- a/Engine/Source/Runtime/Engine/Public/PhysicsReplication.h
+++ b/Engine/Source/Runtime/Engine/Public/PhysicsReplication.h
@@ -11,9 +11,7 @@
 #include "Engine/ReplicatedState.h"
 #include "Physics/PhysicsInterfaceDeclares.h"
 #include "PhysicsProxy/SingleParticlePhysicsProxyFwd.h"
-#if WITH_CHAOS
 #include "Chaos/Particles.h"
-#endif
 
 namespace Chaos
 {
@@ -61,7 +59,6 @@ struct FReplicatedPhysicsTarget
 #endif
 };
 
-#if WITH_CHAOS
 struct ErrorCorrectionData
 {
 	float LinearVelocityCoefficient;
@@ -82,7 +79,6 @@ struct FAsyncPhysicsDesiredState
 	bool bShouldSleep;
 	int32 ServerFrame;
 };
-#endif
 
 struct FBodyInstance;
 struct FRigidBodyErrorCorrection;
@@ -129,20 +125,16 @@ private:
 	/** Get the ping from  */
 	float GetOwnerPing(const AActor* const Owner, const FReplicatedPhysicsTarget& Target) const;
 
-#if WITH_CHAOS
 	static void ApplyAsyncDesiredState(float DeltaSeconds, const FAsyncPhysicsRepCallbackData* Input);
-#endif
 
 private:
 	TMap<TWeakObjectPtr<UPrimitiveComponent>, FReplicatedPhysicsTarget> ComponentToTargets;
 	FPhysScene* PhysScene;
 
-#if WITH_CHAOS
 	FPhysicsReplicationAsyncCallback* AsyncCallback;
 	
 	void PrepareAsyncData_External(const FRigidBodyErrorCorrection& ErrorCorrection);	//prepare async data for writing. Call on external thread (i.e. game thread)
 	FAsyncPhysicsRepCallbackData* CurAsyncData;	//async data being written into before we push into callback
 	friend FPhysicsReplicationAsyncCallback;
-#endif
 
 };
\ No newline at end of file
diff --git a/Engine/Source/Runtime/Experimental/Chaos/Chaos.Build.cs b/Engine/Source/Runtime/Experimental/Chaos/Chaos.Build.cs
--- a/Engine/Source/Runtime/Experimental/Chaos/Chaos.Build.cs
+++ b/Engine/Source/Runtime/Experimental/Chaos/Chaos.Build.cs
@@ -26,15 +26,6 @@ namespace UnrealBuildTool.Rules
             PublicDefinitions.Add("COMPILE_WITHOUT_UNREAL_SUPPORT=0");
 			PrivateDefinitions.Add("CHAOS_INCLUDE_LEVEL_1=1");
 
-            if (Target.bCompileChaos == true || Target.bUseChaos == true)
-            {
-                PublicDefinitions.Add("INCLUDE_CHAOS=1");
-            }
-            else
-            {
-                PublicDefinitions.Add("INCLUDE_CHAOS=0");
-            }
-
 			if (Target.bUseChaosMemoryTracking == true)
 			{
 				PublicDefinitions.Add("CHAOS_MEMORY_TRACKING=1");
diff --git a/Engine/Source/Runtime/Experimental/ChaosCore/ChaosCore.Build.cs b/Engine/Source/Runtime/Experimental/ChaosCore/ChaosCore.Build.cs
--- a/Engine/Source/Runtime/Experimental/ChaosCore/ChaosCore.Build.cs
+++ b/Engine/Source/Runtime/Experimental/ChaosCore/ChaosCore.Build.cs
@@ -17,15 +17,6 @@ namespace UnrealBuildTool.Rules
 
             PublicDefinitions.Add("COMPILE_WITHOUT_UNREAL_SUPPORT=0");
 
-            if (Target.bCompileChaos == true || Target.bUseChaos == true)
-            {
-                PublicDefinitions.Add("INCLUDE_CHAOS=1");
-            }
-            else
-            {
-                PublicDefinitions.Add("INCLUDE_CHAOS=0");
-            }
-
             if (Target.bUseChaosChecked == true)
             {
                 PublicDefinitions.Add("CHAOS_CHECKED=1");
diff --git a/Engine/Source/Runtime/Experimental/ChaosSolverEngine/Private/Chaos/ChaosGameplayEventDispatcher.cpp b/Engine/Source/Runtime/Experimental/ChaosSolverEngine/Private/Chaos/ChaosGameplayEventDispatcher.cpp
--- a/Engine/Source/Runtime/Experimental/ChaosSolverEngine/Private/Chaos/ChaosGameplayEventDispatcher.cpp
+++ b/Engine/Source/Runtime/Experimental/ChaosSolverEngine/Private/Chaos/ChaosGameplayEventDispatcher.cpp
@@ -260,7 +260,6 @@ void UChaosGameplayEventDispatcher::DispatchPendingWakeNotifies()
 
 void UChaosGameplayEventDispatcher::RegisterChaosEvents()
 {
-#if WITH_CHAOS
 	if (FPhysScene* Scene = GetWorld()->GetPhysicsScene())
 	{
 		if (Chaos::FPhysicsSolver* Solver = Scene->GetSolver())
@@ -272,13 +271,10 @@ void UChaosGameplayEventDispatcher::RegisterChaosEvents()
 			EventManager->RegisterHandler<Chaos::FRemovalEventData>(Chaos::EEventType::Removal, this, &UChaosGameplayEventDispatcher::HandleRemovalEvents);
 		}
 	}
-
-#endif
 }
 
 void UChaosGameplayEventDispatcher::UnregisterChaosEvents()
 {
-#if WITH_CHAOS
 	if (GetWorld())
 	{
 		if (FPhysScene* Scene = GetWorld()->GetPhysicsScene())
@@ -293,15 +289,12 @@ void UChaosGameplayEventDispatcher::UnregisterChaosEvents()
 			}
 		}
 	}
-#endif
 }
 
 void UChaosGameplayEventDispatcher::HandleCollisionEvents(const Chaos::FCollisionEventData& Event)
 {
 	SCOPE_CYCLE_COUNTER(STAT_DispatchCollisionEvents);
 
-#if INCLUDE_CHAOS
-
 	FPhysScene_Chaos& Scene = *(GetWorld()->GetPhysicsScene());
 
 	PendingChaosCollisionNotifies.Reset();
@@ -436,8 +429,6 @@ void UChaosGameplayEventDispatcher::HandleCollisionEvents(const Chaos::FCollisio
 
 	// Tell the world and actors about the collisions
 	DispatchPendingCollisionNotifies();
-
-#endif
 }
 
 void UChaosGameplayEventDispatcher::HandleBreakingEvents(const Chaos::FBreakingEventData& Event)
@@ -485,7 +476,6 @@ void UChaosGameplayEventDispatcher::HandleBreakingEvents(const Chaos::FBreakingE
 
 void UChaosGameplayEventDispatcher::HandleSleepingEvents(const Chaos::FSleepingEventData& SleepingData)
 {
-#if WITH_CHAOS
 	FPhysScene& Scene = *(GetWorld()->GetPhysicsScene());
 
 	const Chaos::FSleepingDataArray& SleepingArray = SleepingData.SleepingData;
@@ -506,7 +496,6 @@ void UChaosGameplayEventDispatcher::HandleSleepingEvents(const Chaos::FSleepingE
 	}
 
 	DispatchPendingWakeNotifies();
-#endif
 }
 
 void UChaosGameplayEventDispatcher::AddPendingSleepingNotify(FBodyInstance* BodyInstance, ESleepEvent SleepEventType)
diff --git a/Engine/Source/Runtime/Experimental/ChaosSolverEngine/Private/Chaos/ChaosSolverActor.cpp b/Engine/Source/Runtime/Experimental/ChaosSolverEngine/Private/Chaos/ChaosSolverActor.cpp
--- a/Engine/Source/Runtime/Experimental/ChaosSolverEngine/Private/Chaos/ChaosSolverActor.cpp
+++ b/Engine/Source/Runtime/Experimental/ChaosSolverEngine/Private/Chaos/ChaosSolverActor.cpp
@@ -198,7 +198,6 @@ AChaosSolverActor::AChaosSolverActor(const FObjectInitializer& ObjectInitializer
 	{
 		// Don't spawn solvers on the CDO
 
-#if WITH_CHAOS
 		// @question(Benn) : Does this need to be created on the Physics thread using a queued command?
 		PhysScene = MakeShareable(new FPhysScene_Chaos(this
 #if CHAOS_DEBUG_NAME
@@ -206,7 +205,6 @@ AChaosSolverActor::AChaosSolverActor(const FObjectInitializer& ObjectInitializer
 #endif
 		));
 		Solver = PhysScene->GetSolver();
-#endif
 		// Ticking setup for collision/breaking notifies
 		PrimaryActorTick.TickGroup = TG_PostPhysics;
 		PrimaryActorTick.bCanEverTick = true;
@@ -383,13 +381,12 @@ void AChaosSolverActor::Serialize(FArchive& Ar)
 void AChaosSolverActor::PostRegisterAllComponents()
 {
 	Super::PostRegisterAllComponents();
-#if INCLUDE_CHAOS
+
 	UWorld* const W = GetWorld(); 
 	if (W && !W->PhysicsScene_Chaos)
 	{
 		SetAsCurrentWorldSolver();
 	}
-#endif
 }
 
 void AChaosSolverActor::PostDuplicate(EDuplicateMode::Type DuplicateMode)
@@ -423,9 +420,7 @@ void AChaosSolverActor::SetAsCurrentWorldSolver()
 	UWorld* const W = GetWorld();
 	if (W)
 	{
-#if INCLUDE_CHAOS
 		W->PhysicsScene_Chaos = PhysScene;
-#endif
 	}
 }
 
diff --git a/Engine/Source/Runtime/Experimental/GeometryCollectionEngine/Private/ChaosBlueprint.cpp b/Engine/Source/Runtime/Experimental/GeometryCollectionEngine/Private/ChaosBlueprint.cpp
--- a/Engine/Source/Runtime/Experimental/GeometryCollectionEngine/Private/ChaosBlueprint.cpp
+++ b/Engine/Source/Runtime/Experimental/GeometryCollectionEngine/Private/ChaosBlueprint.cpp
@@ -413,50 +413,42 @@ void UChaosDestructionListener::SortRemovalEvents(TArray<FChaosRemovalEventData>
 
 void UChaosDestructionListener::RegisterChaosEvents(FPhysScene* Scene)
 {
-#if WITH_CHAOS
 	Chaos::FPhysicsSolver* Solver = Scene->GetSolver();
 	Chaos::FEventManager* EventManager = Solver->GetEventManager();
 	EventManager->RegisterHandler<Chaos::FCollisionEventData>(Chaos::EEventType::Collision, this, &UChaosDestructionListener::HandleCollisionEvents);
 	EventManager->RegisterHandler<Chaos::FBreakingEventData>(Chaos::EEventType::Breaking, this, &UChaosDestructionListener::HandleBreakingEvents);
 	EventManager->RegisterHandler<Chaos::FTrailingEventData>(Chaos::EEventType::Trailing, this, &UChaosDestructionListener::HandleTrailingEvents);
 	EventManager->RegisterHandler<Chaos::FRemovalEventData>(Chaos::EEventType::Removal, this, &UChaosDestructionListener::HandleRemovalEvents);
-#endif
 }
 
 void UChaosDestructionListener::UnregisterChaosEvents(FPhysScene* Scene)
 {
-#if WITH_CHAOS
 	Chaos::FPhysicsSolver* Solver = Scene->GetSolver();
 	Chaos::FEventManager* EventManager = Solver->GetEventManager();
 	EventManager->UnregisterHandler(Chaos::EEventType::Collision, this);
 	EventManager->UnregisterHandler(Chaos::EEventType::Breaking, this);
 	EventManager->UnregisterHandler(Chaos::EEventType::Trailing, this);
 	EventManager->UnregisterHandler(Chaos::EEventType::Removal, this);
-#endif
 }
 
 void UChaosDestructionListener::RegisterChaosEvents(TSharedPtr<FPhysScene_Chaos> Scene)
 {
-#if WITH_CHAOS
 	Chaos::FPhysicsSolver* Solver = Scene->GetSolver();
 	Chaos::FEventManager* EventManager = Solver->GetEventManager();
 	EventManager->RegisterHandler<Chaos::FCollisionEventData>(Chaos::EEventType::Collision, this, &UChaosDestructionListener::HandleCollisionEvents);
 	EventManager->RegisterHandler<Chaos::FBreakingEventData>(Chaos::EEventType::Breaking, this, &UChaosDestructionListener::HandleBreakingEvents);
 	EventManager->RegisterHandler<Chaos::FTrailingEventData>(Chaos::EEventType::Trailing, this, &UChaosDestructionListener::HandleTrailingEvents);
 	EventManager->RegisterHandler<Chaos::FRemovalEventData>(Chaos::EEventType::Removal, this, &UChaosDestructionListener::HandleRemovalEvents);
-#endif
 }
 
 void UChaosDestructionListener::UnregisterChaosEvents(TSharedPtr<FPhysScene_Chaos> Scene)
 {
-#if WITH_CHAOS
 	Chaos::FPhysicsSolver* Solver = Scene->GetSolver();
 	Chaos::FEventManager* EventManager = Solver->GetEventManager();
 	EventManager->UnregisterHandler(Chaos::EEventType::Collision, this);
 	EventManager->UnregisterHandler(Chaos::EEventType::Breaking, this);
 	EventManager->UnregisterHandler(Chaos::EEventType::Trailing, this);
 	EventManager->UnregisterHandler(Chaos::EEventType::Removal, this);
-#endif
 }
 
 
diff --git a/Engine/Source/Runtime/Experimental/GeometryCollectionEngine/Private/GeometryCollection/GeometryCollectionActor.cpp b/Engine/Source/Runtime/Experimental/GeometryCollectionEngine/Private/GeometryCollection/GeometryCollectionActor.cpp
--- a/Engine/Source/Runtime/Experimental/GeometryCollectionEngine/Private/GeometryCollection/GeometryCollectionActor.cpp
+++ b/Engine/Source/Runtime/Experimental/GeometryCollectionEngine/Private/GeometryCollection/GeometryCollectionActor.cpp
@@ -57,11 +57,7 @@ void AGeometryCollectionActor::Tick(float DeltaTime)
 
 const Chaos::FPhysicsSolver* GetSolver(const AGeometryCollectionActor& GeomCollectionActor)
 {
-#if INCLUDE_CHAOS
 	return GeomCollectionActor.GetGeometryCollectionComponent()->ChaosSolverActor != nullptr ? GeomCollectionActor.GetGeometryCollectionComponent()->ChaosSolverActor->GetSolver() : GeomCollectionActor.GetWorld()->PhysicsScene_Chaos->GetSolver();
-#else
-	return nullptr;
-#endif
 }
 
 
diff --git a/Engine/Source/Runtime/Experimental/GeometryCollectionEngine/Private/GeometryCollection/GeometryCollectionComponent.cpp b/Engine/Source/Runtime/Experimental/GeometryCollectionEngine/Private/GeometryCollection/GeometryCollectionComponent.cpp
--- a/Engine/Source/Runtime/Experimental/GeometryCollectionEngine/Private/GeometryCollection/GeometryCollectionComponent.cpp
+++ b/Engine/Source/Runtime/Experimental/GeometryCollectionEngine/Private/GeometryCollection/GeometryCollectionComponent.cpp
@@ -136,26 +136,6 @@ int32 GetClusterLevel(const FTransformCollection* Collection, int32 TransformGro
 	return Level;
 }
 
-#if WITH_PHYSX && !WITH_CHAOS_NEEDS_TO_BE_FIXED
-FGeometryCollectionSQAccelerator GlobalGeomCollectionAccelerator;	//todo(ocohen): proper lifetime management needed
-
-void HackRegisterGeomAccelerator(UGeometryCollectionComponent& Component)
-{
-#if TODO_REIMPLEMENT_SCENEQUERY_CROSSENGINE
-	if (UWorld* World = Component.GetWorld())
-	{
-		if (FPhysScene* PhysScene = World->GetPhysicsScene())
-		{
-			if (FSQAcceleratorUnion* SQAccelerationUnion = PhysScene->GetSQAcceleratorUnion())
-			{
-				SQAccelerationUnion->AddSQAccelerator(&GlobalGeomCollectionAccelerator);
-			}
-		}
-	}
-#endif
-}
-#endif
-
 bool FGeometryCollectionRepData::Identical(const FGeometryCollectionRepData* Other, uint32 PortFlags) const
 {
 	return Other && (Version == Other->Version);
@@ -313,7 +293,6 @@ UGeometryCollectionComponent::UGeometryCollectionComponent(const FObjectInitiali
 
 Chaos::FPhysicsSolver* GetSolver(const UGeometryCollectionComponent& GeometryCollectionComponent)
 {
-#if INCLUDE_CHAOS
 	if(GeometryCollectionComponent.ChaosSolverActor)
 	{
 		return GeometryCollectionComponent.ChaosSolverActor->GetSolver();
@@ -325,16 +304,12 @@ Chaos::FPhysicsSolver* GetSolver(const UGeometryCollectionComponent& GeometryCol
 			return Scene->GetSolver();
 		}
 	}
-#endif
 	return nullptr;
 }
 
 void UGeometryCollectionComponent::BeginPlay()
 {
 	Super::BeginPlay();
-#if WITH_PHYSX && !WITH_CHAOS_NEEDS_TO_BE_FIXED
-	HackRegisterGeomAccelerator(*this);
-#endif
 
 #if WITH_EDITOR
 	if (RestCollection != nullptr)
@@ -1232,7 +1207,6 @@ void UGeometryCollectionComponent::RegisterForEvents()
 {
 	if (BodyInstance.bNotifyRigidBodyCollision || bNotifyBreaks || bNotifyCollisions || bNotifyRemovals)
 	{
-#if INCLUDE_CHAOS
 		Chaos::FPhysicsSolver* Solver = GetWorld()->GetPhysicsScene()->GetSolver();
 		if (Solver)
 		{
@@ -1268,7 +1242,6 @@ void UGeometryCollectionComponent::RegisterForEvents()
 
 			}
 		}
-#endif
 	}
 }
 
@@ -1901,12 +1874,10 @@ void UGeometryCollectionComponent::OnUpdateTransform(EUpdateTransformFlags Updat
 {
 	Super::OnUpdateTransform(UpdateTransformFlags, Teleport);
 
-#if WITH_CHAOS
 	if (PhysicsProxy)
 	{
 		PhysicsProxy->SetWorldTransform(GetComponentTransform());
 	}
-#endif
 }
 
 void UGeometryCollectionComponent::TickComponent(float DeltaTime, enum ELevelTick TickType, FActorComponentTickFunction *ThisTickFunction)
@@ -1931,7 +1902,6 @@ void UGeometryCollectionComponent::TickComponent(float DeltaTime, enum ELevelTic
 	}
 #endif
 
-#if WITH_CHAOS
 	//if (bRenderStateDirty && DynamicCollection)	//todo: always send for now
 	if (RestCollection)
 	{
@@ -1971,7 +1941,6 @@ void UGeometryCollectionComponent::TickComponent(float DeltaTime, enum ELevelTic
 			}
 		}
 	}
-#endif
 }
 
 void UGeometryCollectionComponent::AsyncPhysicsTickComponent(float DeltaTime, float SimTime)
@@ -1991,13 +1960,9 @@ void UGeometryCollectionComponent::AsyncPhysicsTickComponent(float DeltaTime, fl
 
 void UGeometryCollectionComponent::OnRegister()
 {
-#if WITH_CHAOS
 	//UE_LOG(UGCC_LOG, Log, TEXT("GeometryCollectionComponent[%p]::OnRegister()[%p]"), this,RestCollection );
 	ResetDynamicCollection();
 
-
-#endif // WITH_CHAOS
-
 	SetIsReplicated(bEnableReplication);
 
 	InitializeEmbeddedGeometry();
@@ -2099,24 +2064,10 @@ void UGeometryCollectionComponent::ResetDynamicCollection()
 
 void UGeometryCollectionComponent::OnCreatePhysicsState()
 {
-/*#if WITH_PHYSX
-	DummyBodySetup = NewObject<UBodySetup>(this, UBodySetup::StaticClass());
-	DummyBodySetup->AggGeom.BoxElems.Add(FKBoxElem(1.0f));
-	DummyBodyInstance.InitBody(DummyBodySetup, GetComponentToWorld(), this, nullptr);
-	DummyBodyInstance.bNotifyRigidBodyCollision = BodyInstance.bNotifyRigidBodyCollision;
-#endif
-*/
 	// Skip the chain - don't care about body instance setup
 	UActorComponent::OnCreatePhysicsState();
 	if (!BodyInstance.bSimulatePhysics) IsObjectLoading = false; // just mark as loaded if we are simulating.
 
-/*#if WITH_PHYSX
-	DummyBodyInstance.SetCollisionEnabled(ECollisionEnabled::QueryOnly);
-	DummyBodyInstance.SetResponseToAllChannels(ECR_Block);
-#endif
-*/
-
-#if WITH_CHAOS
 	// Static mesh uses an init framework that goes through FBodyInstance.  We
 	// do the same thing, but through the geometry collection proxy and lambdas
 	// defined below.  FBodyInstance doesn't work for geometry collections 
@@ -2204,19 +2155,10 @@ void UGeometryCollectionComponent::OnCreatePhysicsState()
 			
 		}
 	}
-
-#if WITH_PHYSX && !WITH_CHAOS_NEEDS_TO_BE_FIXED
-	if (PhysicsProxy)
-	{
-		GlobalGeomCollectionAccelerator.AddComponent(this);
-	}
-#endif
-#endif // WITH_CHAOS
 }
 
 void UGeometryCollectionComponent::RegisterAndInitializePhysicsProxy()
 {
-#if WITH_CHAOS
 	FSimulationParameters SimulationParameters;
 	{
 #if !(UE_BUILD_SHIPPING || UE_BUILD_TEST)
@@ -2274,24 +2216,16 @@ void UGeometryCollectionComponent::RegisterAndInitializePhysicsProxy()
 
 	RegisterForEvents();
 	SetAsyncPhysicsTickEnabled(GetIsReplicated());
-#endif
 }
 
 void UGeometryCollectionComponent::OnDestroyPhysicsState()
 {
 	UActorComponent::OnDestroyPhysicsState();
 
-#if WITH_CHAOS
-#if WITH_PHYSX && !WITH_CHAOS_NEEDS_TO_BE_FIXED
-	GlobalGeomCollectionAccelerator.RemoveComponent(this);
-#endif
-
-#if WITH_PHYSX
 	if(DummyBodyInstance.IsValidBodyInstance())
 	{
 		DummyBodyInstance.TermBody();
 	}
-#endif
 
 	if(PhysicsProxy)
 	{
@@ -2302,7 +2236,6 @@ void UGeometryCollectionComponent::OnDestroyPhysicsState()
 		// Discard the pointer (cleanup happens through the scene or dedicated thread)
 		PhysicsProxy = nullptr;
 	}
-#endif
 }
 
 void UGeometryCollectionComponent::SendRenderDynamicData_Concurrent()
@@ -3147,22 +3080,18 @@ FPhysScene_Chaos* UGeometryCollectionComponent::GetInnerChaosScene() const
 	}
 	else
 	{
-#if INCLUDE_CHAOS
 		if (ensure(GetOwner()) && ensure(GetOwner()->GetWorld()))
 		{
 			return GetOwner()->GetWorld()->GetPhysicsScene();
 		}
+
 		check(GWorld);
 		return GWorld->GetPhysicsScene();
-#else
-		return nullptr;
-#endif
 	}
 }
 
 AChaosSolverActor* UGeometryCollectionComponent::GetPhysicsSolverActor() const
 {
-#if WITH_CHAOS
 	if (ChaosSolverActor)
 	{
 		return ChaosSolverActor;
@@ -3173,7 +3102,6 @@ AChaosSolverActor* UGeometryCollectionComponent::GetPhysicsSolverActor() const
 		return Scene ? Cast<AChaosSolverActor>(Scene->GetSolverActor()) : nullptr;
 	}
 
-#endif
 	return nullptr;
 }
 
diff --git a/Engine/Source/Runtime/Experimental/GeometryCollectionEngine/Private/GeometryCollection/GeometryCollectionSQAccelerator.cpp b/Engine/Source/Runtime/Experimental/GeometryCollectionEngine/Private/GeometryCollection/GeometryCollectionSQAccelerator.cpp
--- a/Engine/Source/Runtime/Experimental/GeometryCollectionEngine/Private/GeometryCollection/GeometryCollectionSQAccelerator.cpp
+++ b/Engine/Source/Runtime/Experimental/GeometryCollectionEngine/Private/GeometryCollection/GeometryCollectionSQAccelerator.cpp
@@ -15,823 +15,3 @@
 #include "ChaosStats.h"
 #include "PhysicsSolver.h"
 #include "PhysicsProxy/GeometryCollectionPhysicsProxy.h"
-
-#if !WITH_CHAOS_NEEDS_TO_BE_FIXED
-
-DECLARE_CYCLE_STAT(TEXT("LowLevelSweep"), STAT_LowLevelSweep, STATGROUP_Chaos);
-DECLARE_CYCLE_STAT(TEXT("LowLevelRaycast"), STAT_LowLevelRaycast, STATGROUP_Chaos);
-DECLARE_CYCLE_STAT(TEXT("LowLevelOverlap"), STAT_LowLevelOverlap, STATGROUP_Chaos);
-
-#if TODO_REIMPLEMENT_SCENEQUERY_CROSSENGINE
-
-bool IsValidIndexAndTransform(const FGeometryCollectionResults& PhysResult, const Chaos::TPBDRigidParticles<Chaos::FReal, 3>& Particles, const TManagedArray<FTransform>& TransformArray, const TArray<bool>& DisabledFlags, const int32 RigidBodyIdx, const bool bCanBeDisabled)
-{
-	if(RigidBodyIdx == -1)
-	{
-		//todo(ocohen): managed to avoid this invalid index, but need to investigate a bit more into whether we can always assume it's valid
-		return false;
-	}
-
-	if (PhysResult.BaseIndex == -1)
-	{
-		//todo(mlentine): Why is this possible?
-		return false;
-	}
-
-	const int32 LocalBodyIndex = RigidBodyIdx - PhysResult.BaseIndex;
-
-	if(LocalBodyIndex < 0 || LocalBodyIndex >= PhysResult.NumParticlesAdded)
-	{
-		// Ignore collisions for other components - need to make this even faster (subset of potential intersections [maybe array view?])
-		return false;
-	}
-
-	const FTransform& CurrentTransform = TransformArray[LocalBodyIndex];
-	const FVector CurrentTranslation = CurrentTransform.GetTranslation();
-
-	//@todo(mlentine): In theory this is no longer needed
-	if (!bCanBeDisabled && DisabledFlags[LocalBodyIndex]) 
-	{ 
-		//disabled particles can actually have stale geometry in them and are clearly not useful anyway
-		return false;
-	}
-
-	if (static_cast<uint32>(RigidBodyIdx) >= Particles.Size())
-	{
-		// @todo(mlentine): Is this a possible situation?
-		return false;
-	}
-
-	if (!(ensure(!FMath::IsNaN(CurrentTranslation[0])) && ensure(!FMath::IsNaN(CurrentTranslation[1])) && ensure(!FMath::IsNaN(CurrentTranslation[2]))))
-	{
-		return false;
-	}
-
-	return true;
-}
-
-bool LowLevelRaycastSingleElement(int32 InParticleIndex, const Chaos::FPhysicsSolver* InSolver, const Chaos::FClusterBuffer& ClusterBuffer, const FGeometryCollectionPhysicsProxy* InObject, const FVector& Start, const FVector& Dir, Chaos::FReal DeltaMag, bool bCanBeDisabled, EHitFlags OutputFlags, FHitRaycast& OutHit)
-{
-	using namespace Chaos;
-
-	// Preconditions from Raycast - shouldn't get in here without valid case
-	checkSlow(InSolver);
-	checkSlow(InObject);
-
-	const FGeometryCollectionResults& PhysResult = InObject->GetPhysicsResults().GetGameDataForRead();
-
-	const TManagedArray<int32>& RigidBodyIdArray = PhysResult.RigidBodyIds;
-	const TManagedArray<FTransform>& TransformArray = PhysResult.Transforms;
-	const TArray<bool>& DisabledFlags = PhysResult.DisabledStates;
-
-	{
-		const TPBDRigidParticles<Chaos::FReal, 3>& Particles = InSolver->GetRigidParticles();
-
-		if(!IsValidIndexAndTransform(PhysResult, Particles, TransformArray, DisabledFlags, InParticleIndex, bCanBeDisabled))
-		{
-			return false;
-		}
-
-		const int32 LocalBodyIndex = InParticleIndex - PhysResult.BaseIndex;
-
-		const FRigidTransform3& TM = PhysResult.ParticleToWorldTransforms[LocalBodyIndex];
-		if (!(ensure(!FMath::IsNaN(TM.GetTranslation().X)) && ensure(!FMath::IsNaN(TM.GetTranslation().Y)) && ensure(!FMath::IsNaN(TM.GetTranslation().Z))))
-		{
-			return false;
-		}
-
-		const FVec3 StartLocal = TM.InverseTransformPositionNoScale(Start);
-		const FVec3 DirLocal = TM.InverseTransformVectorNoScale(Dir);
-		const FVec3 EndLocal = StartLocal + DirLocal * DeltaMag;
-
-		const FImplicitObject* Object = ClusterBuffer.GeometryPtrs[InParticleIndex].Get();	//todo(ocohen): can this ever be null?
-
-		if (!Object)
-		{
-			return false;
-		}
-		Pair<FVec3, bool> Result = Object->FindClosestIntersection(StartLocal, EndLocal, /*Thickness=*/0.f);
-		if(Result.Second)	//todo(ocohen): once we do more than just a bool we need to get the closest point
-		{
-#if WITH_PHYSX
-			//todo(ocohen): check output flags?
-			const Chaos::FReal Distance = (Result.First - StartLocal).Size();
-			if(OutHit.distance == PX_MAX_REAL || Distance < OutHit.distance)
-			{
-				OutHit.distance = Distance;	//todo(ocohen): assuming physx structs for now
-				OutHit.position = U2PVector(TM.TransformPositionNoScale(Result.First));
-				const FVec3 LocalNormal = Object->Normal(Result.First);
-				OutHit.normal = U2PVector(TM.TransformVectorNoScale(LocalNormal));
-				SetFlags(OutHit, EHitFlags::Distance | EHitFlags::Normal | EHitFlags::Position);
-			}
-			return true;
-#endif
-		}
-	}
-
-	return false;
-}
-
-static TAutoConsoleVariable<int32> CVarMaxSweepSteps(
-	TEXT("p.MaxSweepSteps"),
-	3,
-	TEXT("Number of steps during a sweep"),
-	ECVF_Default);
-
-bool LowLevelSweepSingleElement(int32 InParticleIndex, const Chaos::FPhysicsSolver* InSolver, const Chaos::FClusterBuffer& ClusterBuffer, const FGeometryCollectionPhysicsProxy* InObject, const Chaos::FImplicitObject& QueryGeom, const Chaos::FParticles& CollisionParticles, const FTransform& StartPose, const FVector& Dir, Chaos::FReal DeltaMag, const bool bCanBeDisabled, FHitSweep& OutHit)
-{
-	using namespace Chaos;
-
-	checkSlow(InSolver);
-	checkSlow(InObject);
-
-	const FGeometryCollectionResults& PhysResult = InObject->GetPhysicsResults().GetGameDataForRead();
-
-	const TManagedArray<int32>& RigidBodyIdArray = PhysResult.RigidBodyIds;
-	const TManagedArray<FTransform>& TransformArray = PhysResult.Transforms;
-	const TArray<bool>& DisabledFlags = PhysResult.DisabledStates;
-
-	const TPBDRigidParticles<Chaos::FReal, 3>& Particles = InSolver->GetRigidParticles();
-
-	if(!IsValidIndexAndTransform(PhysResult, Particles, TransformArray, DisabledFlags, InParticleIndex, bCanBeDisabled))
-	{
-		return false;
-	}
-
-	const int32 LocalBodyIndex = InParticleIndex - PhysResult.BaseIndex;
-	const FRigidTransform3& TM = PhysResult.ParticleToWorldTransforms[LocalBodyIndex];
-
-	const FImplicitObject* Object = ClusterBuffer.GeometryPtrs[InParticleIndex].Get();
-	if (!Object)
-	{
-		return false;
-	}
-	const FVec3 DirLocal = TM.InverseTransformVectorNoScale(Dir);
-
-	bool bFound = false;
-
-	for(uint32 i = 0; i < CollisionParticles.Size(); ++i)
-	{
-		const FVec3 StartLocal = TM.InverseTransformPositionNoScale(StartPose.TransformPositionNoScale(CollisionParticles.X(i)));
-		const FVec3 EndLocal = StartLocal + DirLocal * DeltaMag;
-
-		Pair<FVec3, bool> Result = Object->FindClosestIntersection(StartLocal, EndLocal, /*Thickness=*/0.f);
-
-		if(Result.Second)
-		{
-#if WITH_PHYSX
-			const Chaos::FReal Distance = (Result.First - StartLocal).Size();
-			if(!bFound || Distance < OutHit.distance)
-			{
-				OutHit.distance = Distance;	//todo(ocohen): assuming physx structs for now
-				OutHit.position = U2PVector(TM.TransformPositionNoScale(Result.First));
-				const FVec3 LocalNormal = Object->Normal(Result.First);
-				OutHit.normal = U2PVector(TM.TransformVectorNoScale(LocalNormal));
-				SetFlags(OutHit, EHitFlags::Distance | EHitFlags::Normal | EHitFlags::Position);
-			}
-			bFound = true;
-#endif
-		}
-	}
-
-	return bFound;
-}
-
-bool LowLevelOverlap(const UGeometryCollectionComponent& GeomCollectionComponent, const TArray<int32>& InPotentialIntersections, const Chaos::FClusterBuffer& ClusterBuffer, const Chaos::FImplicitObject& QueryGeom, const FTransform& GeomPose, FHitOverlap& OutHit)
-{
-	using namespace Chaos;
-
-	SCOPE_CYCLE_COUNTER(STAT_LowLevelOverlap);
-
-	const FGeometryCollectionPhysicsProxy* PhysObject = GeomCollectionComponent.GetPhysicsProxy();
-	if (!ensure(PhysObject))
-	{
-		return false;
-	}
-	const FGeometryCollectionResults& PhysResult = PhysObject->GetPhysicsResults().GetGameDataForRead();
-
-	const TManagedArray<int32>& RigidBodyIdArray = PhysResult.RigidBodyIds;
-	const TManagedArray<FTransform>& TransformArray = PhysResult.Transforms;
-	const TArray<bool>& DisabledFlags = PhysResult.DisabledStates;
-	
-	bool bFound = false;
-
-	if (Chaos::FPhysicsSolver* Solver = GeomCollectionComponent.ChaosSolverActor != nullptr ? GeomCollectionComponent.ChaosSolverActor->GetSolver() : GeomCollectionComponent.GetOwner()->GetWorld()->PhysicsScene_Chaos->GetSolver())
-	{
-		const TPBDRigidParticles<Chaos::FReal, 3>& Particles = Solver->GetRigidParticles();	//todo(ocohen): should these just get passed in instead of hopping through scene?
-
-		check(QueryGeom.HasBoundingBox()); // We do not support unbounded query objects
-		
-		PhysicsParallelFor(InPotentialIntersections.Num(), [&](int32 PotentialIdx)
-		{
-			int32 RigidBodyIdx = InPotentialIntersections[PotentialIdx];
-			if (!IsValidIndexAndTransform(PhysResult, Particles, TransformArray, DisabledFlags, RigidBodyIdx, false))
-			{
-				return;
-			}
-
-			const int32 LocalBodyIndex = RigidBodyIdx - PhysResult.BaseIndex;
-			const FRigidTransform3& TM = PhysResult.ParticleToWorldTransforms[LocalBodyIndex];
-
-			const FImplicitObject* Object = ClusterBuffer.GeometryPtrs[RigidBodyIdx].Get();
-			if (!Object)
-			{
-				return;
-			}
-
-			// Need to do narrow phase
-			Pair<FVec3, bool> Result = QueryGeom.FindDeepestIntersection(Object, Particles.CollisionParticles(RigidBodyIdx).Get(), FRigidTransform3(TM) * FRigidTransform3(GeomPose).Inverse(), 0);
-
-			if (Result.Second)
-			{
-				bFound = true;
-			}
-		});
-	}
-	return bFound;
-}
-
-int32 UseSlowSQ = 0;
-FAutoConsoleVariableRef CVarUseSlowSQ(TEXT("p.UseSlowSQ"), UseSlowSQ, TEXT(""));
-
-void FGeometryCollectionSQAccelerator::Raycast(const FVector& Start, const FVector& Dir, const Chaos::FReal DeltaMagnitude, ChaosInterface::FSQHitBuffer<ChaosInterface::FRaycastHit>& HitBuffer, EHitFlags OutputFlags, const FQueryFilterData& QueryFilterData, ICollisionQueryFilterCallbackBase& QueryCallback) const
-{
-	SCOPE_CYCLE_COUNTER(STAT_GCRaycast);
-	FChaosScopeSolverLock SolverScopeLock;
-
-	using namespace Chaos;
-
-	FChaosSolversModule* Module = FChaosSolversModule::GetModule();
-
-	TMap<const Chaos::FPhysicsSolver*, TArray<int32>> SolverIntersectionSets;
-
-	Chaos::TSpatialRay<Chaos::FReal, 3> Ray(Start, Start + Dir * DeltaMagnitude);
-
-#if WITH_PHYSX
-
-	const TArray<Chaos::FPhysicsSolver*>& Solvers = Module->GetSolvers();
-
-	for(const Chaos::FPhysicsSolver* Solver : Solvers)
-	{
-		if (!Solver)
-		{
-			continue;
-		}
-
-		TArray<int32> IntersectionSet = Solver->GetSpatialAcceleration()->FindAllIntersections(Ray);
-		Solver->ReleaseSpatialAcceleration();
-
-		const Chaos::FClusterBuffer& Buffer = Solver->GetRigidClustering().GetBufferedData();
-
-
-		const Chaos::FPhysicsSolver::FPhysicsProxyReverseMapping& ObjectMap = Solver->GetPhysicsProxyReverseMapping_GameThread();
-
-		FHitRaycast Hit;
-		int32 IntersectionSetSize = IntersectionSet.Num();
-		for(int32 i = 0; i < IntersectionSet.Num(); ++i)
-		{
-			const int32 IntersectParticleIndex = IntersectionSet[i];
-			const PhysicsProxyWrapper& ObjectWrapper = ObjectMap.PhysicsProxyReverseMappingArray[IntersectParticleIndex];
-			if (!ObjectWrapper.PhysicsProxy)
-			{
-				const FImplicitObject* Object = Buffer.GeometryPtrs[IntersectParticleIndex].Get();
-				// Ignore the ground plane
-				if (IntersectParticleIndex == 0 && Object->GetType(true) == ImplicitObjectType::Plane)
-				{
-					continue;
-				}
-				if (Object && !UseSlowSQ && Object->IsUnderlyingUnion())
-				{
-					const FImplicitObjectUnion* Union = static_cast<const TImplicitObjectUnion<Chaos::FReal, 3>*>(Object);
-					//hack: this is terrible because we have no buffered transform so could be off, but most of the time these things are static
-
-					{
-						const FRigidTransform3* TMPtr = Buffer.ClusterParentTransforms.Find(IntersectParticleIndex);
-
-						if (ensure(TMPtr))
-						{
-							if (!(ensure(!FMath::IsNaN(TMPtr->GetTranslation().X)) && ensure(!FMath::IsNaN(TMPtr->GetTranslation().Y)) && ensure(!FMath::IsNaN(TMPtr->GetTranslation().Z))))
-							{
-								continue;
-							}
-
-							const FVec3 StartLocal = TMPtr->InverseTransformPositionNoScale(Start);
-							const FVec3 DirLocal = TMPtr->InverseTransformVectorNoScale(Dir);
-							const FVec3 EndLocal = StartLocal + DirLocal * DeltaMagnitude;
-							Chaos::TSpatialRay<Chaos::FReal, 3> LocalRay(StartLocal, EndLocal);
-							const TArray<int32> IntersectingChildren = Union->FindAllIntersectingChildren(LocalRay);
-							IntersectionSet.Append(IntersectingChildren);
-						}
-						else
-						{
-							UE_LOG(LogChaos, Warning, TEXT("SQ: Could not find a valid transform for a cluster parent for faster child intersections."));
-						}
-					}
-				}
-				else
-				{
-					if (ensure(Buffer.MChildren.Contains(IntersectParticleIndex)))
-					{
-						const TArray<uint32>& Children = Buffer.MChildren[IntersectParticleIndex];
-						for (const uint32 Child : Children)
-						{
-							IntersectionSet.Add(Child);
-						}
-					}
-				}
-				continue;
-			}
-
-			if(ObjectWrapper.Type == EPhysicsProxyType::GeometryCollectionType && ensure(ObjectWrapper.PhysicsProxy))
-			{
-				LowLevelRaycastSingleElement(IntersectParticleIndex, Solver, Buffer, static_cast<FGeometryCollectionPhysicsProxy*>(ObjectWrapper.PhysicsProxy), Start, Dir, DeltaMagnitude, i >= IntersectionSetSize, OutputFlags, Hit);
-
-				// If we registered a hit
-				if(Hit.distance != PX_MAX_REAL && ensure(ObjectWrapper.PhysicsProxy))
-				{
-#if !WITH_IMMEDIATE_PHYSX && PHYSICS_INTERFACE_PHYSX
-					//todo(ocohen):hack placeholder while we convert over to non physx API
-					UGeometryCollectionComponent* Component = Cast<UGeometryCollectionComponent>(ObjectWrapper.PhysicsProxy->GetOwner());
-					check(Component);
-					if (Component->IsRegistered())
-					{
-						const FPhysicsActorHandle& ActorHandle = Component->DummyBodyInstance.GetPhysicsActorHandle();
-						PxRigidActor* PRigidActor = ActorHandle.SyncActor;
-						uint32 PNumShapes = PRigidActor->getNbShapes();
-						TArray<PxShape*> PShapes;
-						PShapes.AddZeroed(PNumShapes);
-						PRigidActor->getShapes(PShapes.GetData(), sizeof(PShapes[0]) * PNumShapes);
-						SetActor(Hit, ActorHandle.SyncActor);
-						SetShape(Hit, PShapes[0]);
-					}
-#else
-					check(false);	//this can't actually return nullptr since higher up API assumes both shape and actor exists in the low level
-					SetActor(Hit, nullptr);
-					SetShape(Hit, nullptr);	//todo(ocohen): what do we return for apeiron?
-#endif
-					Insert(HitBuffer, Hit, true);	//for now assume all blocking hits
-				}
-			}
-		}
-
-		Solver->GetRigidClustering().ReleaseBufferedData();
-		Solver->ReleasePhysicsProxyReverseMapping();
-	}
-
-#endif
-}
-
-DECLARE_CYCLE_STAT(TEXT("Sweep Broadphase"), STAT_SQSweepBroadPhase, STATGROUP_Chaos);
-DECLARE_CYCLE_STAT(TEXT("Sweep Narrowphase"), STAT_SQSweepNarrowPhase, STATGROUP_Chaos);
-
-//@todo(mlentine): Avoid duplicated code between this and overlap
-void FGeometryCollectionSQAccelerator::Sweep(const Chaos::FImplicitObject& QueryGeom, const FTransform& StartTM, const FVector& Dir, const Chaos::FReal DeltaMagnitude, ChaosInterface::FSQHitBuffer<ChaosInterface::FSweepHit>& HitBuffer, EHitFlags OutputFlags, const FQueryFilterData& QueryFilterData, ICollisionQueryFilterCallbackBase& QueryCallback) const
-{
-	SCOPE_CYCLE_COUNTER(STAT_GCSweep);
-	FChaosScopeSolverLock SolverScopeLock;
-
-	using namespace Chaos;
-#if WITH_PHYSX
-
-	TMap<const Chaos::FPhysicsSolver*, TArray<int32>> SolverIntersectionSets;
-
-	// Getter for intersections from the mapping above
-	auto GetIntersectionsFunc = [&](const Chaos::FPhysicsSolver* InSolver, const Chaos::FParticles& InCollisionParticles, Chaos::FReal InDeltaMag, const FTransform& InPose) -> TArray<int32>
-	{
-		SCOPE_CYCLE_COUNTER(STAT_SQSweepBroadPhase)
-
-		TArray<int32> PotentialIntersections;
-
-		if(InSolver->Enabled())
-		{
-			const Chaos::ISpatialAcceleration<Chaos::FReal, 3>* SpacialAcceleration = InSolver->GetSpatialAcceleration();
-
-			const int32 NumCollisionParticles = InCollisionParticles.Size();
-			for(int32 ParticleIndex = 0; ParticleIndex < NumCollisionParticles; ++ParticleIndex)
-			{
-				const Chaos::FVec3 RayStart = InPose.TransformPositionNoScale(InCollisionParticles.X(ParticleIndex));
-				const Chaos::FVec3 RayEnd = RayStart + Dir * InDeltaMag;
-
-				PotentialIntersections.Append(SpacialAcceleration->FindAllIntersections(Chaos::TSpatialRay<Chaos::FReal, 3>(RayStart, RayEnd)));
-			}
-
-			InSolver->ReleaseSpatialAcceleration();
-
-			PotentialIntersections.Sort();
-
-			for(int32 i = PotentialIntersections.Num() - 1; i > 0; i--)
-			{
-				if(PotentialIntersections[i] == PotentialIntersections[i - 1])
-				{
-					PotentialIntersections.RemoveAtSwap(i, 1, false);
-				}
-			}
-		}
-
-		return PotentialIntersections;
-	};
-
-	// Need somewhere to store our translated shape, similar to the PhysX geom holder
-	struct FLocalImplicitStorage
-	{
-		FLocalImplicitStorage()
-			: Capsule(FVector::ZeroVector, FVector::ZeroVector, 0.0f)
-			, Sphere(FVector::ZeroVector, 0.0f)
-			, Box(FVector::ZeroVector, FVector::ZeroVector)
-		{}
-
-		Chaos::FCapsule Capsule;
-		Chaos::TSphere<Chaos::FReal, 3> Sphere;
-		Chaos::FAABB3 Box;
-	};
-	FLocalImplicitStorage ImplicitStorage;
-	Chaos::FImplicitObject* Implicit = nullptr;
-
-	bool bHit = false;
-	FHitSweep Hit;
-	PxGeometryHolder Holder(QueryGeom);
-	Chaos::FParticles CollisionParticles;
-
-	if(Holder.getType() == PxGeometryType::eCAPSULE)
-	{
-		physx::PxCapsuleGeometry& PxCapsule = Holder.capsule();
-		
-		Chaos::FReal Radius = PxCapsule.radius;
-		Chaos::FReal HalfHeight = PxCapsule.halfHeight;
-		Chaos::FVec3 x1(-HalfHeight, 0, 0);
-		Chaos::FVec3 x2(HalfHeight, 0, 0);
-		
-		Chaos::FCapsule Capsule(x1, x2, Radius);
-		ImplicitStorage.Capsule = MoveTemp(Capsule);
-		Implicit = &ImplicitStorage.Capsule;
-
-		{
-			CollisionParticles.AddParticles(14);
-			CollisionParticles.X(0) = Chaos::FVec3(HalfHeight + Radius, 0, 0);
-			CollisionParticles.X(1) = Chaos::FVec3(-HalfHeight - Radius, 0, 0);
-			CollisionParticles.X(2) = Chaos::FVec3(HalfHeight, Radius, Radius);
-			CollisionParticles.X(3) = Chaos::FVec3(HalfHeight, -Radius, Radius);
-			CollisionParticles.X(4) = Chaos::FVec3(HalfHeight, -Radius, -Radius);
-			CollisionParticles.X(5) = Chaos::FVec3(HalfHeight, Radius, -Radius);
-			CollisionParticles.X(6) = Chaos::FVec3(0, Radius, Radius);
-			CollisionParticles.X(7) = Chaos::FVec3(0, -Radius, Radius);
-			CollisionParticles.X(8) = Chaos::FVec3(0, -Radius, -Radius);
-			CollisionParticles.X(9) = Chaos::FVec3(0, Radius, -Radius);
-			CollisionParticles.X(10) = Chaos::FVec3(-HalfHeight, Radius, Radius);
-			CollisionParticles.X(11) = Chaos::FVec3(-HalfHeight, -Radius, Radius);
-			CollisionParticles.X(12) = Chaos::FVec3(-HalfHeight, -Radius, -Radius);
-			CollisionParticles.X(13) = Chaos::FVec3(-HalfHeight, Radius, -Radius);
-		}
-	}
-	else if(Holder.getType() == PxGeometryType::eSPHERE)
-	{
-		physx::PxSphereGeometry& PxSphere = Holder.sphere();
-
-		Chaos::FReal Radius = PxSphere.radius;
-
-		Chaos::TSphere<Chaos::FReal, 3> Sphere(Chaos::FVec3(0), Radius);
-		ImplicitStorage.Sphere = MoveTemp(Sphere);
-		Implicit = &ImplicitStorage.Sphere;
-
-		{
-			CollisionParticles.AddParticles(6);
-			CollisionParticles.X(0) = Chaos::FVec3(Radius, 0, 0);
-			CollisionParticles.X(1) = Chaos::FVec3(-Radius, 0, 0);
-			CollisionParticles.X(2) = Chaos::FVec3(0, Radius, Radius);
-			CollisionParticles.X(3) = Chaos::FVec3(0, -Radius, Radius);
-			CollisionParticles.X(4) = Chaos::FVec3(0, -Radius, -Radius);
-			CollisionParticles.X(5) = Chaos::FVec3(0, Radius, -Radius);
-		}
-	}
-	else if(Holder.getType() == PxGeometryType::eBOX)
-	{
-		physx::PxBoxGeometry& PxBox = Holder.box();
-
-		Chaos::FVec3 x1(-P2UVector(PxBox.halfExtents));
-		Chaos::FVec3 x2(-x1);
-
-		Chaos::FAABB3 Box(x1, x2);
-		ImplicitStorage.Box = MoveTemp(Box);
-		Implicit = &ImplicitStorage.Box;
-
-		{
-			CollisionParticles.AddParticles(8);
-			CollisionParticles.X(0) = Chaos::FVec3(x1.X, x1.Y, x1.Z);
-			CollisionParticles.X(1) = Chaos::FVec3(x1.X, x1.Y, x2.Z);
-			CollisionParticles.X(2) = Chaos::FVec3(x1.X, x2.Y, x1.Z);
-			CollisionParticles.X(3) = Chaos::FVec3(x2.X, x1.Y, x1.Z);
-			CollisionParticles.X(4) = Chaos::FVec3(x2.X, x2.Y, x2.Z);
-			CollisionParticles.X(5) = Chaos::FVec3(x2.X, x2.Y, x1.Z);
-			CollisionParticles.X(6) = Chaos::FVec3(x2.X, x1.Y, x2.Z);
-			CollisionParticles.X(7) = Chaos::FVec3(x1.X, x2.Y, x2.Z);
-		}
-	}
-	else
-	{
-		// We don't support anything else for sweeps
-		ensureMsgf(false, TEXT("Unsupported query type used for sweep"));
-		return;
-	}
-
-	const TArray<Chaos::FPhysicsSolver*>& Solvers = FChaosSolversModule::GetModule()->GetSolvers();
-
-	for(const Chaos::FPhysicsSolver* Solver : Solvers)
-	{		
-		if (!Solver)
-		{
-			continue;
-		}
-
-		TArray<int32> IntersectionSet = GetIntersectionsFunc(Solver, CollisionParticles, DeltaMag, StartTM);
-
-		const Chaos::FPhysicsSolver::FPhysicsProxyReverseMapping& ObjectMap = Solver->GetPhysicsProxyReverseMapping_GameThread();
-		const Chaos::FClusterBuffer& Buffer = Solver->GetRigidClustering().GetBufferedData();
-
-
-		int32 IntersectionSetSize = IntersectionSet.Num();
-		for(int32 i = 0; i < IntersectionSet.Num(); ++i)
-		{
-			SCOPE_CYCLE_COUNTER(STAT_SQSweepNarrowPhase)
-
-			const int32 ParticleIndex = IntersectionSet[i];
-			const PhysicsProxyWrapper& ObjectWrapper = ObjectMap.PhysicsProxyReverseMappingArray[ParticleIndex];
-
-			if (!ObjectWrapper.PhysicsProxy)
-			{
-				const FImplicitObject* Object = Buffer.GeometryPtrs[ParticleIndex].Get();
-				
-				// Ignore ground plane
-				if (ParticleIndex == 0 && Object->GetType(true) == ImplicitObjectType::Plane)
-				{
-					continue;
-				}
-				if (Object && !UseSlowSQ && Object->IsUnderlyingUnion())
-				{
-					const FImplicitObjectUnion* Union = static_cast<const TImplicitObjectUnion<Chaos::FReal, 3>*>(Object);
-					//hack: this is terrible because we have no buffered transform so could be off, but most of the time these things are static
-
-					{
-						const FRigidTransform3* TMPtr = Buffer.ClusterParentTransforms.Find(ParticleIndex);
-
-						if (ensure(TMPtr))
-						{
-							if (!(ensure(!FMath::IsNaN(TMPtr->GetTranslation().X)) && ensure(!FMath::IsNaN(TMPtr->GetTranslation().Y)) && ensure(!FMath::IsNaN(TMPtr->GetTranslation().Z))))
-							{
-								continue;
-							}
-
-							const FVec3 StartLocal = TMPtr->InverseTransformPositionNoScale(StartTM.GetLocation());
-							const FVec3 DirLocal = TMPtr->InverseTransformVectorNoScale(Dir);
-							const FVec3 EndLocal = StartLocal + DirLocal * DeltaMag;
-							Chaos::TSpatialRay<Chaos::FReal, 3> LocalRay(StartLocal, EndLocal);
-							const TArray<int32> IntersectingChildren = Union->FindAllIntersectingChildren(LocalRay);
-							IntersectionSet.Append(IntersectingChildren);
-						}
-						else
-						{
-							UE_LOG(LogChaos, Warning, TEXT("SQ: Could not find a valid transform for a cluster parent for faster child intersections."));
-						}
-					}
-				}
-				else
-				{
-					if (ensure(Buffer.MChildren.Contains(ParticleIndex)))
-					{
-						const TArray<uint32>& Children = Buffer.MChildren[ParticleIndex];
-						for (const uint32 Child : Children)
-						{
-							IntersectionSet.Add(Child);
-						}
-					}
-				}
-				continue;
-			}
-
-			if(ObjectWrapper.Type == EPhysicsProxyType::GeometryCollectionType && ensure(ObjectWrapper.PhysicsProxy))
-			{
-				if(LowLevelSweepSingleElement(ParticleIndex, Solver, Buffer, static_cast<FGeometryCollectionPhysicsProxy*>(ObjectWrapper.PhysicsProxy), *Implicit, CollisionParticles, StartTM, Dir, DeltaMag, i >= IntersectionSetSize, Hit))
-				{
-					//todo(mlentine): This is duplicated from above and should be merged
-					//todo(ocohen):hack placeholder while we convert over to non physx API
-					UGeometryCollectionComponent* Component = Cast<UGeometryCollectionComponent>(ObjectWrapper.PhysicsProxy->GetOwner());
-					check(Component);
-
-					if (Component->IsRegistered())
-					{
-#if !WITH_IMMEDIATE_PHYSX && PHYSICS_INTERFACE_PHYSX
-						const FPhysicsActorHandle& ActorHandle = Component->DummyBodyInstance.GetPhysicsActorHandle();
-						PxRigidActor* PRigidActor = ActorHandle.SyncActor;
-						uint32 PNumShapes = PRigidActor->getNbShapes();
-						TArray<PxShape*> PShapes;
-						PShapes.AddZeroed(PNumShapes);
-						PRigidActor->getShapes(PShapes.GetData(), sizeof(PShapes[0]) * PNumShapes);
-						SetActor(Hit, ActorHandle.SyncActor);
-						SetShape(Hit, PShapes[0]);
-#else
-						check(false);	//this can't actually return nullptr since higher up API assumes both shape and actor exists in the low level
-						SetActor(Hit, nullptr);
-						SetShape(Hit, nullptr);	//todo(ocohen): what do we return for apeiron?
-#endif
-						Insert(HitBuffer, Hit, true);	//for now assume all blocking hits
-					}
-				}
-			}
-		}
-
-		Solver->GetRigidClustering().ReleaseBufferedData();
-
-		Solver->ReleasePhysicsProxyReverseMapping();
-	}
-#endif
-}
-
-bool LowLevelOverlapSingleElement(int32 InParticleIndex, const Chaos::FPhysicsSolver* InSolver, const Chaos::FClusterBuffer& ClusterBuffer, const FGeometryCollectionPhysicsProxy* InObject, const Chaos::FImplicitObject& QueryGeom, const FTransform& InPose, FHitOverlap& OutHit)
-{
-	using namespace Chaos;
-
-	checkSlow(InSolver);
-	checkSlow(InObject);
-
-	const FGeometryCollectionResults& PhysResult = InObject->GetPhysicsResults().GetGameDataForRead();
-
-	const TManagedArray<int32>& RigidBodyIdArray = PhysResult.RigidBodyIds;
-	const TManagedArray<FTransform>& TransformArray = PhysResult.Transforms;
-	const TArray<bool>& DisabledFlags = PhysResult.DisabledStates;
-
-	const TPBDRigidParticles<Chaos::FReal, 3>& Particles = InSolver->GetRigidParticles();
-
-	if(!IsValidIndexAndTransform(PhysResult, Particles, TransformArray, DisabledFlags, InParticleIndex, false))
-	{
-		return false;
-	}
-
-	const int32 LocalBodyIndex = InParticleIndex - PhysResult.BaseIndex;
-	const FRigidTransform3& TM = PhysResult.ParticleToWorldTransforms[LocalBodyIndex];
-	const FImplicitObject* Object = ClusterBuffer.GeometryPtrs[InParticleIndex].Get();
-	
-	if(!Object)
-	{
-		return false;
-	}
-
-	Pair<FVec3, bool> Result = QueryGeom.FindDeepestIntersection(Object, Particles.CollisionParticles(InParticleIndex).Get(), FRigidTransform3(TM) * FRigidTransform3(InPose).Inverse(), 0);
-
-	return Result.Second;
-}
-
-
-void FGeometryCollectionSQAccelerator::Overlap(const Chaos::FImplicitObject& QueryGeom, const FTransform& GeomPose, ChaosInterface::FSQHitBuffer<ChaosInterface::FOverlapHit>& HitBuffer, const FQueryFilterData& QueryFilterData, ICollisionQueryFilterCallbackBase& QueryCallback) const
-{
-	return;	//todo: This is currently broken because it doesn'thandle cluster unions. Need to fix function - disabling entirely for now
-	SCOPE_CYCLE_COUNTER(STAT_GCOverlap);
-	FChaosScopeSolverLock SolverScopeLock;
-
-#if WITH_PHYSX
-
-	TMap<const Chaos::FPhysicsSolver*, TArray<int32>> SolverIntersectionSets;
-
-	// Getter for intersections from the mapping above
-	auto GetIntersectionsFunc = [&](const Chaos::FPhysicsSolver* InSolver, Chaos::FImplicitObject* InGeometry, const FTransform& InPose) -> TArray<int32>*
-	{
-		TArray<int32>* CurrentIntersections = SolverIntersectionSets.Find(InSolver);
-
-		if(!CurrentIntersections)
-		{
-			// This is safe to access here, it's buffered from the physics thread
-			SolverIntersectionSets.Add(InSolver, InSolver->GetSpatialAcceleration()->FindAllIntersections(InGeometry->BoundingBox().TransformedBox(GeomPose)));
-			InSolver->ReleaseSpatialAcceleration();
-		}
-
-		return SolverIntersectionSets.Find(InSolver);
-	};
-
-	// Need somewhere to store our translated shape, similar to the PhysX geom holder
-	struct FLocalImplicitStorage
-	{
-		FLocalImplicitStorage()
-			: Capsule(FVector::ZeroVector, FVector::ZeroVector, 0.0f)
-			, Sphere(FVector::ZeroVector, 0.0f)
-			, Box(FVector::ZeroVector, FVector::ZeroVector)
-		{}
-
-		Chaos::FCapsule Capsule;
-		Chaos::TSphere<Chaos::FReal, 3> Sphere;
-		Chaos::FAABB3 Box;
-	};
-	FLocalImplicitStorage ImplicitStorage;
-
-	bool bHit = false;
-	FHitOverlap Hit;
-	PxGeometryHolder Holder(QueryGeom);
-	Chaos::FImplicitObject* Implicit = nullptr;
-
-	if(Holder.getType() == PxGeometryType::eCAPSULE)
-	{
-		physx::PxCapsuleGeometry& PxCapsule = Holder.capsule();
-		Chaos::FReal Radius = PxCapsule.radius;
-		Chaos::FReal HalfHeight = PxCapsule.halfHeight - Radius;
-		Chaos::FVec3 x1(-HalfHeight, 0, 0);
-		Chaos::FVec3 x2(HalfHeight, 0, 0);
-		Chaos::FCapsule Capsule(x1, x2, Radius);
-		ImplicitStorage.Capsule = MoveTemp(Capsule);
-		Implicit = &ImplicitStorage.Capsule;
-	}
-	else if(Holder.getType() == PxGeometryType::eSPHERE)
-	{
-		physx::PxSphereGeometry& PxSphere = Holder.sphere();
-		Chaos::FReal Radius = PxSphere.radius;
-		Chaos::TSphere<Chaos::FReal, 3> Sphere(Chaos::FVec3(0), Radius);
-		ImplicitStorage.Sphere = MoveTemp(Sphere);
-		Implicit = &ImplicitStorage.Sphere;
-	}
-	else if(Holder.getType() == PxGeometryType::eBOX)
-	{
-		physx::PxBoxGeometry& PxBox = Holder.box();
-		Chaos::FVec3 x1(-P2UVector(PxBox.halfExtents));
-		Chaos::FVec3 x2(-x1);
-		Chaos::FAABB3 Box(x1, x2);
-		ImplicitStorage.Box = MoveTemp(Box);
-		Implicit = &ImplicitStorage.Box;
-	}
-	else
-	{
-		ensureMsgf(false, TEXT("Unsupported query type used for overlap"));
-		return;
-	}
-
-	FChaosSolversModule* Module = FChaosSolversModule::GetModule();
-	const TArray<Chaos::FPhysicsSolver*> Solvers = Module->GetSolvers();
-
-	for(const Chaos::FPhysicsSolver* Solver : Solvers)
-	{
-		// Collect all intersections for this solver
-		TArray<int32> IntersectionSet = Solver->GetSpatialAcceleration()->FindAllIntersections(Implicit->BoundingBox().TransformedBox(GeomPose));
-		Solver->ReleaseSpatialAcceleration();
-
-		const Chaos::FPhysicsSolver::FPhysicsProxyReverseMapping& ObjectMap = Solver->GetPhysicsProxyReverseMapping_GameThread();
-
-		const Chaos::FClusterBuffer& ClusterBuffer = Solver->GetRigidClustering().GetBufferedData();
-
-		for(int32 i = 0; i < IntersectionSet.Num(); ++i)
-		{
-			const int32 IntersectParticleIndex = IntersectionSet[i];
-			const PhysicsProxyWrapper& ObjectWrapper = ObjectMap.PhysicsProxyReverseMappingArray[IntersectParticleIndex];
-
-			if(ObjectWrapper.Type == EPhysicsProxyType::GeometryCollectionType && ensure(ObjectWrapper.PhysicsProxy))
-			{
-				if(LowLevelOverlapSingleElement(IntersectParticleIndex, Solver, ClusterBuffer, static_cast<FGeometryCollectionPhysicsProxy*>(ObjectWrapper.PhysicsProxy), *Implicit, GeomPose, Hit))
-				{
-					//todo(mlentine): This is duplicated from above and should be merged
-					//todo(ocohen):hack placeholder while we convert over to non physx API
-					UGeometryCollectionComponent* Component = Cast<UGeometryCollectionComponent>(ObjectWrapper.PhysicsProxy->GetOwner());
-					check(Component);
-
-					if (Component->IsRegistered())
-					{
-#if !WITH_IMMEDIATE_PHYSX && PHYSICS_INTERFACE_PHYSX
-						const FPhysicsActorHandle& ActorHandle = Component->DummyBodyInstance.GetPhysicsActorHandle();
-						PxRigidActor* PRigidActor = ActorHandle.SyncActor;
-						uint32 PNumShapes = PRigidActor->getNbShapes();
-						TArray<PxShape*> PShapes;
-						PShapes.AddZeroed(PNumShapes);
-						PRigidActor->getShapes(PShapes.GetData(), sizeof(PShapes[0]) * PNumShapes);
-						SetActor(Hit, ActorHandle.SyncActor);
-						SetShape(Hit, PShapes[0]);
-#else
-						check(false);	//this can't actually return nullptr since higher up API assumes both shape and actor exists in the low level
-						SetActor(Hit, nullptr);
-						SetShape(Hit, nullptr);	//todo(ocohen): what do we return for apeiron?
-#endif
-						InsertOverlap(HitBuffer, Hit);
-					}
-				}
-			}
-		}
-
-		Solver->GetRigidClustering().ReleaseBufferedData();
-		Solver->ReleasePhysicsProxyReverseMapping();
-	}
-
-#endif
-}
-
-#endif // TODO_REIMPLEMENT_SCENEQUERY_CROSSENGINE
-
-void FGeometryCollectionSQAccelerator::AddComponent(UGeometryCollectionComponent* Component)
-{
-	if (ensure(Component))
-	{
-//		UE_LOG(LogTemp, Log, TEXT("SQAccel: Adding %s (Outer %s)"), *GetNameSafe(Component), *GetNameSafe(Component->GetOuter()));
-		Components.Add(Component);
-	}
-}
-
-void FGeometryCollectionSQAccelerator::RemoveComponent(UGeometryCollectionComponent* Component)
-{
-	if (ensure(Component))
-	{
-//		UE_LOG(LogTemp, Log, TEXT("SQAccel: Removing %s (Outer %s)"), *GetNameSafe(Component), *GetNameSafe(Component->GetOuter()));
-		Components.Remove(Component);
-	}
-}
-#endif
diff --git a/Engine/Source/Runtime/Experimental/GeometryCollectionEngine/Private/GeometryCollection/GeometryCollectionSQAccelerator.h b/Engine/Source/Runtime/Experimental/GeometryCollectionEngine/Private/GeometryCollection/GeometryCollectionSQAccelerator.h
--- a/Engine/Source/Runtime/Experimental/GeometryCollectionEngine/Private/GeometryCollection/GeometryCollectionSQAccelerator.h
+++ b/Engine/Source/Runtime/Experimental/GeometryCollectionEngine/Private/GeometryCollection/GeometryCollectionSQAccelerator.h
@@ -1,24 +1,2 @@
 // Copyright Epic Games, Inc. All Rights Reserved.
 #pragma once
-
-
-#if !WITH_CHAOS_NEEDS_TO_BE_FIXED
-#include "SQAccelerator.h"
-
-class UGeometryCollectionComponent;
-
-class FGeometryCollectionSQAccelerator
-{
-public:
-	void Raycast(const FVector& Start, const FVector& Dir, const float DeltaMagnitude, ChaosInterface::FSQHitBuffer<ChaosInterface::FRaycastHit>& HitBuffer, EHitFlags OutputFlags, const FQueryFilterData& QueryFilterData, ICollisionQueryFilterCallbackBase& QueryCallback) const;
-	void Sweep(const Chaos::FImplicitObject& QueryGeom, const FTransform& StartTM, const FVector& Dir, const float DeltaMagnitude, ChaosInterface::FSQHitBuffer<ChaosInterface::FSweepHit>& HitBuffer, EHitFlags OutputFlags, const FQueryFilterData& QueryFilterData, ICollisionQueryFilterCallbackBase& QueryCallback) const;
-	void Overlap(const Chaos::FImplicitObject& QueryGeom, const FTransform& GeomPose, ChaosInterface::FSQHitBuffer<ChaosInterface::FOverlapHit>& HitBuffer, const FQueryFilterData& QueryFilterData, ICollisionQueryFilterCallbackBase& QueryCallback) const;
-	virtual ~FGeometryCollectionSQAccelerator() {}
-
-	void AddComponent(UGeometryCollectionComponent* Component);
-	void RemoveComponent(UGeometryCollectionComponent* Component);
-
-private:
-	TSet<UGeometryCollectionComponent*> Components;
-};
-#endif
\ No newline at end of file
diff --git a/Engine/Source/Runtime/Experimental/GeometryCollectionEngine/Private/GeometryCollection/PhysicsAssetSimulation.cpp b/Engine/Source/Runtime/Experimental/GeometryCollectionEngine/Private/GeometryCollection/PhysicsAssetSimulation.cpp
--- a/Engine/Source/Runtime/Experimental/GeometryCollectionEngine/Private/GeometryCollection/PhysicsAssetSimulation.cpp
+++ b/Engine/Source/Runtime/Experimental/GeometryCollectionEngine/Private/GeometryCollection/PhysicsAssetSimulation.cpp
@@ -29,7 +29,6 @@
 
 void FPhysicsAssetSimulationUtil::BuildParams(const UObject* Caller, const AActor* OwningActor, const USkeletalMeshComponent* SkelMeshComponent, const UPhysicsAsset* PhysicsAsset, FSkeletalMeshPhysicsProxyParams& Params)
 {
-#if WITH_PHYSX
 	Params.LocalToWorld = Params.InitialTransform = OwningActor->GetTransform();
 
 	FVector ActorOriginVec, ActorBoxExtentVec;	// LWC_TODO: Perf pessimization
@@ -281,96 +280,8 @@ void FPhysicsAssetSimulationUtil::BuildParams(const UObject* Caller, const AActo
 					Bounds += Points[i];
 				}
 
-#if PHYSICS_INTERFACE_PHYSX
-				if (physx::PxConvexMesh *ConvexMesh = Elem.GetConvexMesh())
-				{
-					ConvexMesh->acquireReference();
-					const PxU32 NumPolygons = ConvexMesh->getNbPolygons();
-					TArray<Chaos::TVec3<int32>> Triangles;
-					Triangles.Reserve(NumPolygons);
-					PxHullPolygon Polygon;
-					for (uint32 i = 0; i < NumPolygons; i++)
-					{
-						if (ConvexMesh->getPolygonData(i, Polygon))
-						{
-							if (Polygon.mNbVerts == 3)
-							{
-								// getPolygonData() augments what getIndexBuffer() returns?
-								const auto LocalIndices = ConvexMesh->getIndexBuffer() + Polygon.mIndexBase;
-								Triangles.Add(
-									Chaos::TVec3<int32>(LocalIndices[0], LocalIndices[1], LocalIndices[2]));
-							}
-							else if (Polygon.mNbVerts > 3)
-							{
-								for (uint32 j = 2; j < Polygon.mNbVerts; j++)
-								{
-									const auto LocalIndices = ConvexMesh->getIndexBuffer() + Polygon.mIndexBase;
-									Triangles.Add(
-										Chaos::TVec3<int32>(LocalIndices[0], LocalIndices[j], LocalIndices[j - 1]));
-								}
-							}
-						}
-					}
-					ConvexMesh->release();
-
-					const Chaos::FVec3 Extents(Bounds.Max - Bounds.Min);
-					const Chaos::FReal LocalMaxExtent = Bounds.GetExtent().GetMax();
-					Bounds = Bounds.ExpandBy(Extents.Size() / 10);
-
-					//TODO: this gets used later which is incorrect. However, it was already broken because it passed a view and then moved the underlying data.
-					//Keeping as is for now, needs fixing
-					Chaos::FParticles Particles(MoveTemp(Points));
-					TUniquePtr<Chaos::FTriangleMesh> TriangleMesh(
-						new Chaos::FTriangleMesh(MoveTemp(Triangles)));
-					Chaos::FErrorReporter ErrorReporter;
-
-					const Chaos::FReal ActorMaxExtent = ActorBoxExtent.Max();
-					const int32 NormalizedMinRes = FMath::CeilToInt(ActorMaxExtent > SMALL_NUMBER ? LocalMaxExtent / ActorMaxExtent * Params.MinRes : Params.MinRes);
-					const int32 NormalizedMaxRes = FMath::CeilToInt(ActorMaxExtent > SMALL_NUMBER ? LocalMaxExtent / ActorMaxExtent * Params.MaxRes : Params.MaxRes);
-
-					Chaos::FLevelSet* LevelSet =
-						static_cast<Chaos::FLevelSet*>(
-							FCollisionStructureManager::NewImplicit(
-								ErrorReporter,
-								Particles,							// mesh particles
-								*TriangleMesh,						// triangle mesh
-								Bounds,								// collision bounds
-								Bounds.GetExtent().Size() * 0.5,	// radius
-								NormalizedMinRes, NormalizedMaxRes,	// min, max resolution
-								100.f,								// CollisionObjectReduction, 100 = no shrinking
-								Params.CollisionType,				// collision type
-								EImplicitTypeEnum::Chaos_Implicit_LevelSet)); // implicit type
-					if (!ErrorReporter.EncounteredAnyErrors())
-					{
-						const int32 Index = AnalyticShapeGroup->Add(Xf, LevelSet);
-						if (DoCollisionGeom)
-						{
-							AnalyticShapeGroup->SetCollisionTopology(
-								Index, MoveTemp(Points), MoveTemp(*TriangleMesh).GetSurfaceElements());
-						}
-						UE_LOG(USkeletalMeshSimulationComponentLogging, Log,
-							TEXT("USkeletalMeshSimulationComponent::OnCreatePhysicsState() - Caller: %p Actor: '%s' - "
-								"ADDED LEVELSET"), Caller, &OwningActor->GetName()[0]);
-					}
-					else
-					{
-						// Something went wrong.  Use a sphere instead.
-						// TODO: report the error.
-						delete LevelSet;
-
-						AnalyticShapeGroup->Add(
-							Xf,
-							new Chaos::TSphere<Chaos::FReal, 3>(Bounds.GetCenter(), LocalMaxExtent * 0.5));
-
-						UE_LOG(USkeletalMeshSimulationComponentLogging, Log,
-							TEXT("USkeletalMeshSimulationComponent::OnCreatePhysicsState() - Caller: %p Actor: '%s' - "
-								"ADDED (LEVELSET) SPHERE"), Caller, &OwningActor->GetName()[0]);
-					}
-				}
-#elif WITH_CHAOS
 				// TODO: Use Chaos convex.
 				CHAOS_ENSURE(false);
-#endif
 			}
 		}
 		Params.BoneHierarchy.Add(MoveTemp(AnalyticShapeGroup));
@@ -458,20 +369,24 @@ void FPhysicsAssetSimulationUtil::BuildParams(const UObject* Caller, const AActo
 			"Bone hierarchy, num shape groups: %d"),
 		Caller, &OwningActor->GetName()[0],
 		Params.BoneHierarchy.GetAnalyticShapeGroups().Num());
-#endif
 }
 
 bool FPhysicsAssetSimulationUtil::UpdateAnimState(const UObject* Caller, const AActor* OwningActor, const USkeletalMeshComponent* SkelMeshComponent, const float Dt, FSkeletalMeshPhysicsProxyParams& Params)
 {
 	int32 UpdatedBones = 0;
 
-#if WITH_PHYSX
 	const USkeletalMesh* SkeletalMesh = SkelMeshComponent->SkeletalMesh;
+	
 	if (!SkeletalMesh)
+	{
 		return false;
+	}
+
 	const FSkeletalMeshRenderData* SkelMeshRenderData = SkeletalMesh->GetResourceForRendering();
 	if (!SkelMeshRenderData || !SkelMeshRenderData->LODRenderData.Num())
+	{
 		return false;
+	}
 
 	// Calling this doesn't seem to matter...
 	//SkelMeshComponent->UpdateKinematicBonesToAnim(
@@ -514,6 +429,6 @@ bool FPhysicsAssetSimulationUtil::UpdateAnimState(const UObject* Caller, const A
 			UpdatedBones++;
 		}
 	}
-#endif
+
 	return (UpdatedBones > 0);
 }
diff --git a/Engine/Source/Runtime/Experimental/GeometryCollectionEngine/Public/GeometryCollection/GeometryCollectionComponent.h b/Engine/Source/Runtime/Experimental/GeometryCollectionEngine/Public/GeometryCollection/GeometryCollectionComponent.h
--- a/Engine/Source/Runtime/Experimental/GeometryCollectionEngine/Public/GeometryCollection/GeometryCollectionComponent.h
+++ b/Engine/Source/Runtime/Experimental/GeometryCollectionEngine/Public/GeometryCollection/GeometryCollectionComponent.h
@@ -847,10 +847,8 @@ private:
 	TArray<FManagedArrayBase**> CopyOnWriteAttributeList;
 
 	// Temporary dummies to interface with Physx expectations of the SQ syatem
-#if WITH_PHYSX
 	friend class FGeometryCollectionSQAccelerator;
 	FBodyInstance DummyBodyInstance;
-#endif
 
 	// Temporary storage for body setup in order to initialise a dummy body instance
 	UPROPERTY(Transient)
diff --git a/Engine/Source/Runtime/Foliage/Public/ProceduralFoliageInstance.h b/Engine/Source/Runtime/Foliage/Public/ProceduralFoliageInstance.h
--- a/Engine/Source/Runtime/Foliage/Public/ProceduralFoliageInstance.h
+++ b/Engine/Source/Runtime/Foliage/Public/ProceduralFoliageInstance.h
@@ -11,13 +11,6 @@ class UActorComponent;
 class UFoliageType;
 struct FProceduralFoliageInstance;
 
-#if WITH_PHYSX
-namespace physx
-{
-	class PxRigidActor;
-}
-#endif
-
 UENUM(BlueprintType)
 namespace ESimulationOverlap
 {
diff --git a/Engine/Source/Runtime/Landscape/Classes/LandscapeHeightfieldCollisionComponent.h b/Engine/Source/Runtime/Landscape/Classes/LandscapeHeightfieldCollisionComponent.h
--- a/Engine/Source/Runtime/Landscape/Classes/LandscapeHeightfieldCollisionComponent.h
+++ b/Engine/Source/Runtime/Landscape/Classes/LandscapeHeightfieldCollisionComponent.h
@@ -1,6 +1,5 @@
 // Copyright Epic Games, Inc. All Rights Reserved.
 
-
 #pragma once
 
 #include "CoreMinimal.h"
@@ -26,16 +25,8 @@ struct FNavigableGeometryExport;
 
 namespace Chaos
 {
-class FHeightField;
-}
-
-#if WITH_PHYSX
-namespace physx
-{
-	class PxMaterial;
-	class PxHeightField;
+	class FHeightField;
 }
-#endif // WITH_PHYSX
 
 UCLASS(MinimalAPI, Within=LandscapeProxy)
 class ULandscapeHeightfieldCollisionComponent : public UPrimitiveComponent
@@ -92,24 +83,13 @@ class ULandscapeHeightfieldCollisionComponent : public UPrimitiveComponent
 	{
 		FGuid Guid;
 
-#if WITH_PHYSX
-		/** List of PxMaterials used on this landscape */
-		TArray<physx::PxMaterial*> UsedPhysicalMaterialArray;
-		physx::PxHeightField* RBHeightfield = nullptr;
-		physx::PxHeightField* RBHeightfieldSimple = nullptr;
-#if WITH_EDITORONLY_DATA
-		physx::PxHeightField* RBHeightfieldEd = nullptr; // Used only by landscape editor, does not have holes in it
-#endif	//WITH_EDITORONLY_DATA
-#endif	//WITH_PHYSX
-
-#if WITH_CHAOS
 		TArray<Chaos::FMaterialHandle> UsedChaosMaterials;
 		TUniquePtr<Chaos::FHeightField> Heightfield;
 	    TUniquePtr<Chaos::FHeightField> HeightfieldSimple;
+
 #if WITH_EDITORONLY_DATA
 		TUniquePtr<Chaos::FHeightField> EditorHeightfield;
 #endif // WITH_EDITORONLY_DATA
-#endif // WITH_CHAOS
 
 		FHeightfieldGeometryRef(FGuid& InGuid);
 
diff --git a/Engine/Source/Runtime/Landscape/Classes/LandscapeMeshCollisionComponent.h b/Engine/Source/Runtime/Landscape/Classes/LandscapeMeshCollisionComponent.h
--- a/Engine/Source/Runtime/Landscape/Classes/LandscapeMeshCollisionComponent.h
+++ b/Engine/Source/Runtime/Landscape/Classes/LandscapeMeshCollisionComponent.h
@@ -18,20 +18,10 @@
 class UPhysicalMaterial;
 struct FNavigableGeometryExport;
 
-#if WITH_PHYSX
-namespace physx
-{
-	class PxMaterial;
-	class PxTriangleMesh;
-}
-#elif WITH_CHAOS
 namespace Chaos
 {
 	class FTriangleMeshImplicitObject;
 }
-#endif
-
-
 
 UCLASS()
 class ULandscapeMeshCollisionComponent : public ULandscapeHeightfieldCollisionComponent
@@ -53,22 +43,11 @@ public:
 	{
 		FGuid Guid;
 
-#if WITH_PHYSX
-		/** List of PxMaterials used on this landscape */
-		TArray<physx::PxMaterial*>	UsedPhysicalMaterialArray;
-		physx::PxTriangleMesh*		RBTriangleMesh;
-#if WITH_EDITORONLY_DATA
-		physx::PxTriangleMesh*		RBTriangleMeshEd; // Used only by landscape editor, does not have holes in it
-#endif	//WITH_EDITORONLY_DATA
-#endif	//WITH_PHYSX
-
-#if WITH_CHAOS
 		TArray<Chaos::FMaterialHandle> UsedChaosMaterials;
 		TUniquePtr<Chaos::FTriangleMeshImplicitObject> Trimesh;
 #if WITH_EDITORONLY_DATA
 		TUniquePtr<Chaos::FTriangleMeshImplicitObject> EditorTrimesh;
 #endif // WITH_EDITORONLY_DATA
-#endif // WITH_CHAOS
 
 		FTriMeshGeometryRef();
 		FTriMeshGeometryRef(FGuid& InGuid);
diff --git a/Engine/Source/Runtime/Landscape/Landscape.Build.cs b/Engine/Source/Runtime/Landscape/Landscape.Build.cs
--- a/Engine/Source/Runtime/Landscape/Landscape.Build.cs
+++ b/Engine/Source/Runtime/Landscape/Landscape.Build.cs
@@ -39,10 +39,6 @@ public class Landscape : ModuleRules
 		);
 
 		SetupModulePhysicsSupport(Target);
-		if (Target.bCompilePhysX && Target.bBuildEditor)
-		{
-			DynamicallyLoadedModuleNames.Add("PhysXCooking");
-		}
 
 		if (Target.Type == TargetType.Editor || Target.Type == TargetType.Program)
 		{
