Fracture Tool clustering operation perf optimization
- Added a cvar ffor result validation ( off by default )
- disable validation when RemoveElements
- 2500 piece flattening operation goes from 1m8s to 9s

diff --git a/Engine/Source/Runtime/Experimental/Chaos/Private/GeometryCollection/GeometryCollectionClusteringUtility.cpp b/Engine/Source/Runtime/Experimental/Chaos/Private/GeometryCollection/GeometryCollectionClusteringUtility.cpp
--- a/Engine/Source/Runtime/Experimental/Chaos/Private/GeometryCollection/GeometryCollectionClusteringUtility.cpp
+++ b/Engine/Source/Runtime/Experimental/Chaos/Private/GeometryCollection/GeometryCollectionClusteringUtility.cpp
@@ -6,6 +6,9 @@
 #include "Containers/Set.h"
 #include "Async/ParallelFor.h"
 
+static int32 ChaosValidateResultsOfEditOperations = 0;
+static FAutoConsoleVariableRef CVarChaosStillCheckDistanceThreshold(TEXT("p.fracture.ValidateResultsOfEditOperations"), ChaosValidateResultsOfEditOperations, TEXT("When on this will enable result validation for fracture tool edit operations (can be slow for large geometry collection) [def:0]"));
+
 void FGeometryCollectionClusteringUtility::ClusterBonesUnderNewNode(FGeometryCollection* GeometryCollection, const int32 InsertAtIndex, const TArray<int32>& SelectedBones, bool CalcNewLocalTransform, bool Validate)
 {
 	check(GeometryCollection);
@@ -738,23 +741,26 @@ void FGeometryCollectionClusteringUtility::CollapseSelectedHierarchy(int8 Level,
 
 void FGeometryCollectionClusteringUtility::ValidateResults(FGeometryCollection* GeometryCollection)
 {
-	const TManagedArray<int32>& Parents = GeometryCollection->Parent;
-	const TManagedArray<TSet<int32>>& Children = GeometryCollection->Children;
-	const TManagedArray<FString>& BoneNames = GeometryCollection->BoneName;
-
-	// there should only ever be one root node
-	int NumRootNodes = 0;
-	for (int i = 0; i < Parents.Num(); i++)
+	if (ChaosValidateResultsOfEditOperations)
 	{
-		if (Parents[i] == FGeometryCollection::Invalid)
+		const TManagedArray<int32>& Parents = GeometryCollection->Parent;
+		const TManagedArray<TSet<int32>>& Children = GeometryCollection->Children;
+		const TManagedArray<FString>& BoneNames = GeometryCollection->BoneName;
+
+		// there should only ever be one root node
+		int NumRootNodes = 0;
+		for (int i = 0; i < Parents.Num(); i++)
 		{
-			NumRootNodes++;
+			if (Parents[i] == FGeometryCollection::Invalid)
+			{
+				NumRootNodes++;
+			}
 		}
-	}
-	check(NumRootNodes == 1);
+		check(NumRootNodes == 1);
 
-	ensure(GeometryCollection->HasContiguousFaces());
-	ensure(GeometryCollection->HasContiguousVertices());
+		ensure(GeometryCollection->HasContiguousFaces());
+		ensure(GeometryCollection->HasContiguousVertices());
+	}
 }
 
 void FGeometryCollectionClusteringUtility::GetLeafBones(const FGeometryCollection* GeometryCollection, int BoneIndex, bool bOnlyRigids, TArray<int32>& LeafBonesOut)
@@ -887,7 +893,9 @@ void FGeometryCollectionClusteringUtility::RemoveDanglingClusters(FGeometryColle
 	if (DeletionList.Num())
 	{
 		DeletionList.Sort();
-		GeometryCollection->RemoveElements(FGeometryCollection::TransformGroup, DeletionList);
+		FManagedArrayCollection::FProcessingParameters Params;
+		Params.bDoValidation = false; // for perf reasons
+		GeometryCollection->RemoveElements(FGeometryCollection::TransformGroup, DeletionList, Params);
 	}
 }
 
diff --git a/Engine/Source/Runtime/Experimental/Chaos/Private/GeometryCollection/GeometryCollectionConvexUtility.cpp b/Engine/Source/Runtime/Experimental/Chaos/Private/GeometryCollection/GeometryCollectionConvexUtility.cpp
--- a/Engine/Source/Runtime/Experimental/Chaos/Private/GeometryCollection/GeometryCollectionConvexUtility.cpp
+++ b/Engine/Source/Runtime/Experimental/Chaos/Private/GeometryCollection/GeometryCollectionConvexUtility.cpp
@@ -1265,7 +1265,9 @@ void FGeometryCollectionConvexUtility::RemoveConvexHulls(FGeometryCollection* Ge
 		if (ConvexIndices.Num())
 		{
 			ConvexIndices.Sort();
-			GeometryCollection->RemoveElements("Convex", ConvexIndices);
+			FManagedArrayCollection::FProcessingParameters Params;
+			Params.bDoValidation = false; // for perf reasons
+			GeometryCollection->RemoveElements("Convex", ConvexIndices, Params);
 		}
 	}
 }
