Chaos Cloth: XPBDSpringConstraint ispc fix when using maps. This is a recent regression (CL20034474). I was not updating Lambda each iteration.

XPBD cannot be enabled via the UI, so this shouldn't have any effect to current users.

diff --git a/Engine/Source/Runtime/Experimental/Chaos/Private/Chaos/XPBDSpringConstraints.ispc b/Engine/Source/Runtime/Experimental/Chaos/Private/Chaos/XPBDSpringConstraints.ispc
--- a/Engine/Source/Runtime/Experimental/Chaos/Private/Chaos/XPBDSpringConstraints.ispc
+++ b/Engine/Source/Runtime/Experimental/Chaos/Private/Chaos/XPBDSpringConstraints.ispc
@@ -105,11 +105,13 @@ export void ApplyXPBDSpringConstraintsWithWeightMaps(uniform FVector4f PandInvM[
 
 		const float CombinedInvMass = InvM2 + InvM1;
 
-		const varying float DLambda = (Offset - Alpha * Lambda) / (CombinedInvMass + Alpha);
+		const varying float ComputedDLambda = (Offset - Alpha * Lambda) / (CombinedInvMass + Alpha);
+		const varying float DLambda = select(Stiffness < MinStiffness, 0.f, ComputedDLambda);
 
 		const FVector3f ComputedDelta = DLambda * Direction;
 
-		const FVector3f Delta = VectorSelect(Stiffness < MinStiffness || (InvM2 == FLOAT_ZERO && InvM1 == FLOAT_ZERO), FloatZeroVector, ComputedDelta);
+		const FVector3f Delta = VectorSelect((InvM2 == FLOAT_ZERO && InvM1 == FLOAT_ZERO), FloatZeroVector, ComputedDelta);
+		Lambdas[i] = Lambda + DLambda;
 
 		if (InvM1 > FLOAT_ZERO)
 		{
