re-initialize the convex properties interface on geometry collection Empty(), to match the initial 'empty' state

diff --git a/Engine/Source/Runtime/Experimental/Chaos/Private/GeometryCollection/GeometryCollection.cpp b/Engine/Source/Runtime/Experimental/Chaos/Private/GeometryCollection/GeometryCollection.cpp
--- a/Engine/Source/Runtime/Experimental/Chaos/Private/GeometryCollection/GeometryCollection.cpp
+++ b/Engine/Source/Runtime/Experimental/Chaos/Private/GeometryCollection/GeometryCollection.cpp
@@ -658,6 +658,8 @@ void FGeometryCollection::Empty()
 	{
 		EmptyGroup(GroupName);
 	}
+	// re-initialize the interface to its default initial state
+	FGeometryCollectionConvexPropertiesInterface::InitializeInterface();
 }
 
 void FGeometryCollection::ReorderElements(FName Group, const TArray<int32>& NewOrder)
