Chaos: stop external strain fields from calling the clustering code directly
- Remove the external strain map and add a ExternalStrains property array for the clustered particles
- ExternalStrain field now just add to the external strain property of the particle and let the clustering code process it once per physics step
- Remove the unecessary maps returned by BreakingModel
- Updated headless chaos

diff --git a/Engine/Source/Programs/HeadlessChaos/Private/GeometryCollection/GeometryCollectionTestClustering.cpp b/Engine/Source/Programs/HeadlessChaos/Private/GeometryCollection/GeometryCollectionTestClustering.cpp
--- a/Engine/Source/Programs/HeadlessChaos/Private/GeometryCollection/GeometryCollectionTestClustering.cpp
+++ b/Engine/Source/Programs/HeadlessChaos/Private/GeometryCollection/GeometryCollectionTestClustering.cpp
@@ -307,8 +307,8 @@ namespace GeometryCollectionTest
 
 			if (Frame == 2)
 			{
-				TMap<FGeometryParticleHandle*, FReal> ExternalStrains = { {ParticleHandles[0], 50.0f} };
-				Clustering.BreakingModel(&ExternalStrains);
+				ParticleHandles[0].AddExternalStrain(50.0f);
+				Clustering.BreakingModel();
 			}
 
 			DisabledFlags.Reset();
diff --git a/Engine/Source/Runtime/Experimental/Chaos/Private/Chaos/PBDRigidClustering.cpp b/Engine/Source/Runtime/Experimental/Chaos/Private/Chaos/PBDRigidClustering.cpp
--- a/Engine/Source/Runtime/Experimental/Chaos/Private/Chaos/PBDRigidClustering.cpp
+++ b/Engine/Source/Runtime/Experimental/Chaos/Private/Chaos/PBDRigidClustering.cpp
@@ -307,7 +307,7 @@ namespace Chaos
 							{
 								// First disable breaking data generation - this is not a break we're just reclustering under a dynamic parent.
 								TGuardValue<bool> BreakFlagGuard(DoGenerateBreakingData, false);
-								Children = ReleaseClusterParticles(ActiveCluster, nullptr, true);
+								Children = ReleaseClusterParticles(ActiveCluster, true);
 							}
 
 							NewClusterGroups[ClusterGroupID].Bodies.Append(Children.Array());
@@ -543,18 +543,16 @@ namespace Chaos
 	DECLARE_CYCLE_STAT(TEXT("TPBDRigidClustering<>::ReleaseClusterParticles(STRAIN)"), STAT_ReleaseClusterParticles_STRAIN, STATGROUP_Chaos);
 	TSet<FPBDRigidParticleHandle*> FRigidClustering::ReleaseClusterParticles(
 		FPBDRigidClusteredParticleHandle* ClusteredParticle,
-		const TMap<FGeometryParticleHandle*, Chaos::FReal>* ExternalStrainMap,
 		bool bForceRelease)
 	{
 		SCOPE_CYCLE_COUNTER(STAT_ReleaseClusterParticles_STRAIN);
 
-		return ReleaseClusterParticlesImpl(ClusteredParticle, ExternalStrainMap, bForceRelease, true /*bCreateNewClusters*/);
+		return ReleaseClusterParticlesImpl(ClusteredParticle, bForceRelease, true /*bCreateNewClusters*/);
 	}
 	
 	DECLARE_CYCLE_STAT(TEXT("TPBDRigidClustering<>::ReleaseClusterParticlesNoInternalCluster"), STAT_ReleaseClusterParticlesNoInternalCluster, STATGROUP_Chaos);
 	TSet<FPBDRigidParticleHandle*> FRigidClustering::ReleaseClusterParticlesNoInternalCluster(
 		FPBDRigidClusteredParticleHandle* ClusteredParticle,
-		const TMap<FGeometryParticleHandle*, Chaos::FReal>* ExternalStrainMap,
 		bool bForceRelease)
 	{
 		/* This is a near duplicate of the ReleaseClusterParticles() method with the internal cluster creation removed.
@@ -564,16 +562,21 @@ namespace Chaos
 		
 		SCOPE_CYCLE_COUNTER(STAT_ReleaseClusterParticlesNoInternalCluster);
 
-		return ReleaseClusterParticlesImpl(ClusteredParticle, ExternalStrainMap, bForceRelease, false /*bCreateNewClusters*/);
+		return ReleaseClusterParticlesImpl(ClusteredParticle, bForceRelease, false /*bCreateNewClusters*/);
 	}
 	
 	TSet<FPBDRigidParticleHandle*> FRigidClustering::ReleaseClusterParticlesImpl(
 		FPBDRigidClusteredParticleHandle* ClusteredParticle,
-		const TMap<FGeometryParticleHandle*, Chaos::FReal>* ExternalStrainMap,
 		bool bForceRelease,
 		bool bCreateNewClusters)
 	{	
 		TSet<FPBDRigidParticleHandle*> ActivatedChildren;
+
+		if (!ensureMsgf((ClusteredParticle->Parent()==nullptr), TEXT("Removing a cluster that still has a parent")))
+		{
+			return ActivatedChildren;;
+		};
+
 		if (!ensureMsgf(MChildren.Contains(ClusteredParticle), TEXT("Removing Cluster that does not exist!")))
 		{
 			return ActivatedChildren;
@@ -589,18 +592,7 @@ namespace Chaos
 				continue;
 			}
 
-			FReal ChildStrain = 0.0;
-
-			if(ExternalStrainMap)
-			{
-				const FReal* MapStrain = ExternalStrainMap->Find(Child);
-				ChildStrain = MapStrain ? *MapStrain : Child->CollisionImpulses();
-			}
-			else
-			{
-				ChildStrain = Child->CollisionImpulses();
-			}
-
+			const FReal ChildStrain = Child->CollisionImpulses() + Child->GetExternalStrains();
 			if (ChildStrain >= Child->Strain() || bForceRelease)
 			{
 				//UE_LOG(LogTemp, Warning, TEXT("Releasing child %d from parent %p due to strain %.5f Exceeding internal strain %.5f (Source: %s)"), ChildIdx, ClusteredParticle, ChildStrain, Child->Strain(), bForceRelease ? TEXT("Forced by caller") : ExternalStrainMap ? TEXT("External") : TEXT("Collision"));
@@ -647,6 +639,10 @@ namespace Chaos
 					SendBreakingEvent(Child);
 				}
 			}
+
+			// no need anymore of the external strain let's clear it
+			// @todo(chaos) impulse collision are right now cleared somewhere else, but we should consolidate that eventually
+			Child->ClearExternalStrains();
 		}
 
 		if (ActivatedChildren.Num() > 0)
@@ -725,7 +721,8 @@ namespace Chaos
 					{
 						if (ensure(!ClusterHandle || ClusteredChildHandle->ClusterIds().Id == ClusterHandle))
 						{
-							FakeStrain.Add(ClusteredChildHandle, TNumericLimits<FReal>::Max());
+							ClusteredChildHandle->ClearExternalStrains();
+							ClusteredChildHandle->AddExternalStrain(TNumericLimits<FReal>::Max());
 							ClusterHandle = ClusteredChildHandle->ClusterIds().Id;
 						}
 						else
@@ -737,7 +734,7 @@ namespace Chaos
 			}
 			if (ClusterHandle)
 			{
-				ActivatedBodies = ReleaseClusterParticles(ClusterHandle->CastToClustered(), &FakeStrain);
+				ActivatedBodies = ReleaseClusterParticles(ClusterHandle->CastToClustered());
 			}
 			DoGenerateBreakingData = bPreDoGenerateData;
 		}
@@ -793,6 +790,7 @@ namespace Chaos
 			//  Monitor the MStrain array for 0 or less values.
 			//  That will trigger a break too.
 			//
+			bool bPotentialBreak = false;
 			{
 				SCOPE_CYCLE_COUNTER(STAT_UpdateDirtyImpulses);
 				for (const auto& ActiveCluster : GetTopLevelClusterParents())
@@ -811,6 +809,10 @@ namespace Chaos
 										ClusteredChild->CollisionImpulse() = FLT_MAX;
 										MCollisionImpulseArrayDirty = true;
 									}
+									else if (!bPotentialBreak && ClusteredChild->GetExternalStrains() > 0)
+									{
+										bPotentialBreak = true;
+									}
 								}
 							}
 						}
@@ -818,11 +820,10 @@ namespace Chaos
 				}
 			}
 
-			if (MCollisionImpulseArrayDirty)
+			if (MCollisionImpulseArrayDirty || bPotentialBreak)
 			{
 				SCOPE_CYCLE_COUNTER(STAT_UpdateDirtyImpulses);
-				TMap<FPBDRigidClusteredParticleHandle*, TSet<FPBDRigidParticleHandle*>> ClusterToActivatedChildren = 
-					BreakingModel();
+				BreakingModel();
 			} // end if MCollisionImpulseArrayDirty
 
 		} // end if MParticles.Size()
@@ -831,9 +832,7 @@ namespace Chaos
 	}
 
 	DECLARE_CYCLE_STAT(TEXT("TPBDRigidClustering<>::BreakingModel()"), STAT_BreakingModel, STATGROUP_Chaos);
-	TMap<FPBDRigidClusteredParticleHandle*, TSet<FPBDRigidParticleHandle*>> 
-	FRigidClustering::BreakingModel(
-		TMap<FGeometryParticleHandle*, Chaos::FReal>* ExternalStrainMap)
+	void FRigidClustering::BreakingModel()
 	{
 		SCOPE_CYCLE_COUNTER(STAT_BreakingModel);
 
@@ -844,15 +843,11 @@ namespace Chaos
 			ClusteredParticlesToProcess.Add(Particle.Handle()->CastToClustered());
 		}
 
-		TMap<FPBDRigidClusteredParticleHandle*, TSet<FPBDRigidParticleHandle*>> AllActivatedChildren;
-
 		for (FPBDRigidClusteredParticleHandle* ClusteredParticle : ClusteredParticlesToProcess)
 		{
 			if (ClusteredParticle->ClusterIds().NumChildren)
 			{
-				AllActivatedChildren.Add(
-					ClusteredParticle,
-					ReleaseClusterParticles(ClusteredParticle, ExternalStrainMap));
+				ReleaseClusterParticles(ClusteredParticle);
 			}
 			else
 			{
@@ -866,8 +861,6 @@ namespace Chaos
 				}
 			}
 		}
-
-		return AllActivatedChildren;
 	}
 
 	DECLARE_CYCLE_STAT(TEXT("FRigidClustering::Visitor"), STAT_ClusterVisitor, STATGROUP_Chaos);
diff --git a/Engine/Source/Runtime/Experimental/Chaos/Private/PhysicsProxy/GeometryCollectionPhysicsProxy.cpp b/Engine/Source/Runtime/Experimental/Chaos/Private/PhysicsProxy/GeometryCollectionPhysicsProxy.cpp
--- a/Engine/Source/Runtime/Experimental/Chaos/Private/PhysicsProxy/GeometryCollectionPhysicsProxy.cpp
+++ b/Engine/Source/Runtime/Experimental/Chaos/Private/PhysicsProxy/GeometryCollectionPhysicsProxy.cpp
@@ -1552,7 +1552,7 @@ void FGeometryCollectionPhysicsProxy::BreakInternalClusterParent(TArray<int32>&&
 					Chaos::FRigidClustering& Clustering = RBDSolver->GetEvolution()->GetRigidClustering();
 					if (!ClusteredParent->Disabled() && Clustering.GetChildrenMap().Contains(ClusteredParent))
 					{
-						Clustering.ReleaseClusterParticles(ClusteredParent, nullptr, true /* bForceRelease */);
+						Clustering.ReleaseClusterParticles(ClusteredParent, true /* bForceRelease */);
 					}
 				}
 			}
diff --git a/Engine/Source/Runtime/Experimental/Chaos/Public/Chaos/PBDRigidClusteredParticles.h b/Engine/Source/Runtime/Experimental/Chaos/Public/Chaos/PBDRigidClusteredParticles.h
--- a/Engine/Source/Runtime/Experimental/Chaos/Public/Chaos/PBDRigidClusteredParticles.h
+++ b/Engine/Source/Runtime/Experimental/Chaos/Public/Chaos/PBDRigidClusteredParticles.h
@@ -87,7 +87,7 @@ class TPBDRigidClusteredParticles : public TPBDRigidParticles<T, d>
 	, MCollisionImpulses(MoveTemp(Other.MCollisionImpulses))
 	, MStrains(MoveTemp(Other.MStrains))
 	, MConnectivityEdges(MoveTemp(Other.MConnectivityEdges))
-
+	, MExternalStrains(MoveTemp(Other.MExternalStrains))
 	{
 		InitHelper();
 	}
@@ -115,6 +115,10 @@ class TPBDRigidClusteredParticles : public TPBDRigidParticles<T, d>
 	auto& CollisionImpulses(int32 Idx) { return MCollisionImpulses[Idx]; }
 	auto& CollisionImpulsesArray() { return MCollisionImpulses; }
 
+	const auto& ExternalStrains(int32 Idx) const { return MExternalStrains[Idx]; }
+	auto& ExternalStrains(int32 Idx) { return MExternalStrains[Idx]; }
+	auto& ExternalStrainsArray() { return MExternalStrains; }
+	
 	const auto& Strains(int32 Idx) const { return MStrains[Idx]; }
 	auto& Strains(int32 Idx) { return MStrains[Idx]; }
 
@@ -159,6 +163,7 @@ class TPBDRigidClusteredParticles : public TPBDRigidParticles<T, d>
 		  TArrayCollection::AddArray(&MCollisionImpulses);
 		  TArrayCollection::AddArray(&MStrains);
 		  TArrayCollection::AddArray(&MConnectivityEdges);
+	  	  TArrayCollection::AddArray(&MExternalStrains);
 	  }
 
 	  TArrayCollectionArray<ClusterId> MClusterIds;
@@ -173,6 +178,10 @@ class TPBDRigidClusteredParticles : public TPBDRigidParticles<T, d>
 	  // Collision Impulses
 	  TArrayCollectionArray<T> MCollisionImpulses;
 
+	  // external strains ( use by fields )
+	  // @todo(chaos) we should eventually merge MCollisionImpulses into MExternalStrains when CLustering code has been updated to not clear the impulses just before processing them 
+	  TArrayCollectionArray<T> MExternalStrains; 
+
 	  // User set parameters
 	  TArrayCollectionArray<T> MStrains;
 
diff --git a/Engine/Source/Runtime/Experimental/Chaos/Public/Chaos/PBDRigidClustering.h b/Engine/Source/Runtime/Experimental/Chaos/Public/Chaos/PBDRigidClustering.h
--- a/Engine/Source/Runtime/Experimental/Chaos/Public/Chaos/PBDRigidClustering.h
+++ b/Engine/Source/Runtime/Experimental/Chaos/Public/Chaos/PBDRigidClustering.h
@@ -121,12 +121,10 @@ public:
 	*/
 	TSet<FPBDRigidParticleHandle*> ReleaseClusterParticles(
 		FPBDRigidClusteredParticleHandle* ClusteredParticle, 
-		const TMap<FGeometryParticleHandle*, Chaos::FReal>* ExternalStrainMap = nullptr,
 		bool bForceRelease = false);
 
 	TSet<FPBDRigidParticleHandle*> ReleaseClusterParticlesNoInternalCluster(
 		FPBDRigidClusteredParticleHandle* ClusteredParticle,
-		const TMap<FGeometryParticleHandle*, Chaos::FReal>* ExternalStrainMap = nullptr,
 		bool bForceRelease = false);
 
 	/*
@@ -168,9 +166,7 @@ public:
 	*    encoded strain. The remainder strains are propagated back down to
 	*    the children clusters.
 	*/
-	TMap<FPBDRigidClusteredParticleHandle*, TSet<FPBDRigidParticleHandle*>> BreakingModel(
-		TMap<FGeometryParticleHandle*, Chaos::FReal>* ExternalStrainMap = nullptr);
-
+	void BreakingModel();
 
 	//
 	// Access
@@ -305,7 +301,6 @@ public:
 
 	TSet<FPBDRigidParticleHandle*> ReleaseClusterParticlesImpl(
 		FPBDRigidClusteredParticleHandle* ClusteredParticle, 
-		const TMap<FGeometryParticleHandle*, FReal>* ExternalStrainMap,
 		bool bForceRelease,
 		bool bCreateNewClusters);
 	
diff --git a/Engine/Source/Runtime/Experimental/Chaos/Public/Chaos/ParticleHandle.h b/Engine/Source/Runtime/Experimental/Chaos/Public/Chaos/ParticleHandle.h
--- a/Engine/Source/Runtime/Experimental/Chaos/Public/Chaos/ParticleHandle.h
+++ b/Engine/Source/Runtime/Experimental/Chaos/Public/Chaos/ParticleHandle.h
@@ -1208,6 +1208,10 @@ public:
 	T& CollisionImpulses() { return PBDRigidClusteredParticles->CollisionImpulses(ParticleIdx); }
 	void SetCollisionImpulses(const T Value) { PBDRigidClusteredParticles->CollisionImpulses(ParticleIdx) = Value; }
 
+	T GetExternalStrains() const { return PBDRigidClusteredParticles->ExternalStrains(ParticleIdx); }
+	void AddExternalStrain(const T Value) { PBDRigidClusteredParticles->ExternalStrains(ParticleIdx) += Value; }
+	void ClearExternalStrains() { PBDRigidClusteredParticles->ExternalStrains(ParticleIdx) = static_cast<T>(0); }
+	
 	const T& Strain() const { return PBDRigidClusteredParticles->Strains(ParticleIdx); }
 	T& Strain() { return PBDRigidClusteredParticles->Strains(ParticleIdx); }
 	void SetStrain(const T Value) { PBDRigidClusteredParticles->Strains(ParticleIdx) = Value; }
diff --git a/Engine/Source/Runtime/Experimental/Chaos/Public/PhysicsProxy/FieldSystemProxyHelper.h b/Engine/Source/Runtime/Experimental/Chaos/Public/PhysicsProxy/FieldSystemProxyHelper.h
--- a/Engine/Source/Runtime/Experimental/Chaos/Public/PhysicsProxy/FieldSystemProxyHelper.h
+++ b/Engine/Source/Runtime/Experimental/Chaos/Public/PhysicsProxy/FieldSystemProxyHelper.h
@@ -291,34 +291,6 @@ namespace Chaos
 		}
 	}
 
-	/**
-	 * Update the solver breaking model based on external strain
-	 * @param    Rigidsolver Rigid solver owning the breaking model
-	 * @param    ExternalStrain Strain to be used to update the breaking model
-	 */
-	static void UpdateSolverBreakingModel(Chaos::FPBDRigidsSolver* RigidSolver, TMap<Chaos::FGeometryParticleHandle*, Chaos::FReal>& ExternalStrain)
-	{
-		// Capture the results from the breaking model to post-process
-		TMap<Chaos::FPBDRigidClusteredParticleHandle*, TSet<Chaos::FPBDRigidParticleHandle*>> BreakResults =
-			RigidSolver->GetEvolution()->GetRigidClustering().BreakingModel(&ExternalStrain);
-
-		// If clusters broke apart then we'll have activated new particles that have no relationship to the proxy that now owns them
-		// Here we attach each new particle to the proxy of the parent particle that owns it.
-		for (const TPair<Chaos::TPBDRigidClusteredParticleHandle<Chaos::FReal, 3>*, TSet<Chaos::FPBDRigidParticleHandle*>> & Iter : BreakResults)
-		{
-			const TSet<Chaos::FPBDRigidParticleHandle*>& Activated = Iter.Value;
-
-			for (Chaos::FPBDRigidParticleHandle* Handle : Activated)
-			{
-				if (Handle->PhysicsProxy() == nullptr)
-				{
-					IPhysicsProxyBase* ParentProxy = Iter.Key->PhysicsProxy();
-					Handle->SetPhysicsProxy(ParentProxy);
-				}
-			}
-		}
-	}
-
 	/**
 	 * Update the handle sleeping linear and angular theshold
 	 * @param    Rigidsolver Rigid solver owning the particle handle
@@ -514,17 +486,18 @@ namespace Chaos
 		{
 			SCOPE_CYCLE_COUNTER(STAT_ParamUpdateField_ExternalClusterStrain);
 			{
-				TMap<Chaos::FGeometryParticleHandle*, Chaos::FReal> ExternalStrain;
-
 				static_cast<const FFieldNode<float>*>(FieldCommand.RootNode.Get())->Evaluate(FieldContext, ResultsView);
 				for (const FFieldContextIndex& Index : FieldContext.GetEvaluatedSamples())
 				{
-					if (ResultsView[Index.Result] > 0)
+					const float ExternalStrainValue = ResultsView[Index.Result];
+					if (ExternalStrainValue > 0)
 					{
-						ExternalStrain.Add(ParticleHandles[Index.Sample], ResultsView[Index.Result]);
+						if (Chaos::FPBDRigidClusteredParticleHandle* ClusteredParticle = ParticleHandles[Index.Sample]->CastToClustered())
+						{
+							ClusteredParticle->AddExternalStrain(ExternalStrainValue);
+						}
 					}
 				}
-				UpdateSolverBreakingModel(RigidSolver, ExternalStrain);
 			}
 		}
 		else if (FieldCommand.PhysicsType == EFieldPhysicsType::Field_Kill)
