Generate IncidentElements for TetrahedralCollection at initialization. Implement TManagedArray for FIntVector2, TArray<FIntVector2>, and FIntVector4.

Depends on 20489689 which adds operator<< to FIntVector2.

diff --git a/Engine/Source/Runtime/Experimental/Chaos/Public/GeometryCollection/ManagedArray.h b/Engine/Source/Runtime/Experimental/Chaos/Public/GeometryCollection/ManagedArray.h
--- a/Engine/Source/Runtime/Experimental/Chaos/Public/GeometryCollection/ManagedArray.h
+++ b/Engine/Source/Runtime/Experimental/Chaos/Public/GeometryCollection/ManagedArray.h
@@ -67,6 +67,16 @@ inline void TryBulkSerializeManagedArray(Chaos::FChaosArchive& Ar, TArray<uint8>
 	Array.BulkSerialize(Ar);
 }
 
+inline void TryBulkSerializeManagedArray(Chaos::FChaosArchive& Ar, TArray<FIntVector2>& Array)
+{
+	Array.BulkSerialize(Ar);
+}
+
+inline void TryBulkSerializeManagedArray(Chaos::FChaosArchive& Ar, TArray<FIntVector4>& Array)
+{
+	Array.BulkSerialize(Ar);
+}
+
 /***
 *  Managed Array Base
 *
@@ -250,11 +260,16 @@ public:
 	*/
 	FORCEINLINE TManagedArrayBase& operator=(TManagedArrayBase<ElementType>&& Other)
 	{
-		Array = MoveTemp(Other.Array);
-		return *this;
+		return this->operator=(MoveTemp(Other.Array));
 	}
 
-
+	FORCEINLINE TManagedArrayBase& operator=(TArray<ElementType>&& Other)
+	{
+		// ryan - is it okay to check that the size matches?
+		ensureMsgf(Array.Num() == 0 || Array.Num() == Other.Num(), TEXT("TManagedArrayBase<T>::operator=(TArray<T>&&) : Invalid array size."));
+		Array = MoveTemp(Other);
+		return *this;
+	}
 
 	/**
 	* Virtual Destructor 
@@ -781,3 +796,187 @@ public:
 	}
 };
 
+template<>
+class TManagedArray<FIntVector2> : public TManagedArrayBase<FIntVector2>
+{
+public:
+	using TManagedArrayBase<FIntVector2>::Num;
+
+	FORCEINLINE TManagedArray()
+	{}
+
+	FORCEINLINE TManagedArray(const TArray<FIntVector2>& Other)
+		: TManagedArrayBase<FIntVector2>(Other)
+	{}
+
+	FORCEINLINE TManagedArray(const TManagedArray<FIntVector2>& Other) = delete;
+	FORCEINLINE TManagedArray(TManagedArray<FIntVector2>&& Other) = default;
+	FORCEINLINE TManagedArray(TArray<FIntVector2>&& Other)
+		: TManagedArrayBase<FIntVector2>(MoveTemp(Other))
+	{}
+	FORCEINLINE TManagedArray& operator=(TManagedArray<FIntVector2>&& Other) = default;
+
+	virtual ~TManagedArray()
+	{}
+
+	virtual void Reindex(const TArray<int32>& Offsets, const int32& FinalSize, const TArray<int32>& SortedDeletionList) override
+	{
+		UE_LOG(UManagedArrayLogging, Log, TEXT("TManagedArray<FIntVector>[%p]::Reindex()"), this);
+		int32 ArraySize = Num(), MaskSize = Offsets.Num();
+		for (int32 Index = 0; Index < ArraySize; Index++)
+		{
+			const FIntVector2& RemapVal = this->operator[](Index);
+			for (int i = 0; i < 2; i++)
+			{
+				if (0 <= RemapVal[i])
+				{
+					ensure(RemapVal[i] < MaskSize);
+					this->operator[](Index)[i] -= Offsets[RemapVal[i]];
+					ensure(-1 <= this->operator[](Index)[i] && this->operator[](Index)[i] <= FinalSize);
+				}
+			}
+		}
+	}
+
+	virtual void ReindexFromLookup(const TArray<int32>& InverseNewOrder) override
+	{
+		int32 ArraySize = Num();
+		for (int32 Index = 0; Index < ArraySize; Index++)
+		{
+			FIntVector2& RemapVal = this->operator[](Index);
+			for (int i = 0; i < 2; i++)
+			{
+				if (RemapVal[i] >= 0)
+				{
+					RemapVal[i] = InverseNewOrder[RemapVal[i]];
+				}
+			}
+		}
+	}
+};
+
+template<>
+class TManagedArray<TArray<FIntVector2>> : public TManagedArrayBase<TArray<FIntVector2>>
+{
+public:
+	using TManagedArrayBase<TArray<FIntVector2>>::Num;
+
+	FORCEINLINE TManagedArray()
+	{}
+
+	FORCEINLINE TManagedArray(const TArray<TArray<FIntVector2>>& Other)
+		: TManagedArrayBase<TArray<FIntVector2>>(Other)
+	{}
+
+	FORCEINLINE TManagedArray(const TManagedArray<TArray<FIntVector2>>& Other) = delete;
+	FORCEINLINE TManagedArray(TManagedArray<TArray<FIntVector2>>&& Other) = default;
+	FORCEINLINE TManagedArray(TArray<TArray<FIntVector2>>&& Other)
+		: TManagedArrayBase<TArray<FIntVector2>>(MoveTemp(Other))
+	{}
+	FORCEINLINE TManagedArray& operator=(TManagedArray<TArray<FIntVector2>>&& Other) = default;
+
+	virtual ~TManagedArray()
+	{}
+
+	virtual void Reindex(const TArray<int32>& Offsets, const int32& FinalSize, const TArray<int32>& SortedDeletionList) override
+	{
+		UE_LOG(UManagedArrayLogging, Log, TEXT("TManagedArray<FIntVector>[%p]::Reindex()"), this);
+		int32 ArraySize = Num(), MaskSize = Offsets.Num();
+		for (int32 Index = 0; Index < ArraySize; Index++)
+		{
+			const TArray<FIntVector2>& RemapValArray = this->operator[](Index);
+			for (int32 ArrayIndex = 0; ArrayIndex < RemapValArray.Num(); ArrayIndex++)
+			{
+				const FIntVector2& RemapVal = RemapValArray[ArrayIndex];
+				for (int i = 0; i < 2; i++)
+				{
+					if (0 <= RemapVal[i])
+					{
+						ensure(RemapVal[i] < MaskSize);
+						this->operator[](Index)[ArrayIndex][i] -= Offsets[RemapVal[i]];
+						ensure(-1 <= this->operator[](Index)[ArrayIndex][i] && this->operator[](Index)[ArrayIndex][i] <= FinalSize);
+					}
+				}
+			}
+		}
+	}
+
+	virtual void ReindexFromLookup(const TArray<int32>& InverseNewOrder) override
+	{
+		int32 ArraySize = Num();
+		for (int32 Index = 0; Index < ArraySize; Index++)
+		{
+			TArray<FIntVector2>& RemapValArray = this->operator[](Index);
+			for (int32 ArrayIndex = 0; ArrayIndex < RemapValArray.Num(); ArrayIndex++)
+			{
+				FIntVector2& RemapVal = RemapValArray[ArrayIndex];
+				for (int i = 0; i < 2; i++)
+				{
+					if (RemapVal[i] >= 0)
+					{
+						RemapVal[i] = InverseNewOrder[RemapVal[i]];
+					}
+				}
+			}
+		}
+	}
+};
+
+template<>
+class TManagedArray<FIntVector4> : public TManagedArrayBase<FIntVector4>
+{
+public:
+	using TManagedArrayBase<FIntVector4>::Num;
+
+	FORCEINLINE TManagedArray()
+	{}
+
+	FORCEINLINE TManagedArray(const TArray<FIntVector4>& Other)
+		: TManagedArrayBase<FIntVector4>(Other)
+	{}
+
+	FORCEINLINE TManagedArray(const TManagedArray<FIntVector4>& Other) = delete;
+	FORCEINLINE TManagedArray(TManagedArray<FIntVector4>&& Other) = default;
+	FORCEINLINE TManagedArray(TArray<FIntVector4>&& Other)
+		: TManagedArrayBase<FIntVector4>(MoveTemp(Other))
+	{}
+	FORCEINLINE TManagedArray& operator=(TManagedArray<FIntVector4>&& Other) = default;
+
+	virtual ~TManagedArray()
+	{}
+
+	virtual void Reindex(const TArray<int32>& Offsets, const int32& FinalSize, const TArray<int32>& SortedDeletionList) override
+	{
+		UE_LOG(UManagedArrayLogging, Log, TEXT("TManagedArray<FIntVector>[%p]::Reindex()"), this);
+		int32 ArraySize = Num(), MaskSize = Offsets.Num();
+		for (int32 Index = 0; Index < ArraySize; Index++)
+		{
+			const FIntVector4& RemapVal = this->operator[](Index);
+			for (int i = 0; i < 4; i++)
+			{
+				if (0 <= RemapVal[i])
+				{
+					ensure(RemapVal[i] < MaskSize);
+					this->operator[](Index)[i] -= Offsets[RemapVal[i]];
+					ensure(-1 <= this->operator[](Index)[i] && this->operator[](Index)[i] <= FinalSize);
+				}
+			}
+		}
+	}
+
+	virtual void ReindexFromLookup(const TArray<int32>& InverseNewOrder) override
+	{
+		int32 ArraySize = Num();
+		for (int32 Index = 0; Index < ArraySize; Index++)
+		{
+			FIntVector4& RemapVal = this->operator[](Index);
+			for (int i = 0; i < 4; i++)
+			{
+				if (RemapVal[i] >= 0)
+				{
+					RemapVal[i] = InverseNewOrder[RemapVal[i]];
+				}
+			}
+		}
+	}
+};
diff --git a/Engine/Source/Runtime/Experimental/Chaos/Public/GeometryCollection/ManagedArrayTypeValues.inl b/Engine/Source/Runtime/Experimental/Chaos/Public/GeometryCollection/ManagedArrayTypeValues.inl
--- a/Engine/Source/Runtime/Experimental/Chaos/Public/GeometryCollection/ManagedArrayTypeValues.inl
+++ b/Engine/Source/Runtime/Experimental/Chaos/Public/GeometryCollection/ManagedArrayTypeValues.inl
@@ -44,6 +44,9 @@ MANAGED_ARRAY_TYPE(TArray<FVector2f>, Vector2DArray)
 MANAGED_ARRAY_TYPE(double, Double)
 MANAGED_ARRAY_TYPE(FIntVector4, IntVector4)
 MANAGED_ARRAY_TYPE(FVector3d, Vector3d)
+MANAGED_ARRAY_TYPE(FIntVector2, IntVector2)
+MANAGED_ARRAY_TYPE(TArray<FIntVector2>, IntVector2Array)
+
 
 // NOTE: new types must be added at the bottom to keep serialization from breaking
 
