Fix crash when enabling group clustering and replication

diff --git a/Engine/Source/Runtime/Experimental/GeometryCollectionEngine/Private/GeometryCollection/GeometryCollectionComponent.cpp b/Engine/Source/Runtime/Experimental/GeometryCollectionEngine/Private/GeometryCollection/GeometryCollectionComponent.cpp
--- a/Engine/Source/Runtime/Experimental/GeometryCollectionEngine/Private/GeometryCollection/GeometryCollectionComponent.cpp
+++ b/Engine/Source/Runtime/Experimental/GeometryCollectionEngine/Private/GeometryCollection/GeometryCollectionComponent.cpp
@@ -1962,8 +1962,15 @@ void UGeometryCollectionComponent::OnRegister()
 {
 	//UE_LOG(UGCC_LOG, Log, TEXT("GeometryCollectionComponent[%p]::OnRegister()[%p]"), this,RestCollection );
 	ResetDynamicCollection();
-
-	SetIsReplicated(bEnableReplication);
+	
+	if (ensureMsgf(ClusterGroupIndex != 0, TEXT("Replication with cluster groups is not supported - disabling replication")))
+	{
+		SetIsReplicated(false);	
+	}
+	else
+	{
+		SetIsReplicated(bEnableReplication);	
+	}
 
 	InitializeEmbeddedGeometry();
 
