Dataflow : Serialization
- Add support for dynamic property serialization.

diff --git a/Engine/Source/Runtime/Experimental/Dataflow/Core/Private/Dataflow/DataflowNode.cpp b/Engine/Source/Runtime/Experimental/Dataflow/Core/Private/Dataflow/DataflowNode.cpp
--- a/Engine/Source/Runtime/Experimental/Dataflow/Core/Private/Dataflow/DataflowNode.cpp
+++ b/Engine/Source/Runtime/Experimental/Dataflow/Core/Private/Dataflow/DataflowNode.cpp
@@ -79,12 +79,102 @@ namespace Dataflow
 		return nullptr;
 	}
 
+	TMap<FName, FProperty*> FNode::GetPropertyMap()
+	{
+		TMap<FName, FProperty*> Map;
+		for (FProperty* Prop : GetProperties())
+		{
+			Map.Add(TTuple<FName, FProperty*>(Prop->GetName(),Prop));
+		}
+		return Map;
+	}
 
-	void FNode::SerializeInternal(FArchive& Ar) 
+	const TMap<FName, const FProperty*> FNode::GetPropertyMap() const
 	{
+		TMap<FName, const FProperty*> Map;
 		for (FProperty* Prop : GetProperties())
 		{
-			Prop->Serialize(Ar);
+			Map.Add(TTuple<FName, const FProperty*>(Prop->GetName(), Prop));
+		}
+		return Map;
+	}
+
+	void FNode::SerializeInternal(FArchive& Ar)
+	{
+		TMap<FName, FProperty*> Map = GetPropertyMap();
+		TArray< FProperty* >& CurrentProperties = GetProperties();
+		uint8 NumProperties = CurrentProperties.Num();
+		Ar << NumProperties;
+
+		for (int32 i = NumProperties - 1; i >= 0; i--)
+		{
+			uint8 Type;
+			FString StrName;
+			int64 PropertySizeOf = 0, PropertyDataSize = 0;
+
+			if (Ar.IsSaving())
+			{
+				FProperty* Prop = CurrentProperties[i];
+				Type = (uint8)Prop->GetType();
+				StrName = Prop->GetName().ToString();
+				PropertySizeOf = Prop->SizeOf();
+				Ar << Type << StrName << PropertySizeOf;
+
+				// The following code serializes the size of the PropertyData to serialize so that when loading,
+				// we can skip over it if the function could not be loaded.
+				// Serialize a temporary value for the delta in order to end up with an archive of the right size.
+				// Then serialize the PropertyData in order to get its size.
+				const int64 PropertyDataSizePos = Ar.Tell();
+				PropertyDataSize = 0;
+				Ar << PropertyDataSize;
+
+				const int64 PropBegin = Ar.Tell();
+				Prop->Serialize(Ar);
+
+				// Only go back and serialize the number of argument bytes if there is actually an underlying buffer to seek to.
+				// Come back to the temporary value we wrote and overwrite it with the PropertyData size we just calculated.
+				// And finally seek back to the end.
+				if (PropertyDataSizePos != INDEX_NONE)
+				{
+					const int64 PropEnd = Ar.Tell();
+					PropertyDataSize = (PropEnd - PropBegin);
+					Ar.Seek(PropertyDataSizePos);
+					Ar << PropertyDataSize;
+					Ar.Seek(PropEnd);
+				}
+			}
+			else if (Ar.IsLoading())
+			{
+
+				Ar << Type << StrName << PropertySizeOf << PropertyDataSize;
+				FName PropName(StrName);
+
+				if (FProperty* NewProperty = FProperty::NewProperty((FProperty::EType)Type, PropName))
+				{
+					if (Map.Contains(PropName) && Map[PropName]->GetType() == (FProperty::EType)Type)
+					{
+						Map[PropName]->Serialize(Ar);
+						delete NewProperty;
+					}
+					else
+					{
+						NewProperty->Serialize(Ar);
+
+						if (!Map.Contains(PropName))
+						{
+							AddProperty(NewProperty);
+						}
+						else
+						{
+							delete NewProperty;
+						}
+					}
+				}
+				else
+				{
+					Ar.Seek(Ar.Tell() + PropertyDataSize);
+				}
+			}
 		}
 	}
 
diff --git a/Engine/Source/Runtime/Experimental/Dataflow/Core/Private/Dataflow/DataflowProperty.cpp b/Engine/Source/Runtime/Experimental/Dataflow/Core/Private/Dataflow/DataflowProperty.cpp
--- a/Engine/Source/Runtime/Experimental/Dataflow/Core/Private/Dataflow/DataflowProperty.cpp
+++ b/Engine/Source/Runtime/Experimental/Dataflow/Core/Private/Dataflow/DataflowProperty.cpp
@@ -1,13 +1,37 @@
-// Copyright Epic Games, Inc. All Rights Reserved.
-
-#include "Dataflow/DataflowProperty.h"
-#include "Dataflow/DataflowNode.h"
-
-namespace Dataflow
-{
-	void FProperty::BindProperty(FNode* InNode, FProperty* That)
-	{
-		InNode->AddProperty(That);
-	}
-}
-
+// Copyright Epic Games, Inc. All Rights Reserved.
+
+#include "Dataflow/DataflowProperty.h"
+#include "Dataflow/DataflowNode.h"
+
+namespace Dataflow
+{
+	FProperty* FProperty::NewProperty(EType InType, FName InName, FNode* InNode)
+	{
+		switch (InType)
+		{
+		case EType::BOOL:
+			return new TProperty<bool>({ InName,bool(), InNode });
+		case EType::INT:
+			return new TProperty<int32>({ InName, int32(), InNode });
+		case EType::FLOAT:
+			return new TProperty<float>({ InName, float(), InNode });
+		case EType::DOUBLE:
+			return new TProperty<double>({ InName, double(), InNode });
+		case EType::STRING:
+			return new TProperty<FString>({ InName, FString(), InNode });
+		case EType::NAME:
+			return new TProperty<FName>({ InName, FName(), InNode });
+		default:
+			return nullptr;
+		}
+	}
+
+	void FProperty::BindProperty(FNode* InNode)
+	{
+		if (InNode)
+		{
+			InNode->AddProperty(this);
+		}
+	}
+}
+
diff --git a/Engine/Source/Runtime/Experimental/Dataflow/Core/Public/Dataflow/DataflowNode.h b/Engine/Source/Runtime/Experimental/Dataflow/Core/Public/Dataflow/DataflowNode.h
--- a/Engine/Source/Runtime/Experimental/Dataflow/Core/Public/Dataflow/DataflowNode.h
+++ b/Engine/Source/Runtime/Experimental/Dataflow/Core/Public/Dataflow/DataflowNode.h
@@ -58,6 +58,9 @@ namespace Dataflow
 		const TArray< FProperty* >& GetProperties() const { return Properties; }
 		TArray< FProperty* >& GetProperties() { return Properties; }
 
+		const TMap<FName, const FProperty*> GetPropertyMap() const;
+ 		TMap<FName, FProperty*> GetPropertyMap();
+
 		void AddInput(FConnection* InPtr);
 		const TArray< FConnection* >& GetInputs() const { return Inputs; }
 		TArray< FConnection* >& GetInputs() { return Inputs; }
diff --git a/Engine/Source/Runtime/Experimental/Dataflow/Core/Public/Dataflow/DataflowProperty.h b/Engine/Source/Runtime/Experimental/Dataflow/Core/Public/Dataflow/DataflowProperty.h
--- a/Engine/Source/Runtime/Experimental/Dataflow/Core/Public/Dataflow/DataflowProperty.h
+++ b/Engine/Source/Runtime/Experimental/Dataflow/Core/Public/Dataflow/DataflowProperty.h
@@ -36,23 +36,26 @@ namespace Dataflow
 		};
 
 
-		FProperty(EType InType, FName InName, FNode* Node)
+		FProperty(EType InType, FName InName, FNode* Node = nullptr)
 			: Type(InType)
 			, Name(InName)
 			, Category("")
 		{
-			BindProperty(Node, this);
+			BindProperty(Node);
 		}
 		virtual ~FProperty() {}
 
+		static FProperty* NewProperty(EType InType,FName InName = FName(""), FNode* InNode = nullptr);
 
 		const FName& GetName() const { return Name; }
 		const EType GetType() const { return Type; }
 
+		virtual SIZE_T SizeOf() const { return 0; }
+
 		void GetCategory(const FName& InName) { Name = InName; }
 		const FName& GetCategory() const { return Category; }
 
-		void BindProperty(FNode* InNode, FProperty* That);
+		void BindProperty(FNode* InNode);
 
 		virtual void Serialize(FArchive& Ar) {}
 
@@ -77,11 +80,12 @@ namespace Dataflow
 		{}
 
 		static FProperty::EType StaticType();
+		virtual SIZE_T SizeOf() const override final;
 
 		const T& GetValue() const { return Value; }
 		void SetValue(const T& InValue) { Value = InValue; }
 
-		virtual void Serialize(FArchive& Ar) override
+		virtual void Serialize(FArchive& Ar) override final
 		{
 			Ar << Value;
 		}
@@ -90,17 +94,16 @@ namespace Dataflow
 
 } // Dataflow
 
-#define DATAFLOW_PROPERTY(DECL, TYPE,ETYPE)																			\
+#define DATAFLOW_PROPERTY(DECL, TYPE,ETYPE, SIZEOF)																	\
 	template<> inline Dataflow::FProperty::EType Dataflow::TProperty<TYPE>::StaticType() { return EType::ETYPE; }	\
+	template<> inline SIZE_T Dataflow::TProperty<TYPE>::SizeOf() const {return SIZEOF;}								\
 	template struct DECL Dataflow::TPropertyParameters<TYPE>;														\
 	template class Dataflow::TProperty<TYPE>;
 
-DATAFLOW_PROPERTY(DATAFLOWCORE_API, bool, BOOL)
-DATAFLOW_PROPERTY(DATAFLOWCORE_API, int, INT)
-DATAFLOW_PROPERTY(DATAFLOWCORE_API, float, FLOAT)
-DATAFLOW_PROPERTY(DATAFLOWCORE_API, double, DOUBLE)
-DATAFLOW_PROPERTY(DATAFLOWCORE_API, FString, STRING)
-DATAFLOW_PROPERTY(DATAFLOWCORE_API, FName, NAME)
-//DATAFLOW_PROPERTY(DATAFLOWCORE_API, Dataflow::FFilename, FILE)
-
+DATAFLOW_PROPERTY(DATAFLOWCORE_API, bool, BOOL, sizeof(bool))
+DATAFLOW_PROPERTY(DATAFLOWCORE_API, int32, INT, sizeof(int32))
+DATAFLOW_PROPERTY(DATAFLOWCORE_API, float, FLOAT, sizeof(float))
+DATAFLOW_PROPERTY(DATAFLOWCORE_API, double, DOUBLE, sizeof(double))
+DATAFLOW_PROPERTY(DATAFLOWCORE_API, FString, STRING, Value.GetAllocatedSize())
+DATAFLOW_PROPERTY(DATAFLOWCORE_API, FName, NAME, Value.ToString().GetAllocatedSize())
 
