clean up + optimize non-overlapping convex hull logic and make it parallel

diff --git a/Engine/Source/Runtime/Experimental/Chaos/Private/GeometryCollection/GeometryCollectionConvexUtility.cpp b/Engine/Source/Runtime/Experimental/Chaos/Private/GeometryCollection/GeometryCollectionConvexUtility.cpp
--- a/Engine/Source/Runtime/Experimental/Chaos/Private/GeometryCollection/GeometryCollectionConvexUtility.cpp
+++ b/Engine/Source/Runtime/Experimental/Chaos/Private/GeometryCollection/GeometryCollectionConvexUtility.cpp
@@ -343,7 +343,7 @@ void CreateNonoverlappingConvexHulls(
 
 	TArray<TSet<int32>> LeafProximity;
 	LeafProximity.SetNum(NumBones);
-	if (GeoProximity)
+	if (ensure(GeoProximity))
 	{
 		for (int32 GeomIdx = 0; GeomIdx < GeoProximity->Num(); GeomIdx++)
 		{
@@ -354,10 +354,6 @@ void CreateNonoverlappingConvexHulls(
 			}
 		}
 	}
-	else
-	{
-		// TODO: fill out LeafProximity by doing IsColliding calls (may have to construct Children array before this)
-	}
 
 	auto IsColliding = [&Convexes](int32 ConvexA, int32 ConvexB)
 	{
@@ -584,6 +580,19 @@ void CreateNonoverlappingConvexHulls(
 		}
 	}
 
+	TArray<int32> ByDepthOrder;
+	ByDepthOrder.Reserve(NumBones);
+	for (int32 ProcessDepth = MaxDepth; ProcessDepth >= 0; --ProcessDepth)
+	{
+		for (int32 Bone = 0; Bone < NumBones; Bone++)
+		{
+			if (Depths[Bone] == ProcessDepth)
+			{
+				ByDepthOrder.Add(Bone);
+			}
+		}
+	}
+
 	auto AddLeaves = [&Children, &SimulationType, LeafType](int32 Bone, TArray<int32>& Leaves)
 	{
 		TArray<int32> ToExpand;
@@ -620,55 +629,110 @@ void CreateNonoverlappingConvexHulls(
 		}
 	};
 
-	auto FixLeafCollisions = [&AddLeaves, &FixCollisionWithCut, &OnlyConvex, &HasCustomConvexFn](int32 BoneA, int32 BoneB)
+	// Use initial leaf proximity to compute which cluster bones are in proximity to which neighbors at the same level of the bone hierarchy
+
+	TArray<TSet<int32>> SameDepthClusterProximity;
+	SameDepthClusterProximity.SetNum(NumBones);
+	for (int32 Bone = 0; Bone < NumBones; Bone++)
 	{
-		TArray<int32> LeavesA, LeavesB;
-		AddLeaves(BoneA, LeavesA);
-		AddLeaves(BoneB, LeavesB);
-		bool bAnyCollides = false;
-		for (int32 LeafA : LeavesA)
+		for (int32 NbrBone : LeafProximity[Bone])
 		{
-			if (HasCustomConvexFn(LeafA))
+			if (NbrBone > Bone)
 			{
 				continue;
 			}
-			for (int32 LeafB : LeavesB)
+
+			auto FindCommonParent = [&Parents](int32 BoneA, int32 BoneB)
 			{
-				if (HasCustomConvexFn(LeafB))
+				int AParent = Parents[BoneA];
+				int BParent = Parents[BoneB];
+				if (AParent == BParent) // early out if they're in the same cluster
 				{
-					continue;
+					return AParent;
+				}
+
+				TSet<int32> AParents;
+				while (AParent != INDEX_NONE)
+				{
+					AParents.Add(AParent);
+					AParent = Parents[AParent];
+				}
+
+				while (BParent != INDEX_NONE && !AParents.Contains(BParent))
+				{
+					BParent = Parents[BParent];
+				}
+				return BParent;
+			};
+			int32 CommonParent = FindCommonParent(Bone, NbrBone);
+
+			auto ConnectAtMatchingDepth = [&SameDepthClusterProximity, &Parents, CommonParent, &Depths](int32 BoneA, int32 BoneToWalkParents)
+			{
+				int32 DepthA = Depths[BoneA];
+				int32 DepthB = Depths[BoneToWalkParents];
+				while (BoneToWalkParents != INDEX_NONE && BoneToWalkParents != CommonParent && DepthB >= DepthA)
+				{
+					if (DepthB == DepthA)
+					{
+						bool bWasInSet = false;
+						SameDepthClusterProximity[BoneToWalkParents].Add(BoneA, &bWasInSet);
+						if (bWasInSet)
+						{
+							break;
+						}
+						SameDepthClusterProximity[BoneA].Add(BoneToWalkParents);
+					}
+
+					BoneToWalkParents = Parents[BoneToWalkParents];
+					DepthB--;
 				}
-				// TODO: Instead of skipping all custom convex, if only one of the leaves is custom, we 
-				// could iterate through the custom leaf convexes and clip the non-custom convex.
-				// But need to decide a fallback if this would delete the whole non-custom convex.
-				// For now, if you have custom convex leaves we just assume you they're set up to not explode the sim.
-				bool bCollides = FixCollisionWithCut(OnlyConvex(LeafA), OnlyConvex(LeafB));
-				bAnyCollides |= bCollides;
+			};
+
+			// connect bone to neighboring bone's clusters
+			ConnectAtMatchingDepth(Bone, Parents[NbrBone]);
+
+			// walk chain of parents of bone and connect each to the chain of neighboring bone + its parents
+			int32 BoneParent = Parents[Bone];
+			while (BoneParent != INDEX_NONE && BoneParent != CommonParent)
+			{
+				ConnectAtMatchingDepth(BoneParent, NbrBone);
+				BoneParent = Parents[BoneParent];
 			}
 		}
-		return bAnyCollides;
-	};
+	}
 
-	// Compute initial hulls at all levels, and use these to fill out the full proximity links
-	// Fix collisions between any two leaf hulls
+	// Compute initial hulls at all levels and filter out any that are too large relative to the geometry they contain
 
 	TArray<TSet<int32>> ClusterProximity;
 	ClusterProximity.SetNum(NumBones);
-	int ProcessDepth = MaxDepth;
-	TArray<Chaos::FConvex::FVec3Type> JoinedHullPts;
-	TArray<int32> ChildrenWithHulls;
-	while (--ProcessDepth >= 0)
+	TArrayView<int32> DepthOrderView(ByDepthOrder);
+
+	for (int32 Idx = 0, SliceEndIdx = Idx + 1; Idx < ByDepthOrder.Num(); Idx = SliceEndIdx)
 	{
-		TSet<int32> ToProcess;
-		for (int Bone = 0; Bone < Parents.Num(); Bone++)
+		FCriticalSection ConvexesCS;
+
+		int32 ProcessDepth = Depths[ByDepthOrder[Idx]];
+		for (SliceEndIdx = Idx + 1; SliceEndIdx < ByDepthOrder.Num(); SliceEndIdx++)
 		{
-			if (Depths[Bone] == ProcessDepth)
+			if (Depths[ByDepthOrder[SliceEndIdx]] != ProcessDepth)
+			{
+				break;
+			}
+		}
+		TArrayView<int32> SameDepthSlice = DepthOrderView.Slice(Idx, SliceEndIdx - Idx);
+		
+		ParallelFor(SameDepthSlice.Num(), [&](int32 SliceIdx)
+		{
+			int32 Bone = SameDepthSlice[SliceIdx];
+			checkSlow(Depths[Bone] == ProcessDepth);
+			if (!SkipBone(Bone))
 			{
 				if (TransformToConvexIndices[Bone].Num() == 0 && !HasCustomConvexFn(Bone))
 				{
-					JoinedHullPts.Reset();
-					ChildrenWithHulls.Reset();
+					TArray<Chaos::FConvex::FVec3Type> JoinedHullPts;
+					TArray<int32> ChildrenWithHulls;
 					AddDescendentsWithHulls(Bone, ChildrenWithHulls);
+					ConvexesCS.Lock();
 					for (int32 Child : ChildrenWithHulls)
 					{
 						for (int32 ConvexIdx : TransformToConvexIndices[Child])
@@ -676,6 +740,7 @@ void CreateNonoverlappingConvexHulls(
 							JoinedHullPts.Append(Convexes[ConvexIdx]->GetVertices());
 						}
 					}
+					ConvexesCS.Unlock();
 					if (JoinedHullPts.Num() > 0)
 					{
 						FilterHullPoints(JoinedHullPts, SimplificationDistanceThreshold);
@@ -691,93 +756,72 @@ void CreateNonoverlappingConvexHulls(
 						}
 						if (!bIsTooBig)
 						{
+							ConvexesCS.Lock();
 							int32 ConvexIdx = Convexes.Add(MoveTemp(Hull));
+							ConvexesCS.Unlock();
 							TransformToConvexIndices[Bone].Add(ConvexIdx);
-							ToProcess.Add(Bone);
 						}
 					}
 				}
-				else
-				{
-					ToProcess.Add(Bone);
-				}
 			}
-		}
+		});
 
-		// We don't have proximity at the cluster level, so reconstruct from n^2 collision tests for now
-		for (int32 BoneA : ToProcess)
+		// Compute cluster proximity
+		for (int32 BoneA : SameDepthSlice)
 		{
-			for (int32 BoneB : ToProcess)
+			if (SkipBone(BoneA) || TransformToConvexIndices[BoneA].Num() == 0)
 			{
-				if (BoneB < BoneA)
+				continue;
+			}
+			for (int32 BoneB : SameDepthClusterProximity[BoneA])
+			{
+				if (BoneB < BoneA && IsBoneColliding(BoneA, BoneB))
 				{
-					if (IsBoneColliding(BoneA, BoneB))
+					int32 Bones[2]{ BoneA, BoneB };
+					for (int32 BoneIdx = 0; BoneIdx < 2; BoneIdx++)
 					{
-						FixLeafCollisions(BoneA, BoneB);
-						int32 Bones[2]{ BoneA, BoneB };
-						for (int32 BoneIdx = 0; BoneIdx < 2; BoneIdx++)
+						int32 ParentBone = Bones[BoneIdx];
+						int32 OtherBone = Bones[1 - BoneIdx];
+						TArray<int32> TraverseBones;
+						TraverseBones.Append(Children[ParentBone]);
+						while (TraverseBones.Num() > 0)
 						{
-							int32 ParentBone = Bones[BoneIdx];
-							int32 OtherBone = Bones[1 - BoneIdx];
-							// if leaves changed flag that all hulls using those leaves must also change
-							TArray<int32> TraverseBones;
-							TraverseBones.Append(Children[ParentBone]);
-							while (TraverseBones.Num() > 0)
+							int32 ToProc = TraverseBones.Pop(false);
+							if (IsBoneColliding(OtherBone, ToProc))
 							{
-								int32 ToProc = TraverseBones.Pop(false);
-								if (IsBoneColliding(OtherBone, ToProc))
-								{
-									ClusterProximity[OtherBone].Add(ToProc);
-									ClusterProximity[ToProc].Add(OtherBone);
-									TraverseBones.Append(Children[ToProc]);
-								}
+								ClusterProximity[OtherBone].Add(ToProc);
+								ClusterProximity[ToProc].Add(OtherBone);
+								TraverseBones.Append(Children[ToProc]);
 							}
 						}
-
-						ClusterProximity[BoneA].Add(BoneB);
-						ClusterProximity[BoneB].Add(BoneA);
 					}
+
+					ClusterProximity[BoneA].Add(BoneB);
+					ClusterProximity[BoneB].Add(BoneA);
 				}
 			}
 		}
 	}
 
-	// Now that leaves don't intersect, re-compute all non-leaf hulls on the clipped leaves
-	// and record volumes for each hull
+
+
+	// Compute all initial non-leaf hull volumes
 
 	TArray<double> NonLeafVolumes; // Original volumes of non-leaf hulls (to compare against progressively cut-down volume as intersections are removed)
 	NonLeafVolumes.SetNumZeroed(Convexes.Num());
-	ProcessDepth = MaxDepth;
-	while (--ProcessDepth > 0)
+	ParallelFor(NumBones, [&](int32 Bone)
 	{
-		TSet<int32> ToProcess;
-		for (int32 Bone = 0; Bone < NumBones; Bone++)
+		bool bCustom = HasCustomConvexFn(Bone);
+		if (!bCustom && // if we need an automatic hull
+			Children[Bone].Num() > 0 && !SkipBone(Bone) && // and we have children && aren't embedded geo
+			TransformToConvexIndices[Bone].Num() == 1 // and hull wasn't already ruled out by CanExceedFraction
+			)
 		{
-			bool bCustom = HasCustomConvexFn(Bone);
-			if (!bCustom && Depths[Bone] == ProcessDepth && // if we need an automatic hull and are at the current depth
-				Children[Bone].Num() > 0 && !SkipBone(Bone) && // and we have children && aren't embedded geo
-				TransformToConvexIndices[Bone].Num() == 1 // and hull wasn't already ruled out by CanExceedFraction
-				)
-			{
-				JoinedHullPts.Reset();
-				ChildrenWithHulls.Reset();
-				AddDescendentsWithHulls(Bone, ChildrenWithHulls);
-				for (int32 Child : ChildrenWithHulls)
-				{
-					for (int32 ConvexIdx : TransformToConvexIndices[Child])
-					{
-						JoinedHullPts.Append(Convexes[ConvexIdx]->GetVertices());
-					}
-				}
-				int32 ConvexIdx = OnlyConvex(Bone);
-				if (ConvexIdx)
-				{
-					*Convexes[ConvexIdx] = MakeHull(JoinedHullPts, SimplificationDistanceThreshold);
-					NonLeafVolumes[ConvexIdx] = Convexes[ConvexIdx]->GetVolume();
-				}
-			}
+			int32 ConvexIdx = OnlyConvex(Bone);
+			checkSlow(ConvexIdx > -1); // safe to assume because the if-statement verified "TransformToConvexIndices[Bone].Num() == 1" 
+			NonLeafVolumes[ConvexIdx] = Convexes[ConvexIdx]->GetVolume();
 		}
-	}
+	});
 
 	// if bOneSidedCut, then only ConvexA is cut; ConvexB is left unchanged
 	auto CutIfOk = [&Convexes, &NonLeafVolumes, &FindCutPlane, &FracAllowRemove, &SimplificationDistanceThreshold](bool bOneSidedCut, int32 ConvexA, int32 ConvexB) -> bool
@@ -848,24 +892,24 @@ void CreateNonoverlappingConvexHulls(
 	};
 
 	// re-process all non-leaf bones
-	ProcessDepth = MaxDepth;
-	while (--ProcessDepth > 0)
+	for (int32 Idx = 0, SliceEndIdx = Idx + 1; Idx < DepthOrderView.Num(); Idx = SliceEndIdx)
 	{
-		TSet<int32> ToProcess;
-		for (int32 Bone = 0; Bone < NumBones; Bone++)
+		int32 ProcessDepth = Depths[DepthOrderView[Idx]];
+		for (SliceEndIdx = Idx + 1; SliceEndIdx < DepthOrderView.Num(); SliceEndIdx++)
 		{
-			if (Depths[Bone] == ProcessDepth && Children[Bone].Num() > 0)
+			if (Depths[DepthOrderView[SliceEndIdx]] != ProcessDepth)
 			{
-				ToProcess.Add(Bone);
+				break;
 			}
 		}
+		TArrayView<int32> SameDepthSlice = DepthOrderView.Slice(Idx, SliceEndIdx - Idx);
 
 		TSet<int32> WasNotOk;
 
-		for (int32 Bone : ToProcess)
+		for (int32 Bone : SameDepthSlice)
 		{
 			bool bCustom = HasCustomConvexFn(Bone);
-			if (bCustom || TransformToConvexIndices[Bone].Num() == 0 || WasNotOk.Contains(Bone))
+			if (bCustom || Children[Bone].Num() == 0 || TransformToConvexIndices[Bone].Num() == 0 || WasNotOk.Contains(Bone))
 			{
 				continue;
 			}
@@ -952,8 +996,8 @@ void HullsFromGeometry(
 
 	int32 NumBones = Geometry.TransformToGeometryIndex.Num();
 	TransformToConvexIndices.SetNum(NumBones);
-	TArray<Chaos::FConvex::FVec3Type> HullPts;
-	for (int32 Idx = 0; Idx < NumBones; Idx++)
+	FCriticalSection HullCS;
+	ParallelFor(NumBones, [&](int32 Idx)
 	{
 		int32 GeomIdx = Geometry.TransformToGeometryIndex[Idx];
 		if (OrigConvexData.IsSet() && HasCustomConvexFn(Idx))
@@ -962,13 +1006,17 @@ void HullsFromGeometry(
 			const FTransform& Transform = GlobalTransformArray[Idx];
 			for (int32 OrigConvexIdx : OrigConvexData->TransformToConvexIndices[Idx])
 			{
-				HullPts.Reset();
+				TArray<Chaos::FConvex::FVec3Type> HullPts;
+				HullPts.Reserve(OrigConvexData->ConvexHull[OrigConvexIdx]->NumVertices());
 				for (const Chaos::FConvex::FVec3Type& P : OrigConvexData->ConvexHull[OrigConvexIdx]->GetVertices())
 				{
-					HullPts.Add(FVector(Transform.TransformPosition((FVector)P)));
+					HullPts.Add(Chaos::FConvex::FVec3Type(Transform.TransformPosition((FVector)P)));
 				}
 				// Do not simplify hulls when we're just trying to transform them
-				int32 NewConvexIdx = Convexes.Add(MakeUnique<Chaos::FConvex>(HullPts, UE_KINDA_SMALL_NUMBER));
+				TUniquePtr Hull = MakeUnique<Chaos::FConvex>(HullPts, UE_KINDA_SMALL_NUMBER);
+				HullCS.Lock();
+				int32 NewConvexIdx = Convexes.Add(MoveTemp(Hull));
+				HullCS.Unlock();
 				TransformToConvexIndices[Idx].Add(NewConvexIdx);
 			}
 		}
@@ -977,7 +1025,7 @@ void HullsFromGeometry(
 			int32 VStart = Geometry.VertexStart[GeomIdx];
 			int32 VCount = Geometry.VertexCount[GeomIdx];
 			int32 VEnd = VStart + VCount;
-			HullPts.Reset();
+			TArray<Chaos::FConvex::FVec3Type> HullPts;
 			HullPts.Reserve(VCount);
 			for (int32 VIdx = VStart; VIdx < VEnd; VIdx++)
 			{
@@ -985,42 +1033,13 @@ void HullsFromGeometry(
 			}
 			ensure(HullPts.Num() > 0);
 			FilterHullPoints(HullPts, SimplificationDistanceThreshold);
-			int32 ConvexIdx = Convexes.Add(MakeUnique<Chaos::FConvex>(HullPts, UE_KINDA_SMALL_NUMBER));
-			if (Convexes[ConvexIdx]->NumVertices() == 0 && HullPts.Num() > 0)
-			{
-				// if we've failed to make a convex hull, add a tiny bounding box to try to give it a hull anyway
-				// (note it may still end up with an empty hull after cutting the hull to remove collisions with neighboring hulls; this just reduces the chance/number of empty hulls)
-				Chaos::FConvex::FAABB3Type AABB = Convexes[ConvexIdx]->GetLocalBoundingBox();
-				Chaos::FVec3 Extents = AABB.Extents();
-				constexpr Chaos::FReal SmallestConvexWidth = (Chaos::FReal).11; // Convexes that span less distance than this will become empty hulls due to the size thresholds in FConvex
-				if (Extents.GetMax() < SmallestConvexWidth * 2) // only apply this logic if the bounding box was actually small (in case FConvex fails to make a hull even for larger hulls)
-				{
-					// expand the hull to cover a .11 ^ 3 volume, so it won't disappear
-					constexpr Chaos::FReal MinExpand = (Chaos::FReal).01;
-					Chaos::FVec3 GrowVec;
-					for (int Axis = 0; Axis < 3; Axis++)
-					{
-						GrowVec[Axis] = FMath::Max((Chaos::FReal).5 * (SmallestConvexWidth - Extents[Axis]), MinExpand);
-					}
-					AABB.ThickenSymmetrically(GrowVec);
-					Chaos::FConvex::FVec3Type Min = AABB.Min();
-					Chaos::FConvex::FVec3Type Max = AABB.Max();
-					HullPts.Reset(); // we're adding the bounding box, so no need to keep the points inside
-					HullPts.Add(Min);
-					HullPts.Add(Max);
-					HullPts.Add({ Min.X, Min.Y, Max.Z });
-					HullPts.Add({ Min.X, Max.Y, Max.Z });
-					HullPts.Add({ Max.X, Min.Y, Max.Z });
-					HullPts.Add({ Max.X, Max.Y, Min.Z });
-					HullPts.Add({ Max.X, Min.Y, Min.Z });
-					HullPts.Add({ Min.X, Max.Y, Min.Z });
-					// note: Do not use SimplificationDistanceThreshold for this fixed tiny hull
-					*Convexes[ConvexIdx] = Chaos::FConvex(HullPts, UE_KINDA_SMALL_NUMBER);
-				}
-			}
+			TUniquePtr Hull = MakeUnique<Chaos::FConvex>(HullPts, UE_KINDA_SMALL_NUMBER);
+			HullCS.Lock();
+			int32 ConvexIdx = Convexes.Add(MoveTemp(Hull));
+			HullCS.Unlock();
 			TransformToConvexIndices[Idx].Add(ConvexIdx);
 		}
-	}
+	});
 }
 
 void TransformHullsToLocal(
@@ -1029,7 +1048,7 @@ void TransformHullsToLocal(
 	TArray<TSet<int32>>& TransformToConvexIndices
 )
 {
-	for (int32 Bone = 0; Bone < TransformToConvexIndices.Num(); Bone++)
+	ParallelFor(TransformToConvexIndices.Num(), [&](int32 Bone)
 	{
 		FTransform& Transform = GlobalTransformArray[Bone];
 		TArray<Chaos::FConvex::FVec3Type> HullPts;
@@ -1043,7 +1062,7 @@ void TransformHullsToLocal(
 			// Do not simplify hulls when we're just trying to transform them
 			*Convexes[ConvexIdx] = Chaos::FConvex(HullPts, UE_KINDA_SMALL_NUMBER);
 		}
-	}
+	});
 }
 
 // copy hulls from InBone to OutBone
@@ -1190,7 +1209,9 @@ FGeometryCollectionConvexUtility::FGeometryCollectionConvexData FGeometryCollect
 			EmptyConvex.Add(ConvexIdx);
 		}
 	}
-	GeometryCollection->RemoveElements("Convex", EmptyConvex);
+	FManagedArrayCollection::FProcessingParameters Params;
+	Params.bDoValidation = false; // for perf reasons
+	GeometryCollection->RemoveElements("Convex", EmptyConvex, Params);
 
 	checkSlow(FGeometryCollectionConvexUtility::ValidateConvexData(GeometryCollection));
 
@@ -1437,8 +1458,20 @@ void FGeometryCollectionConvexUtility::SetVolumeAttributes(FGeometryCollection*
 
 	TArray<FTransform> Transforms;
 	GeometryCollectionAlgo::GlobalMatrices(Collection->Transform, Collection->Parent, Transforms);
-	TArray<int32> RecursiveOrder = GeometryCollectionAlgo::ComputeRecursiveOrder(*Collection);
+
+	TArray<float> GeoVolumes;
+	GeoVolumes.SetNumZeroed(Collection->NumElements(FGeometryCollection::GeometryGroup));
+	ParallelFor(GeoVolumes.Num(), [&](int32 GeoIdx)
+	{
+		int32 Bone = Collection->TransformIndex[GeoIdx];
+		if (SimulationType[Bone] == FGeometryCollection::ESimulationTypes::FST_Rigid)
+		{
+			GeoVolumes[GeoIdx] = (float)ComputeGeometryVolume(Collection, GeoIdx, Transforms[Bone], 1.0);
+		}
+	});
+
 	float MaxGeoVolume = 0.0f;
+	TArray<int32> RecursiveOrder = GeometryCollectionAlgo::ComputeRecursiveOrder(*Collection);
 	Volumes.Fill(0.0f);
 	for (const int32 Bone : RecursiveOrder)
 	{
@@ -1452,7 +1485,7 @@ void FGeometryCollectionConvexUtility::SetVolumeAttributes(FGeometryCollection*
 			}
 			else
 			{
-				const float GeoVolume = (float)ComputeGeometryVolume(Collection, GeoIdx, Transforms[Bone], 1.0);
+				const float GeoVolume = GeoVolumes[GeoIdx];
 				MaxGeoVolume = FMath::Max(MaxGeoVolume, GeoVolume);
 				Volumes[Bone] = GeoVolume;
 			}
