Chaos : fix verification logic with replication and cluster groups

diff --git a/Engine/Source/Runtime/Experimental/GeometryCollectionEngine/Private/GeometryCollection/GeometryCollectionComponent.cpp b/Engine/Source/Runtime/Experimental/GeometryCollectionEngine/Private/GeometryCollection/GeometryCollectionComponent.cpp
--- a/Engine/Source/Runtime/Experimental/GeometryCollectionEngine/Private/GeometryCollection/GeometryCollectionComponent.cpp
+++ b/Engine/Source/Runtime/Experimental/GeometryCollectionEngine/Private/GeometryCollection/GeometryCollectionComponent.cpp
@@ -1962,15 +1962,17 @@ void UGeometryCollectionComponent::OnRegister()
 {
 	//UE_LOG(UGCC_LOG, Log, TEXT("GeometryCollectionComponent[%p]::OnRegister()[%p]"), this,RestCollection );
 	ResetDynamicCollection();
-	
-	if (ensureMsgf(ClusterGroupIndex != 0, TEXT("Replication with cluster groups is not supported - disabling replication")))
-	{
-		SetIsReplicated(false);	
-	}
-	else
+
+	bool bIsReplicated = false;
+	const bool bHasClusterGroup = (ClusterGroupIndex != 0);
+	if (bEnableReplication)
 	{
-		SetIsReplicated(bEnableReplication);	
+		if (ensureMsgf(!bHasClusterGroup, TEXT("Replication with cluster groups is not supported - disabling replication")))
+		{
+			bIsReplicated = true; 
+		}
 	}
+	SetIsReplicated(bIsReplicated);
 
 	InitializeEmbeddedGeometry();
 
