Dataflow : Connection Serialization
- Updates connections to include node guids.

diff --git a/Engine/Source/Runtime/Experimental/Dataflow/Core/Private/Dataflow/DataflowGraph.cpp b/Engine/Source/Runtime/Experimental/Dataflow/Core/Private/Dataflow/DataflowGraph.cpp
--- a/Engine/Source/Runtime/Experimental/Dataflow/Core/Private/Dataflow/DataflowGraph.cpp
+++ b/Engine/Source/Runtime/Experimental/Dataflow/Core/Private/Dataflow/DataflowGraph.cpp
@@ -79,7 +79,9 @@ namespace Dataflow
 		{
 			Input->AddConnection(Output);
 			Output->AddConnection(Input);
-			Connections.Add(FLink(Input->GetGuid(), Output->GetGuid()));
+			Connections.Add(FLink(
+				Input->GetOwningNode()->GetGuid(), Input->GetGuid(), 
+				Output->GetOwningNode()->GetGuid(), Output->GetGuid()));
 		}
 	}
 
@@ -87,7 +89,10 @@ namespace Dataflow
 	{
 		Input->RemoveConnection(Output);
 		Output->RemoveConnection(Input);
-		Connections.RemoveSwap(FLink(Input->GetGuid(), Output->GetGuid()));
+		Connections.RemoveSwap(FLink(
+			Input->GetOwningNode()->GetGuid(), Input->GetGuid(),
+			Output->GetOwningNode()->GetGuid(), Output->GetGuid()));
+
 	}
 
 	void FGraph::Serialize(FArchive& Ar)
@@ -98,6 +103,7 @@ namespace Dataflow
 			FGuid ArGuid;
 			FName ArType, ArName;
 			int32 ArNum = Nodes.Num();
+
 			Ar << ArNum;
 			for (TSharedPtr<FNode> Node : Nodes)
 			{
@@ -106,19 +112,37 @@ namespace Dataflow
 				ArName = Node->GetName();
 				Ar << ArGuid << ArType << ArName;
 
-
-				TArray< FConnection* > IO = Node->GetOutputs();  IO.Append(Node->GetInputs());
-				ArNum = IO.Num();
-				Ar << ArNum;
-				for (FConnection* Conn : IO)
+				// Data Jump Serialization (on Node and connections)
 				{
-					ArGuid = Conn->GetGuid();
-					ArType = Conn->GetType();
-					ArName = Conn->GetName();
-					Ar << ArGuid << ArType << ArName;
-				}
+					const int64 NodeDataSizePos = Ar.Tell();
+					int64 NodeDataSize = 0;
+					Ar << NodeDataSize;
+
+					const int64 NodeBegin = Ar.Tell();
+					{ // Data we can jump past on load
+						TArray< FConnection* > IO = Node->GetOutputs();  IO.Append(Node->GetInputs());
+						ArNum = IO.Num();
+						Ar << ArNum;
+						for (FConnection* Conn : IO)
+						{
+							ArGuid = Conn->GetGuid();
+							ArType = Conn->GetType();
+							ArName = Conn->GetName();
+							Ar << ArGuid << ArType << ArName;
+						}
+
+						Node->SerializeInternal(Ar);
+					}
 
-				Node->SerializeInternal(Ar);
+					if (NodeDataSizePos != INDEX_NONE)
+					{
+						const int64 NodeEnd = Ar.Tell();
+						NodeDataSize = (NodeEnd - NodeBegin);
+						Ar.Seek(NodeDataSizePos);
+						Ar << NodeDataSize;
+						Ar.Seek(NodeEnd);
+					}
+				}
 			}
 
 			Ar << Connections;
@@ -129,15 +153,22 @@ namespace Dataflow
 			FGuid ArGuid;
 			FName ArType, ArName;
 			int32 ArNum = 0;
+			int64 NodeDataSize = 0;
 
+			TMap<FGuid, TSharedPtr<FNode> > NodeGuidMap;
 			TMap<FGuid, FConnection* > ConnectionGuidMap;
 
 			Ar << ArNum;
 			for (int32 Ndx = ArNum; Ndx > 0; Ndx--)
 			{
-				Ar << ArGuid << ArType << ArName;
+				Ar << ArGuid << ArType << ArName << NodeDataSize;
+
 				if (TSharedPtr<FNode> Node = FNodeFactory::GetInstance()->NewNodeFromRegisteredType(*this, { ArGuid,ArType,ArName }))
 				{
+					ensure(!NodeGuidMap.Contains(ArGuid));
+					NodeGuidMap.Add(ArGuid,Node);
+
+
 					int ArNumInner;
 					Ar << ArNumInner;
 					TArray< FConnection* > IO = Node->GetOutputs();  IO.Append(Node->GetInputs());
@@ -164,7 +195,11 @@ namespace Dataflow
 				}
 				else
 				{
-					ensureMsgf(false, TEXT("Warning: Missing registered node type on load. Graph will fail to evaluate (%s %s)"), *ArType.ToString(), *ArName.ToString());
+					Ar.Seek(Ar.Tell() + NodeDataSize); // Data Jump Serialization (on mising Node)
+					DisabledNodes.Add(ArName);
+					ensureMsgf(false, 
+						TEXT("Error: Missing registered node type (%s) will be removed from graph on load. Graph will fail to evaluate due to missing node (%s).")
+						, *ArType.ToString(), *ArName.ToString());
 				}
 				
 			}
@@ -173,11 +208,14 @@ namespace Dataflow
 			Ar << LocalConnections;
 			for (const FLink& Con : LocalConnections)
 			{
-				if (ensure(ConnectionGuidMap.Contains(Con.Input) && ConnectionGuidMap.Contains(Con.Output)))
+				if (ensure(NodeGuidMap.Contains(Con.InputNode) && NodeGuidMap.Contains(Con.OutputNode)))
 				{
-					if (ensure(ConnectionGuidMap[Con.Input]->GetType() == ConnectionGuidMap[Con.Output]->GetType()))
+					if (ensure(ConnectionGuidMap.Contains(Con.Input) && ConnectionGuidMap.Contains(Con.Output)))
 					{
-						Connect(ConnectionGuidMap[Con.Input], ConnectionGuidMap[Con.Output]);
+						if (ensure(ConnectionGuidMap[Con.Input]->GetType() == ConnectionGuidMap[Con.Output]->GetType()))
+						{
+							Connect(ConnectionGuidMap[Con.Input], ConnectionGuidMap[Con.Output]);
+						}
 					}
 				}
 			}
diff --git a/Engine/Source/Runtime/Experimental/Dataflow/Core/Public/Dataflow/DataflowGraph.h b/Engine/Source/Runtime/Experimental/Dataflow/Core/Public/Dataflow/DataflowGraph.h
--- a/Engine/Source/Runtime/Experimental/Dataflow/Core/Public/Dataflow/DataflowGraph.h
+++ b/Engine/Source/Runtime/Experimental/Dataflow/Core/Public/Dataflow/DataflowGraph.h
@@ -12,16 +12,20 @@ namespace Dataflow
 	class FConnection;
 
 	struct FLink {
+		FGuid InputNode;
 		FGuid Input;
+		FGuid OutputNode;
 		FGuid Output;
 
 		FLink() {}
 
-		FLink(FGuid InInput, FGuid InOutput)
-			: Input(InInput), Output(InOutput) {}
+		FLink(FGuid InInputNode, FGuid InInput, FGuid InOutputNode, FGuid InOutput)
+			: InputNode(InInputNode), Input(InInput)
+			, OutputNode(InOutputNode), Output(InOutput) {}
 
 		FLink(const FLink& Other)
-			: Input(Other.Input), Output(Other.Output) {}
+			: InputNode(Other.InputNode), Input(Other.Input)
+			, OutputNode(Other.OutputNode), Output(Other.Output) {}
 
 		bool operator==(const FLink& Other) const
 		{
@@ -30,8 +34,8 @@ namespace Dataflow
 
 		bool Equals(const FLink& Other) const
 		{
-			return Input == Other.Input
-				&& Output == Other.Output;
+			return Input == Other.Input && InputNode == Other.InputNode
+				&& Output == Other.Output && OutputNode == Other.OutputNode;
 		}
 	};
 
@@ -45,7 +49,7 @@ namespace Dataflow
 		FGuid  Guid;
 		TArray< TSharedPtr<FNode> > Nodes;
 		TArray< FLink > Connections;
-
+		TSet< FName > DisabledNodes;
 	public:
 		FGraph(FGuid InGuid = FGuid::NewGuid());
 		virtual ~FGraph() {}
@@ -96,19 +100,20 @@ namespace Dataflow
 		void Disconnect(FConnection* Input, FConnection* Output);
 
 		virtual void Serialize(FArchive& Ar);
+		const TSet<FName>& GetDisabledNodes() const { return DisabledNodes; }
 
 	};
 }
 
 FORCEINLINE FArchive& operator<<(FArchive& Ar, Dataflow::FLink& Value)
 {
-	Ar << Value.Input << Value.Output;
+	Ar << Value.InputNode << Value.OutputNode << Value.Input << Value.Output;
 	return Ar;
 }
 
 FORCEINLINE FArchive& operator<<(Chaos::FChaosArchive& Ar, Dataflow::FLink& Value)
 {
-	Ar << Value.Input << Value.Output;
+	Ar << Value.InputNode << Value.OutputNode << Value.Input << Value.Output;
 	return Ar;
 }
 
diff --git a/Engine/Source/Runtime/Experimental/Dataflow/Engine/Private/Dataflow/DataflowObject.cpp b/Engine/Source/Runtime/Experimental/Dataflow/Engine/Private/Dataflow/DataflowObject.cpp
--- a/Engine/Source/Runtime/Experimental/Dataflow/Engine/Private/Dataflow/DataflowObject.cpp
+++ b/Engine/Source/Runtime/Experimental/Dataflow/Engine/Private/Dataflow/DataflowObject.cpp
@@ -50,6 +50,8 @@ void UDataflow::PostEditChangeProperty(struct FPropertyChangedEvent& PropertyCha
 void UDataflow::PostLoad()
 {
 #if WITH_EDITOR
+	const TSet<FName>& DisabledNodes = Dataflow->GetDisabledNodes();
+
 	for (UEdGraphNode* EdNode : Nodes)
 	{
 		UDataflowEdNode* DataflowEdNode = Cast<UDataflowEdNode>(EdNode);
@@ -57,6 +59,11 @@ void UDataflow::PostLoad()
 		{
 			DataflowEdNode->SetDataflowGraph(Dataflow);
 		}
+
+		if (DisabledNodes.Contains(FName(EdNode->GetName())))
+		{
+			EdNode->SetEnabledState(ENodeEnabledState::Disabled);
+		}
 	}
 #endif
 
