Chaos destruction : import external collision data from create a geometry collection from other geometry collections

diff --git a/Engine/Source/Runtime/Experimental/GeometryCollectionEngine/Private/GeometryCollection/GeometryCollectionEngineConversion.cpp b/Engine/Source/Runtime/Experimental/GeometryCollectionEngine/Private/GeometryCollection/GeometryCollectionEngineConversion.cpp
--- a/Engine/Source/Runtime/Experimental/GeometryCollectionEngine/Private/GeometryCollection/GeometryCollectionEngineConversion.cpp
+++ b/Engine/Source/Runtime/Experimental/GeometryCollectionEngine/Private/GeometryCollection/GeometryCollectionEngineConversion.cpp
@@ -465,6 +465,8 @@ void FGeometryCollectionEngineConversion::AppendGeometryCollection(const UGeomet
 		TargetMaterialIndex[FaceOffset] = FaceOffset;
 	}
 
+	using FCollisionType = FGeometryDynamicCollection::FSharedImplicit;
+	
 	// source transform information
 	const TManagedArray<FTransform>& SourceTransform = SourceGeometryCollectionPtr->Transform;
 	const TManagedArray<FString>& SourceBoneName = SourceGeometryCollectionPtr->BoneName;
@@ -475,6 +477,7 @@ void FGeometryCollectionEngineConversion::AppendGeometryCollection(const UGeomet
 	const TManagedArray<int32>& SourceSimulationType = SourceGeometryCollectionPtr->SimulationType;
 	const TManagedArray<int32>& SourceStatusFlags = SourceGeometryCollectionPtr->StatusFlags;
 	const TManagedArray<int32>& SourceInitialDynamicState = SourceGeometryCollectionPtr->InitialDynamicState;
+	const TManagedArray<FCollisionType>* SourceExternalCollisions = SourceGeometryCollectionPtr->FindAttribute<FCollisionType>("ExternalCollisions", FGeometryCollection::TransformGroup);
 	
 	// target transform information
 	TManagedArray<FTransform>& TargetTransform = GeometryCollection->Transform;
@@ -486,6 +489,7 @@ void FGeometryCollectionEngineConversion::AppendGeometryCollection(const UGeomet
 	TManagedArray<int32>& TargetSimulationType = GeometryCollection->SimulationType;
 	TManagedArray<int32>& TargetStatusFlags = GeometryCollection->StatusFlags;
 	TManagedArray<int32>& TargetInitialDynamicState = GeometryCollection->InitialDynamicState;
+	TManagedArray<FCollisionType>& TargetExternalCollisions = GeometryCollection->AddAttribute<FCollisionType>("ExternalCollisions", FGeometryCollection::TransformGroup);
 
 	// append transform hierarchy
 	for (int32 TransformIndex = 0; TransformIndex < TransformCount; ++TransformIndex)
@@ -522,6 +526,12 @@ void FGeometryCollectionEngineConversion::AppendGeometryCollection(const UGeomet
 		TargetSimulationType[TransformOffset] = SourceSimulationType[TransformIndex];
 		TargetStatusFlags[TransformOffset] = SourceStatusFlags[TransformIndex];
 		TargetInitialDynamicState[TransformOffset] = SourceInitialDynamicState[TransformIndex];
+
+		TargetExternalCollisions[TransformOffset] = nullptr;
+		if (SourceExternalCollisions)
+		{
+			TargetExternalCollisions[TransformOffset] = (*SourceExternalCollisions)[TransformIndex];
+		}
 	}
 
 	// source geometry information
