Dataflow : Optional Serialize Macro
- Moved optional blocks in to macro wrappers/bookends.

diff --git a/Engine/Source/Runtime/Experimental/Dataflow/Core/Private/Dataflow/DataflowGraph.cpp b/Engine/Source/Runtime/Experimental/Dataflow/Core/Private/Dataflow/DataflowGraph.cpp
--- a/Engine/Source/Runtime/Experimental/Dataflow/Core/Private/Dataflow/DataflowGraph.cpp
+++ b/Engine/Source/Runtime/Experimental/Dataflow/Core/Private/Dataflow/DataflowGraph.cpp
@@ -6,6 +6,7 @@
 #include "Dataflow/DataflowNodeParameters.h"
 #include "Dataflow/DataflowNodeFactory.h"
 #include "Logging/LogMacros.h"
+#include "Dataflow/DataflowArchive.h"
 
 DEFINE_LOG_CATEGORY_STATIC(DATAFLOW_LOG, Error, All);
 
@@ -112,37 +113,22 @@ namespace Dataflow
 				ArName = Node->GetName();
 				Ar << ArGuid << ArType << ArName;
 
-				// Data Jump Serialization (on Node and connections)
+				DATAFLOW_OPTIONAL_BLOCK_WRITE_BEGIN()
 				{
-					const int64 NodeDataSizePos = Ar.Tell();
-					int64 NodeDataSize = 0;
-					Ar << NodeDataSize;
-
-					const int64 NodeBegin = Ar.Tell();
-					{ // Data we can jump past on load
-						TArray< FConnection* > IO = Node->GetOutputs();  IO.Append(Node->GetInputs());
-						ArNum = IO.Num();
-						Ar << ArNum;
-						for (FConnection* Conn : IO)
-						{
-							ArGuid = Conn->GetGuid();
-							ArType = Conn->GetType();
-							ArName = Conn->GetName();
-							Ar << ArGuid << ArType << ArName;
-						}
-
-						Node->SerializeInternal(Ar);
-					}
-
-					if (NodeDataSizePos != INDEX_NONE)
+					TArray< FConnection* > IO = Node->GetOutputs();
+					IO.Append(Node->GetInputs());
+					ArNum = IO.Num();
+					Ar << ArNum;
+					for (FConnection* Conn : IO)
 					{
-						const int64 NodeEnd = Ar.Tell();
-						NodeDataSize = (NodeEnd - NodeBegin);
-						Ar.Seek(NodeDataSizePos);
-						Ar << NodeDataSize;
-						Ar.Seek(NodeEnd);
+						ArGuid = Conn->GetGuid();
+						ArType = Conn->GetType();
+						ArName = Conn->GetName();
+						Ar << ArGuid << ArType << ArName;
 					}
+					Node->SerializeInternal(Ar);
 				}
+				DATAFLOW_OPTIONAL_BLOCK_WRITE_END();
 			}
 
 			Ar << Connections;
@@ -153,7 +139,6 @@ namespace Dataflow
 			FGuid ArGuid;
 			FName ArType, ArName;
 			int32 ArNum = 0;
-			int64 NodeDataSize = 0;
 
 			TMap<FGuid, TSharedPtr<FNode> > NodeGuidMap;
 			TMap<FGuid, FConnection* > ConnectionGuidMap;
@@ -161,30 +146,22 @@ namespace Dataflow
 			Ar << ArNum;
 			for (int32 Ndx = ArNum; Ndx > 0; Ndx--)
 			{
-				Ar << ArGuid << ArType << ArName << NodeDataSize;
+				Ar << ArGuid << ArType << ArName;
 
-				if (TSharedPtr<FNode> Node = FNodeFactory::GetInstance()->NewNodeFromRegisteredType(*this, { ArGuid,ArType,ArName }))
+				TSharedPtr<FNode> Node = FNodeFactory::GetInstance()->NewNodeFromRegisteredType(*this, { ArGuid,ArType,ArName });
+				DATAFLOW_OPTIONAL_BLOCK_READ_BEGIN(Node != nullptr)
 				{
 					ensure(!NodeGuidMap.Contains(ArGuid));
-					NodeGuidMap.Add(ArGuid,Node);
-
+					NodeGuidMap.Add(ArGuid, Node);
 
 					int ArNumInner;
 					Ar << ArNumInner;
 					TArray< FConnection* > IO = Node->GetOutputs();  IO.Append(Node->GetInputs());
-					for (int Cdx=0;Cdx< ArNumInner; Cdx++)
+					for (int Cdx = 0; Cdx < ArNumInner; Cdx++)
 					{
 						Ar << ArGuid << ArType << ArName;
 						if (Cdx < IO.Num())
 						{
-							//if ((EGraphConnectionType)ArIntType!=IO[Cdx]->GetType())
-							//{
-							//	FName TypeName = GraphConnectionTypeName((EGraphConnectionType)ArIntType);
-							//	FName FromType = GraphConnectionTypeName(IO[Cdx]->GetType());
-							//	UE_LOG(DATAFLOW_LOG, Fatal, TEXT("Type mismatch in input file for Ed::FGraphNode Input/Output (%s,%s) to type (%s)"),
-							//		*TypeName.ToString(), *ArName.ToString(), *FromType.ToString());
-							//}
-
 							IO[Cdx]->SetGuid(ArGuid);
 							ensure(!ConnectionGuidMap.Contains(ArGuid));
 							ConnectionGuidMap.Add(ArGuid, IO[Cdx]);
@@ -193,15 +170,14 @@ namespace Dataflow
 
 					Node->SerializeInternal(Ar);
 				}
-				else
+				DATAFLOW_OPTIONAL_BLOCK_READ_ELSE()
 				{
-					Ar.Seek(Ar.Tell() + NodeDataSize); // Data Jump Serialization (on mising Node)
 					DisabledNodes.Add(ArName);
-					ensureMsgf(false, 
+					ensureMsgf(false,
 						TEXT("Error: Missing registered node type (%s) will be removed from graph on load. Graph will fail to evaluate due to missing node (%s).")
 						, *ArType.ToString(), *ArName.ToString());
 				}
-				
+				DATAFLOW_OPTIONAL_BLOCK_READ_END();
 			}
 
 			TArray< FLink > LocalConnections;
diff --git a/Engine/Source/Runtime/Experimental/Dataflow/Core/Private/Dataflow/DataflowNode.cpp b/Engine/Source/Runtime/Experimental/Dataflow/Core/Private/Dataflow/DataflowNode.cpp
--- a/Engine/Source/Runtime/Experimental/Dataflow/Core/Private/Dataflow/DataflowNode.cpp
+++ b/Engine/Source/Runtime/Experimental/Dataflow/Core/Private/Dataflow/DataflowNode.cpp
@@ -4,6 +4,7 @@
 
 #include "Dataflow/DataflowNode.h"
 #include "Dataflow/DataflowProperty.h"
+#include "Dataflow/DataflowArchive.h"
 
 namespace Dataflow
 {
@@ -118,38 +119,23 @@ namespace Dataflow
 				Type = (uint8)Prop->GetType();
 				StrName = Prop->GetName().ToString();
 				PropertySizeOf = Prop->SizeOf();
-				Ar << Type << StrName << PropertySizeOf;
-
-				// The following code serializes the size of the PropertyData to serialize so that when loading,
-				// we can skip over it if the function could not be loaded.
-				// Serialize a temporary value for the delta in order to end up with an archive of the right size.
-				// Then serialize the PropertyData in order to get its size.
-				const int64 PropertyDataSizePos = Ar.Tell();
-				PropertyDataSize = 0;
-				Ar << PropertyDataSize;
-
-				const int64 PropBegin = Ar.Tell();
-				Prop->Serialize(Ar);
-
-				// Only go back and serialize the number of argument bytes if there is actually an underlying buffer to seek to.
-				// Come back to the temporary value we wrote and overwrite it with the PropertyData size we just calculated.
-				// And finally seek back to the end.
-				if (PropertyDataSizePos != INDEX_NONE)
+				Ar << Type << StrName;
+
+				DATAFLOW_OPTIONAL_BLOCK_WRITE_BEGIN()
 				{
-					const int64 PropEnd = Ar.Tell();
-					PropertyDataSize = (PropEnd - PropBegin);
-					Ar.Seek(PropertyDataSizePos);
-					Ar << PropertyDataSize;
-					Ar.Seek(PropEnd);
+					Prop->Serialize(Ar);
 				}
+				DATAFLOW_OPTIONAL_BLOCK_WRITE_END()
+
 			}
 			else if (Ar.IsLoading())
 			{
 
-				Ar << Type << StrName << PropertySizeOf << PropertyDataSize;
+				Ar << Type << StrName;
 				FName PropName(StrName);
 
-				if (FProperty* NewProperty = FProperty::NewProperty((FProperty::EType)Type, PropName))
+				FProperty* NewProperty = FProperty::NewProperty((FProperty::EType)Type, PropName);
+				DATAFLOW_OPTIONAL_BLOCK_READ_BEGIN(NewProperty != nullptr)
 				{
 					if (Map.Contains(PropName) && Map[PropName]->GetType() == (FProperty::EType)Type)
 					{
@@ -170,10 +156,11 @@ namespace Dataflow
 						}
 					}
 				}
-				else
+				DATAFLOW_OPTIONAL_BLOCK_READ_ELSE()
 				{
-					Ar.Seek(Ar.Tell() + PropertyDataSize);
+					// Silently skip
 				}
+				DATAFLOW_OPTIONAL_BLOCK_READ_END()
 			}
 		}
 	}
diff --git a/Engine/Source/Runtime/Experimental/Dataflow/Core/Public/Dataflow/DataflowArchive.h b/Engine/Source/Runtime/Experimental/Dataflow/Core/Public/Dataflow/DataflowArchive.h
new file mode 100644
--- /dev/null
+++ b/Engine/Source/Runtime/Experimental/Dataflow/Core/Public/Dataflow/DataflowArchive.h
@@ -0,0 +1,42 @@
+// Copyright Epic Games, Inc. All Rights Reserved.
+
+#pragma once
+
+// The following code serializes the size of the PropertyData to serialize so that when loading,
+// we can skip over it if the function could not be loaded.
+// Serialize a temporary value for the delta in order to end up with an archive of the right size.
+// Then serialize the PropertyData in order to get its size.
+#define DATAFLOW_OPTIONAL_BLOCK_WRITE_BEGIN(){\
+	const int64 NodeDataSizePos = Ar.Tell();\
+	int64 NodeDataSize = 0;\
+	Ar << NodeDataSize;\
+	const int64 NodeBegin = Ar.Tell();
+
+
+// Only go back and serialize the number of argument bytes if there is actually an underlying buffer to seek to.
+// Come back to the temporary value we wrote and overwrite it with the PropertyData size we just calculated.
+// And finally seek back to the end.
+#define DATAFLOW_OPTIONAL_BLOCK_WRITE_END()\
+	if (NodeDataSizePos != INDEX_NONE)\
+	{\
+		const int64 NodeEnd = Ar.Tell();\
+		NodeDataSize = (NodeEnd - NodeBegin);\
+		Ar.Seek(NodeDataSizePos);\
+		Ar << NodeDataSize;\
+		Ar.Seek(NodeEnd);\
+	}\
+	}
+
+
+#define DATAFLOW_OPTIONAL_BLOCK_READ_BEGIN(COND) \
+	{\
+	int64 NodeDataSize = 0;\
+	Ar << NodeDataSize;\
+	if (COND) 
+
+#define DATAFLOW_OPTIONAL_BLOCK_READ_ELSE() \
+	else{\
+	Ar.Seek(Ar.Tell() + NodeDataSize);\
+
+#define DATAFLOW_OPTIONAL_BLOCK_READ_END() }}
+
