Chaos geometry collection : fix import of external collision to ignore non implicit shapes

diff --git a/Engine/Source/Runtime/Experimental/GeometryCollectionEngine/Private/GeometryCollection/GeometryCollectionEngineConversion.cpp b/Engine/Source/Runtime/Experimental/GeometryCollectionEngine/Private/GeometryCollection/GeometryCollectionEngineConversion.cpp
--- a/Engine/Source/Runtime/Experimental/GeometryCollectionEngine/Private/GeometryCollection/GeometryCollectionEngineConversion.cpp
+++ b/Engine/Source/Runtime/Experimental/GeometryCollectionEngine/Private/GeometryCollection/GeometryCollectionEngineConversion.cpp
@@ -65,6 +65,12 @@ FORCEINLINE uint32 GetTypeHash(const FUniqueVertex& UniqueVertex)
 	return VertexHash;
 }
 
+static bool IsImportableImplicitObjectType(Chaos::EImplicitObjectType Type)
+{
+	const Chaos::EImplicitObjectType InnerType = Type & (~(Chaos::ImplicitObjectType::IsScaled | Chaos::ImplicitObjectType::IsInstanced));
+	return (InnerType == Chaos::ImplicitObjectType::Box || InnerType == Chaos::ImplicitObjectType::Sphere || InnerType == Chaos::ImplicitObjectType::Capsule || InnerType == Chaos::ImplicitObjectType::Convex);
+}
+
 void FGeometryCollectionEngineConversion::AppendStaticMesh(const UStaticMesh* StaticMesh, const TArray<UMaterialInterface*>& Materials, const FTransform& StaticMeshTransform, UGeometryCollection* GeometryCollectionObject, bool ReindexMaterials)
 {
 #if WITH_EDITORONLY_DATA
@@ -271,17 +277,23 @@ void FGeometryCollectionEngineConversion::AppendStaticMesh(const UStaticMesh* St
 
 			using FCollisionType = FGeometryDynamicCollection::FSharedImplicit;
 			TManagedArray<FCollisionType>& ExternaCollisions = GeometryCollection->AddAttribute<FCollisionType>("ExternalCollisions", FGeometryCollection::TransformGroup);
-			if (Geoms.Num() == 0)
-			{
-				ExternaCollisions[TransformIndex1] = nullptr;
-			}
-			else if (Geoms.Num() == 1)
+
+			ExternaCollisions[TransformIndex1] = nullptr;
+			for (int32 GeomIndex = 0; GeomIndex < Geoms.Num();)
 			{
-				ExternaCollisions[TransformIndex1] = TSharedPtr<Chaos::FImplicitObject>(Geoms[0].Release());  
+				// make sure we only import box, sphere, capsule or convex
+				if (IsImportableImplicitObjectType(Geoms[GeomIndex]->GetType()))
+				{
+					GeomIndex++;
+				}
+				else
+				{
+					Geoms.RemoveAtSwap(GeomIndex);
+				}
 			}
-			else
+			if (Geoms.Num() > 0)
 			{
-				ExternaCollisions[TransformIndex1] = MakeShared<Chaos::FImplicitObjectUnion>(MoveTemp(Geoms)); 
+				ExternaCollisions[TransformIndex1] = MakeShared<Chaos::FImplicitObjectUnion>(MoveTemp(Geoms));
 			}
 		}
 		
