Chaos : Fix headless chaos

diff --git a/Engine/Source/Runtime/Experimental/Chaos/Public/PhysicsProxy/GeometryCollectionPhysicsProxy.h b/Engine/Source/Runtime/Experimental/Chaos/Public/PhysicsProxy/GeometryCollectionPhysicsProxy.h
--- a/Engine/Source/Runtime/Experimental/Chaos/Public/PhysicsProxy/GeometryCollectionPhysicsProxy.h
+++ b/Engine/Source/Runtime/Experimental/Chaos/Public/PhysicsProxy/GeometryCollectionPhysicsProxy.h
@@ -79,7 +79,7 @@ public:
 		const FSimulationParameters& SimulationParameters,
 		FCollisionFilterData InSimFilter,
 		FCollisionFilterData InQueryFilter,
-		FGuid InCollectorGuid,
+		FGuid InCollectorGuid = FGuid::NewGuid(),
 		const Chaos::EMultiBufferMode BufferMode=Chaos::EMultiBufferMode::TripleGuarded);
 	virtual ~FGeometryCollectionPhysicsProxy();
 
