Interchange: Serialize the full class path name instead of just the class name when serializing to json. This fixes warnings when deserializing error messages between the interchange worker and the UE editor.

diff --git a/Engine/Source/Runtime/Experimental/Interchange/Core/Private/InterchangeResult.cpp b/Engine/Source/Runtime/Experimental/Interchange/Core/Private/InterchangeResult.cpp
--- a/Engine/Source/Runtime/Experimental/Interchange/Core/Private/InterchangeResult.cpp
+++ b/Engine/Source/Runtime/Experimental/Interchange/Core/Private/InterchangeResult.cpp
@@ -15,7 +15,7 @@ FString UInterchangeResult::ToJson()
 	TSharedPtr<FJsonObject> MsgObject = MakeShared<FJsonObject>();
 	TSharedPtr<FJsonObject> PropertiesObject = MakeShared<FJsonObject>();
 
-	MsgObject->SetStringField(TEXT("Type"), GetClass()->GetName());
+	MsgObject->SetStringField(TEXT("Type"), GetClass()->GetPathName());
 
 	TArray<TSharedPtr<FJsonValue>> Properties;
 	for (TFieldIterator<FProperty> FieldIterator(GetClass()); FieldIterator; ++FieldIterator)
