GeometryCollection : Foreign Keys
Fixes a bug in the reindexing of the geometry collection were some of the array types were not setting INDEX_NONE for their value when a dependent array entry was removed. This change updates all int arrays typed to set INDEX_NONE if they have a dependency. Doing this also flushed out a few cases where the vertices count, and face count, were not updated properly.
There are two concerning changes in here:
-	The new method for calculating vertex and face ranges after a resize.
-	Setting INDEX_NONE in an array when its referenced attribute entry is removed.
Also:
-	Turns out that FaceStart and VertexStart are not actually �simple� keys into the FaceGroup or the VertexGroup. They are indices that point to the start of a range. While it still makes sense to adjust all the start indices during a resize. If a start index is removed, then start should be moved to the new start index, not just set to INDEX_NONE.
-	Added support for TManagedArray<TArray<int>>

diff --git a/Engine/Source/Programs/HeadlessChaos/Private/GeometryCollection/GeometryCollectionTestCreation.cpp b/Engine/Source/Programs/HeadlessChaos/Private/GeometryCollection/GeometryCollectionTestCreation.cpp
--- a/Engine/Source/Programs/HeadlessChaos/Private/GeometryCollection/GeometryCollectionTestCreation.cpp
+++ b/Engine/Source/Programs/HeadlessChaos/Private/GeometryCollection/GeometryCollectionTestCreation.cpp
@@ -1256,4 +1256,185 @@ namespace GeometryCollectionTest
 
 		delete Collection;
 	}
+
+
+	void IntListReindexOnDeletionTest()
+	{
+		//--gtest_filter=GeometryCollection_CreationTest.IntListReindexOnDeletionTest
+
+		//
+		// Build two arrays in GroupB:
+		//   - One array is an index into GroupA. (dependent on  A)
+		//   - The other array stores a local copy of the value in A. 
+		// As elements are deleted from GroupA, the indices will be 
+		// shifted within GroupB, and missing references will be 
+		// updated to INDEX_NONE. 
+		//
+		// At the end of the test all indices in GroupB will be INDEX_NONE.
+
+		FManagedArrayCollection* Collection = new FManagedArrayCollection();
+
+		FManagedArrayCollection::FConstructionParameters DependencyOnA("GroupA");
+		Collection->AddGroup("GroupA");
+		Collection->AddGroup("GroupB");
+
+		Collection->AddElements(10, "GroupA");
+		Collection->AddElements(10, "GroupB");
+
+		TManagedArray<float>& AttrInA = Collection->AddAttribute<float>("Floats", "GroupA");
+
+		// int32
+		TManagedArray<int32>& IntInB = Collection->AddAttribute<int32>("Int", "GroupB", DependencyOnA);
+		TManagedArray<float>& IntInBVal = Collection->AddAttribute<float>("Intf", "GroupB");
+		//FIntVector
+		TManagedArray<FIntVector>& VecInB = Collection->AddAttribute<FIntVector>("IntVec", "GroupB", DependencyOnA);
+		TManagedArray< TArray<float> >& VecInBVal = Collection->AddAttribute< TArray<float> >("IntVecf", "GroupB");
+		//FIntVector2
+		TManagedArray<FIntVector2>& Vec2InB = Collection->AddAttribute<FIntVector2>("IntVec2", "GroupB", DependencyOnA);
+		TManagedArray< TArray<float> >& Vec2InBVal = Collection->AddAttribute< TArray<float> >("IntVec2f", "GroupB");
+		// TArray<FIntVector2>
+		TManagedArray<TArray<FIntVector2>>& VecArray2InB = Collection->AddAttribute<TArray<FIntVector2>>("IntVecArray2", "GroupB", DependencyOnA);
+		TManagedArray< TArray<FVector2f> >& VecArray2InBVal = Collection->AddAttribute< TArray<FVector2f> >("IntVecArray2f", "GroupB");
+		// FIntVector4
+		TManagedArray<FIntVector4>& Vec4InB = Collection->AddAttribute<FIntVector4>("Int4", "GroupB", DependencyOnA);
+		TManagedArray< TArray<float> >& Vec4InBVal = Collection->AddAttribute< TArray<float> >("Int4f", "GroupB");
+		// TArray<int32>
+		TManagedArray<TArray<int32>>& ArrayNInB = Collection->AddAttribute<TArray<int32>>("IntArray", "GroupB", DependencyOnA);
+		TManagedArray<TArray<float>>& ArrayNInBVal = Collection->AddAttribute<TArray<float>>("IntArrayf", "GroupB");
+
+
+		for (int i = 0; i < 10; i++)
+		{
+			AttrInA[i] = i * 0.1;
+		}
+
+		for (int i = 0; i < 10; i++)
+		{
+			// int32
+			IntInB[i] = i;
+			IntInBVal[i] = AttrInA[IntInB[i]];
+
+			//FIntVector
+			VecInB[i] = FIntVector(FMath::Clamp(i - 1, 0, 9), FMath::Clamp(i, 0, 9), FMath::Clamp(i + 1, 0, 9));
+			VecInBVal[i].SetNum(3);
+			VecInBVal[i][0] = AttrInA[VecInB[i][0]];
+			VecInBVal[i][1] = AttrInA[VecInB[i][1]];
+			VecInBVal[i][2] = AttrInA[VecInB[i][2]];
+
+			//FIntVector4
+			Vec2InB[i] = FIntVector2(FMath::Clamp(i - 1, 0, 9), FMath::Clamp(i, 0, 9));
+			Vec2InBVal[i].SetNum(2);
+			Vec2InBVal[i][0] = AttrInA[Vec2InB[i][0]];
+			Vec2InBVal[i][1] = AttrInA[Vec2InB[i][1]];
+
+			// TArray<FIntVector2>
+			VecArray2InB[i].SetNum(i);
+			VecArray2InBVal[i].SetNum(i);
+			for (int j = 0; j < i; j++)
+			{
+				VecArray2InB[i][j] = FIntVector2(FMath::Clamp(j - 1, 0, 9), FMath::Clamp(j, 0, 9));
+				VecArray2InBVal[i][j] = FVector2f(AttrInA[VecArray2InB[i][j][0]], AttrInA[VecArray2InB[i][j][1]]);
+			}
+
+			//FIntVector4
+			Vec4InB[i] = FIntVector4(FMath::Clamp(i - 1, 0, 9), FMath::Clamp(i, 0, 9), FMath::Clamp(i + 1, 0, 9), FMath::Clamp(i + 2, 0, 9));
+			Vec4InBVal[i].SetNum(4);
+			Vec4InBVal[i][0] = AttrInA[Vec4InB[i][0]];
+			Vec4InBVal[i][1] = AttrInA[Vec4InB[i][1]];
+			Vec4InBVal[i][2] = AttrInA[Vec4InB[i][2]];
+			Vec4InBVal[i][3] = AttrInA[Vec4InB[i][3]];
+
+			//TArray<int32>
+			ArrayNInB[i].SetNum(i);
+			ArrayNInBVal[i].SetNum(i);
+			for (int j = 0; j < i; j++)
+			{
+				ArrayNInB[i][j] = FMath::Clamp(j, 0, 9);
+				ArrayNInBVal[i][j] = AttrInA[ArrayNInB[i][j]];
+			}
+		}
+
+		auto RemoveIndicesOn = [&](const TArray<int32>& Elements)
+		{
+			bool HasValidValues = false;
+			Collection->RemoveElements("GroupA", Elements);
+			for (int i = 0; i < 10; i++)
+			{
+				// int32
+				if (IntInB[i] != INDEX_NONE)
+				{
+					HasValidValues = true;
+					EXPECT_NEAR(AttrInA[IntInB[i]], IntInBVal[i], FLT_EPSILON);
+				}
+
+				//FIntVector
+				for (int j = 0; j < 3; j++)
+				{
+					if (VecInB[i][j] != INDEX_NONE)
+					{
+						HasValidValues = true;
+						EXPECT_NEAR(AttrInA[VecInB[i][j]], VecInBVal[i][j], FLT_EPSILON);
+					}
+				}
+
+				//FIntVector2
+				for (int j = 0; j < 2; j++)
+				{
+					if (Vec2InB[i][j] != INDEX_NONE)
+					{
+						HasValidValues = true;
+						EXPECT_NEAR(AttrInA[Vec2InB[i][j]], Vec2InBVal[i][j], FLT_EPSILON);
+					}
+				}
+
+				// TArray<FIntVector2>
+				for (int j = 0; j < i; j++)
+				{
+					if (VecArray2InB[i][j][0] != INDEX_NONE)
+					{
+						HasValidValues = true;
+						EXPECT_NEAR(AttrInA[VecArray2InB[i][j][0]], VecArray2InBVal[i][j][0], FLT_EPSILON);
+					}
+					if (VecArray2InB[i][j][1] != INDEX_NONE)
+					{
+						HasValidValues = true;
+						ensure(FMath::IsNearlyEqual(AttrInA[VecArray2InB[i][j][1]], VecArray2InBVal[i][j][1], FLT_EPSILON));
+						EXPECT_NEAR(AttrInA[VecArray2InB[i][j][1]], VecArray2InBVal[i][j][1], FLT_EPSILON);
+					}
+				}
+
+
+				//FIntVector4
+				for (int j = 0; j < 4; j++)
+				{
+					if (Vec4InB[i][j] != INDEX_NONE)
+					{
+						HasValidValues = true;
+						EXPECT_NEAR(AttrInA[Vec4InB[i][j]], Vec4InBVal[i][j], FLT_EPSILON);
+					}
+				}
+
+				//TArray<int32>
+				for (int j = 0; j < i; j++)
+				{
+					if (ArrayNInB[i][j] != INDEX_NONE)
+					{
+						HasValidValues = true;
+						EXPECT_NEAR(AttrInA[ArrayNInB[i][j]], ArrayNInBVal[i][j], FLT_EPSILON);
+					}
+				}
+			}
+			return HasValidValues;
+		};
+
+		EXPECT_TRUE(RemoveIndicesOn({ 3 }));
+		EXPECT_TRUE(RemoveIndicesOn({ 3,4,5 }));
+		EXPECT_TRUE(RemoveIndicesOn({ 5 }));
+		EXPECT_TRUE(RemoveIndicesOn({ 0 }));
+		EXPECT_TRUE(RemoveIndicesOn({ 1,3 }));
+		EXPECT_TRUE(RemoveIndicesOn({ 1 }));
+		EXPECT_FALSE(RemoveIndicesOn({ 0 }));
+
+		delete Collection;
+	}
 }
diff --git a/Engine/Source/Programs/HeadlessChaos/Private/HeadlessChaosApp.cpp b/Engine/Source/Programs/HeadlessChaos/Private/HeadlessChaosApp.cpp
--- a/Engine/Source/Programs/HeadlessChaos/Private/HeadlessChaosApp.cpp
+++ b/Engine/Source/Programs/HeadlessChaos/Private/HeadlessChaosApp.cpp
@@ -335,7 +335,8 @@ TEST(GeometryCollection_CreationTest,DeleteRootLeafMiddle) { GeometryCollectionT
 TEST(GeometryCollection_CreationTest,DeleteEverything) { GeometryCollectionTest::DeleteEverything(); SUCCEED(); }
 TEST(GeometryCollection_CreationTest,ReindexMaterialsTest) { GeometryCollectionTest::ReindexMaterialsTest(); SUCCEED(); }
 TEST(GeometryCollection_CreationTest,ContiguousElementsTest) { GeometryCollectionTest::ContiguousElementsTest(); SUCCEED(); }
-TEST(GeometryCollection_CreationTest,AttributeDependencyTest) { GeometryCollectionTest::AttributeDependencyTest(); SUCCEED(); }
+TEST(GeometryCollection_CreationTest, AttributeDependencyTest) { GeometryCollectionTest::AttributeDependencyTest(); SUCCEED(); }
+TEST(GeometryCollection_CreationTest, IntListReindexOnDeletionTest) { GeometryCollectionTest::IntListReindexOnDeletionTest(); SUCCEED(); }
 
 
 // Proximity Tests
diff --git a/Engine/Source/Programs/HeadlessChaos/Public/GeometryCollection/GeometryCollectionTestCreation.h b/Engine/Source/Programs/HeadlessChaos/Public/GeometryCollection/GeometryCollectionTestCreation.h
--- a/Engine/Source/Programs/HeadlessChaos/Public/GeometryCollection/GeometryCollectionTestCreation.h
+++ b/Engine/Source/Programs/HeadlessChaos/Public/GeometryCollection/GeometryCollectionTestCreation.h
@@ -35,4 +35,6 @@ namespace GeometryCollectionTest
 
 	void  AttributeDependencyTest();
 
+	void IntListReindexOnDeletionTest();
+
 }
diff --git a/Engine/Source/Runtime/Experimental/Chaos/Private/GeometryCollection/GeometryCollection.cpp b/Engine/Source/Runtime/Experimental/Chaos/Private/GeometryCollection/GeometryCollection.cpp
--- a/Engine/Source/Runtime/Experimental/Chaos/Private/GeometryCollection/GeometryCollection.cpp
+++ b/Engine/Source/Runtime/Experimental/Chaos/Private/GeometryCollection/GeometryCollection.cpp
@@ -543,11 +543,13 @@ void FGeometryCollection::RemoveElements(const FName & Group, const TArray<int32
 		}
 		else if( Group == FGeometryCollection::FacesGroup)
 		{
+			BuildFaceToGeometryMapping();
 			Super::RemoveElements(Group, SortedDeletionList);
 			UpdateFaceGroupElements();
 		}
 		else if (Group == FGeometryCollection::VerticesGroup)
 		{
+			BuildVertexToGeometryMapping();
 			Super::RemoveElements(Group, SortedDeletionList);
 			UpdateVerticesGroupElements();
 		}
@@ -778,56 +780,105 @@ void FGeometryCollection::ReorderGeometryElements(const TArray<int32>& NewOrder)
 	Super::ReorderElements(GeometryGroup, NewOrder);
 }
 
+bool FGeometryCollection::BuildVertexToGeometryMapping(bool InSaved )
+{
+	bool AttributeAdded = false;
+
+	if (!FindAttribute<int32>("VertexToGeometryIndex", VerticesGroup))
+	{
+		FConstructionParameters NotSaved = { FName(""), InSaved };
+		AddAttribute<int32>("VertexToGeometryIndex", VerticesGroup, NotSaved);
+		AttributeAdded = true;
+	}
+
+	TManagedArray<int32>* VertexGeometryMap = FindAttribute<int32>("VertexToGeometryIndex", VerticesGroup);
+	if (ensure(VertexGeometryMap))
+	{
+		for (int32 GeometryIndex = NumElements(GeometryGroup) - 1; GeometryIndex >= 0; GeometryIndex--)
+		{
+			int VertexEnd = VertexStart[GeometryIndex] + VertexCount[GeometryIndex];
+			for (int32 VertexIndex = VertexStart[GeometryIndex]; VertexIndex < VertexEnd; VertexIndex++)
+			{
+				(*VertexGeometryMap)[VertexIndex] = GeometryIndex;
+			}
+		}
+	}
+	return AttributeAdded;
+}
+
 void FGeometryCollection::UpdateVerticesGroupElements()
 {
 	//
 	//  Reset the VertexCount array
 	//
-	int32 NumberOfVertices = Vertex.Num();
-	for (int32 GeometryIndex = 0, ng = TransformIndex.Num(); GeometryIndex < ng; ++GeometryIndex)
+	TManagedArray<int32>* VertexGeometryMap = FindAttribute<int32>("VertexToGeometryIndex", VerticesGroup);
+	if(ensure(VertexGeometryMap))
+	{ 
+		for (int32 GeometryIndex = NumElements(FGeometryCollection::GeometryGroup) - 1; GeometryIndex >= 0; GeometryIndex--)
+		{
+			VertexStart[GeometryIndex] = INT_MAX;
+			VertexCount[GeometryIndex] = 0;
+		}
+
+		for (int32 VertexIndex = NumElements(VerticesGroup) - 1; VertexIndex >= 0; VertexIndex--)
+		{
+			VertexStart[ (*VertexGeometryMap)[VertexIndex] ] = FMath::Min(VertexStart[ (*VertexGeometryMap)[VertexIndex] ], VertexIndex);
+			VertexCount[ (*VertexGeometryMap)[VertexIndex] ]++;
+		}
+	}
+
+	ensure(HasContiguousVertices());
+}
+
+bool FGeometryCollection::BuildFaceToGeometryMapping(bool InSaved)
+{
+	bool AttributeAdded = false;
+
+	if (!FindAttribute<int32>("FaceToGeometryIndex", FacesGroup))
 	{
-		int32 VertexIndex = VertexStart[GeometryIndex];
-		if (VertexIndex != INDEX_NONE)
+		FConstructionParameters NotSaved = { FName(""), InSaved };
+		AddAttribute<int32>("FaceToGeometryIndex", FacesGroup, NotSaved);
+		AttributeAdded = true;
+	}
+
+	TManagedArray<int32>* FaceGeometryMap = FindAttribute<int32>("FaceToGeometryIndex", FacesGroup);
+	if (ensure(FaceGeometryMap))
+	{
+		for (int32 GeometryIndex = NumElements(GeometryGroup) - 1; GeometryIndex >= 0; GeometryIndex--)
 		{
-			int32 StartBoneMapTransformValue = BoneMap[VertexIndex];
-			int32 CurrentBoneMapTransformValue = StartBoneMapTransformValue;
-			while ((CurrentBoneMapTransformValue == StartBoneMapTransformValue) && (++VertexIndex < NumberOfVertices))
+			int FaceEnd = FaceStart[GeometryIndex] + FaceCount[GeometryIndex];
+			for (int32 FaceIndex = FaceStart[GeometryIndex]; FaceIndex < FaceEnd; FaceIndex++)
 			{
-				CurrentBoneMapTransformValue = BoneMap[VertexIndex];
+				(*FaceGeometryMap)[FaceIndex] = GeometryIndex;
 			}
-			VertexCount[GeometryIndex] = VertexIndex - VertexStart[GeometryIndex];
-		}
-		else
-		{
-			VertexCount[GeometryIndex] = 0;
 		}
 	}
+	return AttributeAdded;
 }
 
 void FGeometryCollection::UpdateFaceGroupElements()
 {
 	//
-	//  Reset the FaceCount array
+	//  Reset the FaceCount and NumFaces array
 	//
-	int32 NumberOfFaces = Indices.Num();
-	for (int32 GeometryIndex = 0, ng = TransformIndex.Num(); GeometryIndex < ng; ++GeometryIndex)
+
+	TManagedArray<int32>* FaceGeometryMap = FindAttribute<int32>("FaceToGeometryIndex", FacesGroup);
+	if (ensure(FaceGeometryMap))
 	{
-		int32 FaceIndex = FaceStart[GeometryIndex];
-		if (FaceIndex != INDEX_NONE)
+		for (int32 GeometryIndex = NumElements(FGeometryCollection::GeometryGroup) - 1; GeometryIndex >= 0; GeometryIndex--)
 		{
-			int32 StartBoneMapTransformValue = BoneMap[Indices[FaceIndex][0]];
-			int32 CurrentBoneMapTransformValue = StartBoneMapTransformValue;
-			while ((CurrentBoneMapTransformValue == StartBoneMapTransformValue) && (++FaceIndex < NumberOfFaces))
-			{
-				CurrentBoneMapTransformValue = BoneMap[Indices[FaceIndex][0]];
-			}
-			FaceCount[GeometryIndex] = FaceIndex - FaceStart[GeometryIndex];
+			FaceStart[GeometryIndex] = INT_MAX;
+			FaceCount[GeometryIndex] = 0;
 		}
-		else
+
+		for (int32 FaceIndex = NumElements(FacesGroup) - 1; FaceIndex >= 0; FaceIndex--)
 		{
-			FaceCount[GeometryIndex] = 0;
+			FaceStart[(*FaceGeometryMap)[FaceIndex]] = FMath::Min(FaceStart[(*FaceGeometryMap)[FaceIndex]], FaceIndex);
+			FaceCount[(*FaceGeometryMap)[FaceIndex]]++;
 		}
 	}
+
+	ensure(HasContiguousVertices());
 }
 
 
diff --git a/Engine/Source/Runtime/Experimental/Chaos/Private/GeometryCollection/ManagedArrayCollection.cpp b/Engine/Source/Runtime/Experimental/Chaos/Private/GeometryCollection/ManagedArrayCollection.cpp
--- a/Engine/Source/Runtime/Experimental/Chaos/Private/GeometryCollection/ManagedArrayCollection.cpp
+++ b/Engine/Source/Runtime/Experimental/Chaos/Private/GeometryCollection/ManagedArrayCollection.cpp
@@ -37,6 +37,7 @@ void FManagedArrayCollection::RemoveElements(const FName& Group, const TArray<in
 		TArray<int32> Offsets;
 		GeometryCollectionAlgo::BuildIncrementMask(SortedDeletionList, GroupSize, Offsets);
 
+		TSet<int32> DeletionSet(SortedDeletionList);
 		for (const TTuple<FKeyType, FValueType>& Entry : Map)
 		{
 			//
@@ -44,7 +45,7 @@ void FManagedArrayCollection::RemoveElements(const FName& Group, const TArray<in
 			//
 			if (Entry.Value.GroupIndexDependency == Group && Params.bReindexDependentAttibutes)
 			{
-				Entry.Value.Value->Reindex(Offsets, GroupSize - DelListNum, SortedDeletionList);
+				Entry.Value.Value->Reindex(Offsets, GroupSize - DelListNum, SortedDeletionList, DeletionSet);
 			}
 
 			//
diff --git a/Engine/Source/Runtime/Experimental/Chaos/Public/GeometryCollection/GeometryCollection.h b/Engine/Source/Runtime/Experimental/Chaos/Public/GeometryCollection/GeometryCollection.h
--- a/Engine/Source/Runtime/Experimental/Chaos/Public/GeometryCollection/GeometryCollection.h
+++ b/Engine/Source/Runtime/Experimental/Chaos/Public/GeometryCollection/GeometryCollection.h
@@ -311,11 +311,13 @@ protected:
 	/**
 	* Update Face Attributes based on changes in the group.
 	*/
+	bool BuildFaceToGeometryMapping(bool InSaved=false);
 	void UpdateFaceGroupElements();
 
 	/**
-	* Update Vertex Attributes based on changes in the group.
+	* Build and Update Vertex Attributes based on changes in the group.
 	*/
+	bool BuildVertexToGeometryMapping(bool InSaved = false);
 	void UpdateVerticesGroupElements();
 
 	/** 
diff --git a/Engine/Source/Runtime/Experimental/Chaos/Public/GeometryCollection/ManagedArray.h b/Engine/Source/Runtime/Experimental/Chaos/Public/GeometryCollection/ManagedArray.h
--- a/Engine/Source/Runtime/Experimental/Chaos/Public/GeometryCollection/ManagedArray.h
+++ b/Engine/Source/Runtime/Experimental/Chaos/Public/GeometryCollection/ManagedArray.h
@@ -188,7 +188,7 @@ public:
 	*   Offsets is the size of the dependent group;
 	*   Final is post resize of dependent group used for bounds checking on remapped indices.
 	*/
-	virtual void Reindex(const TArray<int32> & Offsets, const int32 & FinalSize, const TArray<int32> & SortedDeletionList) { }
+	virtual void Reindex(const TArray<int32> & Offsets, const int32 & FinalSize, const TArray<int32> & SortedDeletionList, const TSet<int32> & DeletionSet) { }
 
 #if 0 //not needed until per instance serialization
 	/** Swap elements*/
@@ -628,10 +628,10 @@ public:
 	virtual ~TManagedArray()
 	{}
 
-	virtual void Reindex(const TArray<int32> & Offsets, const int32 & FinalSize, const TArray<int32> & SortedDeletionList) override
+	virtual void Reindex(const TArray<int32> & Offsets, const int32 & FinalSize, const TArray<int32> & SortedDeletionList, const TSet<int32>& DeletionSet) override
 	{
 		UE_LOG(UManagedArrayLogging, Log, TEXT("TManagedArray<int32>[%p]::Reindex()"),this);
-	
+
 		int32 ArraySize = Num(), MaskSize = Offsets.Num();
 		for (int32 Index = 0; Index < ArraySize; Index++)
 		{
@@ -639,14 +639,14 @@ public:
 			if (0 <= RemapVal)
 			{
 				ensure(RemapVal < MaskSize);
-				this->operator[](Index) -= Offsets[RemapVal];
-
-				// #todo Reindexing is currently incorrect in a few cases where it leaves dangling indices that
-				// should be set to INDEX_NONE.
-				// This ensure is currently suppressed to handle the dangling index problem specifically 
-				// in the ConvexGroup case (GeometryCollectionConvexUtility::RemoveConvexHulls) but should be reinstated 
-				// once we have a better solution in place for Reindexing.
-				//ensure(-1 <= this->operator[](Index) && this->operator[](Index) < FinalSize);
+				if (DeletionSet.Contains(this->operator[](Index)))
+				{
+					this->operator[](Index) = INDEX_NONE;
+				}
+				else
+				{
+					this->operator[](Index) -= Offsets[RemapVal];
+				}
 				ensure(-1 <= this->operator[](Index));
 			}
 		}
@@ -690,12 +690,11 @@ public:
 	virtual ~TManagedArray()
 	{}
 	
-	virtual void Reindex(const TArray<int32> & Offsets, const int32 & FinalSize, const TArray<int32> & SortedDeletionList) override
+	virtual void Reindex(const TArray<int32> & Offsets, const int32 & FinalSize, const TArray<int32> & SortedDeletionList, const TSet<int32>& DeletionSet) override
 	{
 		UE_LOG(UManagedArrayLogging, Log, TEXT("TManagedArray<TArray<int32>>[%p]::Reindex()"), this);
 		
 		int32 ArraySize = Num(), MaskSize = Offsets.Num();
-		TSet<int32> SortedDeletionSet(SortedDeletionList);
 
 		for (int32 Index = 0; Index < ArraySize; Index++)
 		{
@@ -760,7 +759,7 @@ public:
 	virtual ~TManagedArray()
 	{}
 
-	virtual void Reindex(const TArray<int32> & Offsets, const int32 & FinalSize, const TArray<int32> & SortedDeletionList) override
+	virtual void Reindex(const TArray<int32> & Offsets, const int32 & FinalSize, const TArray<int32> & SortedDeletionList, const TSet<int32>& DeletionSet) override
 	{
 		UE_LOG(UManagedArrayLogging, Log, TEXT("TManagedArray<FIntVector>[%p]::Reindex()"), this);
 		int32 ArraySize = Num(), MaskSize = Offsets.Num();
@@ -772,7 +771,14 @@ public:
 				if (0 <= RemapVal[i])
 				{
 					ensure(RemapVal[i] < MaskSize);
-					this->operator[](Index)[i] -= Offsets[RemapVal[i]];
+					if (DeletionSet.Contains(this->operator[](Index)[i]))
+					{
+						this->operator[](Index)[i] = INDEX_NONE;
+					}
+					else
+					{
+						this->operator[](Index)[i] -= Offsets[RemapVal[i]];
+					}
 					ensure(-1 <= this->operator[](Index)[i] && this->operator[](Index)[i] <= FinalSize);
 				}
 			}
@@ -819,7 +825,7 @@ public:
 	virtual ~TManagedArray()
 	{}
 
-	virtual void Reindex(const TArray<int32>& Offsets, const int32& FinalSize, const TArray<int32>& SortedDeletionList) override
+	virtual void Reindex(const TArray<int32>& Offsets, const int32& FinalSize, const TArray<int32>& SortedDeletionList, const TSet<int32>& DeletionSet) override
 	{
 		UE_LOG(UManagedArrayLogging, Log, TEXT("TManagedArray<FIntVector>[%p]::Reindex()"), this);
 		int32 ArraySize = Num(), MaskSize = Offsets.Num();
@@ -831,7 +837,14 @@ public:
 				if (0 <= RemapVal[i])
 				{
 					ensure(RemapVal[i] < MaskSize);
-					this->operator[](Index)[i] -= Offsets[RemapVal[i]];
+					if (DeletionSet.Contains(this->operator[](Index)[i]))
+					{
+						this->operator[](Index)[i] = INDEX_NONE;
+					}
+					else
+					{
+						this->operator[](Index)[i] -= Offsets[RemapVal[i]];
+					}
 					ensure(-1 <= this->operator[](Index)[i] && this->operator[](Index)[i] <= FinalSize);
 				}
 			}
@@ -878,7 +891,7 @@ public:
 	virtual ~TManagedArray()
 	{}
 
-	virtual void Reindex(const TArray<int32>& Offsets, const int32& FinalSize, const TArray<int32>& SortedDeletionList) override
+	virtual void Reindex(const TArray<int32>& Offsets, const int32& FinalSize, const TArray<int32>& SortedDeletionList, const TSet<int32>& DeletionSet) override
 	{
 		UE_LOG(UManagedArrayLogging, Log, TEXT("TManagedArray<FIntVector>[%p]::Reindex()"), this);
 		int32 ArraySize = Num(), MaskSize = Offsets.Num();
@@ -893,7 +906,14 @@ public:
 					if (0 <= RemapVal[i])
 					{
 						ensure(RemapVal[i] < MaskSize);
-						this->operator[](Index)[ArrayIndex][i] -= Offsets[RemapVal[i]];
+						if (DeletionSet.Contains(this->operator[](Index)[ArrayIndex][i]))
+						{
+							this->operator[](Index)[ArrayIndex][i] = INDEX_NONE;
+						}
+						else
+						{
+							this->operator[](Index)[ArrayIndex][i] -= Offsets[RemapVal[i]];
+						}
 						ensure(-1 <= this->operator[](Index)[ArrayIndex][i] && this->operator[](Index)[ArrayIndex][i] <= FinalSize);
 					}
 				}
@@ -945,7 +965,7 @@ public:
 	virtual ~TManagedArray()
 	{}
 
-	virtual void Reindex(const TArray<int32>& Offsets, const int32& FinalSize, const TArray<int32>& SortedDeletionList) override
+	virtual void Reindex(const TArray<int32>& Offsets, const int32& FinalSize, const TArray<int32>& SortedDeletionList, const TSet<int32>& DeletionSet) override
 	{
 		UE_LOG(UManagedArrayLogging, Log, TEXT("TManagedArray<FIntVector>[%p]::Reindex()"), this);
 		int32 ArraySize = Num(), MaskSize = Offsets.Num();
@@ -957,7 +977,14 @@ public:
 				if (0 <= RemapVal[i])
 				{
 					ensure(RemapVal[i] < MaskSize);
-					this->operator[](Index)[i] -= Offsets[RemapVal[i]];
+					if (DeletionSet.Contains(this->operator[](Index)[i]))
+					{
+						this->operator[](Index)[i] = INDEX_NONE;
+					}
+					else
+					{
+						this->operator[](Index)[i] -= Offsets[RemapVal[i]];
+					}
 					ensure(-1 <= this->operator[](Index)[i] && this->operator[](Index)[i] <= FinalSize);
 				}
 			}
@@ -980,3 +1007,69 @@ public:
 		}
 	}
 };
+
+template<>
+class TManagedArray<TArray<int32>> : public TManagedArrayBase<TArray<int32>>
+{
+public:
+	using TManagedArrayBase<TArray<int32>>::Num;
+
+	FORCEINLINE TManagedArray()
+	{}
+
+	FORCEINLINE TManagedArray(const TArray<TArray<int32>>& Other)
+		: TManagedArrayBase<TArray<int32>>(Other)
+	{}
+
+	FORCEINLINE TManagedArray(const TManagedArray<TArray<int32>>& Other) = delete;
+	FORCEINLINE TManagedArray(TManagedArray<TArray<int32>>&& Other) = default;
+	FORCEINLINE TManagedArray(TArray<TArray<int32>>&& Other)
+		: TManagedArrayBase<TArray<int32>>(MoveTemp(Other))
+	{}
+	FORCEINLINE TManagedArray& operator=(TManagedArray<TArray<int32>>&& Other) = default;
+
+	virtual ~TManagedArray()
+	{}
+
+	virtual void Reindex(const TArray<int32>& Offsets, const int32& FinalSize, const TArray<int32>& SortedDeletionList, const TSet<int32>& DeletionSet) override
+	{
+		UE_LOG(UManagedArrayLogging, Log, TEXT("TManagedArray<FIntVector>[%p]::Reindex()"), this);
+		int32 ArraySize = Num(), MaskSize = Offsets.Num();
+		for (int32 Index = 0; Index < ArraySize; Index++)
+		{
+			const TArray<int32>& RemapValArray = this->operator[](Index);
+			for (int32 ArrayIndex = 0; ArrayIndex < RemapValArray.Num(); ArrayIndex++)
+			{
+				if (0 <= RemapValArray[ArrayIndex])
+				{
+					ensure(RemapValArray[ArrayIndex] < MaskSize);
+					if (DeletionSet.Contains(this->operator[](Index)[ArrayIndex]))
+					{
+						this->operator[](Index)[ArrayIndex] = INDEX_NONE;
+					}
+					else
+					{
+						this->operator[](Index)[ArrayIndex] -= Offsets[RemapValArray[ArrayIndex]];
+					}
+					ensure(-1 <= this->operator[](Index)[ArrayIndex] && this->operator[](Index)[ArrayIndex] <= FinalSize);
+				}
+			}
+		}
+	}
+
+	virtual void ReindexFromLookup(const TArray<int32>& InverseNewOrder) override
+	{
+		int32 ArraySize = Num();
+		for (int32 Index = 0; Index < ArraySize; Index++)
+		{
+			TArray<int32>& RemapValArray = this->operator[](Index);
+			for (int32 ArrayIndex = 0; ArrayIndex < RemapValArray.Num(); ArrayIndex++)
+			{
+				if (RemapValArray[ArrayIndex] >= 0)
+				{
+					RemapValArray[ArrayIndex] = InverseNewOrder[RemapValArray[ArrayIndex]];
+				}
+			}
+		}
+	}
+};
diff --git a/Engine/Source/Runtime/Experimental/Chaos/Public/GeometryCollection/ManagedArrayTypeValues.inl b/Engine/Source/Runtime/Experimental/Chaos/Public/GeometryCollection/ManagedArrayTypeValues.inl
--- a/Engine/Source/Runtime/Experimental/Chaos/Public/GeometryCollection/ManagedArrayTypeValues.inl
+++ b/Engine/Source/Runtime/Experimental/Chaos/Public/GeometryCollection/ManagedArrayTypeValues.inl
@@ -46,6 +46,8 @@ MANAGED_ARRAY_TYPE(FIntVector4, IntVector4)
 MANAGED_ARRAY_TYPE(FVector3d, Vector3d)
 MANAGED_ARRAY_TYPE(FIntVector2, IntVector2)
 MANAGED_ARRAY_TYPE(TArray<FIntVector2>, IntVector2Array)
+MANAGED_ARRAY_TYPE(TArray<int32>, Int32Array)
+MANAGED_ARRAY_TYPE(TArray<float>, FloatArray)
 
 
 // NOTE: new types must be added at the bottom to keep serialization from breaking
