Dataflow : Input/Output
- Cache copies from GetValue of the input
- some pointer protection.

diff --git a/Engine/Source/Runtime/Experimental/Chaos/Public/GeometryCollection/ManagedArrayCollection.h b/Engine/Source/Runtime/Experimental/Chaos/Public/GeometryCollection/ManagedArrayCollection.h
--- a/Engine/Source/Runtime/Experimental/Chaos/Public/GeometryCollection/ManagedArrayCollection.h
+++ b/Engine/Source/Runtime/Experimental/Chaos/Public/GeometryCollection/ManagedArrayCollection.h
@@ -117,20 +117,23 @@ public:
 	T* NewCopy() const
 	{
 		T* Collection = new T();
-		for (const TTuple<FKeyType, FValueType>& Entry : Map)
+		if (!Map.IsEmpty())
 		{
-			if (!Collection->HasGroup(Entry.Key.Get<1>()))
+			for (const TTuple<FKeyType, FValueType>& Entry : Map)
 			{
-				Collection->AddGroup(Entry.Key.Get<1>());
+				if (!Collection->HasGroup(Entry.Key.Get<1>()))
+				{
+					Collection->AddGroup(Entry.Key.Get<1>());
+				}
+
+				if (NumElements(Entry.Key.Get<1>()) != Collection->NumElements(Entry.Key.Get<1>()))
+				{
+					ensure(!Collection->NumElements(Entry.Key.Get<1>()));
+					Collection->AddElements(NumElements(Entry.Key.Get<1>()), Entry.Key.Get<1>());
+				}
+
+				Collection->CopyAttribute(*this, Entry.Key.Get<0>(), Entry.Key.Get<1>());
 			}
-
-			if (NumElements(Entry.Key.Get<1>()) != Collection->NumElements(Entry.Key.Get<1>()))
-			{
-				ensure(!Collection->NumElements(Entry.Key.Get<1>()));
-				Collection->AddElements(NumElements(Entry.Key.Get<1>()), Entry.Key.Get<1>());
-			}
-
-			Collection->CopyAttribute(*this, Entry.Key.Get<0>(), Entry.Key.Get<1>());
 		}
 		return Collection;
 	}
diff --git a/Engine/Source/Runtime/Experimental/Dataflow/Core/Public/Dataflow/DataflowGraph.h b/Engine/Source/Runtime/Experimental/Dataflow/Core/Public/Dataflow/DataflowGraph.h
--- a/Engine/Source/Runtime/Experimental/Dataflow/Core/Public/Dataflow/DataflowGraph.h
+++ b/Engine/Source/Runtime/Experimental/Dataflow/Core/Public/Dataflow/DataflowGraph.h
@@ -65,21 +65,6 @@ namespace Dataflow
 			return NewNode;
 		}
 
-		template<class T> T* FindNode(FName InName)
-		{
-			for (TSharedPtr<FNode> Node : Nodes)
-			{
-				if (Node->GetName().ToString().Equals(InName.ToString()))
-				{
-					if (Node->GetType().ToString().Equals(T::Type.ToString()))
-					{
-						return reinterpret_cast<T*>(Node.Get()); // @todo(eg) : Can we do better here?
-					}
-				}
-			}
-			return nullptr;
-		}
-
 		TSharedPtr<FNode> FindBaseNode(FGuid InGuid)
 		{
 			for (TSharedPtr<FNode> Node : Nodes)
diff --git a/Engine/Source/Runtime/Experimental/Dataflow/Core/Public/Dataflow/DataflowInputOutput.h b/Engine/Source/Runtime/Experimental/Dataflow/Core/Public/Dataflow/DataflowInputOutput.h
--- a/Engine/Source/Runtime/Experimental/Dataflow/Core/Public/Dataflow/DataflowInputOutput.h
+++ b/Engine/Source/Runtime/Experimental/Dataflow/Core/Public/Dataflow/DataflowInputOutput.h
@@ -167,7 +167,7 @@ namespace Dataflow
 		void SetValue(T InVal, const FContext& Context)
 		{
 			CacheKey = Context.GetTypeHash();
-			Cache.Data = DeepCopy<T>(InVal);
+			Cache.Data = InVal;
 		}
 
 		T GetValue(const FContext& Context)
@@ -177,7 +177,7 @@ namespace Dataflow
 				OwningNode->Evaluate(Context, this);
 			}
 			ensure(CacheKey == Context.GetTypeHash());
-			return Cache.Data;
+			return DeepCopy<T>(Cache.Data);
 		}
 
 		virtual void Invalidate() override
