Voronoi library improvements:
 - Add a 'compute helper' class to allow efficient point location; Voronoi class should be safe to use across multiple threads as long as threads do not share their compute helper.
 - Add a 'Voronoi field' class that implements common queries one might need for sampling a voronoi diagram as a field + add test coverage of this
 - Switch to TPimplPtr for the implementation-hiding storage of voro++ data
 - Add move constructors to voro++ classes to account for their memory management
 - Switch tolerances and offsets from float to double, to match the rest of the computation

diff --git a/Engine/Source/Runtime/Experimental/Voronoi/Private/Voronoi/Voronoi.cpp b/Engine/Source/Runtime/Experimental/Voronoi/Private/Voronoi/Voronoi.cpp
--- a/Engine/Source/Runtime/Experimental/Voronoi/Private/Voronoi/Voronoi.cpp
+++ b/Engine/Source/Runtime/Experimental/Voronoi/Private/Voronoi/Voronoi.cpp
@@ -9,6 +9,8 @@ THIRD_PARTY_INCLUDES_END
 
 namespace {
 
+	constexpr double VoronoiDefaultBoundingBoxSlack = .01f;
+
 	// initialize AABB, ignoring NaNs
 	FBox SafeInitBounds(const TArrayView<const FVector>& Points)
 	{
@@ -42,7 +44,7 @@ namespace {
 	}
 
 	// Add sites to Voronoi container, with contiguous ids, ignoring NaNs, ignoring Sites that are on top of existing sites
-	int32 PutSitesWithDistanceCheck(voro::container *Container, const TArrayView<const FVector>& Sites, int32 Offset, float SquaredDistThreshold = 1e-4)
+	int32 PutSitesWithDistanceCheck(voro::container *Container, const TArrayView<const FVector>& Sites, int32 Offset, double SquaredDistThreshold = 1e-4)
 	{
 		int32 SkippedPts = 0;
 		voro::voro_compute<voro::container> VoroCompute = Container->make_compute();
@@ -60,9 +62,9 @@ namespace {
 				int ExistingPtID;
 				if (Container->find_voronoi_cell(V.X, V.Y, V.Z, EX, EY, EZ, ExistingPtID, VoroCompute))
 				{
-					float dx = V.X - EX;
-					float dy = V.Y - EY;
-					float dz = V.Z - EZ;
+					double dx = V.X - EX;
+					double dy = V.Y - EY;
+					double dz = V.Z - EZ;
 					if (dx*dx + dy * dy + dz * dz < SquaredDistThreshold)
 					{
 						SkippedPts++;
@@ -75,7 +77,8 @@ namespace {
 		return SkippedPts;
 	}
 
-	TUniquePtr<voro::container> StandardVoroContainerInit(const TArrayView<const FVector>& Sites, FBox& BoundingBox, float BoundingBoxSlack = 0.1f, float SquaredDistSkipPtThreshold = 0.0f)
+	TPimplPtr<voro::container> EmptyVoroContainerInit(int32 ExpectedNumSites, FBox& BoundingBox,
+		double BoundingBoxSlack = VoronoiDefaultBoundingBoxSlack)
 	{
 		BoundingBox = BoundingBox.ExpandBy(BoundingBoxSlack);
 		FVector BoundingBoxSize = BoundingBox.GetSize();
@@ -84,18 +87,23 @@ namespace {
 		// TODO: Figure out reasonable bounds / behavior for this case
 		ensure(BoundingBoxSize.GetMax() < HUGE_VALF);
 
-
-		int NumSites = Sites.Num();
-
 		int GridCellsX, GridCellsY, GridCellsZ;
-		voro::guess_optimal(Sites.Num(), BoundingBoxSize.X, BoundingBoxSize.Y, BoundingBoxSize.Z, GridCellsX, GridCellsY, GridCellsZ);
+		voro::guess_optimal(ExpectedNumSites, BoundingBoxSize.X, BoundingBoxSize.Y, BoundingBoxSize.Z, GridCellsX, GridCellsY, GridCellsZ);
 
-		TUniquePtr<voro::container> Container = MakeUnique<voro::container>(
+		TPimplPtr<voro::container> Container = MakePimpl<voro::container>(
 			BoundingBox.Min.X, BoundingBox.Max.X, BoundingBox.Min.Y,
 			BoundingBox.Max.Y, BoundingBox.Min.Z, BoundingBox.Max.Z,
 			GridCellsX, GridCellsY, GridCellsZ, false, false, false, 10
-		);
+			);
+
+		return Container;
+	}
 
+	TPimplPtr<voro::container> StandardVoroContainerInit(const TArrayView<const FVector>& Sites, FBox& BoundingBox, 
+		double BoundingBoxSlack = VoronoiDefaultBoundingBoxSlack, double SquaredDistSkipPtThreshold = 0.0f)
+	{
+		TPimplPtr<voro::container> Container = EmptyVoroContainerInit(Sites.Num(), BoundingBox, BoundingBoxSlack);
+		
 		if (SquaredDistSkipPtThreshold > 0)
 		{
 			PutSitesWithDistanceCheck(Container.Get(), Sites, 0, SquaredDistSkipPtThreshold);
@@ -107,34 +115,55 @@ namespace {
 
 		return Container;
 	}
+
+
 };
 
 FVoronoiDiagram::~FVoronoiDiagram() = default;
 
-// TODO: ExtraBoundingSpace unused
-FVoronoiDiagram::FVoronoiDiagram(const TArrayView<const FVector>& Sites, float ExtraBoundingSpace, float SquaredDistSkipPtThreshold) : NumSites(0), Container(nullptr)
+FVoronoiDiagram::FVoronoiDiagram(const TArrayView<const FVector>& Sites, double ExtraBoundingSpace, double SquaredDistSkipPtThreshold) : NumSites(0), Container(nullptr)
 {
-	FBox BoundingBox = SafeInitBounds(Sites);
-	Container = StandardVoroContainerInit(Sites, BoundingBox, DefaultBoundingBoxSlack, SquaredDistSkipPtThreshold);
+	FBox BoundingBox = SafeInitBounds(Sites).ExpandBy(FMath::Max(0, ExtraBoundingSpace));
+	Container = StandardVoroContainerInit(Sites, BoundingBox, 0, SquaredDistSkipPtThreshold);
 	this->Bounds = BoundingBox;
 	NumSites = Sites.Num();
 }
 
-// TODO: ExtraBoundingSpace unused
-FVoronoiDiagram::FVoronoiDiagram(const TArrayView<const FVector>& Sites, const FBox &Bounds, float ExtraBoundingSpace, float SquaredDistSkipPtThreshold) : NumSites(0), Container(nullptr)
+FVoronoiDiagram::FVoronoiDiagram(const TArrayView<const FVector>& Sites, const FBox &Bounds, double ExtraBoundingSpace, double SquaredDistSkipPtThreshold) : NumSites(0), Container(nullptr)
 {
-	FBox BoundingBox(Bounds);
-	Container = StandardVoroContainerInit(Sites, BoundingBox, DefaultBoundingBoxSlack, SquaredDistSkipPtThreshold);
+	FBox BoundingBox = Bounds.ExpandBy(FMath::Max(0, ExtraBoundingSpace));
+	Container = StandardVoroContainerInit(Sites, BoundingBox, 0, SquaredDistSkipPtThreshold);
 	this->Bounds = BoundingBox;
 	NumSites = Sites.Num();
 }
 
+FVoronoiDiagram::FVoronoiDiagram(int32 ExpectedNumSites, const FBox& Bounds, double ExtraBoundingSpace) : NumSites(0), Container(nullptr)
+{
+	FBox BoundingBox = Bounds.ExpandBy(FMath::Max(0, ExtraBoundingSpace));
+	Container = EmptyVoroContainerInit(ExpectedNumSites, BoundingBox, 0);
+	this->Bounds = BoundingBox;
+}
 
 
-bool VoronoiNeighbors(const TArrayView<const FVector> &Sites, TArray<TArray<int>> &Neighbors, bool bExcludeBounds, float SquaredDistSkipPtThreshold) 
+FBox FVoronoiDiagram::GetBounds(const TArrayView<const FVector>& Sites, double ExtraBoundingSpace)
+{
+	FBox Bounds = SafeInitBounds(Sites);
+	return Bounds.ExpandBy(ExtraBoundingSpace);
+}
+
+void FVoronoiDiagram::Initialize(const TArrayView<const FVector>& Sites, const FBox& BoundsIn, double ExtraBoundingSpace, double SquaredDistSkipPtThreshold)
+{
+	FBox BoundingBox = BoundsIn.ExpandBy(FMath::Max(0, ExtraBoundingSpace));
+	Container = StandardVoroContainerInit(Sites, BoundingBox, 0, SquaredDistSkipPtThreshold);
+	Bounds = BoundingBox;
+	NumSites = Sites.Num();
+}
+
+
+bool VoronoiNeighbors(const TArrayView<const FVector> &Sites, TArray<TArray<int>> &Neighbors, bool bExcludeBounds, double SquaredDistSkipPtThreshold) 
 {
 	FBox BoundingBox = SafeInitBounds(Sites);
-	auto Container = StandardVoroContainerInit(Sites, BoundingBox, FVoronoiDiagram::DefaultBoundingBoxSlack, SquaredDistSkipPtThreshold);
+	auto Container = StandardVoroContainerInit(Sites, BoundingBox, VoronoiDefaultBoundingBoxSlack, SquaredDistSkipPtThreshold);
 	int32 NumSites = Sites.Num();
 
 	Neighbors.Empty(NumSites);
@@ -161,18 +190,18 @@ bool VoronoiNeighbors(const TArrayView<const FVector> &Sites, TArray<TArray<int>
 	return true;
 }
 
-bool GetVoronoiEdges(const TArrayView<const FVector>& Sites, const FBox& Bounds, TArray<TTuple<FVector, FVector>>& Edges, TArray<int32>& CellMember, float SquaredDistSkipPtThreshold)
+bool GetVoronoiEdges(const TArrayView<const FVector>& Sites, const FBox& Bounds, TArray<TTuple<FVector, FVector>>& Edges, TArray<int32>& CellMember, double SquaredDistSkipPtThreshold)
 {
 	int32 NumSites = Sites.Num();
 	FBox BoundingBox(Bounds);
-	FVoronoiDiagram Diagram(Sites, Bounds, FVoronoiDiagram::DefaultBoundingBoxSlack, SquaredDistSkipPtThreshold);
+	FVoronoiDiagram Diagram(Sites, Bounds, VoronoiDefaultBoundingBoxSlack, SquaredDistSkipPtThreshold);
 	Diagram.ComputeCellEdges(Edges, CellMember, Diagram.ApproxSitesPerThreadWithDefault(-1));
 	return true;
 }
 
 
 // TODO: maybe make this a "SetSites" instead of an Add?
-void FVoronoiDiagram::AddSites(const TArrayView<const FVector>& AddSites, float SquaredDistSkipPtThreshold)
+void FVoronoiDiagram::AddSites(const TArrayView<const FVector>& AddSites, double SquaredDistSkipPtThreshold)
 {
 	int32 OrigSitesNum = NumSites;
 	if (SquaredDistSkipPtThreshold > 0)
@@ -292,6 +321,38 @@ void FVoronoiDiagram::ComputeAllCells(TArray<FVoronoiCellInfo>& AllCells, int32
 	});
 }
 
+void FVoronoiDiagram::ComputeAllNeighbors(TArray<TArray<int32>>& AllNeighbors, bool bExcludeBounds, int32 ApproxSitesPerThread)
+{
+	check(Container);
+
+	TArray<int32> ChunkBounds = GetParallelBlockRanges(ApproxSitesPerThread);
+
+	AllNeighbors.SetNum(NumSites);
+
+	ParallelFor(ChunkBounds.Num() - 1, [&](int ChunkIdx)
+	{
+		voro::voro_compute<voro::container> VoroCompute = Container->make_compute();
+		voro::c_loop_block_range CellIterator(*Container);
+		CellIterator.setup_range(ChunkBounds[ChunkIdx], ChunkBounds[ChunkIdx + 1]);
+		voro::voronoicell_neighbor cell;
+		if (!CellIterator.start())
+		{
+			return;
+		}
+		do
+		{
+			bool bCouldComputeCell = Container->compute_cell(cell, CellIterator, VoroCompute);
+			ensureMsgf(bCouldComputeCell, TEXT("Failed to compute a Voronoi cell -- this may indicate sites positioned directly on top of other sites, which is not valid for a Voronoi diagram"));
+			if (bCouldComputeCell)
+			{
+				int32 id = CellIterator.pid();
+
+				cell.neighborsTArray(AllNeighbors[id], bExcludeBounds);
+			}
+		} while (CellIterator.inc());
+	});
+}
+
 void FVoronoiDiagram::ComputeCellEdgesSerial(TArray<TTuple<FVector, FVector>>& Edges, TArray<int32>& CellMember)
 {
 	check(Container);
@@ -419,16 +480,21 @@ void FVoronoiDiagram::ComputeCellEdges(TArray<TTuple<FVector, FVector>>& Edges,
 	}
 }
 
-//// TODO: expose separate compute data as something we can pass into this fn, and uncomment it
-//int32 FVoronoiDiagram::FindCell(const FVector& Pos)
-//{
-//	check(Container);
-//	voro::voro_compute<voro::container> VoroCompute = Container->make_compute();
-//	double rx, ry, rz; // holds position of the found Voronoi site (if any); currently unused
-//	int pid;
-//	bool found = Container->find_voronoi_cell(Pos.X, Pos.Y, Pos.Z, rx, ry, rz, pid, VoroCompute);
-//	return found ? pid : -1;
-//}
-
-const float FVoronoiDiagram::DefaultBoundingBoxSlack = .1f;
+int32 FVoronoiDiagram::FindCell(const FVector& Pos, FVoronoiComputeHelper& ComputeHelper, FVector& OutFoundSite) const
+{
+	check(Container && ComputeHelper.Compute && &ComputeHelper.Compute->con == Container.Get());
+	
+	int pid;
+	bool found = Container->find_voronoi_cell(Pos.X, Pos.Y, Pos.Z, OutFoundSite.X, OutFoundSite.Y, OutFoundSite.Z, pid, *ComputeHelper.Compute.Get());
+	
+	return found ? pid : -1;
+}
+
+FVoronoiComputeHelper FVoronoiDiagram::GetComputeHelper() const
+{
+	FVoronoiComputeHelper Helper;
+	Helper.Compute = MakePimpl<voro::voro_compute<voro::container>>(Container->make_compute());
+	return Helper;
+}
+
 const int FVoronoiDiagram::MinDefaultSitesPerThread = 150;
\ No newline at end of file
diff --git a/Engine/Source/Runtime/Experimental/Voronoi/Private/voro++/src/container.cc b/Engine/Source/Runtime/Experimental/Voronoi/Private/voro++/src/container.cc
--- a/Engine/Source/Runtime/Experimental/Voronoi/Private/voro++/src/container.cc
+++ b/Engine/Source/Runtime/Experimental/Voronoi/Private/voro++/src/container.cc
@@ -46,10 +46,16 @@ container_base::container_base(double ax_,double bx_,double ay_,double by_,doubl
 /** The container destructor frees the dynamically allocated memory. */
 container_base::~container_base() {
 	int l;
-	for(l=0;l<nxyz;l++) delete [] p[l];
-	for(l=0;l<nxyz;l++) delete [] id[l];
-	delete [] id;
-	delete [] p;
+	if (p)
+	{
+		for (l = 0; l < nxyz; l++) delete[] p[l];
+		delete[] p;
+	}
+	if (id)
+	{
+		for (l = 0; l < nxyz; l++) delete[] id[l];
+		delete[] id;
+	}
 	delete [] co;
 	delete [] mem;
 }
diff --git a/Engine/Source/Runtime/Experimental/Voronoi/Public/Voronoi/Voronoi.h b/Engine/Source/Runtime/Experimental/Voronoi/Public/Voronoi/Voronoi.h
--- a/Engine/Source/Runtime/Experimental/Voronoi/Public/Voronoi/Voronoi.h
+++ b/Engine/Source/Runtime/Experimental/Voronoi/Public/Voronoi/Voronoi.h
@@ -1,9 +1,10 @@
 // Copyright Epic Games, Inc. All Rights Reserved.
 #pragma once
 #include "CoreMinimal.h"
+#include "Templates/PimplPtr.h"
 
-bool VORONOI_API VoronoiNeighbors(const TArrayView<const FVector> &Sites, TArray<TArray<int>> &Neighbors, bool bExcludeBounds = true, float SquaredDistSkipPtThreshold = UE_KINDA_SMALL_NUMBER);
-bool VORONOI_API GetVoronoiEdges(const TArrayView<const FVector> &Sites, const FBox& Bounds, TArray<TTuple<FVector, FVector>> &Edges, TArray<int32> &CellMember, float SquaredDistSkipPtThreshold = UE_KINDA_SMALL_NUMBER);
+bool VORONOI_API VoronoiNeighbors(const TArrayView<const FVector> &Sites, TArray<TArray<int>> &Neighbors, bool bExcludeBounds = true, double SquaredDistSkipPtThreshold = UE_KINDA_SMALL_NUMBER);
+bool VORONOI_API GetVoronoiEdges(const TArrayView<const FVector> &Sites, const FBox& Bounds, TArray<TTuple<FVector, FVector>> &Edges, TArray<int32> &CellMember, double SquaredDistSkipPtThreshold = UE_KINDA_SMALL_NUMBER);
 
 // All the info you would typically want about a single cell in the Voronoi diagram, in the format that is easiest to compute
 struct FVoronoiCellInfo
@@ -17,54 +18,125 @@ struct FVoronoiCellInfo
 // third party library voro++'s container class; stores voronoi sites in a uniform grid accel structure
 namespace voro {
 	class container;
+	template <class ContainerType> class voro_compute;
 }
 
+class FVoronoiDiagram;
+
+// The Voronoi Compute Helper caches information to help compute cells quickly
+// To safely parallelize voronoi diagram construction, create a separate compute helper per-thread
+class VORONOI_API FVoronoiComputeHelper
+{
+	TPimplPtr<voro::voro_compute<voro::container>> Compute;
+public:
+
+	FVoronoiComputeHelper() {}
+	FVoronoiComputeHelper(const FVoronoiComputeHelper& Other) = delete;
+	FVoronoiComputeHelper& operator=(const FVoronoiComputeHelper& Other) = delete;
+	FVoronoiComputeHelper(FVoronoiComputeHelper&& Other) = default;
+	FVoronoiComputeHelper& operator=(FVoronoiComputeHelper&& Other) = default;
+
+	void Init();
+
+	friend class FVoronoiDiagram;
+};
+
 class VORONOI_API FVoronoiDiagram
 {
 	int32 NumSites;
-	TUniquePtr<voro::container> Container;
+	TPimplPtr<voro::container> Container;
 	FBox Bounds;
 
 public:
-	// we typically add extra space to the bounds of the Voronoi diagram, to avoid numerical issues of a Voronoi site being on the boundary
-	const static float DefaultBoundingBoxSlack;
 	const static int MinDefaultSitesPerThread;
 
 	/**
+	 * Create a Voronoi diagram with the given sites, in a box bounding containing those sites
+	 * 
 	 * @param Sites							Voronoi sites for the diagram
 	 * @param ExtraBoundingSpace			Voronoi diagram will be computed within the bounding box of the sites + this amount of extra space in each dimension
 	 * @param SquaredDistSkipPtThreshold	A safety threshold to avoid creating invalid cells: sites that are within this distance of an already-added site will not be added
 	 *										(If you know there will be no duplicate sites, can set to zero for faster perf.)
 	 */
-	FVoronoiDiagram(const TArrayView<const FVector>& Sites, float ExtraBoundingSpace, float SquaredDistSkipPtThreshold = UE_KINDA_SMALL_NUMBER);
+	FVoronoiDiagram(const TArrayView<const FVector>& Sites, double ExtraBoundingSpace, double SquaredDistSkipPtThreshold = UE_KINDA_SMALL_NUMBER);
+	/**
+	 * Create a Voronoi diagram with the given sites, in a given bounding box
+	 * 
+	 * @param Sites							Voronoi sites for the diagram
+	 * @param Bounds						Bounding box within which to compute the Voronoi diagram
+	 * @param ExtraBoundingSpace			Voronoi diagram will be computed within the input Bounds + this amount of extra space in each dimension
+	 * @param SquaredDistSkipPtThreshold	A safety threshold to avoid creating invalid cells: sites that are within this distance of an already-added site will not be added.
+	 *										(If you know there will be no duplicate sites, can set to zero for faster perf.)
+	 */
+	FVoronoiDiagram(const TArrayView<const FVector>& Sites, const FBox &Bounds, double ExtraBoundingSpace, double SquaredDistSkipPtThreshold = UE_KINDA_SMALL_NUMBER);
+
+	/**
+	 * Make a container for a Voronoi diagram without any points.  Call AddSites() after to finish creating the diagram.
+	 * 
+	 * @param ExpectedNumSites		How many sites you expect to add in total
+	 * @param Bounds				Bounding box within which to compute the Voronoi diagram
+	 * @param ExtraBoundingSpace	Amount to expand the bounding box before computing the Voronoi diagram
+	 */
+	FVoronoiDiagram(int32 ExpectedNumSites, const FBox& Bounds, double ExtraBoundingSpace);
+
+	/**
+	 * Make an empty diagram, to be filled in with a subsequent call to "Initialize()"
+	 */
+	FVoronoiDiagram() : NumSites(0), Container(nullptr), Bounds(EForceInit::ForceInit)
+	{}
+
+	FVoronoiDiagram(const FVoronoiDiagram& Other) = delete;
+	FVoronoiDiagram& operator=(const FVoronoiDiagram& Other) = delete;
+	FVoronoiDiagram(FVoronoiDiagram&& Other) = default;
+	FVoronoiDiagram& operator=(FVoronoiDiagram&& Other) = default;
+
 	/**
-	* @param Sites							Voronoi sites for the diagram
-	* @param Bounds							Bounding box within which to compute the Voronoi diagram
-	* @param ExtraBoundingSpace				Voronoi diagram will be computed within the input Bounds + this amount of extra space in each dimension
-	* @param SquaredDistSkipPtThreshold		A safety threshold to avoid creating invalid cells: sites that are within this distance of an already-added site will not be added.
-	*										(If you know there will be no duplicate sites, can set to zero for faster perf.)
-	*/
-	FVoronoiDiagram(const TArrayView<const FVector>& Sites, const FBox &Bounds, float ExtraBoundingSpace, float SquaredDistSkipPtThreshold = UE_KINDA_SMALL_NUMBER);
+	 * Create a Voronoi diagram with the given sites, in a given bounding box.  Any previous data is discarded.
+	 *
+	 * @param Sites							Voronoi sites for the diagram
+	 * @param Bounds						Bounding box within which to compute the Voronoi diagram
+	 * @param ExtraBoundingSpace			Voronoi diagram will be computed within the input Bounds + this amount of extra space in each dimension
+	 * @param SquaredDistSkipPtThreshold	A safety threshold to avoid creating invalid cells: sites that are within this distance of an already-added site will not be added.
+	 *										(If you know there will be no duplicate sites, can set to zero for faster perf.)
+	 */
+	void Initialize(const TArrayView<const FVector>& Sites, const FBox& Bounds, double ExtraBoundingSpace, double SquaredDistSkipPtThreshold = UE_KINDA_SMALL_NUMBER);
+
 	~FVoronoiDiagram();
 
+	static FBox GetBounds(const TArrayView<const FVector>& Sites, double ExtraBoundingSpace = UE_DOUBLE_KINDA_SMALL_NUMBER);
+
 	int32 Num() const
 	{
 		return NumSites;
 	}
 
-	void AddSites(const TArrayView<const FVector>& Sites, float SquaredDistSkipPtThreshold = 0.0f);
+	void AddSites(const TArrayView<const FVector>& Sites, double SquaredDistSkipPtThreshold = 0.0f);
 
 	void ComputeAllCellsSerial(TArray<FVoronoiCellInfo>& AllCells);
 	void ComputeAllCells(TArray<FVoronoiCellInfo>& AllCells, int32 ApproxSitesPerThread = -1);
 
+	void ComputeAllNeighbors(TArray<TArray<int32>>& AllNeighbors, bool bExcludeBounds = true, int32 ApproxSitesPerThread = -1);
+
 	void ComputeCellEdgesSerial(TArray<TTuple<FVector, FVector>>& Edges, TArray<int32>& CellMember);
 	void ComputeCellEdges(TArray<TTuple<FVector, FVector>>& Edges, TArray<int32>& CellMember, int32 ApproxSitesPerThread = -1);
 
 	/**
-	 * Find the id of the Voronoi cell containing the given position (or -1 if position is outside diagram)
-	 * TODO: expose a "compute" parameter to store cached state across calls to this function, and un-comment it back
+	 * Find the closest Voronoi cell to a given position
+	 * 
+	 * @param Pos				Position to query for closest point
+	 * @param ComputeHelper		A valid helper as returned by GetComputeHelper() for this voronoi diagram
+	 * @param OutFoundSite		The position of the closest Voronoi site, if any
+	 * @return					The id of the Voronoi cell containing the given position, or -1 if position is outside diagram
 	 */
-	//int32 FindCell(const FVector& Pos);
+	int32 FindCell(const FVector& Pos, FVoronoiComputeHelper& ComputeHelper, FVector& OutFoundSite) const;
+
+	int32 FindCell(const FVector& Pos, FVoronoiComputeHelper& ComputeHelper) const
+	{
+		FVector OutFoundSite;
+		return FindCell(Pos, ComputeHelper, OutFoundSite);
+	}
+
+	FVoronoiComputeHelper GetComputeHelper() const;
 
 	int32 ApproxSitesPerThreadWithDefault(int32 ApproxSitesPerThreadIn)
 	{
@@ -79,4 +151,116 @@ private:
 	TArray<int32> GetParallelBlockRanges(int32 ApproxSitesPerThread);
 };
 
+/**
+ * Use a Voronoi diagram to support faster querying of Voronoi-like data from arbitrary sample points:
+ *  - Distance to closest point
+ *  - ID of closest point
+ *  - IDs of closest two points
+ *  - Distance to closest Voronoi cell boundary (not considering bounding box walls)
+ * 
+ * Note: Queries are only thread safe if each thread has its own FVoronoiComputeHelper
+ */
+class VORONOI_API FVoronoiDiagramField
+{
+	TArray<FVector> Sites;
+	TArray<TArray<int32>> Neighbors;
+	FVoronoiDiagram Diagram;
+
+public:
+	FVoronoiDiagramField() = default;
 
+	FVoronoiDiagramField(const TArray<FVector>& SitesIn, const FBox& Bounds, double SquaredDistSkipPtThreshold = DBL_EPSILON)
+	{
+		Initialize(SitesIn, Bounds, SquaredDistSkipPtThreshold);
+	}
+
+	void Initialize(const TArray<FVector>& SitesIn, const FBox& Bounds, double SquaredDistSkipPtThreshold = DBL_EPSILON)
+	{
+		Sites = SitesIn;
+		Diagram.Initialize(Sites, Bounds, 0, SquaredDistSkipPtThreshold);
+		Diagram.ComputeAllNeighbors(Neighbors, true);
+	}
+
+	FVoronoiDiagramField(const FVoronoiDiagramField& Other) = delete;
+	FVoronoiDiagramField& operator=(const FVoronoiDiagramField& Other) = delete;
+	FVoronoiDiagramField(FVoronoiDiagramField&& Other) = default;
+	FVoronoiDiagramField& operator=(FVoronoiDiagramField&& Other) = default;
+
+	inline FVoronoiComputeHelper GetComputeHelper() const
+	{
+		return Diagram.GetComputeHelper();
+	}
+
+	// @return Distance to the nearest Voronoi cell boundary (ignoring bounding box walls), or InvalidValue if outside of bounds
+	double DistanceToCellWall(const FVector& Sample, FVoronoiComputeHelper& ComputeHelper, double InvalidValue = TNumericLimits<double>::Lowest())
+	{
+		FVector CloseSite;
+		int32 Cell = Diagram.FindCell(Sample, ComputeHelper, CloseSite);
+		if (Cell < 0)
+		{
+			return InvalidValue;
+		}
+		double BestDistSq = DBL_MAX;
+		int32 BestNbr = -1;
+		for (int32 NbrCell : Neighbors[Cell])
+		{
+			double DistSq = FVector::DistSquared(Sample, Sites[NbrCell]);
+			if (DistSq < BestDistSq)
+			{
+				BestNbr = NbrCell;
+				BestDistSq = DistSq;
+			}
+		}
+		if (BestNbr == -1)
+		{
+			return InvalidValue;
+		}
+		FVector NbrPos = Sites[BestNbr];
+		FVector Mid = (NbrPos + CloseSite) * .5;
+		FVector Normal = NbrPos - CloseSite;
+		bool bNormalizeSuccess = Normal.Normalize();
+		checkSlow(bNormalizeSuccess); // expect this can't happen because the Voronoi diagram can't be constructed w/ duplicate points
+		// Note: Returns a signed plane distance, but should always be positive because we know Sample is closer to CloseSite
+		return (Mid - Sample).Dot(Normal);
+	}
+
+	TPair<int32, int32> ClosestTwoIDs(const FVector& Sample, FVoronoiComputeHelper& ComputeHelper)
+	{
+		TPair<int32, int32> ToRet(-1, -1);
+		FVector CloseSite;
+		ToRet.Key = Diagram.FindCell(Sample, ComputeHelper, CloseSite);
+		if (ToRet.Key < 0)
+		{
+			return ToRet;
+		}
+		double BestDistSq = DBL_MAX;
+		ToRet.Value = -1;
+		for (int32 NbrCell : Neighbors[ToRet.Key])
+		{
+			double DistSq = FVector::DistSquared(Sample, Sites[NbrCell]);
+			if (DistSq < BestDistSq)
+			{
+				ToRet.Value = NbrCell;
+				BestDistSq = DistSq;
+			}
+		}
+		return ToRet;
+	}
+
+	// @return Distance to the nearest Voronoi site, or InvalidValue if outside of bounds
+	inline double DistanceToClosest(const FVector& Sample, FVoronoiComputeHelper& ComputeHelper, double InvalidValue = TNumericLimits<double>::Lowest())
+	{
+		FVector CloseSite;
+		int32 Cell = Diagram.FindCell(Sample, ComputeHelper, CloseSite);
+		if (Cell < 0)
+		{
+			return InvalidValue;
+		}
+		return FVector::Distance(Sample, CloseSite);
+	}
+
+	inline int32 ClosestID(const FVector& Sample, FVoronoiComputeHelper& ComputeHelper)
+	{
+		return Diagram.FindCell(Sample, ComputeHelper);
+	}
+};
