Chaos : fix CCD issue on collision with geometry collection ( and in general with unions )
Particles were not sorted properly asserting about the key computed from the constraint particles

diff --git a/Engine/Source/Runtime/Experimental/Chaos/Private/Chaos/CollisionResolution.cpp b/Engine/Source/Runtime/Experimental/Chaos/Private/Chaos/CollisionResolution.cpp
--- a/Engine/Source/Runtime/Experimental/Chaos/Private/Chaos/CollisionResolution.cpp
+++ b/Engine/Source/Runtime/Experimental/Chaos/Private/Chaos/CollisionResolution.cpp
@@ -110,6 +110,8 @@
 bool bChaos_Collision_ShapesArrayMode = true;
 FAutoConsoleVariableRef CVarChaos_Collision_ShapesArrayMode(TEXT("p.Chaos.Collision.ShapesArrayMode"), bChaos_Collision_ShapesArrayMode, TEXT(""));
 
+bool bChaos_Collision_SortParticlesOnConstraintConstruct = true;
+FAutoConsoleVariableRef CVarChaos_Collision_SortParticlesOnConstraintConstruct(TEXT("p.Chaos.Collision.SortParticlesOnConstraintConstruct"), bChaos_Collision_SortParticlesOnConstraintConstruct, TEXT(""));
 
 namespace Chaos
 {
@@ -2348,7 +2350,24 @@
 
 			// If we get here, we have a pair of concrete shapes (i.e., no wrappers or containers)
 			// Create a constraint for the shape pair
-			ConstructConstraintsImpl<T_TRAITS>(Particle0, Particle1, Implicit0, Shape0, Simplicial0, Implicit1, Shape1, Simplicial1, ParticleWorldTransform0, LocalTransform0, ParticleWorldTransform1, LocalTransform1, CullDistance, Dt, Context);
+			{
+				// let's make sure the order of particle is correct
+				// this needs to be done at this level because we have the concrete type available here ( and no union or other wrappers )
+				// @todo(chaos) : with this order enforced, ConstructConstraintsImpl can be optimized by remove reverse test cases ( keeping sphere|convex but removing convex|sphere for example )  
+				bool bShouldSwapParticles = false;
+				if (bChaos_Collision_SortParticlesOnConstraintConstruct)
+				{
+					CalculateShapePairType(Implicit0, Simplicial0, Implicit1, Simplicial1, bShouldSwapParticles);
+				}
+				if (bShouldSwapParticles)
+				{
+					ConstructConstraintsImpl<T_TRAITS>(Particle1, Particle0, Implicit1, Shape1, Simplicial1, Implicit0, Shape0, Simplicial0, ParticleWorldTransform1, LocalTransform1, ParticleWorldTransform0, LocalTransform0, CullDistance, Dt, Context);	
+				}
+				else
+				{
+					ConstructConstraintsImpl<T_TRAITS>(Particle0, Particle1, Implicit0, Shape0, Simplicial0, Implicit1, Shape1, Simplicial1, ParticleWorldTransform0, LocalTransform0, ParticleWorldTransform1, LocalTransform1, CullDistance, Dt, Context);
+				}
+			}
 		}
 
 		EContactShapesType CalculateShapePairType(const EImplicitObjectType Implicit0Type, const EImplicitObjectType Implicit1Type, const bool bIsConvex0, const bool bIsConvex1, const bool bIsBVH0, const bool bIsBVH1,  bool &bOutSwap)
