Chaos : fix fields not applying velocity after applying external strain

diff --git a/Engine/Source/Runtime/Experimental/Chaos/Private/Chaos/PBDRigidClustering.cpp b/Engine/Source/Runtime/Experimental/Chaos/Private/Chaos/PBDRigidClustering.cpp
--- a/Engine/Source/Runtime/Experimental/Chaos/Private/Chaos/PBDRigidClustering.cpp
+++ b/Engine/Source/Runtime/Experimental/Chaos/Private/Chaos/PBDRigidClustering.cpp
@@ -461,10 +461,10 @@ namespace Chaos
 			//todo(ocohen): in the presence of collisions, this will leave all children with the post-collision
 			// velocity. This should be controlled by material properties so we can allow the broken pieces to
 			// maintain the clusters pre-collision velocity.
-			Child->SetV(ClusteredParent->V());
-			Child->SetW(ClusteredParent->W());
-			Child->SetPreV(ClusteredParent->PreV());
-			Child->SetPreW(ClusteredParent->PreW());
+			Child->SetV(Child->V() + ClusteredParent->V());
+			Child->SetW(Child->W() + ClusteredParent->W());
+			Child->SetPreV(Child->PreV() + ClusteredParent->PreV());
+			Child->SetPreW(Child->PreW() + ClusteredParent->PreW());
 		}
 	};
 
