fix fracture crash on building proximity w/ leaf proximity to root

diff --git a/Engine/Source/Runtime/Experimental/Chaos/Private/GeometryCollection/GeometryCollectionConvexUtility.cpp b/Engine/Source/Runtime/Experimental/Chaos/Private/GeometryCollection/GeometryCollectionConvexUtility.cpp
--- a/Engine/Source/Runtime/Experimental/Chaos/Private/GeometryCollection/GeometryCollectionConvexUtility.cpp
+++ b/Engine/Source/Runtime/Experimental/Chaos/Private/GeometryCollection/GeometryCollectionConvexUtility.cpp
@@ -637,7 +637,7 @@ void CreateNonoverlappingConvexHulls(
 	{
 		for (int32 NbrBone : LeafProximity[Bone])
 		{
-			if (NbrBone > Bone)
+			if (Bone > NbrBone)
 			{
 				continue;
 			}
@@ -668,6 +668,11 @@ void CreateNonoverlappingConvexHulls(
 
 			auto ConnectAtMatchingDepth = [&SameDepthClusterProximity, &Parents, CommonParent, &Depths](int32 BoneA, int32 BoneToWalkParents)
 			{
+				if (BoneA == INDEX_NONE || BoneToWalkParents == INDEX_NONE)
+				{
+					return;
+				}
+
 				int32 DepthA = Depths[BoneA];
 				int32 DepthB = Depths[BoneToWalkParents];
 				while (BoneToWalkParents != INDEX_NONE && BoneToWalkParents != CommonParent && DepthB >= DepthA)
