Chaos : add index to break event

diff --git a/Engine/Source/Runtime/Experimental/ChaosSolverEngine/Private/Chaos/ChaosGameplayEventDispatcher.cpp b/Engine/Source/Runtime/Experimental/ChaosSolverEngine/Private/Chaos/ChaosGameplayEventDispatcher.cpp
--- a/Engine/Source/Runtime/Experimental/ChaosSolverEngine/Private/Chaos/ChaosGameplayEventDispatcher.cpp
+++ b/Engine/Source/Runtime/Experimental/ChaosSolverEngine/Private/Chaos/ChaosGameplayEventDispatcher.cpp
@@ -22,6 +22,7 @@ FChaosBreakEvent::FChaosBreakEvent()
 	, Velocity(FVector::ZeroVector)
 	, AngularVelocity(FVector::ZeroVector)
 	, Mass(0.0f)
+	, Index(INDEX_NONE)
 {
 
 }
@@ -464,6 +465,7 @@ void UChaosGameplayEventDispatcher::HandleBreakingEvents(const Chaos::FBreakingE
 						BreakEvent.Velocity = BreakingDataItem.Velocity;
 						BreakEvent.AngularVelocity = BreakingDataItem.AngularVelocity;
 						BreakEvent.Mass = BreakingDataItem.Mass;
+						BreakEvent.Index = BreakingDataItem.TransformGroupIndex;
 					}
 				}
 			}
diff --git a/Engine/Source/Runtime/Experimental/ChaosSolverEngine/Public/Chaos/ChaosGameplayEventDispatcher.h b/Engine/Source/Runtime/Experimental/ChaosSolverEngine/Public/Chaos/ChaosGameplayEventDispatcher.h
--- a/Engine/Source/Runtime/Experimental/ChaosSolverEngine/Public/Chaos/ChaosGameplayEventDispatcher.h
+++ b/Engine/Source/Runtime/Experimental/ChaosSolverEngine/Public/Chaos/ChaosGameplayEventDispatcher.h
@@ -26,20 +26,29 @@ public:
 
 	FChaosBreakEvent();
 
+	/** primitive component involved in the break event */
 	UPROPERTY(BlueprintReadOnly, Category = "Break Event")
 	TObjectPtr<UPrimitiveComponent> Component = nullptr;
 
+	/** World location of the break */
 	UPROPERTY(BlueprintReadOnly, Category = "Break Event")
 	FVector Location;
 
+	/** Linear Velocity of the breaking particle  */
 	UPROPERTY(BlueprintReadOnly, Category = "Break Event")
 	FVector Velocity;
 
+	/** Angular Velocity of the breaking particle  */
 	UPROPERTY(BlueprintReadOnly, Category = "Break Event")
 	FVector AngularVelocity;
 
+	/** Mass of the breaking particle  */
 	UPROPERTY(BlueprintReadOnly, Category = "Break Event")
 	float Mass;
+
+	/** Index of the geometry collection bone if positive */
+	UPROPERTY(BlueprintReadOnly, Category = "Break Event")
+	int32 Index;
 };
 
 USTRUCT(BlueprintType)
