Chaos: Hoisted branch evaluation allocations out of AABB inner loops, reduces allocation churn within a frame when objects move in Chaos scenes

diff --git a/Engine/Source/Runtime/Experimental/Chaos/Public/Chaos/AABBTree.h b/Engine/Source/Runtime/Experimental/Chaos/Public/Chaos/AABBTree.h
--- a/Engine/Source/Runtime/Experimental/Chaos/Public/Chaos/AABBTree.h
+++ b/Engine/Source/Runtime/Experimental/Chaos/Public/Chaos/AABBTree.h
@@ -1100,11 +1100,8 @@ public:
 		NewBounds.Thicken(FAABBTreeCVars::DynamicTreeBoundingBoxPadding);
 
 		//Priority Q of indices to explore
-		TArray<int32> PriorityQ;
-		PriorityQ.Reserve(10);
-
-		TArray<FReal> SumDeltaCostQ;
-		SumDeltaCostQ.Reserve(10);
+		PriorityQ.Reset();
+		SumDeltaCostQ.Reset();
 
 		int32 QIndex = 0;
 		
@@ -3108,9 +3105,10 @@ private:
 		, OverlappingPairs(Other.OverlappingPairs)
 		, OverlappingCounts(Other.OverlappingCounts)
 	{
-
 		ensure(bDynamicTree == Other.IsTreeDynamic());
 
+		PriorityQ.Reserve(32);
+		SumDeltaCostQ.Reserve(32);
 	}
 
 	virtual ISpatialAcceleration<TPayloadType, T, 3>& operator=(const ISpatialAcceleration<TPayloadType, T, 3>& Other) override
@@ -3215,6 +3213,10 @@ private:
 
 	/** Number of overlapping leaves per leaf */
 	TArray<int32> OverlappingCounts;
+
+	/** Buffers for calculating the best candidate for tree insertion, reset per-calculation to avoid over allocation */
+	TArray<int32> PriorityQ;
+	TArray<FReal> SumDeltaCostQ;
 };
 
 template<typename TPayloadType, typename TLeafType, bool bMutable, typename T>
