Dataflow : Skeletal updates

diff --git a/Engine/Source/Runtime/Experimental/Dataflow/Core/Private/Dataflow/DataflowGraph.cpp b/Engine/Source/Runtime/Experimental/Dataflow/Core/Private/Dataflow/DataflowGraph.cpp
--- a/Engine/Source/Runtime/Experimental/Dataflow/Core/Private/Dataflow/DataflowGraph.cpp
+++ b/Engine/Source/Runtime/Experimental/Dataflow/Core/Private/Dataflow/DataflowGraph.cpp
@@ -185,11 +185,11 @@ namespace Dataflow
 			Ar << LocalConnections;
 			for (const FLink& Con : LocalConnections)
 			{
-				if (ensure(NodeGuidMap.Contains(Con.InputNode) && NodeGuidMap.Contains(Con.OutputNode)))
+				if (NodeGuidMap.Contains(Con.InputNode) && NodeGuidMap.Contains(Con.OutputNode))
 				{
-					if (ensure(ConnectionGuidMap.Contains(Con.Input) && ConnectionGuidMap.Contains(Con.Output)))
+					if (ConnectionGuidMap.Contains(Con.Input) && ConnectionGuidMap.Contains(Con.Output))
 					{
-						if (ensure(ConnectionGuidMap[Con.Input]->GetType() == ConnectionGuidMap[Con.Output]->GetType()))
+						if (ConnectionGuidMap[Con.Input]->GetType() == ConnectionGuidMap[Con.Output]->GetType())
 						{
 							Connect(ConnectionGuidMap[Con.Input], ConnectionGuidMap[Con.Output]);
 						}
diff --git a/Engine/Source/Runtime/Experimental/Dataflow/Engine/Private/Dataflow/DataflowEngineUtil.cpp b/Engine/Source/Runtime/Experimental/Dataflow/Engine/Private/Dataflow/DataflowEngineUtil.cpp
--- a/Engine/Source/Runtime/Experimental/Dataflow/Engine/Private/Dataflow/DataflowEngineUtil.cpp
+++ b/Engine/Source/Runtime/Experimental/Dataflow/Engine/Private/Dataflow/DataflowEngineUtil.cpp
@@ -10,19 +10,23 @@ namespace Dataflow
 	{
 		void GlobalTransformsInternal(int32 Index, const FReferenceSkeleton& Ref, TArray<FTransform>& Mat, TArray<bool>& Visited)
 		{
-			const TArray<FTransform>& RefMat = Ref.GetRefBonePose();
-
-			int32 ParentIndex = Ref.GetParentIndex(Index);
-			if (ParentIndex != INDEX_NONE && ParentIndex!=Index) // why self check?
-			{
-				GlobalTransformsInternal(ParentIndex, Ref, Mat, Visited);
-				Mat[Index] = Mat[ParentIndex].GetRelativeTransform(RefMat[Index]);
-			}
-			else
+			if (!Visited[Index])
 			{
-				Mat[Index] = RefMat[Index];
+				const TArray<FTransform>& RefMat = Ref.GetRefBonePose();
+
+				int32 ParentIndex = Ref.GetParentIndex(Index);
+				if (ParentIndex != INDEX_NONE && ParentIndex != Index) // why self check?
+				{
+					GlobalTransformsInternal(ParentIndex, Ref, Mat, Visited);
+					Mat[Index].SetFromMatrix(RefMat[Index].ToMatrixWithScale()*Mat[ParentIndex].ToMatrixWithScale());
+				}
+				else
+				{
+					Mat[Index] = RefMat[Index];
+				}
+
+				Visited[Index] = true;
 			}
-			Visited[Index] = true;
 		}
 
 		void GlobalTransforms(const FReferenceSkeleton& Ref, TArray<FTransform>& Mat)
@@ -32,7 +36,7 @@ namespace Dataflow
 			Mat.SetNum(Ref.GetNum());
 
 			int32 Index = Ref.GetNum() - 1;
-			while (Index >= 0 && !Visited[Index])
+			while (Index >= 0)
 			{
 				GlobalTransformsInternal(Index, Ref, Mat, Visited);
 				Index--;
