Dataflow : FactoryParameters

diff --git a/Engine/Plugins/Experimental/Dataflow/Source/DataflowEditor/Private/Dataflow/DataflowEditorActions.cpp b/Engine/Plugins/Experimental/Dataflow/Source/DataflowEditor/Private/Dataflow/DataflowEditorActions.cpp
--- a/Engine/Plugins/Experimental/Dataflow/Source/DataflowEditor/Private/Dataflow/DataflowEditorActions.cpp
+++ b/Engine/Plugins/Experimental/Dataflow/Source/DataflowEditor/Private/Dataflow/DataflowEditorActions.cpp
@@ -16,20 +16,20 @@ void FDataflowEditorCommandsImpl::RegisterCommands()
 
 	if (Dataflow::FNodeFactory* Factory = Dataflow::FNodeFactory::GetInstance())
 	{
-		for (FName NodeName : Factory->RegisteredNodes())
+		for (Dataflow::FFactoryParameters Parameters : Factory->RegisteredParameters())
 		{
 			TSharedPtr< FUICommandInfo > AddNode;
 			FUICommandInfo::MakeCommandInfo(
 				this->AsShared(),
 				AddNode,
-				NodeName, //FName("UseCreationFormToggle"),
-				NSLOCTEXT("DataFlow1", "DataflowButton", "New Dataflow Node"),
-				NSLOCTEXT("DataFlow2", "NewDataflowNodeTooltip", "New Dataflow Node Tooltip"),
+				Parameters.TypeName,
+				LOCTEXT("DataflowButton", "New Dataflow Node"),
+				LOCTEXT("NewDataflowNodeTooltip", "New Dataflow Node Tooltip"),
 				FSlateIcon(),
 				EUserInterfaceActionType::Button,
 				FInputChord()
 			);
-			CreateNodesMap.Add(NodeName, AddNode);
+			CreateNodesMap.Add(Parameters.TypeName, AddNode);
 		}
 	}
 }
diff --git a/Engine/Plugins/Experimental/Dataflow/Source/DataflowEditor/Private/Dataflow/DataflowSNode.cpp b/Engine/Plugins/Experimental/Dataflow/Source/DataflowEditor/Private/Dataflow/DataflowSNode.cpp
--- a/Engine/Plugins/Experimental/Dataflow/Source/DataflowEditor/Private/Dataflow/DataflowSNode.cpp
+++ b/Engine/Plugins/Experimental/Dataflow/Source/DataflowEditor/Private/Dataflow/DataflowSNode.cpp
@@ -32,12 +32,21 @@ FReply SDataflowEdNode::OnMouseButtonDoubleClick(const FGeometry& InMyGeometry,
 //
 TSharedPtr<FAssetSchemaAction_Dataflow_CreateNode_DataflowEdNode> FAssetSchemaAction_Dataflow_CreateNode_DataflowEdNode::CreateAction(UEdGraph* ParentGraph, const FName & InNodeTypeName)
 {
- 	const FText AddToolTip = LOCTEXT("DataflowNodeTooltip_Example", "Add a Dataflow node.");
-	const FText NodeName = FText::FromString(*InNodeTypeName.ToString());
-	const FText Catagory = LOCTEXT("DataflowNodeDescription_Example", "Dataflow"); 	
-	TSharedPtr<FAssetSchemaAction_Dataflow_CreateNode_DataflowEdNode> NewNodeAction(
-		new FAssetSchemaAction_Dataflow_CreateNode_DataflowEdNode(InNodeTypeName, Catagory, NodeName, AddToolTip, 0));
-	return NewNodeAction;
+	if (Dataflow::FNodeFactory* Factory = Dataflow::FNodeFactory::GetInstance())
+	{
+		const Dataflow::FFactoryParameters& Param = Factory->GetParameters(InNodeTypeName);
+		if (Param.IsValid())
+		{
+			const FText ToolTip = FText::FromString(Param.ToolTip.IsEmpty() ? FString("Add a Dataflow node.") : Param.ToolTip);
+			const FText NodeName = FText::FromString(Param.DisplayName.ToString());
+			const FText Catagory = FText::FromString(Param.Category.ToString().IsEmpty() ? FString("Dataflow") : Param.Category.ToString());
+			const FText Tags = FText::FromString(Param.Tags);
+			TSharedPtr<FAssetSchemaAction_Dataflow_CreateNode_DataflowEdNode> NewNodeAction(
+				new FAssetSchemaAction_Dataflow_CreateNode_DataflowEdNode(InNodeTypeName, Catagory, NodeName, ToolTip, Tags));
+			return NewNodeAction;
+		}
+	}
+	return TSharedPtr<FAssetSchemaAction_Dataflow_CreateNode_DataflowEdNode>(nullptr);
 }
 
 //
diff --git a/Engine/Plugins/Experimental/Dataflow/Source/DataflowEditor/Private/Dataflow/DataflowSchema.cpp b/Engine/Plugins/Experimental/Dataflow/Source/DataflowEditor/Private/Dataflow/DataflowSchema.cpp
--- a/Engine/Plugins/Experimental/Dataflow/Source/DataflowEditor/Private/Dataflow/DataflowSchema.cpp
+++ b/Engine/Plugins/Experimental/Dataflow/Source/DataflowEditor/Private/Dataflow/DataflowSchema.cpp
@@ -40,11 +40,15 @@ void UDataflowSchema::GetGraphContextActions(FGraphContextMenuBuilder& ContextMe
 {
 	if (Dataflow::FNodeFactory* Factory = Dataflow::FNodeFactory::GetInstance())
 	{
-		for (FName NodeTypeName : Factory->RegisteredNodes())
+		for (Dataflow::FFactoryParameters NodeParameters : Factory->RegisteredParameters())
 		{
-			if (FDataflowEditorCommands::Get().CreateNodesMap.Contains(NodeTypeName))
+			if (FDataflowEditorCommands::Get().CreateNodesMap.Contains(NodeParameters.TypeName))
 			{
-				ContextMenuBuilder.AddAction(FAssetSchemaAction_Dataflow_CreateNode_DataflowEdNode::CreateAction(ContextMenuBuilder.OwnerOfTemporaries, NodeTypeName));
+				if (TSharedPtr<FAssetSchemaAction_Dataflow_CreateNode_DataflowEdNode> Action =
+					FAssetSchemaAction_Dataflow_CreateNode_DataflowEdNode::CreateAction(ContextMenuBuilder.OwnerOfTemporaries, NodeParameters.TypeName))
+				{
+					ContextMenuBuilder.AddAction(Action);
+				}
 			}
 		}
 	}
diff --git a/Engine/Plugins/Experimental/Dataflow/Source/DataflowEditor/Public/Dataflow/DataflowSNode.h b/Engine/Plugins/Experimental/Dataflow/Source/DataflowEditor/Public/Dataflow/DataflowSNode.h
--- a/Engine/Plugins/Experimental/Dataflow/Source/DataflowEditor/Public/Dataflow/DataflowSNode.h
+++ b/Engine/Plugins/Experimental/Dataflow/Source/DataflowEditor/Public/Dataflow/DataflowSNode.h
@@ -44,8 +44,8 @@ struct DATAFLOWEDITOR_API FAssetSchemaAction_Dataflow_CreateNode_DataflowEdNode
 public:
 	FAssetSchemaAction_Dataflow_CreateNode_DataflowEdNode() : FEdGraphSchemaAction() {}
 
-	FAssetSchemaAction_Dataflow_CreateNode_DataflowEdNode(const FName& InType, const FText& InNodeCategory, const FText& InMenuDesc, const FText& InToolTip, const int32 InGrouping)
-		: FEdGraphSchemaAction(InNodeCategory, InMenuDesc, InToolTip, InGrouping), NodeTypeName(InType) {}
+	FAssetSchemaAction_Dataflow_CreateNode_DataflowEdNode(const FName& InType, FText InNodeCategory, FText InMenuDesc, FText InToolTip, FText InKeywords)
+		: FEdGraphSchemaAction(InNodeCategory, InMenuDesc, InToolTip, 0, InKeywords), NodeTypeName(InType) {}
 
 	static TSharedPtr<FAssetSchemaAction_Dataflow_CreateNode_DataflowEdNode> CreateAction(UEdGraph* ParentGraph, const FName & NodeTypeName);
 
diff --git a/Engine/Source/Runtime/Experimental/Dataflow/Core/Public/Dataflow/DataflowNode.h b/Engine/Source/Runtime/Experimental/Dataflow/Core/Public/Dataflow/DataflowNode.h
--- a/Engine/Source/Runtime/Experimental/Dataflow/Core/Public/Dataflow/DataflowNode.h
+++ b/Engine/Source/Runtime/Experimental/Dataflow/Core/Public/Dataflow/DataflowNode.h
@@ -46,15 +46,16 @@ struct DATAFLOWCORE_API FDataflowNode
 
 	virtual ~FDataflowNode() {}
 
+	FGuid GetGuid() const { return Guid; }
 	FName GetName() const { return Name; }
 	void SetName(FName InName) { Name = InName; }
 
-	virtual FString GetDisplayName() { return ""; }
-	virtual FString GetCatagory() { return ""; }
-	virtual FString GetTags() { return ""; }
-
 	virtual FName GetType() const { check(true); return FName("invalid"); }
-	FGuid GetGuid() const { return Guid; }
+	virtual FName GetDisplayName() const { return ""; }
+	virtual FName GetCategory() const { return ""; }
+	virtual FString GetTags() const { return ""; }
+	virtual FString GetToolTip() const { return ""; }
+
 
 	TArray<Dataflow::FPin> GetPins() const;
 	Dataflow::FConnection* FindInput(FName Name) const;
@@ -81,27 +82,32 @@ namespace Dataflow
 	// Used these macros to register dataflow nodes. 
 	//
 
-#define DATAFLOW_NODE_REGISTER_CREATION_FACTORY(A)						\
-	FNodeFactory::GetInstance()->RegisterNode(							\
-		A::StaticType(),													\
-		[](const FNewNodeParameters& InParam)							\
+#define DATAFLOW_NODE_REGISTER_CREATION_FACTORY(A)								\
+	FNodeFactory::GetInstance()->RegisterNode(									\
+		{A::StaticType(),A::StaticDisplay(),A::StaticCategory(),				\
+			A::StaticTags(),A::StaticToolTip()},								\
+		[](const FNewNodeParameters& InParam)									\
 			{return new A({InParam.Name}, InParam.Guid); });
 
-#define DATAFLOW_NODE_DEFINE_INTERNAL(TYPE, DISPLAY_NAME, CATAGORY, TAGS)\
-public:																	\
-	static FName StaticType() {return #TYPE;}							\
-	static FName StaticDisplay() {return DISPLAY_NAME;}					\
-	virtual FName GetType() const { return #TYPE; }						\
-	virtual FStructOnScope* NewScructOnScope() override {				\
-	   return new FStructOnScope(TYPE::StaticStruct(), (uint8*)this);}	\
-	virtual void SerializeInternal(FArchive& Ar) override {				\
-		UScriptStruct* const Struct = TYPE::StaticStruct();				\
-		Struct->SerializeTaggedProperties(Ar, (uint8*)this,				\
-		Struct, nullptr);}												\
-	virtual FString GetDisplayName() override { return DISPLAY_NAME; }	\
-	virtual FString GetCatagory() override { return CATAGORY; }			\
-	virtual FString GetTags() override { return TAGS; }					\
-	TYPE() {}															\
+#define DATAFLOW_NODE_DEFINE_INTERNAL(TYPE, DISPLAY_NAME, CATEGORY, TAGS)		\
+public:																			\
+	static FName StaticType() {return #TYPE;}									\
+	static FName StaticDisplay() {return DISPLAY_NAME;}							\
+	static FName StaticCategory() {return CATEGORY;}							\
+	static FString StaticTags() {return TAGS;}									\
+	static FString StaticToolTip() {return FString("Create a dataflow node.");}	\
+	virtual FName GetType() const { return #TYPE; }								\
+	virtual FStructOnScope* NewScructOnScope() override {						\
+	   return new FStructOnScope(TYPE::StaticStruct(), (uint8*)this);}			\
+	virtual void SerializeInternal(FArchive& Ar) override {						\
+		UScriptStruct* const Struct = TYPE::StaticStruct();						\
+		Struct->SerializeTaggedProperties(Ar, (uint8*)this,						\
+		Struct, nullptr);}														\
+	virtual FName GetDisplayName() const override { return TYPE::StaticDisplay(); }	\
+	virtual FName GetCategory() const override { return TYPE::StaticCategory(); }	\
+	virtual FString GetTags() const override { return TYPE::StaticTags(); }			\
+	virtual FString GetToolTip() const override { return TYPE::StaticToolTip(); }	\
+	TYPE() {}																		\
 private:
 
 
diff --git a/Engine/Source/Runtime/Experimental/Dataflow/Core/Public/Dataflow/DataflowNodeFactory.h b/Engine/Source/Runtime/Experimental/Dataflow/Core/Public/Dataflow/DataflowNodeFactory.h
--- a/Engine/Source/Runtime/Experimental/Dataflow/Core/Public/Dataflow/DataflowNodeFactory.h
+++ b/Engine/Source/Runtime/Experimental/Dataflow/Core/Public/Dataflow/DataflowNodeFactory.h
@@ -4,6 +4,7 @@
 
 #include "CoreMinimal.h"
 #include "Dataflow/DataflowGraph.h"
+#include "ChaosLog.h"
 
 struct FDataflowNode;
 
@@ -14,7 +15,23 @@ namespace Dataflow
 	struct DATAFLOWCORE_API FNewNodeParameters {
 		FGuid Guid;
 		FName Type;
-		FName Name; 
+		FName Name;
+	};
+
+	struct DATAFLOWCORE_API FFactoryParameters {
+		FFactoryParameters() {}
+		FFactoryParameters(FName InTypeName, FName InDisplayName, FName InCategory, FString InTags, FString InToolTip)
+		: TypeName(InTypeName), DisplayName(InDisplayName), Category(InCategory), Tags(InTags), ToolTip(InToolTip) {}
+
+		FName TypeName = FName("");
+		FName DisplayName = FName("");
+		FName Category = FName("");
+		FString Tags = FString("");
+		FString ToolTip = FString("");
+
+		bool IsValid() const {
+			return !TypeName.ToString().IsEmpty() && !DisplayName.ToString().IsEmpty();
+		}
 	};
 
 	//
@@ -24,7 +41,10 @@ namespace Dataflow
 	{
 		typedef TFunction<FDataflowNode* (const FNewNodeParameters&)> FNewNodeFunction;
 
-		TMap<FName, FNewNodeFunction > ClassMap;
+		// All Maps indexed by TypeName
+		TMap<FName, FNewNodeFunction > ClassMap;		// [TypeName] -> NewNodeFunction
+		TMap<FName, FFactoryParameters > ParametersMap;	// [TypeName] -> Parameters
+		TMap<FName, FName > DisplayMap;					// [DisplayName] -> TypeName
 
 		static FNodeFactory* Instance;
 		FNodeFactory() {}
@@ -41,12 +61,62 @@ namespace Dataflow
 			return Instance;
 		}
 
-		void RegisterNode(const FName& Type, FNewNodeFunction NewFunction)
+		void RegisterNode(const FFactoryParameters& Parameters, FNewNodeFunction NewFunction)
+		{
+			bool bRegisterNode = true;
+			if (ClassMap.Contains(Parameters.TypeName) || DisplayMap.Contains(Parameters.DisplayName))
+			{
+				if (ParametersMap[Parameters.TypeName].DisplayName.IsEqual(Parameters.DisplayName) )
+				{
+					UE_LOG(LogChaos, Warning, 
+						TEXT("Warning : Dataflow node registration mismatch with type(%s).The \
+						nodes have inconsistent display names(%s) vs(%s).There are two nodes \
+						with the same type being registered."), *Parameters.TypeName.ToString(),
+						*ParametersMap[Parameters.TypeName].DisplayName.ToString(), 
+						*Parameters.DisplayName.ToString(), *Parameters.TypeName.ToString());
+				}
+				if (ParametersMap[Parameters.TypeName].Category.IsEqual(Parameters.Category))
+				{
+					UE_LOG(LogChaos, Warning, 
+						TEXT("Warning : Dataflow node registration mismatch with type (%s). The nodes \
+						have inconsistent categories names (%s) vs (%s). There are two different nodes \
+						with the same type being registered. "), *Parameters.TypeName.ToString(),
+						*ParametersMap[Parameters.TypeName].DisplayName.ToString(), 
+						*Parameters.DisplayName.ToString(),*Parameters.TypeName.ToString());
+				}
+				if (!ClassMap.Contains(Parameters.TypeName))
+				{
+					UE_LOG(LogChaos, Warning,
+						TEXT("Warning: Attempted to register node type(%s) with display name (%s) \
+						that conflicts with an existing nodes display name (%s)."), 
+						*Parameters.TypeName.ToString(),*Parameters.DisplayName.ToString(), 
+						*ParametersMap[Parameters.TypeName].DisplayName.ToString());
+				}
+			}
+			else
+			{
+				ClassMap.Add(Parameters.TypeName, NewFunction);
+				ParametersMap.Add(Parameters.TypeName, Parameters);
+				DisplayMap.Add(Parameters.DisplayName, Parameters.TypeName);
+			}
+		}
+
+		FName TypeNameFromDisplayName(const FName& DisplayName)
 		{
-			if (ensure(!ClassMap.Contains(Type)))
+			if (DisplayMap.Contains(DisplayName))
 			{
-				ClassMap.Add(Type, NewFunction);
+				return DisplayMap[DisplayName];
 			}
+			return "";
+		}
+
+		FFactoryParameters GetParameters(FName InTypeName) 
+		{
+			if (ParametersMap.Contains(InTypeName))
+			{
+				return ParametersMap[InTypeName];
+			}
+			return FFactoryParameters();
 		}
 
 		TSharedPtr<FDataflowNode> NewNodeFromRegisteredType(FGraph& Graph, const FNewNodeParameters& Param);
@@ -56,12 +126,13 @@ namespace Dataflow
 			return Graph.AddNode(new T(Param.Name, Param.Guid));
 		}
 
-		TArray<FName> RegisteredNodes() const
+		TArray<FFactoryParameters> RegisteredParameters() const
 		{
-			TArray<FName> ReturnNames;
-			for (auto Elem : ClassMap) ReturnNames.Add(Elem.Key);
-			return ReturnNames;
+			TArray<FFactoryParameters> RetVal;
+			for (auto Elem : ParametersMap) RetVal.Add(Elem.Value);
+			return RetVal;
 		}
+
 	};
 
 }
