Chaos
- improved one sided collision against heightfield and tri mesh
- imrpoved contact point pruning
- unified heightfield and tri mesh contact generation and pruning
- Add lambda-based VisitTriangle methods to heightfield and trimesh
- remove Manifold functions from TriMesh and HeightField
- Set initial search direction for convex-to-triangle GJK
- Add debug visualiuzation of mesh contact pruning

diff --git a/Engine/Source/Programs/HeadlessChaos/Private/HeadlessChaosTestBarycentric.cpp b/Engine/Source/Programs/HeadlessChaos/Private/HeadlessChaosTestBarycentric.cpp
new file mode 100644
--- /dev/null
+++ b/Engine/Source/Programs/HeadlessChaos/Private/HeadlessChaosTestBarycentric.cpp
@@ -0,0 +1,188 @@
+// Copyright Epic Games, Inc. All Rights Reserved.
+
+#include "HeadlessChaos.h"
+
+#include "Chaos/Utilities.h"
+#include "Chaos/Triangle.h"
+#include "ChaosLog.h"
+
+namespace ChaosTest
+{
+	using namespace Chaos;
+
+	// For non-degenerate triangles, the berycentric coords should be uniquely defined and match the input.
+	// For degenerate triangles or external points, the barycentric coords are not unique but the calculated coords should generate the same point.
+	void TestToFromBarycentric(const FVec3 Verts[], const FVec3& ExpectedBary, const bool bIsDegenerate, const bool bIsOutside)
+	{
+		const FVec3 ExpectedP = FromBarycentric(ExpectedBary, Verts[0], Verts[1], Verts[2]);
+
+		const FVec3 CalculatedBary = ToBarycentric(ExpectedP, Verts[0], Verts[1], Verts[2]);
+
+		const FVec3 CalculatedP = FromBarycentric(CalculatedBary, Verts[0], Verts[1], Verts[2]);
+
+		const FReal Tolerance = 1.e-4;
+
+		// Regardless of degeneracy, barycentric coords should sum to 1
+		EXPECT_NEAR(CalculatedBary.X + CalculatedBary.Y + CalculatedBary.Z, 1.0, Tolerance);
+
+		// If we are inside the triangle/line/point barycentric coords should be in [0,1]
+		if (!bIsOutside)
+		{
+			EXPECT_GE(CalculatedBary.X, 0);
+			EXPECT_GE(CalculatedBary.Y, 0);
+			EXPECT_GE(CalculatedBary.Z, 0);
+			EXPECT_LE(CalculatedBary.X, 1);
+			EXPECT_LE(CalculatedBary.Y, 1);
+			EXPECT_LE(CalculatedBary.Z, 1);
+		}
+
+		// Regardless of degeneracy, barycentric coords should reproduce the point
+		EXPECT_NEAR(CalculatedP.X, ExpectedP.X, Tolerance);
+		EXPECT_NEAR(CalculatedP.Y, ExpectedP.Y, Tolerance);
+		EXPECT_NEAR(CalculatedP.Z, ExpectedP.Z, Tolerance);
+
+		// for non-degenerate triangles we should calculate the barycentric coords exactly
+		if (!bIsDegenerate)
+		{
+			EXPECT_NEAR(CalculatedBary.X + CalculatedBary.Y + CalculatedBary.Z, 1.0, Tolerance);
+			EXPECT_NEAR(CalculatedBary.X, ExpectedBary.X, Tolerance);
+			EXPECT_NEAR(CalculatedBary.Y, ExpectedBary.Y, Tolerance);
+			EXPECT_NEAR(CalculatedBary.Z, ExpectedBary.Z, Tolerance);
+		}
+	}
+
+	// Runa set of points through the barycentric test for the specified triangle
+	void TestBarycentricPoints(const FVec3 Verts[], const bool bIsDegenerate)
+	{
+		TestToFromBarycentric(Verts, FVec3(0.2, 0.4, 0.4), bIsDegenerate, false);	// Inside
+		TestToFromBarycentric(Verts, FVec3(0.2, 0.8, 0.0), bIsDegenerate, false);	// AB Edge
+		TestToFromBarycentric(Verts, FVec3(0.0, 0.5, 0.5), bIsDegenerate, false);	// BC Edge
+		TestToFromBarycentric(Verts, FVec3(0.3, 0.0, 0.7), bIsDegenerate, false);	// CA Edge
+		TestToFromBarycentric(Verts, FVec3(1.0, 0.0, 0.0), bIsDegenerate, false);	// A Vertex
+		TestToFromBarycentric(Verts, FVec3(0.0, 1.0, 0.0), bIsDegenerate, false);	// B Vertex
+		TestToFromBarycentric(Verts, FVec3(0.0, 0.0, 1.0), bIsDegenerate, false);	// C Vertex
+	}
+
+	// Non-degenerate triangle test
+	GTEST_TEST(BarycentricTests, TestBarycentric_Triangle_Fwd)
+	{
+		const FVec3 Verts[] = {
+			FVec3(0,0,0),
+			FVec3(100,0,0),
+			FVec3(100,100,0)
+		};
+
+		TestBarycentricPoints(Verts, false);
+	}
+
+	// Non-degenerate triangle test - reverse winding
+	GTEST_TEST(BarycentricTests, TestBarycentric_Triangle2_Rev)
+	{
+		const FVec3 Verts[] = {
+			FVec3(0,0,0),
+			FVec3(100,100,0),
+			FVec3(100,0,0)
+		};
+
+		TestBarycentricPoints(Verts, false);
+	}
+
+	// Test Barycentric when the triangle is degenerate (a line A-B-C)
+	GTEST_TEST(BarycentricTests, TestBarycentric_LineABC)
+	{
+		const FVec3 Verts[] = {
+			FVec3(0,0,0),
+			FVec3(100,0,0),
+			FVec3(200,0,0)
+		};
+
+		TestBarycentricPoints(Verts, true);
+	}
+
+	// Test Barycentric when the triangle is degenerate - reverse winding (a line A-C-B)
+	GTEST_TEST(BarycentricTests, TestBarycentric_LineACB)
+	{
+		const FVec3 Verts[] = {
+			FVec3(0,0,0),
+			FVec3(200,0,0),
+			FVec3(100,0,0)
+		};
+
+		TestBarycentricPoints(Verts, true);
+	}
+
+	// Test Barycentric when the triangle is degenerate (a point A-B-C)
+	GTEST_TEST(BarycentricTests, TestBarycentric_PointABC)
+	{
+		const FVec3 Verts[] = {
+			FVec3(0,0,0),
+			FVec3(0,0,0),
+			FVec3(0,0,0)
+		};
+
+		TestBarycentricPoints(Verts, true);
+	}
+
+	// Test Barycentric when the triangle is degenerate (a line and point A-B)
+	GTEST_TEST(BarycentricTests, TestBarycentric_PointAB)
+	{
+		const FVec3 Verts[] = {
+			FVec3(0,0,0),
+			FVec3(0,0,0),
+			FVec3(100,0,0)
+		};
+
+		TestBarycentricPoints(Verts, true);
+	}
+
+	// Test Barycentric when the triangle is degenerate (a line and point B-C)
+	GTEST_TEST(BarycentricTests, TestBarycentric_PointBC)
+	{
+		const FVec3 Verts[] = {
+			FVec3(0,0,0),
+			FVec3(100,0,0),
+			FVec3(100,0,0)
+		};
+
+		TestBarycentricPoints(Verts, true);
+	}
+
+	// Test Barycentric when the triangle is degenerate (a line and point A-C)
+	GTEST_TEST(BarycentricTests, TestBarycentric_PointCA)
+	{
+		const FVec3 Verts[] = {
+			FVec3(0,0,0),
+			FVec3(100,0,0),
+			FVec3(0,0,0)
+		};
+
+		TestBarycentricPoints(Verts, true);
+	}
+
+	// Non-degenerate triangle test with outside points
+	GTEST_TEST(BarycentricTests, TestBarycentric_Triangle_Outside)
+	{
+		const FVec3 Verts[] = {
+			FVec3(50,0,0),
+			FVec3(100,0,0),
+			FVec3(100,100,0)
+		};
+
+		TestToFromBarycentric(Verts, FVec3(1.5, -0.5, 0.0), false, true);	// Outside
+		TestToFromBarycentric(Verts, FVec3(9.0, -5.0, -3.0), false, true);	// Outside
+	}
+
+	// Degenerate triangle test with outside points (on line)
+	GTEST_TEST(BarycentricTests, TestBarycentric_LineACB_Outside)
+	{
+		const FVec3 Verts[] = {
+			FVec3(50,0,0),
+			FVec3(200,0,0),
+			FVec3(100,0,0)
+		};
+
+		TestToFromBarycentric(Verts, FVec3(1.5, -0.5, 0.0), true, true);	// Outside
+		TestToFromBarycentric(Verts, FVec3(9.0, -5.0, -3.0), true, true);	// Outside
+	}
+
+}
\ No newline at end of file
diff --git a/Engine/Source/Runtime/Experimental/Chaos/Private/Chaos/ChaosDebugDraw.cpp b/Engine/Source/Runtime/Experimental/Chaos/Private/Chaos/ChaosDebugDraw.cpp
--- a/Engine/Source/Runtime/Experimental/Chaos/Private/Chaos/ChaosDebugDraw.cpp
+++ b/Engine/Source/Runtime/Experimental/Chaos/Private/Chaos/ChaosDebugDraw.cpp
@@ -324,15 +324,15 @@
 			const FAABB3 WorldQueryBounds = FAABB3(WorldQueryCenter - FVec3(WorldQueryRadius), WorldQueryCenter + FVec3(WorldQueryRadius));
 			const FAABB3 LocalQueryBounds = WorldQueryBounds.InverseTransformedAABB(ShapeTransform);
 
-			Shape->VisitTriangles(LocalQueryBounds, [&](const FTriangle& Tri)
-				{
-					FVec3 A = ShapeTransform.TransformPosition(Tri[0]);
-					FVec3 B = ShapeTransform.TransformPosition(Tri[1]);
-					FVec3 C = ShapeTransform.TransformPosition(Tri[2]);
-					FDebugDrawQueue::GetInstance().DrawDebugLine(A, B, Color, false, Duration, 0, Settings.ShapeThicknesScale * Settings.LineThickness);
-					FDebugDrawQueue::GetInstance().DrawDebugLine(B, C, Color, false, Duration, 0, Settings.ShapeThicknesScale * Settings.LineThickness);
-					FDebugDrawQueue::GetInstance().DrawDebugLine(C, A, Color, false, Duration, 0, Settings.ShapeThicknesScale * Settings.LineThickness);
-				});
+			Shape->VisitTriangles(LocalQueryBounds, ShapeTransform, [&](const FTriangle& Tri, const int32 TriangleIndex, const int32 VertexIndex0, const int32 VertexIndex1, const int32 VertexIndex2)
+			{
+				const FVec3& A = Tri[0];
+				const FVec3& B = Tri[1];
+				const FVec3& C = Tri[2];
+				FDebugDrawQueue::GetInstance().DrawDebugLine(A, B, Color, false, Duration, 0, Settings.ShapeThicknesScale * Settings.LineThickness);
+				FDebugDrawQueue::GetInstance().DrawDebugLine(B, C, Color, false, Duration, 0, Settings.ShapeThicknesScale * Settings.LineThickness);
+				FDebugDrawQueue::GetInstance().DrawDebugLine(C, A, Color, false, Duration, 0, Settings.ShapeThicknesScale * Settings.LineThickness);
+			});
 		}
 
 		void DrawShapesTriangleMeshImpl(const TGeometryParticleHandle<FReal, 3>* Particle, const FRigidTransform3& ShapeTransform, const FTriangleMeshImplicitObject* Shape, const FColor& Color, const FRealSingle Duration, const FChaosDebugDrawSettings& Settings)
@@ -342,11 +342,11 @@
 			const FAABB3 WorldQueryBounds = FAABB3(WorldQueryCenter - FVec3(WorldQueryRadius), WorldQueryCenter + FVec3(WorldQueryRadius));
 			const FAABB3 LocalQueryBounds = WorldQueryBounds.InverseTransformedAABB(ShapeTransform);
 
-			Shape->VisitTriangles(LocalQueryBounds, [&](const FTriangle& Tri)
+			Shape->VisitTriangles(LocalQueryBounds, ShapeTransform, [&](const FTriangle& Tri, const int32 TriangleIndex, const int32 VertexIndex0, const int32 VertexIndex1, const int32 VertexIndex2)
 			{
-				FVec3 A = ShapeTransform.TransformPosition(Tri[0]);
-				FVec3 B = ShapeTransform.TransformPosition(Tri[1]);
-				FVec3 C = ShapeTransform.TransformPosition(Tri[2]);
+				const FVec3& A = Tri[0];
+				const FVec3& B = Tri[1];
+				const FVec3& C = Tri[2];
 				FDebugDrawQueue::GetInstance().DrawDebugLine(A, B, Color, false, Duration, 0, Settings.ShapeThicknesScale * Settings.LineThickness);
 				FDebugDrawQueue::GetInstance().DrawDebugLine(B, C, Color, false, Duration, 0, Settings.ShapeThicknesScale * Settings.LineThickness);
 				FDebugDrawQueue::GetInstance().DrawDebugLine(C, A, Color, false, Duration, 0, Settings.ShapeThicknesScale * Settings.LineThickness);
diff --git a/Engine/Source/Runtime/Experimental/Chaos/Private/Chaos/Collision/ContactPointsMiscShapes.cpp b/Engine/Source/Runtime/Experimental/Chaos/Private/Chaos/Collision/ContactPointsMiscShapes.cpp
--- a/Engine/Source/Runtime/Experimental/Chaos/Private/Chaos/Collision/ContactPointsMiscShapes.cpp
+++ b/Engine/Source/Runtime/Experimental/Chaos/Private/Chaos/Collision/ContactPointsMiscShapes.cpp
@@ -237,28 +237,6 @@
 		return Contact;
 	}
 
-	template <typename GeometryA, typename GeometryB>
-	void GJKImplicitManifold(const FImplicitObject& A, const FRigidTransform3& ATransform, const GeometryB& B, const FRigidTransform3& BTransform, const FReal CullDistance, const FReal ShapePadding, TArray<FContactPoint>& ContactPoints)
-	{
-		const FRigidTransform3 AToBTM = ATransform.GetRelativeTransform(BTransform);
-
-		if (const TImplicitObjectScaled<GeometryA>* ScaledConvexImplicit = A.template GetObject<const TImplicitObjectScaled<GeometryA> >())
-		{
-			B.ContactManifold(*ScaledConvexImplicit, AToBTM, CullDistance, ContactPoints);
-		}
-		else if (const TImplicitObjectInstanced<GeometryA>* InstancedConvexImplicit = A.template GetObject<const TImplicitObjectInstanced<GeometryA> >())
-		{
-			if (const GeometryA* InstancedInnerObject = static_cast<const GeometryA*>(InstancedConvexImplicit->GetInstancedObject()))
-			{
-				B.ContactManifold(*InstancedInnerObject, AToBTM, CullDistance, ContactPoints);
-			}
-		}
-		else if (const GeometryA* ConvexImplicit = A.template GetObject<const GeometryA>())
-		{
-			B.ContactManifold(*ConvexImplicit, AToBTM, CullDistance, ContactPoints);
-		}
-	}
-
 	// A is the implicit here, we want to return a contact point on B (trimesh)
 	template <typename GeometryA>
 	FContactPoint GJKImplicitScaledTriMeshSweptContactPoint(const FImplicitObject& A, const FRigidTransform3& AStartTransform, const TImplicitObjectScaled<FTriangleMeshImplicitObject>& B, const FRigidTransform3& BTransform, const FVec3& Dir, const FReal Length, const FReal IgnorePenetration, const FReal TargetPenetration, FReal& OutTOI)
@@ -546,18 +524,4 @@
 	template FContactPoint SphereTriangleMeshContactPoint<FTriangleMeshImplicitObject>(const TSphere<FReal, 3>& A, const FRigidTransform3& ATransform, const FTriangleMeshImplicitObject& B, const FRigidTransform3& BTransform, const FReal CullDistance, const FReal ShapePadding);
 	template FContactPoint CapsuleTriangleMeshContactPoint<TImplicitObjectScaled<class Chaos::FTriangleMeshImplicitObject, 1>>(const FCapsule& A, const FRigidTransform3& ATransform, const TImplicitObjectScaled<class Chaos::FTriangleMeshImplicitObject, 1>& B, const FRigidTransform3& BTransform, const FReal CullDistance, const FReal ShapePadding);
 	template FContactPoint CapsuleTriangleMeshContactPoint<FTriangleMeshImplicitObject>(const FCapsule& A, const FRigidTransform3& ATransform, const FTriangleMeshImplicitObject& B, const FRigidTransform3& BTransform, const FReal CullDistance, const FReal ShapePadding);
-
-	template void GJKImplicitManifold<FConvex, FHeightField>(const FImplicitObject& A, const FRigidTransform3& ATransform, const FHeightField& B, const FRigidTransform3& BTransform, const FReal CullDistance, const FReal ShapePadding, TArray<FContactPoint>& ContactPoints);
-	template void GJKImplicitManifold<FCapsule, FHeightField>(const FImplicitObject& A, const FRigidTransform3& ATransform, const FHeightField& B, const FRigidTransform3& BTransform, const FReal CullDistance, const FReal ShapePadding, TArray<FContactPoint>& ContactPoints);
-
-	template void GJKImplicitManifold<FConvex, TImplicitObjectScaled<class Chaos::FTriangleMeshImplicitObject, 1>>(const FImplicitObject& A, const FRigidTransform3& ATransform, const TImplicitObjectScaled<class Chaos::FTriangleMeshImplicitObject, 1>& B, const FRigidTransform3& BTransform, const FReal CullDistance, const FReal ShapePadding, TArray<FContactPoint>& ContactPoints);
-	template void GJKImplicitManifold<FConvex, FTriangleMeshImplicitObject>(const FImplicitObject& A, const FRigidTransform3& ATransform, const FTriangleMeshImplicitObject& B, const FRigidTransform3& BTransform, const FReal CullDistance, const FReal ShapePadding, TArray<FContactPoint>& ContactPoints);
-
-	template void GJKImplicitManifold<FCapsule, TImplicitObjectScaled<class Chaos::FTriangleMeshImplicitObject, 1>>(const FImplicitObject& A, const FRigidTransform3& ATransform, const TImplicitObjectScaled<class Chaos::FTriangleMeshImplicitObject, 1>& B, const FRigidTransform3& BTransform, const FReal CullDistance, const FReal ShapePadding, TArray<FContactPoint>& ContactPoints);
-	template void GJKImplicitManifold<FCapsule, FTriangleMeshImplicitObject>(const FImplicitObject& A, const FRigidTransform3& ATransform, const FTriangleMeshImplicitObject& B, const FRigidTransform3& BTransform, const FReal CullDistance, const FReal ShapePadding, TArray<FContactPoint>& ContactPoints);
-
-	template void GJKImplicitManifold<FImplicitBox3, FHeightField>(const FImplicitObject& A, const FRigidTransform3& ATransform, const FHeightField& B, const FRigidTransform3& BTransform, const FReal CullDistance, const FReal ShapePadding, TArray<FContactPoint>& ContactPoints);
-
-	template void GJKImplicitManifold<FImplicitBox3, TImplicitObjectScaled<class Chaos::FTriangleMeshImplicitObject, 1>>(const FImplicitObject& A, const FRigidTransform3& ATransform, const TImplicitObjectScaled<class Chaos::FTriangleMeshImplicitObject, 1>& B, const FRigidTransform3& BTransform, const FReal CullDistance, const FReal ShapePadding, TArray<FContactPoint>& ContactPoints);
-	template void GJKImplicitManifold<FImplicitBox3, FTriangleMeshImplicitObject>(const FImplicitObject& A, const FRigidTransform3& ATransform, const FTriangleMeshImplicitObject& B, const FRigidTransform3& BTransform, const FReal CullDistance, const FReal ShapePadding, TArray<FContactPoint>& ContactPoints);
 }
diff --git a/Engine/Source/Runtime/Experimental/Chaos/Private/Chaos/Collision/ContactTriangles.cpp b/Engine/Source/Runtime/Experimental/Chaos/Private/Chaos/Collision/ContactTriangles.cpp
new file mode 100644
--- /dev/null
+++ b/Engine/Source/Runtime/Experimental/Chaos/Private/Chaos/Collision/ContactTriangles.cpp
@@ -0,0 +1,554 @@
+// Copyright Epic Games, Inc.All Rights Reserved.
+#include "Chaos/Collision/ContactTriangles.h"
+#include "Chaos/DebugDrawQueue.h"
+
+namespace Chaos
+{
+	namespace CVars
+	{
+		extern int32 ChaosSolverDebugDrawMeshContacts;
+	}
+	extern bool bChaos_Collision_EnableEdgePrune;
+
+
+	void FContactTriangleCollector::SetNumContacts(const int32 NumContacts)
+	{
+		TriangleContactPoints.SetNum(NumContacts, false);
+	}
+
+	void FContactTriangleCollector::RemoveContact(const int32 ContactIndex)
+	{
+		TriangleContactPoints.RemoveAtSwap(ContactIndex, 1, false);
+	}
+
+	int32 FContactTriangleCollector::AddTriangle(const FTriangle& Triangle, const int32 TriangleIndex, const int32 VertexIndex0, const int32 VertexIndex1, const int32 VertexIndex2)
+	{
+		const int32 ContactTriangleIndex = ContactTriangles.AddUninitialized();
+		FContactTriangle& ContactTriangle = ContactTriangles[ContactTriangleIndex];
+		ContactTriangle.Vertices[0] = Triangle.GetVertex(0);
+		ContactTriangle.Vertices[1] = Triangle.GetVertex(1);
+		ContactTriangle.Vertices[2] = Triangle.GetVertex(2);
+		ContactTriangle.FaceNormal = Triangle.GetNormal();
+		ContactTriangle.VertexIndices[0] = VertexIndex0;
+		ContactTriangle.VertexIndices[1] = VertexIndex1;
+		ContactTriangle.VertexIndices[2] = VertexIndex2;
+		ContactTriangle.TriangleIndex = TriangleIndex;
+		return ContactTriangleIndex;
+	}
+
+	// Find the other triangle in the list that shares the specified edge with the specified triangle
+	int32 FContactTriangleCollector::GetOtherTriangleIndexSharingEdge(const int32 KnownTriangleIndex, const int32 EdgeVertexIndexA, const int32 EdgeVertexIndexB) const
+	{
+		check(EdgeVertexIndexA != INDEX_NONE);
+		check(EdgeVertexIndexA != EdgeVertexIndexB);
+
+		for (int32 OtherTriangleIndex = 0; OtherTriangleIndex < ContactTriangles.Num(); ++OtherTriangleIndex)
+		{
+			if (OtherTriangleIndex != KnownTriangleIndex)
+			{
+				const FContactTriangle& OtherTriangle = ContactTriangles[OtherTriangleIndex];
+				if (OtherTriangle.HasVertexIndex(EdgeVertexIndexA))
+				{
+					if (EdgeVertexIndexB == INDEX_NONE)
+					{
+						// @todo(chaos): this is just the first triangle using the vertex, there will be at least 3 faces sharing a vertex and could be more...we should handle that
+						return OtherTriangleIndex;
+					}
+
+					if (OtherTriangle.HasVertexIndex(EdgeVertexIndexB))
+					{
+						// This triangle also uses the edge
+						return OtherTriangleIndex;
+					}
+				}
+			}
+		}
+
+		return INDEX_NONE;
+	}
+
+	int32 FContactTriangleCollector::GetNextTriangleIndexSharingVertex(const int32 LastContactTriangleIndex, const int32 VertexIndex) const
+	{
+		check(VertexIndex != INDEX_NONE);
+		for (int32 OtherTriangleIndex = LastContactTriangleIndex + 1; OtherTriangleIndex < ContactTriangles.Num(); ++OtherTriangleIndex)
+		{
+			const FContactTriangle& OtherTriangle = ContactTriangles[OtherTriangleIndex];
+			if (OtherTriangle.HasVertexIndex(VertexIndex))
+			{
+				return OtherTriangleIndex;
+			}
+		}
+
+		return INDEX_NONE;
+	}
+
+	void FContactTriangleCollector::BuildContactFeatureSets()
+	{
+		const int32 NumContactPoints = TriangleContactPoints.Num();
+
+		// @todo(chaos): use sorted array and binary search rather than TSet (which won't hash well for edge IDs)
+		ContactEdges.Reserve(3 * NumContactPoints);
+		ContactVertices.Reserve(3 * NumContactPoints);
+
+		for (int32 ContactIndex = 0; ContactIndex < NumContactPoints; ++ContactIndex)
+		{
+			FTriangleContactPoint& ContactPoint = TriangleContactPoints[ContactIndex];
+			const FContactTriangle& ContactTriangle = ContactTriangles[ContactPoint.ContactTriangleIndex];
+
+			if (ContactPoint.ContactType == EContactPointType::VertexPlane)
+			{
+				// Since we collided with this face, we cannot collide with these edges or vertices
+				ContactEdges.Add(FContactEdgeID(ContactTriangle.VertexIndices[0], ContactTriangle.VertexIndices[1]));
+				ContactEdges.Add(FContactEdgeID(ContactTriangle.VertexIndices[1], ContactTriangle.VertexIndices[2]));
+				ContactEdges.Add(FContactEdgeID(ContactTriangle.VertexIndices[2], ContactTriangle.VertexIndices[0]));
+				ContactVertices.Add(FContactVertexID(ContactTriangle.VertexIndices[0]));
+				ContactVertices.Add(FContactVertexID(ContactTriangle.VertexIndices[1]));
+				ContactVertices.Add(FContactVertexID(ContactTriangle.VertexIndices[2]));
+			}
+			else if ((ContactPoint.ContactType == EContactPointType::EdgeEdge) || (ContactPoint.ContactType == EContactPointType::PlaneVertex))
+			{
+				int32 VertexIndexA, VertexIndexB;
+				FVec3 VertexA, VertexB;
+				if (ContactTriangle.GetEdgeVerticesAtPosition(ContactPoint.ShapeContactPoints[1], VertexA, VertexB, VertexIndexA, VertexIndexB, FReal(0.01)))
+				{
+					if ((VertexIndexA != INDEX_NONE) && (VertexIndexB != INDEX_NONE))
+					{
+						// Since we collided with this edge, we cannot also collide with these verts
+						ContactVertices.Add(FContactVertexID(VertexIndexA));
+						ContactVertices.Add(FContactVertexID(VertexIndexB));
+
+						// Store the edge ID so we can see if any other faces add it to the don'tcollide list
+						ContactPoint.EdgeID = FContactEdgeID(VertexIndexA, VertexIndexB);
+					}
+					else if (VertexIndexA != INDEX_NONE)
+					{
+						// Store the vertex ID so we can see if any other faces or edges add it to the don'tcollide list
+						ContactPoint.VertexID = FContactVertexID(VertexIndexA);
+					}
+				}
+			}
+		}
+	}
+
+	void FContactTriangleCollector::ProcessContacts(const FRigidTransform3& MeshToConvexTransform)
+	{
+		DebugDrawContactPoints(FColor::White, 0.2);
+
+		// Build the feature set from the full list of contacts before pruning
+		BuildContactFeatureSets();
+
+		// Remove edge contacts that should not be possible because we are also colliding with a face that used the edge.
+		if (TriangleContactPoints.Num() > 1)
+		{
+			PruneEdgeContactPoints();
+		}
+
+		DebugDrawContactPoints(FColor::Cyan, 0.35);
+
+		// Remove contacts that should not be able to occur if we are a one-sided mesh
+		// NOTE: does not require sorted contacts, and does not maintain array order
+		if (bOneSidedCollision)
+		{
+			PruneInfacingContactPoints();
+		}
+
+		SortContactPointsForPruning();
+
+		// Remove close contacts and those not near the deepest Phi
+		// NOTE: relies on the sort above, and retains order
+		if (TriangleContactPoints.Num() > 4)
+		{
+			PruneUnnecessaryContactPoints(FReal(0.1), FReal(0.1));
+		}
+
+		DebugDrawContactPoints(FColor::Yellow, 0.5);
+
+		// Reduce to only 4 contact points from here
+		// NOTE: relies on the sort above and may not retain order
+		if (TriangleContactPoints.Num() > 4)
+		{
+			ReduceManifoldContactPointsTriangeMesh();
+		}
+
+		// Fix contact normals for edge and vertex collisions
+		if (TriangleContactPoints.Num() > 1)
+		{
+			FixInvalidNormalContactPoints();
+		}
+
+		// Sort the contacts so that faces are processed before edges, edges before vertices, and then deepest are processed first
+		SortContactPointsForSolving();
+
+		DebugDrawContactPoints(FColor::Red, 1.0);
+
+		FinalizeContacts(MeshToConvexTransform);
+	}
+
+	void FContactTriangleCollector::SortContactPointsForPruning()
+	{
+		if (TriangleContactPoints.Num() > 1)
+		{
+			std::sort(
+				&TriangleContactPoints[0],
+				&TriangleContactPoints[0] + TriangleContactPoints.Num(),
+				[](const FTriangleContactPoint& L, const FTriangleContactPoint& R)
+				{
+					return L.Phi < R.Phi;
+				}
+			);
+		}
+	}
+
+	// Sort contacts on a shape pair in the order we like to solve them.
+	// NOTE: This relies on the enum order of EContactPointType.
+	void FContactTriangleCollector::SortContactPointsForSolving()
+	{
+		if (TriangleContactPoints.Num() > 1)
+		{
+			std::sort(
+				&TriangleContactPoints[0],
+				&TriangleContactPoints[0] + TriangleContactPoints.Num(),
+				[](const FTriangleContactPoint& L, const FTriangleContactPoint& R)
+				{
+					if (L.ContactType == R.ContactType)
+					{
+						return L.Phi < R.Phi;
+					}
+					return L.ContactType < R.ContactType;
+				}
+			);
+		}
+	}
+
+	// For each edge contact, find the valid range of normal and reject the contact if the normal is outside the range.
+	// This rejects collisions against concave edges, as well as edge contacts that should be hidden by face contacts on adjacent triangles.
+	// NOTE: This only works if both triangle sharing the edge are in the ContactPointTriangles list. They usually will be because you won't collide with an edge unless you overlap the adjacent triangle.
+	// NOTE: Assumes the ContactPoints array has both contacts in the same space (Tri Mesh and Heightfields are this way until the final step)
+	void FContactTriangleCollector::PruneEdgeContactPoints()
+	{
+		if (!bChaos_Collision_EnableEdgePrune)
+		{
+			return;
+		}
+
+		// Remove all vertex collisions when we have also an edge collision that includes the vertex
+		// Remove all edge collisions when we have a face collision that includes the edge
+		int32 NumContactPoints = TriangleContactPoints.Num();
+		for (int32 ContactIndex = 0; ContactIndex < NumContactPoints; ++ContactIndex)
+		{
+			const FTriangleContactPoint& ContactPoint = TriangleContactPoints[ContactIndex];
+
+			bool bRemoveContact = false;
+			if (ContactPoint.VertexID != INDEX_NONE)
+			{
+				bRemoveContact = ContactVertices.Contains(ContactPoint.VertexID);
+			}
+			else if (ContactPoint.EdgeID.IsValid())
+			{
+				bRemoveContact = ContactEdges.Contains(ContactPoint.EdgeID);
+			}
+
+			if (bRemoveContact)
+			{
+				RemoveContact(ContactIndex);
+				--ContactIndex;
+				--NumContactPoints;
+			}
+		}
+	}
+
+	// Examine all edge and vertex contacts. If the contact normal is invalid (i.e., it points into the prism formed by 
+	// infinitly extending the face along the normal) then fix it to be in the valid range.
+	void FContactTriangleCollector::FixInvalidNormalContactPoints()
+	{
+		const int32 NumContactPoints = TriangleContactPoints.Num();
+		for (int32 ContactIndex = 0; ContactIndex < NumContactPoints; ++ContactIndex)
+		{
+			FTriangleContactPoint& ContactPoint = TriangleContactPoints[ContactIndex];
+
+			// Do we have a collision with the triangle edge or vertex?
+			// NOTE: second shape is always the triangle
+			const bool bIsTriangleEdgeOrVertex = (ContactPoint.ContactType == EContactPointType::EdgeEdge) || (ContactPoint.ContactType == EContactPointType::PlaneVertex);
+			if (bIsTriangleEdgeOrVertex)
+			{
+				const int32 ContactTriangleIndex = ContactPoint.ContactTriangleIndex;
+				const FContactTriangle& ContactTriangle = ContactTriangles[ContactTriangleIndex];
+
+				// Find the vertex positions and indices for this contact
+				int32 VertexIndexA, VertexIndexB;
+				FVec3 VertexA, VertexB;
+				if (ContactTriangle.GetEdgeVerticesAtPosition(ContactPoint.ShapeContactPoints[1], VertexA, VertexB, VertexIndexA, VertexIndexB))
+				{
+					// If VertexIndexA and VertexIndexB are valid we have an edge collision.
+					// If only VertexIndexA is valid (VertexIndexB is INDEX_NONE) we have a vertex collision. 
+					// @todo(chaos): Add check for valid normals with vertex collisions. This is much harder than the edge case...
+					if ((VertexIndexA != INDEX_NONE) && (VertexIndexB != INDEX_NONE))
+					{
+						// Get the other triangle that shares this edge
+						const int32 OtherContactTriangleIndex = GetOtherTriangleIndexSharingEdge(ContactTriangleIndex, VertexIndexA, VertexIndexB);
+						if (OtherContactTriangleIndex != INDEX_NONE)
+						{
+							const FContactTriangle& OtherContactTriangle = ContactTriangles[OtherContactTriangleIndex];
+							const FReal MinContactDotNormal = FVec3::DotProduct(OtherContactTriangle.FaceNormal, ContactTriangle.FaceNormal);
+							const FReal ContactDotNormal = FVec3::DotProduct(ContactPoint.ShapeContactNormal, ContactTriangle.FaceNormal);
+							if (ContactDotNormal < MinContactDotNormal)
+							{
+								// We are outside the valid normal range for this edge
+								// Convert the edge collision to a face collision on one of the faces, selected to get the smallest depth
+								// NOTE: We keep the contact depth as it is because we know that the depth is a lower-bound. 
+								// We have to update the ShapeContactPoint[0] because Phi is actually derived from the positions (the value in ContactPoint is just a cache of current state)
+								// @todo(chaos): face selection logic might be better if we knew the contact velocity
+								const FReal OtherContactDotNormal = FVec3::DotProduct(ContactPoint.ShapeContactNormal, ContactTriangle.FaceNormal);
+								if (ContactDotNormal >= OtherContactDotNormal)
+								{
+									ContactPoint.ShapeContactNormal = (ContactDotNormal > -SMALL_NUMBER) ? ContactTriangle.FaceNormal : -ContactTriangle.FaceNormal;
+									ContactPoint.ShapeContactPoints[0] = ContactPoint.ShapeContactPoints[1] + ContactPoint.Phi * ContactPoint.ShapeContactNormal;
+								}
+								else
+								{
+									ContactPoint.ShapeContactNormal = (OtherContactDotNormal > -SMALL_NUMBER) ? OtherContactTriangle.FaceNormal : -OtherContactTriangle.FaceNormal;
+									ContactPoint.ShapeContactPoints[0] = ContactPoint.ShapeContactPoints[1] + ContactPoint.Phi * ContactPoint.ShapeContactNormal;
+									ContactPoint.ContactTriangleIndex = OtherContactTriangleIndex;
+								}
+							}
+						}
+					}
+					else if (VertexIndexA != INDEX_NONE)
+					{
+						for (int32 TriangleIndex = 0; TriangleIndex < ContactTriangles.Num(); ++TriangleIndex)
+						{
+							const FContactTriangle& OtherContactTriangle = ContactTriangles[TriangleIndex];
+							if (OtherContactTriangle.HasVertexIndex(VertexIndexA))
+							{
+								FVec3 EdgeVertex0, EdgeVertex1;
+								if (OtherContactTriangle.GetOtherVertexPositions(VertexIndexA, EdgeVertex0, EdgeVertex1))
+								{
+									// Does the contact normal point into the infinite prism formed by extruding the triangle along the face normal?
+									const FVec3 EdgeDelta0 = EdgeVertex0 - VertexA;
+									const FVec3 EdgeDelta1 = EdgeVertex1 - VertexA;
+									const FReal EdgeSign0 = FVec3::DotProduct(FVec3::CrossProduct(ContactPoint.ShapeContactNormal, EdgeVertex0 - VertexA), OtherContactTriangle.FaceNormal);
+									const FReal EdgeSign1 = FVec3::DotProduct(FVec3::CrossProduct(ContactPoint.ShapeContactNormal, EdgeVertex1 - VertexA), OtherContactTriangle.FaceNormal);
+									if (FMath::Sign(EdgeSign0) == FMath::Sign(EdgeSign1))
+									{
+										const FVec3 Centroid = OtherContactTriangle.GetCentroid();
+										const FReal NormalDotCentroid = FVec3::DotProduct(ContactPoint.ShapeContactNormal, Centroid - ContactPoint.ShapeContactPoints[1]);
+										if (NormalDotCentroid > 0)
+										{
+											const FReal OtherContactDotNormal = FVec3::DotProduct(ContactPoint.ShapeContactNormal, ContactTriangle.FaceNormal);
+											ContactPoint.ShapeContactNormal = (OtherContactDotNormal > -SMALL_NUMBER) ? OtherContactTriangle.FaceNormal : -OtherContactTriangle.FaceNormal;
+											ContactPoint.ShapeContactPoints[0] = ContactPoint.ShapeContactPoints[1] + ContactPoint.Phi * ContactPoint.ShapeContactNormal;
+										}
+									}
+								}
+							}
+						}
+					}
+				}
+			}
+		}
+	}
+
+	// Remove contact points that have a normal pointing inthe opposite direction of the face normal
+	void FContactTriangleCollector::PruneInfacingContactPoints()
+	{
+		for (int32 ContactIndex = 0; ContactIndex < TriangleContactPoints.Num(); ++ContactIndex)
+		{
+			const FTriangleContactPoint& ContactPoint = TriangleContactPoints[ContactIndex];
+			const FContactTriangle& ContactTriangle = ContactTriangles[ContactPoint.ContactTriangleIndex];
+
+			const FReal ContactNormalDotFaceNormal = FVec3::DotProduct(ContactPoint.ShapeContactNormal, ContactTriangle.FaceNormal);
+			if (ContactNormalDotFaceNormal < 0)
+			{
+				RemoveContact(ContactIndex);
+				--ContactIndex;
+			}
+		}
+	}
+
+	// Remove contact points that are much more shallow that the deepest
+	// Remove contact points that are too close to others.
+	// NOTE: ContactPoints must be sorted on Phi on entry, and will be sorted on exit
+	// @todo(chaos): we should only remove shallow contacts that are hidden by nearby deeper contacts. It should be ok to have shallow contacts
+	// far from the deeper ones. E.g., a box landing on plane at an angle.
+	void FContactTriangleCollector::PruneUnnecessaryContactPoints(const FReal PhiTolerance, const FReal DistanceTolerance)
+	{
+		// Remove all points except for the deepest one, and ones with phis similar to it
+		int32 NewContactPointCount = TriangleContactPoints.Num() > 0 ? 1 : 0;
+		for (int32 Index = 1; Index < TriangleContactPoints.Num(); Index++)
+		{
+			if ((TriangleContactPoints[Index].Phi < 0) || ((TriangleContactPoints[Index].Phi - TriangleContactPoints[0].Phi) < PhiTolerance))
+			{
+				NewContactPointCount++;
+			}
+			else
+			{
+				break;
+			}
+		}
+		SetNumContacts(NewContactPointCount);
+
+		// Remove points that are very close together. Always keep the one that is least likely to be an edge contact
+		const FReal DistanceToleranceSq = FMath::Square(DistanceTolerance);
+		for (int32 ContactIndex0 = 0; ContactIndex0 < TriangleContactPoints.Num(); ++ContactIndex0)
+		{
+			const FTriangleContactPoint& ContactPoint0 = TriangleContactPoints[ContactIndex0];
+			const FContactTriangle& ContactTriangle0 = ContactTriangles[ContactPoint0.ContactTriangleIndex];
+
+			for (int32 ContactIndex1 = ContactIndex0 + 1; ContactIndex1 < TriangleContactPoints.Num(); ++ContactIndex1)
+			{
+				const FTriangleContactPoint& ContactPoint1 = TriangleContactPoints[ContactIndex1];
+				const FContactTriangle& ContactTriangle1 = ContactTriangles[ContactPoint1.ContactTriangleIndex];
+
+				// NOTE: ShapeContactPoints[1] is on the triangle
+				const FVec3 ContactDelta = ContactPoint0.ShapeContactPoints[1] - ContactPoint1.ShapeContactPoints[1];
+				const FReal ContactDeltaSq = ContactDelta.SizeSquared();
+				if (ContactDeltaSq < DistanceToleranceSq)
+				{
+					const FReal NormalDot0 = FVec3::DotProduct(ContactPoint0.ShapeContactNormal, ContactTriangle0.FaceNormal);
+					const FReal NormalDot1 = FVec3::DotProduct(ContactPoint1.ShapeContactNormal, ContactTriangle1.FaceNormal);
+					if (FMath::Abs(NormalDot0) >= FMath::Abs(NormalDot1))
+					{
+						TriangleContactPoints.RemoveAt(ContactIndex1, 1, false);
+						--ContactIndex1;
+						continue;
+					}
+					else
+					{
+						TriangleContactPoints.RemoveAt(ContactIndex0, 1, false);
+						--ContactIndex0;
+						break;
+					}
+				}
+			}
+		}
+	}
+
+	// Reduce the number of contact points (in place)
+	// Prerequisites to calling this function:
+	// ContactPoints are sorted on phi (ascending)
+	void FContactTriangleCollector::ReduceManifoldContactPointsTriangeMesh()
+	{
+		if (TriangleContactPoints.Num() <= 4)
+		{
+			return;
+		}
+
+		// Point 1) is the deepest contact point
+		// It is already in position
+
+		//
+		// @todo(chaos): actually this isn't really good enough. We want to take the outer 4 points for the most stable manifold, but also
+		// must take the deepest for correct behaviour. We also have a problem right now where if we have a number of contacts at the same 
+		// depth, the "deepest" one is arbitrary and may be a poor choice as the anchor of the manifold building. E.g., Consider 5 contact 
+		// points in a X config where the center point is first in the list - this leads to a triangular/unstable manifold.
+		// 
+		// Maybe we should do this:
+		//	- build the manifold using the current process
+		//	- revisit point 0 at the end and replace with the best choice for stability if the depth is similar
+		// Or this:
+		//	- build manifold for best stabilty (largest convex area)
+		//	- if manifold does not contain deepest point, replace one point so minimize area reduction
+		//
+
+		// Point 2) Find the point with the largest distance to the deepest contact point
+		{
+			uint32 FarthestPointIndex = 1;
+			FReal FarthestPointDistanceSQR = -1.0f;
+			for (int32 PointIndex = 1; PointIndex < TriangleContactPoints.Num(); PointIndex++)
+			{
+				FReal PointAToPointBSizeSQR = (TriangleContactPoints[PointIndex].ShapeContactPoints[1] - TriangleContactPoints[0].ShapeContactPoints[1]).SizeSquared();
+				if (PointAToPointBSizeSQR > FarthestPointDistanceSQR)
+				{
+					FarthestPointIndex = PointIndex;
+					FarthestPointDistanceSQR = PointAToPointBSizeSQR;
+				}
+			}
+			// Farthest point will be added now
+			Swap(TriangleContactPoints[1], TriangleContactPoints[FarthestPointIndex]);
+		}
+
+		// Point 3) Largest triangle area
+		{
+			uint32 LargestTrianglePointIndex = 2;
+			FReal LargestTrianglePointSignedArea = 0.0f; // This will actually be double the signed area
+			FVec3 P0to1 = TriangleContactPoints[1].ShapeContactPoints[1] - TriangleContactPoints[0].ShapeContactPoints[1];
+			for (int32 PointIndex = 2; PointIndex < TriangleContactPoints.Num(); PointIndex++)
+			{
+				FReal TriangleSignedArea = FVec3::DotProduct(FVec3::CrossProduct(P0to1, TriangleContactPoints[PointIndex].ShapeContactPoints[1] - TriangleContactPoints[0].ShapeContactPoints[1]), TriangleContactPoints[0].ShapeContactNormal);
+				if (FMath::Abs(TriangleSignedArea) > FMath::Abs(LargestTrianglePointSignedArea))
+				{
+					LargestTrianglePointIndex = PointIndex;
+					LargestTrianglePointSignedArea = TriangleSignedArea;
+				}
+			}
+			// Point causing the largest triangle will be added now
+			Swap(TriangleContactPoints[2], TriangleContactPoints[LargestTrianglePointIndex]);
+			// Ensure the winding order is consistent
+			if (LargestTrianglePointSignedArea < 0)
+			{
+				Swap(TriangleContactPoints[0], TriangleContactPoints[1]);
+			}
+		}
+
+		// Point 4) Find the largest triangle connecting with our current triangle
+		{
+			uint32 LargestTrianglePointIndex = 3;
+			FReal LargestPositiveTrianglePointSignedArea = 0.0f;
+			for (int32 PointIndex = 3; PointIndex < TriangleContactPoints.Num(); PointIndex++)
+			{
+				for (uint32 EdgeIndex = 0; EdgeIndex < 3; EdgeIndex++)
+				{
+					FReal TriangleSignedArea = FVec3::DotProduct(FVec3::CrossProduct(TriangleContactPoints[PointIndex].ShapeContactPoints[1] - TriangleContactPoints[EdgeIndex].ShapeContactPoints[1], TriangleContactPoints[(EdgeIndex + 1) % 3].ShapeContactPoints[1] - TriangleContactPoints[EdgeIndex].ShapeContactPoints[1]), TriangleContactPoints[0].ShapeContactNormal);
+					if (TriangleSignedArea > LargestPositiveTrianglePointSignedArea)
+					{
+						LargestTrianglePointIndex = PointIndex;
+						LargestPositiveTrianglePointSignedArea = TriangleSignedArea;
+					}
+				}
+			}
+			// Point causing the largest positive triangle area will be added now
+			Swap(TriangleContactPoints[3], TriangleContactPoints[LargestTrianglePointIndex]);
+		}
+
+		// Revisit Point 0 - if we have a choice that increases area at similar depth, use it
+
+		// Will end up with 4 points
+		SetNumContacts(4);
+	}
+
+	void FContactTriangleCollector::FinalizeContacts(const FRigidTransform3& MeshToConvexTransform)
+	{
+		FinalContactPoints.Reset(TriangleContactPoints.Num());
+
+		// Transform contact data back into shape-local space and build the output array
+		for (int32 ContactIndex = 0; ContactIndex < TriangleContactPoints.Num(); ++ContactIndex)
+		{
+			FTriangleContactPoint& ContactPoint = TriangleContactPoints[ContactIndex];
+			const FContactTriangle& ContactTriangle = ContactTriangles[ContactPoint.ContactTriangleIndex];
+
+			ContactPoint.ShapeContactPoints[1] = MeshToConvexTransform.InverseTransformPositionNoScale(ContactPoint.ShapeContactPoints[1]);
+			ContactPoint.ShapeContactNormal = MeshToConvexTransform.InverseTransformVectorNoScale(ContactPoint.ShapeContactNormal);
+
+			FinalContactPoints.Add(ContactPoint);
+		}
+	}
+
+	void FContactTriangleCollector::DebugDrawContactPoints(const FColor& Color, const FReal LineThickness)
+	{
+#if CHAOS_DEBUG_DRAW
+		if (CVars::ChaosSolverDebugDrawMeshContacts && FDebugDrawQueue::GetInstance().IsDebugDrawingEnabled())
+		{
+			const FReal Duration = 0;
+			const int32 DrawPriority = 10;
+			for (const FTriangleContactPoint& ContactPoint : TriangleContactPoints)
+			{
+				const FVec3 P1 = ConvexTransform.TransformPosition(ContactPoint.ShapeContactPoints[1]);
+				const FVec3 N = ConvexTransform.TransformVectorNoScale(ContactPoint.ShapeContactNormal);
+
+				FDebugDrawQueue::GetInstance().DrawDebugLine(P1, P1 + FReal(15) * N, Color, false, FRealSingle(Duration), DrawPriority, FRealSingle(LineThickness));
+			}
+		}
+#endif
+	}
+
+}
\ No newline at end of file
diff --git a/Engine/Source/Runtime/Experimental/Chaos/Private/Chaos/CollisionOneShotManifolds.cpp b/Engine/Source/Runtime/Experimental/Chaos/Private/Chaos/CollisionOneShotManifolds.cpp
--- a/Engine/Source/Runtime/Experimental/Chaos/Private/Chaos/CollisionOneShotManifolds.cpp
+++ b/Engine/Source/Runtime/Experimental/Chaos/Private/Chaos/CollisionOneShotManifolds.cpp
@@ -4,6 +4,7 @@
 
 #include "Chaos/Box.h"
 #include "Chaos/CollisionResolution.h"
+#include "Chaos/Collision/ContactTriangles.h"
 #include "Chaos/Collision/PBDCollisionConstraint.h"
 #include "Chaos/Convex.h"
 #include "Chaos/Defines.h"
@@ -78,14 +79,15 @@
 	FAutoConsoleVariableRef CVarChaos_Collision_EnableEdgePrune(TEXT("p.Chaos.Collision.EnableEdgePrune"), bChaos_Collision_EnableEdgePrune, TEXT(""));
 	FAutoConsoleVariableRef CVarChaos_Collision_EdgePrunePlaneDistance(TEXT("p.Chaos.Collision.EdgePrunePlaneDistance"), Chaos_Collision_EdgePrunePlaneDistance, TEXT(""));
 
-	// Whether to use the new index-less GJK. This should be enabled full time once tested
-	bool bChaos_Collision_UseGJK2 = false;
+	// Whether to use the new index-less GJK. 
+	// @todo(chaos): This should be removed once soaked for a bit (enabled 7 June 2022)
+	bool bChaos_Collision_UseGJK2 = true;
 	FAutoConsoleVariableRef CVarChaos_Collision_UseGJK2(TEXT("p.Chaos.Collision.UseGJK2"), bChaos_Collision_UseGJK2, TEXT(""));
 
 	bool bChaos_Collision_OneSidedTriangleMesh = true;
-	bool bChaos_Collision_OneSidedHeightfield = true;
+	bool bChaos_Collision_OneSidedHeightField = true;
 	FAutoConsoleVariableRef CVarChaos_Collision_OneSidedTriangleMesh(TEXT("p.Chaos.Collision.OneSidedTriangleMesh"), bChaos_Collision_OneSidedTriangleMesh, TEXT(""));
-	FAutoConsoleVariableRef CVarChaos_Collision_OneSidedHeightfield(TEXT("p.Chaos.Collision.OneSidedHeightfield"), bChaos_Collision_OneSidedHeightfield, TEXT(""));
+	FAutoConsoleVariableRef CVarChaos_Collision_OneSidedHeightfield(TEXT("p.Chaos.Collision.OneSidedHeightField"), bChaos_Collision_OneSidedHeightField, TEXT(""));
 
 	namespace Collisions
 	{
@@ -274,165 +276,6 @@
 			return OutPointCount; // This should always be 4
 		}
 
-		// Remove spurious edge contact by removing edge-edge contact that are in the plane of another vertex-plane contact.
-		// Note: this assumes that the first shape in the contacts is the same convex shape (e.g., a box), and the second
-		// can be from any other shape(s), e.g., triangles of a non-convex shape.
-		// We treat PlaneVertex contacts the same because these are contacts on the convex face with a vertex of the triangle
-		// which require the same culling rule be applied.
-		template<bool bMaintainOrder>
-		void PruneEdgeContactPointsImp(TArray<FContactPoint>& ContactPoints, const FReal MaxPlaneDistance)
-		{
-			if (!bChaos_Collision_EnableEdgePrune)
-			{
-				return;
-			}
-
-			int32 NumContactPoints = ContactPoints.Num();
-			for (int32 EdgeContactIndex = 0; EdgeContactIndex < NumContactPoints; ++EdgeContactIndex)
-			{
-				const FContactPoint& EdgeContactPoint = ContactPoints[EdgeContactIndex];
-				const bool bIsTriangleEdgeOrVertex = (EdgeContactPoint.ContactType == EContactPointType::EdgeEdge) || (EdgeContactPoint.ContactType == EContactPointType::PlaneVertex);
-
-				if (bIsTriangleEdgeOrVertex)
-				{
-					for (int32 PlaneContactIndex = 0; PlaneContactIndex < NumContactPoints; ++PlaneContactIndex)
-					{
-						const FContactPoint PlaneContactPoint = ContactPoints[PlaneContactIndex];
-						const bool bIsTriangleFace = (PlaneContactPoint.ContactType == EContactPointType::VertexPlane);
-
-						// If this is a triangle face contact, does the edge/vertex contact lie in the plane?
-						if (bIsTriangleFace)
-						{
-							const int32 PlaneShapeIndex = 1;
-							const FVec3 DeltaPos = PlaneContactPoint.ShapeContactPoints[PlaneShapeIndex] - EdgeContactPoint.ShapeContactPoints[PlaneShapeIndex];
-							const FReal PlaneDistance = FVec3::DotProduct(DeltaPos, PlaneContactPoint.ShapeContactNormal);
-							if (FMath::Abs(PlaneDistance) < MaxPlaneDistance)
-							{
-								// @todo(chaos): Pretty sure we'll need to do an absolute distance check here too so we don;t reject edge
-								// contacts because of plane contacts a long way away.
-
-								// Note: Cannot use RemoveAtSwap unless we re-sort at the end. See ReduceManifoldContactPointsTriangeMesh
-								const bool bAllowShrinking = false;
-								if (bMaintainOrder)
-								{
-									ContactPoints.RemoveAt(EdgeContactIndex, 1, bAllowShrinking);
-								}
-								else
-								{
-									ContactPoints.RemoveAtSwap(EdgeContactIndex, 1, bAllowShrinking);
-								}
-								--EdgeContactIndex;
-								--NumContactPoints;
-								break;
-							}
-						}
-					}
-				}
-			}
-		}
-
-		void PruneEdgeContactPointsOrdered(TArray<FContactPoint>& ContactPoints, const FReal MaxPlaneDistance)
-		{
-			PruneEdgeContactPointsImp<true>(ContactPoints, MaxPlaneDistance);
-		}
-
-		void PruneEdgeContactPointsUnordered(TArray<FContactPoint>& ContactPoints, const FReal MaxPlaneDistance)
-		{
-			PruneEdgeContactPointsImp<false>(ContactPoints, MaxPlaneDistance);
-		}
-
-		// Reduce the number of contact points (in place)
-		// Prerequisites to calling this function:
-		// ContactPoints are sorted on phi (ascending)
-		void ReduceManifoldContactPointsTriangeMesh(TArray<FContactPoint>& ContactPoints)
-		{
-			if (ContactPoints.Num() <= 4)
-			{
-				return;
-			}
-
-			// Point 1) is the deepest contact point
-			// It is already in position
-			
-			//
-			// @todo(chaos): actually this isn't really good enough. We want to take the outer 4 points for the most stable manifold, but also
-			// must take the deepest for correct behaviour. We also have a problem right now where if we have a number of contacts at the same 
-			// depth, the "deepest" one is arbitrary and may be a poor choice as the anchor of the manifold building. E.g., Consider 5 contact 
-			// points in a X config where the center point is first in the list - this leads to a triangular/unstable manifold.
-			// 
-			// Maybe we should do this:
-			//	- build the manifold using the current process
-			//	- revisit point 0 at the end and replace with the best choice for stability if the depth is similar
-			// Or this:
-			//	- build manifold for best stabilty (largest convex area)
-			//	- if manifold does not contain deepest point, replace one point so minimize area reduction
-			//
-			
-			// Point 2) Find the point with the largest distance to the deepest contact point
-			{
-				uint32 FarthestPointIndex = 1;
-				FReal FarthestPointDistanceSQR = -1.0f;
-				for (int32 PointIndex = 1; PointIndex < ContactPoints.Num(); PointIndex++)
-				{
-					FReal PointAToPointBSizeSQR = (ContactPoints[PointIndex].ShapeContactPoints[1] - ContactPoints[0].ShapeContactPoints[1]).SizeSquared();
-					if (PointAToPointBSizeSQR > FarthestPointDistanceSQR)
-					{
-						FarthestPointIndex = PointIndex;
-						FarthestPointDistanceSQR = PointAToPointBSizeSQR;
-					}
-				}
-				// Farthest point will be added now
-				Swap(ContactPoints[1], ContactPoints[FarthestPointIndex]);
-			}
-
-			// Point 3) Largest triangle area
-			{
-				uint32 LargestTrianglePointIndex = 2;
-				FReal LargestTrianglePointSignedArea = 0.0f; // This will actually be double the signed area
-				FVec3 P0to1 = ContactPoints[1].ShapeContactPoints[1] - ContactPoints[0].ShapeContactPoints[1];
-				for (int32 PointIndex = 2; PointIndex < ContactPoints.Num(); PointIndex++)
-				{
-					FReal TriangleSignedArea = FVec3::DotProduct(FVec3::CrossProduct(P0to1, ContactPoints[PointIndex].ShapeContactPoints[1] - ContactPoints[0].ShapeContactPoints[1]), ContactPoints[0].ShapeContactNormal);
-					if (FMath::Abs(TriangleSignedArea) > FMath::Abs(LargestTrianglePointSignedArea))
-					{
-						LargestTrianglePointIndex = PointIndex;
-						LargestTrianglePointSignedArea = TriangleSignedArea;
-					}
-				}
-				// Point causing the largest triangle will be added now
-				Swap(ContactPoints[2], ContactPoints[LargestTrianglePointIndex]);
-				// Ensure the winding order is consistent
-				if (LargestTrianglePointSignedArea < 0)
-				{
-					Swap(ContactPoints[0], ContactPoints[1]);
-				}
-			}
-
-			// Point 4) Find the largest triangle connecting with our current triangle
-			{
-				uint32 LargestTrianglePointIndex = 3;
-				FReal LargestPositiveTrianglePointSignedArea = 0.0f;
-				for (int32 PointIndex = 3; PointIndex < ContactPoints.Num(); PointIndex++)
-				{
-					for (uint32 EdgeIndex = 0; EdgeIndex < 3; EdgeIndex++)
-					{
-						FReal TriangleSignedArea = FVec3::DotProduct(FVec3::CrossProduct(ContactPoints[PointIndex].ShapeContactPoints[1] - ContactPoints[EdgeIndex].ShapeContactPoints[1], ContactPoints[(EdgeIndex + 1) % 3].ShapeContactPoints[1] - ContactPoints[EdgeIndex].ShapeContactPoints[1]), ContactPoints[0].ShapeContactNormal);
-						if (TriangleSignedArea > LargestPositiveTrianglePointSignedArea)
-						{
-							LargestTrianglePointIndex = PointIndex;
-							LargestPositiveTrianglePointSignedArea = TriangleSignedArea;
-						}
-					}
-				}
-				// Point causing the largest positive triangle area will be added now
-				Swap(ContactPoints[3], ContactPoints[LargestTrianglePointIndex]);
-			}
-
-			// Revisit Point 0 - if we have a choice that increases area at similar depth, use it
-
-			ContactPoints.SetNum(4); // Will end up with 4 points
-		}
-
 		// This function will clip the input vertices by a reference shape's planes
 		// more vertices may be added to outputVertexBuffer by this function
 		// This is the core of the Sutherland-Hodgman algorithm
@@ -651,7 +494,7 @@
 		// GJK contact point between two GJKShape wrapped implicits
 		// This assumes that the both shapes are in the same space
 		template <typename GJKShapeA, typename GJKShapeB>
-		FContactPoint GJKContactPointSameSpace(const GJKShapeA& A, const GJKShapeB& B, FReal& OutMaxMarginDelta, int32& VertexIndexA, int32& VertexIndexB)
+		FContactPoint GJKContactPointSameSpace(const GJKShapeA& A, const GJKShapeB& B, FReal& OutMaxMarginDelta, int32& VertexIndexA, int32& VertexIndexB, const FVec3 InitialGJKDir = FVec3(-1,0,0))
 		{
 			SCOPE_CYCLE_COUNTER_MANIFOLD_GJK();
 
@@ -667,11 +510,11 @@
 			bool bHaveContact = false;
 			if (bChaos_Collision_UseGJK2)
 			{
-				bHaveContact = GJKPenetrationSameSpace2(A, B, Penetration, ClosestA, ClosestB, Normal, VertexIndexA, VertexIndexB, OutMaxMarginDelta, GJKEpsilon, EPAEpsilon);
+				bHaveContact = GJKPenetrationSameSpace2(A, B, Penetration, ClosestA, ClosestB, Normal, VertexIndexA, VertexIndexB, OutMaxMarginDelta, InitialGJKDir, GJKEpsilon, EPAEpsilon);
 			}
 			else
 			{
-				bHaveContact = GJKPenetrationSameSpace(A, B, Penetration, ClosestA, ClosestB, Normal, VertexIndexA, VertexIndexB, OutMaxMarginDelta, GJKEpsilon, EPAEpsilon);
+				bHaveContact = GJKPenetrationSameSpace(A, B, Penetration, ClosestA, ClosestB, Normal, VertexIndexA, VertexIndexB, OutMaxMarginDelta, InitialGJKDir, GJKEpsilon, EPAEpsilon);
 			}
 
 			if (bHaveContact)
@@ -1356,9 +1199,22 @@
 			}
 		}
 
-		// @todo(chaos): don't use GJK/EPA for triangle-convex so we can avoid using margains and avoid errors related to misidentifying edge contacts
+		void ConstructCapsuleTriangleOneShotManifold(const FImplicitCapsule3& Capsule, const FTriangle& Triangle, const FReal CullDistance, TCArray<FContactPoint, 4>& OutContactPoints)
+		{
+			// @todo(chaos): make custom capsule-triangle manifold function.
+			// NOTE: ConstructCapsuleConvexOneShotManifold could be used but it has issues when we collide a triangle edge with the capsule cylinder - that needs fixing...
+			FPBDCollisionConstraint Constraint = FPBDCollisionConstraint::MakeTriangle(&Capsule);
+			ConstructConvexConvexOneShotManifold(Capsule, FRigidTransform3::Identity, Triangle, FRigidTransform3::Identity, 0.0f, Constraint);
+			for (const FManifoldPoint& ManifoldPoint : Constraint.GetManifoldPoints())
+			{
+				OutContactPoints.Add(ManifoldPoint.ContactPoint);
+			}
+			return;
+		}
+
+		// @todo(chaos): don't use GJK/EPA for triangle-convex so we can avoid using margins and errors related to misidentifying edge contacts
 		template <typename ConvexType>
-		void ConstructPlanarConvexTriangleOneShotManifold(const ConvexType& Convex, const FTriangle& Triangle, const bool bOneSided, const FReal CullDistance, TCArray<FContactPoint, 4>& OutContactPoints)
+		void ConstructPlanarConvexTriangleOneShotManifold(const ConvexType& Convex, const FTriangle& Triangle, const FReal CullDistance, TCArray<FContactPoint, 4>& OutContactPoints)
 		{
 			SCOPE_CYCLE_COUNTER_MANIFOLD();
 			check(OutContactPoints.Num() == 0);
@@ -1371,7 +1227,8 @@
 			int32 VertexIndexA = INDEX_NONE, VertexIndexB = INDEX_NONE;
 			const TGJKCoreShape<ConvexType> GJKConvex(Convex, ConvexMargin);
 			const TGJKShape<FTriangle> GJKTriangle(Triangle);
-			FContactPoint GJKContactPoint = GJKContactPointSameSpace(GJKConvex, GJKTriangle, MaxMarginDelta, VertexIndexA, VertexIndexB);
+			const FVec3 InitialGJKDir = -Triangle.GetCentroid();
+			FContactPoint GJKContactPoint = GJKContactPointSameSpace(GJKConvex, GJKTriangle, MaxMarginDelta, VertexIndexA, VertexIndexB, InitialGJKDir);
 
 			PHYSICS_CSV_CUSTOM_EXPENSIVE(PhysicsCounters, NumManifoldsGJKCalled, 1, ECsvCustomStatOp::Accumulate);
 
@@ -1392,12 +1249,6 @@
 			const FVec3 TrianglePlanePosition = Triangle[0];
 			const FReal TrianglePlaneNormalDotContactNormal = FVec3::DotProduct(TrianglePlaneNormal, GJKContactPoint.ShapeContactNormal);
 
-			// One sided triangles ignore collisions with the inside face
-			if (bOneSided && (TrianglePlaneNormalDotContactNormal < FReal(0)))
-			{
-				return;
-			}
-
 			// Find the best opposing plane on the convex
 			// @todo(chaos): handle zero margins...
 			const FReal PlaneSearchDistance = ConvexMargin;
@@ -1767,7 +1618,6 @@
 		void ConstructPlanarConvexTriangleOneShotManifold(
 			const FImplicitConvex3& Convex, 
 			const FTriangle& Triangle,
-			const bool bOneSided,
 			const FReal CullDistance, 
 			TCArray<FContactPoint, 4>& OutContactPoints);
 
@@ -1775,7 +1625,6 @@
 		void ConstructPlanarConvexTriangleOneShotManifold(
 			const TImplicitObjectInstanced<FImplicitConvex3>& Convex,
 			const FTriangle& Triangle,
-			const bool bOneSided,
 			const FReal CullDistance,
 			TCArray<FContactPoint, 4>& OutContactPoints);
 
@@ -1783,7 +1632,6 @@
 		void ConstructPlanarConvexTriangleOneShotManifold(
 			const TImplicitObjectScaled<FImplicitConvex3>& Convex,
 			const FTriangle& Triangle,
-			const bool bOneSided,
 			const FReal CullDistance,
 			TCArray<FContactPoint, 4>& OutContactPoints);
 
@@ -1791,7 +1639,6 @@
 		void ConstructPlanarConvexTriangleOneShotManifold(
 			const FImplicitBox3& Convex,
 			const FTriangle& Triangle,
-			const bool bOneSided,
 			const FReal CullDistance,
 			TCArray<FContactPoint, 4>& OutContactPoints);
 
@@ -1799,7 +1646,6 @@
 		void ConstructPlanarConvexTriangleOneShotManifold(
 			const TImplicitObjectScaled<FImplicitBox3>& Convex,
 			const FTriangle& Triangle,
-			const bool bOneSided,
 			const FReal CullDistance,
 			TCArray<FContactPoint, 4>& OutContactPoints);
 
@@ -1807,7 +1653,6 @@
 		void ConstructPlanarConvexTriangleOneShotManifold(
 			const TImplicitObjectInstanced<FImplicitBox3>& Convex,
 			const FTriangle& Triangle,
-			const bool bOneSided,
 			const FReal CullDistance,
 			TCArray<FContactPoint, 4>& OutContactPoints);
 
diff --git a/Engine/Source/Runtime/Experimental/Chaos/Private/Chaos/CollisionOneShotManifolds.h b/Engine/Source/Runtime/Experimental/Chaos/Private/Chaos/CollisionOneShotManifolds.h
--- a/Engine/Source/Runtime/Experimental/Chaos/Private/Chaos/CollisionOneShotManifolds.h
+++ b/Engine/Source/Runtime/Experimental/Chaos/Private/Chaos/CollisionOneShotManifolds.h
@@ -14,9 +14,6 @@ namespace Chaos
 	{
 		uint32 BoxBoxClipVerticesAgainstPlane(const FVec3* InputVertexBuffer, FVec3* outputVertexBuffer, uint32 ClipPointCount, int32 ClippingAxis, FReal Distance);
 		uint32 ReduceManifoldContactPoints(FVec3* Points, uint32 PointCount);
-		void PruneEdgeContactPointsOrdered(TArray<FContactPoint>& ContactPoints, const FReal MaxPlaneDistance);
-		void PruneEdgeContactPointsUnordered(TArray<FContactPoint>& ContactPoints, const FReal MaxPlaneDistance);
-		void ReduceManifoldContactPointsTriangeMesh(TArray<FContactPoint>& ContactPoints);
 
 		void ConstructBoxBoxOneShotManifold(
 			const FImplicitBox3& Box1,
@@ -44,11 +41,16 @@ namespace Chaos
 			const FReal CullDistance,
 			TCArray<FContactPoint, 4>& OutContactPoints);
 
+		void ConstructCapsuleTriangleOneShotManifold(
+			const FImplicitCapsule3& Capsule,
+			const FTriangle& Triangle,
+			const FReal CullDistance,
+			TCArray<FContactPoint, 4>& OutContactPoints);
+
 		template <typename ConvexType>
 		void ConstructPlanarConvexTriangleOneShotManifold(
 			const ConvexType& Convex, 
 			const FTriangle& Triangle, 
-			const bool bOneSided,
 			const FReal CullDistance, 
 			TCArray<FContactPoint, 4>& OutContactPoints);
 	}
diff --git a/Engine/Source/Runtime/Experimental/Chaos/Private/Chaos/CollisionOneShotManifoldsMiscShapes.cpp b/Engine/Source/Runtime/Experimental/Chaos/Private/Chaos/CollisionOneShotManifoldsMiscShapes.cpp
--- a/Engine/Source/Runtime/Experimental/Chaos/Private/Chaos/CollisionOneShotManifoldsMiscShapes.cpp
+++ b/Engine/Source/Runtime/Experimental/Chaos/Private/Chaos/CollisionOneShotManifoldsMiscShapes.cpp
@@ -6,12 +6,14 @@
 #include "Chaos/CollisionResolution.h"
 #include "Chaos/Collision/CapsuleConvexContactPoint.h"
 #include "Chaos/Collision/ContactPointsMiscShapes.h"
+#include "Chaos/Collision/ContactTriangles.h"
 #include "Chaos/Collision/SphereConvexContactPoint.h"
 #include "Chaos/Collision/PBDCollisionConstraint.h"
 #include "Chaos/Convex.h"
 #include "Chaos/Defines.h"
 #include "Chaos/Framework/UncheckedArray.h"
 #include "Chaos/GJK.h"
+#include "Chaos/HeightField.h"
 #include "Chaos/ImplicitObjectScaled.h"
 #include "Chaos/Transform.h"
 #include "Chaos/Triangle.h"
@@ -35,9 +37,11 @@ namespace Chaos
 	extern FRealSingle Chaos_Collision_EPAEpsilon;
 
 	extern bool bChaos_Collision_OneSidedTriangleMesh;
+	extern bool bChaos_Collision_OneSidedHeightField;
 
 	namespace Collisions
 	{
+
 		void ConstructSphereSphereOneShotManifold(
 			const TSphere<FReal, 3>& SphereA,
 			const FRigidTransform3& SphereATransform, //world
@@ -251,6 +255,7 @@ namespace Chaos
 			// GJK and EPA tolerances. See comments in GJKContactPointMargin
 			const FReal GJKEpsilon = Chaos_Collision_GJKEpsilon;
 			const FReal EPAEpsilon = Chaos_Collision_EPAEpsilon;
+			const FVec3 InitialDir = FVec3(-1,0,0);
 			FReal ClosestPenetration;
 			int32 ClosestVertexIndexSphere = INDEX_NONE, ClosestVertexIndexConvex = INDEX_NONE;
 			FReal ClosestSupportMaxDelta = FReal(0);
@@ -270,6 +275,7 @@ namespace Chaos
 					ClosestVertexIndexConvex,
 					ClosestVertexIndexSphere,
 					ClosestSupportMaxDelta,
+					InitialDir,
 					GJKEpsilon, 
 					EPAEpsilon);
 			}
@@ -285,6 +291,7 @@ namespace Chaos
 					ClosestVertexIndexConvex,
 					ClosestVertexIndexSphere,
 					ClosestSupportMaxDelta,
+					InitialDir,
 					GJKEpsilon,
 					EPAEpsilon);
 			}
@@ -657,209 +664,194 @@ namespace Chaos
 			}
 		}
 
-		template <typename TriMeshType>
-		void ConstructCapsuleTriMeshOneShotManifold(const FCapsule& Capsule, const FRigidTransform3& CapsuleWorldTransform, const TriMeshType& TriangleMesh, const FRigidTransform3& TriMeshWorldTransform, const FReal Dt, FPBDCollisionConstraint& Constraint)
+		/**
+		 * @brief Generate a manifold between a convex shape and a single triangle
+		 * Templated so we can specialize for some shape types
+		*/
+		template<typename ConvexType>
+		void GenerateConvexTriangleOneShotManifold(const ConvexType& Convex, const FTriangle& Triangle, const FReal CullDistance, TCArray<FContactPoint, 4>& OutContactPoints)
 		{
-			// We only build one shot manifolds once
-			// All convexes are pre-scaled, or wrapped in TImplicitObjectScaled
-			ensure(CapsuleWorldTransform.GetScale3D() == FVec3(1.0f, 1.0f, 1.0f));
-			ensure(TriMeshWorldTransform.GetScale3D() == FVec3(1.0f, 1.0f, 1.0f));
+			ConstructPlanarConvexTriangleOneShotManifold(Convex, Triangle, CullDistance, OutContactPoints);
+		}
 
-			// @todo(chaos): support manifold maintenance
-			Constraint.ResetActiveManifoldContacts();
+		template<>
+		void GenerateConvexTriangleOneShotManifold<FCapsule>(const FCapsule& Convex, const FTriangle& Triangle, const FReal CullDistance, TCArray<FContactPoint, 4>& OutContactPoints)
+		{
+			ConstructCapsuleTriangleOneShotManifold(Convex, Triangle, CullDistance, OutContactPoints);
+		}
 
-			TArray<FContactPoint> ContactPoints;
-			GJKImplicitManifold<FCapsule, TriMeshType>(Capsule, CapsuleWorldTransform, TriangleMesh, TriMeshWorldTransform, Constraint.GetCullDistance(), 0.0f, ContactPoints);
-			for (FContactPoint& ContactPoint : ContactPoints)
+		/**
+		* @brief Create a minimized set of contact points between a convex polyhedron (box, convex) and a non-convex mesh (trimesh, heightfield)
+		*
+		* @tparam ConvexType any convex type (Sphere, Capsule, Box, Convex, possibly with a scaled or instanced)
+		* @tparam MeshType any triangle mesh type (HeightField or TriangleMesh, without a scaled or instanced wrapper)
+		* @param MeshQueryBounds Triangles overlapping this box will be tested. Should be in the space of the mesh.
+		* @param MeshToConvexTransform The transform from Mesh space to Convex space. This low-level convex-triangle collision detection is performed in Convex space.
+		*
+		* @see ConstructPlanarConvexTriMeshOneShotManifoldImp, ConstructPlanarConvexHeightfieldOneShotManifoldImp
+		*/
+		template<typename ConvexType, typename MeshType>
+		void GenerateConvexMeshContactPoints(const ConvexType& Convex, const MeshType& Mesh, const FAABB3& MeshQueryBounds, const FRigidTransform3& MeshToConvexTransform, const FReal CullDistance, FContactTriangleCollector& MeshContacts)
+		{
+			FContactTriangles ContactTriangles;
+			TCArray<FContactPoint, 4> TriangleManifoldPoints;
+
+			// Loop over all the triangles, build a manifold and add the points to the total manifold
+			// NOTE: contact points will be in the space of the convex until the end of the function when we convert into shape local space
+			Mesh.VisitTriangles(MeshQueryBounds, MeshToConvexTransform, [&](const FTriangle& Triangle, const int32 TriangleIndex, const int32 VertexIndex0, const int32 VertexIndex1, const int32 VertexIndex2)
 			{
-				if (ContactPoint.Phi < Constraint.GetCullDistance())
+				// Generate the manifold for this triangle
+				TriangleManifoldPoints.Reset();
+				GenerateConvexTriangleOneShotManifold(Convex, Triangle, CullDistance, TriangleManifoldPoints);
+
+				if (TriangleManifoldPoints.Num() > 0)
 				{
-					Constraint.AddOneshotManifoldContact(ContactPoint);
+					// Add the points into the main contact array
+					// NOTE: The Contacts' FaceIndices will be an index into the ContactTriangles not the original tri mesh (this will get mapped back to the mesh index below)
+					MeshContacts.AddTriangleContacts(MakeArrayView(&TriangleManifoldPoints[0], TriangleManifoldPoints.Num()), Triangle, TriangleIndex, VertexIndex0, VertexIndex1, VertexIndex2, CullDistance);
 				}
-			}
+			});
+
+			// Reduce contacts to a minimum manifold abd transform contact data back into shape-local space
+			MeshContacts.ProcessContacts(MeshToConvexTransform);
 		}
 
-		void ConstructCapsuleHeightFieldOneShotManifold(const FCapsule& Capsule, const FRigidTransform3& CapsuleTransform, const FHeightField& HeightField, const FRigidTransform3& HeightFieldTransform, const FReal Dt, FPBDCollisionConstraint& Constraint)
+		/**
+		 * @brief Used by all the convex types to generate a manifold against any mesh type
+		*/
+		template<typename ConvexType, typename MeshType>
+		void ConstructConvexMeshOneShotManifold(const ConvexType& Convex, const FRigidTransform3& ConvexTransform, const MeshType& Mesh, const FRigidTransform3& MeshTransform, const FVec3& MeshScale, const FReal CullDistance, FContactTriangleCollector& MeshContacts)
 		{
-			// We only build one shot manifolds once
-			// All convexes are pre-scaled, or wrapped in TImplicitObjectScaled
-			ensure(CapsuleTransform.GetScale3D() == FVec3(1.0f, 1.0f, 1.0f));
-			ensure(HeightFieldTransform.GetScale3D() == FVec3(1.0f, 1.0f, 1.0f));
+			FRigidTransform3 MeshToConvexTransform = MeshTransform.GetRelativeTransformNoScale(ConvexTransform);
+			MeshToConvexTransform.SetScale3D(MeshScale);
 
-			// @todo(chaos): support manifold maintenance
-			Constraint.ResetActiveManifoldContacts();
+			// @todo(chaos): add Convex.CalculateInverseTransformed bounds with scale support (to optimize sphere and capsule)
+			const FAABB3 MeshQueryBounds = Convex.BoundingBox().InverseTransformedAABB(MeshToConvexTransform).Thicken(CullDistance);
 
-			TArray<FContactPoint> ContactPoints;
-			GJKImplicitManifold<FCapsule>(Capsule, CapsuleTransform, HeightField, HeightFieldTransform, Constraint.GetCullDistance(), 0.0f, ContactPoints);
-			for (FContactPoint& ContactPoint : ContactPoints)
-			{
-				if (ContactPoint.Phi < Constraint.GetCullDistance())
-				{
-					Constraint.AddOneshotManifoldContact(ContactPoint);
-				}
-			}
+			// Create the minimal manifold from all the overlapping triangles
+			GenerateConvexMeshContactPoints(Convex, Mesh, MeshQueryBounds, MeshToConvexTransform, CullDistance, MeshContacts);
 		}
 
-		template<typename ConvexType>
-		void ConstructConvexHeightFieldOneShotManifold(const FImplicitObject& Convex, const FRigidTransform3& ConvexTransform, const FHeightField& HeightField, const FRigidTransform3& HeightFieldTransform, const FReal Dt, FPBDCollisionConstraint& Constraint)
+		void ConstructCapsuleTriMeshOneShotManifold(const FCapsule& Capsule, const FRigidTransform3& CapsuleTransform, const FImplicitObject& InMesh, const FRigidTransform3& MeshTransform, const FReal Dt, FPBDCollisionConstraint& Constraint)
 		{
-			// We only build one shot manifolds once
 			// All convexes are pre-scaled, or wrapped in TImplicitObjectScaled
-			ensure(ConvexTransform.GetScale3D() == FVec3(1.0f, 1.0f, 1.0f));
-			ensure(HeightFieldTransform.GetScale3D() == FVec3(1.0f, 1.0f, 1.0f));
+			ensure(CapsuleTransform.GetScale3D() == FVec3(1));
+			ensure(MeshTransform.GetScale3D() == FVec3(1));
 
-			// @todo(chaos): support manifold maintenance
-			Constraint.ResetActiveManifoldContacts();
+			// Unwrap the tri mesh (remove Scaled or Instanced) and get the scale
+			FVec3 MeshScale;
+			FReal MeshMargin;	// Not used - will be zero for meshes
+			const FTriangleMeshImplicitObject* Mesh = UnwrapImplicit<FTriangleMeshImplicitObject>(InMesh, MeshScale, MeshMargin);
+			check(Mesh != nullptr);
 
-			//FContactPoint ContactPoint = ConvexHeightFieldContactPoint(Convex, ConvexTransform, HeightField, HeightFieldTransform, Constraint.GetCullDistance(), 0.0f);
+			FContactTriangleCollector MeshContacts(bChaos_Collision_OneSidedTriangleMesh, CapsuleTransform);
+			const FReal CullDistance = Constraint.GetCullDistance();
 
-			TArray<FContactPoint> ContactPoints;
-			GJKImplicitManifold<ConvexType>(Convex, ConvexTransform, HeightField, HeightFieldTransform, Constraint.GetCullDistance(), 0.0f, ContactPoints);
-			for (FContactPoint& ContactPoint : ContactPoints)
-			{
-				if (ContactPoint.Phi < Constraint.GetCullDistance())
-				{
-					Constraint.AddOneshotManifoldContact(ContactPoint);
-				}
-			}
+			ConstructConvexMeshOneShotManifold(Capsule, CapsuleTransform, *Mesh, MeshTransform, MeshScale, CullDistance, MeshContacts);
+
+			Constraint.SetOneShotManifoldContacts(MeshContacts.GetContactPoints());
 		}
 
-		template<typename ConvexType>
-		void ConstructPlanarConvexTriMeshOneShotManifoldImp(const ConvexType& Convex, const FRigidTransform3& ConvexTransform, const FImplicitObject& TriMesh, const FRigidTransform3& TriMeshTransform, FPBDCollisionConstraint& Constraint)
+		void ConstructCapsuleHeightFieldOneShotManifold(const FCapsule& Capsule, const FRigidTransform3& CapsuleTransform, const FHeightField& Mesh, const FRigidTransform3& MeshTransform, const FReal Dt, FPBDCollisionConstraint& Constraint)
 		{
+			// We only build one shot manifolds once
 			// All convexes are pre-scaled, or wrapped in TImplicitObjectScaled
-			ensure(ConvexTransform.GetScale3D() == FVec3(1.0f, 1.0f, 1.0f));
-			ensure(TriMeshTransform.GetScale3D() == FVec3(1.0f, 1.0f, 1.0f));
+			ensure(CapsuleTransform.GetScale3D() == FVec3(1));
+			ensure(MeshTransform.GetScale3D() == FVec3(1));
 
-			// Unwrap the tri mesh (remove Scaled or Instanced) and get the scale
-			FVec3 TriMeshScale;
-			FReal TriMeshMargin;	// Not used - will be zero
-			const FTriangleMeshImplicitObject* UnscaledTriMesh = UnwrapImplicit<FTriangleMeshImplicitObject>(TriMesh, TriMeshScale, TriMeshMargin);
-			check(UnscaledTriMesh != nullptr);
+			FContactTriangleCollector MeshContacts(bChaos_Collision_OneSidedHeightField, CapsuleTransform);
+			const FReal CullDistance = Constraint.GetCullDistance();
 
-			// @todo(chaos): should be an option on the Tri Mesh?
-			const bool bOneSidedCollision = bChaos_Collision_OneSidedTriangleMesh;
+			ConstructConvexMeshOneShotManifold(Capsule, CapsuleTransform, Mesh, MeshTransform, FVec3(1), CullDistance, MeshContacts);
 
-			const FRigidTransform3 TriangleMeshToConvexTransform = TriMeshTransform.GetRelativeTransformNoScale(ConvexTransform);
+			Constraint.SetOneShotManifoldContacts(MeshContacts.GetContactPoints());
+		}
 
-			// Calculate the query bounds in trimesh space
-			// NOTE: to handle negative scales, we need to include it in the AABB transform (cannot use FAABB3::Scale)
-			const FRigidTransform3 QueryTransform = FRigidTransform3(TriangleMeshToConvexTransform.GetTranslation(), TriangleMeshToConvexTransform.GetRotation(), TriMeshScale);
-			const FAABB3 TriMeshQueryBounds = Convex.BoundingBox().InverseTransformedAABB(QueryTransform);
-			const FReal CullDistance = Constraint.GetCullDistance();
+		/**
+		* @brief Populate the Constraint with a manifold of contacts between a Convex and a TriangleMesh
+		* @param Convex A convex polyhedron (Box, Convex) that may be wrapped in Scaled or Instanced
+		* @param InMesh A TriangleMesh ImplicitObject that may be wrapped in Scaled or Instaned
+		*/
+		void ConstructPlanarConvexTriMeshOneShotManifold(const FImplicitObject& Convex, const FRigidTransform3& ConvexTransform, const FImplicitObject& InMesh, const FRigidTransform3& MeshTransform, FPBDCollisionConstraint& Constraint)
+		{
+			// All convexes are pre-scaled, or wrapped in TImplicitObjectScaled
+			ensure(ConvexTransform.GetScale3D() == FVec3(1));
+			ensure(MeshTransform.GetScale3D() == FVec3(1));
 
-			// A set of contact points which contains points from all triangle-convex manifolds
-			TArray<FContactPoint> ContactPoints;
+			// Unwrap the tri mesh (remove Scaled or Instanced) and get the scale
+			FVec3 MeshScale;
+			FReal MeshMargin;	// Not used - will be zero for meshes
+			const FTriangleMeshImplicitObject* Mesh = UnwrapImplicit<FTriangleMeshImplicitObject>(InMesh, MeshScale, MeshMargin);
+			check(Mesh != nullptr);
 
-			// Prime the triangle producer with overlappig indices
-			FTriangleMeshTriangleProducer TriangleProducer;
-			int32 TriangleIndex;
-			FTriangle Triangle;
-			TCArray<FContactPoint, 4> TriangleManifoldPoints;
-			TriangleProducer.Reset(*UnscaledTriMesh, TriMeshQueryBounds);
+			FContactTriangleCollector MeshContacts(bChaos_Collision_OneSidedTriangleMesh, ConvexTransform);
+			const FReal CullDistance = Constraint.GetCullDistance();
 
-			// Loop over all the triangles, build a manifold and add the points to the total manifold
-			while (TriangleProducer.NextTriangle(*UnscaledTriMesh, QueryTransform, Triangle, TriangleIndex))
+			if (const FImplicitBox3* RawBox = Convex.template GetObject<FImplicitBox3>())
 			{
-				TriangleManifoldPoints.Reset();
-				ConstructPlanarConvexTriangleOneShotManifold(Convex, Triangle, bOneSidedCollision, CullDistance, TriangleManifoldPoints);
-					
-				for (int32 TriangleContactIndex = 0; TriangleContactIndex < TriangleManifoldPoints.Num(); ++TriangleContactIndex)
-				{
-					FContactPoint& ContactPoint = ContactPoints[ContactPoints.AddUninitialized()];
-					ContactPoint.ShapeContactPoints[0] = TriangleManifoldPoints[TriangleContactIndex].ShapeContactPoints[0];
-					ContactPoint.ShapeContactPoints[1] = TriangleMeshToConvexTransform.InverseTransformPositionNoScale(TriangleManifoldPoints[TriangleContactIndex].ShapeContactPoints[1]);
-					ContactPoint.ShapeContactNormal = TriangleMeshToConvexTransform.InverseTransformVectorNoScale(TriangleManifoldPoints[TriangleContactIndex].ShapeContactNormal);
-					ContactPoint.Phi = TriangleManifoldPoints[TriangleContactIndex].Phi;
-					ContactPoint.FaceIndex = TriangleIndex;
-					ContactPoint.ContactType = TriangleManifoldPoints[TriangleContactIndex].ContactType;
-				}
+				ConstructConvexMeshOneShotManifold(*RawBox, ConvexTransform, *Mesh, MeshTransform, MeshScale, CullDistance, MeshContacts);
 			}
-
-			// Remove edge contacts that are "hidden" by face contacts
-			// @todo(chaos): EdgePruneDistance should be some fraction of the convex margin...
-			if (ContactPoints.Num() > 0)
+			else if (const TImplicitObjectScaled<FImplicitConvex3>* ScaledConvex = Convex.template GetObject<TImplicitObjectScaled<FImplicitConvex3>>())
 			{
-				const FReal EdgePruneDistance = Chaos_Collision_EdgePrunePlaneDistance;
-				Collisions::PruneEdgeContactPointsUnordered(ContactPoints, EdgePruneDistance);
+				ConstructConvexMeshOneShotManifold(*ScaledConvex, ConvexTransform, *Mesh, MeshTransform, MeshScale, CullDistance, MeshContacts);
 			}
-
-			// Whittle the manifold down to 4 points
-			if (ContactPoints.Num() > 4)
+			else if (const TImplicitObjectInstanced<FImplicitConvex3>* InstancedConvex = Convex.template GetObject<TImplicitObjectInstanced<FImplicitConvex3>>())
 			{
-				std::sort(&ContactPoints[0], &ContactPoints[0] + ContactPoints.Num(), [](const FContactPoint& L, const FContactPoint& R) { return L.Phi < R.Phi; });
-
-				// Remove all points (except for the deepest one, and ones with phis similar to it)
-				// NOTE: relies on the sort above
-				const FReal CullMargin = 0.1f;
-				int32 NewContactPointCount = ContactPoints.Num() > 0 ? 1 : 0;
-				for (int32 Index = 1; Index < ContactPoints.Num(); Index++)
-				{
-					if (ContactPoints[Index].Phi < 0 || ContactPoints[Index].Phi - ContactPoints[0].Phi < CullMargin)
-					{
-						NewContactPointCount++;
-					}
-					else
-					{
-						break;
-					}
-				}
-				ContactPoints.SetNum(NewContactPointCount, false);
+				ConstructConvexMeshOneShotManifold(*InstancedConvex, ConvexTransform, *Mesh, MeshTransform, MeshScale, CullDistance, MeshContacts);
 			}
-
-			if (ContactPoints.Num() > 4)
+			else if (const FImplicitConvex3* RawConvex = Convex.template GetObject<FImplicitConvex3>())
 			{
-				// Reduce to only 4 contact points from here
-				// NOTE: relies on the sort above
-				Collisions::ReduceManifoldContactPointsTriangeMesh(ContactPoints);
+				ConstructConvexMeshOneShotManifold(*RawConvex, ConvexTransform, *Mesh, MeshTransform, MeshScale, CullDistance, MeshContacts);
 			}
-
-			// Add the manifold points to the constraint
-			Constraint.ResetActiveManifoldContacts();
-			for (FContactPoint& ContactPoint : ContactPoints)
+			else
 			{
-				// NOTE: We don't reuse manifolds between frames for Convex-TriMesh so it's not too bad to
-				// skip manifold points that are beyond the cull distance
-				if (ContactPoint.Phi < Constraint.GetCullDistance())
-				{
-					Constraint.AddOneshotManifoldContact(ContactPoint);
-				}
+				check(false);
 			}
+
+			Constraint.SetOneShotManifoldContacts(MeshContacts.GetContactPoints());
 		}
 
-		void ConstructPlanarConvexTriMeshOneShotManifold(const FImplicitObject& Convex, const FRigidTransform3& ConvexTransform, const FImplicitObject& TriangleMesh, const FRigidTransform3& TriangleMeshTransform, FPBDCollisionConstraint& Constraint)
+		/**
+		* @brief Populate the Constraint with a manifold of contacts between a Convex and a HeightField
+		* @param Convex A convex polyhedron (Box, Convex) that may be wrapped in Scaled or Instanced
+		*/
+		void ConstructPlanarConvexHeightFieldOneShotManifold(const FImplicitObject& Convex, const FRigidTransform3& ConvexTransform, const FHeightField& Mesh, const FRigidTransform3& MeshTransform, FPBDCollisionConstraint& Constraint)
 		{
+			// All convexes are pre-scaled, or wrapped in TImplicitObjectScaled
+			ensure(ConvexTransform.GetScale3D() == FVec3(1));
+			ensure(MeshTransform.GetScale3D() == FVec3(1));
+
+			FContactTriangleCollector MeshContacts(bChaos_Collision_OneSidedHeightField, ConvexTransform);
+			const FReal CullDistance = Constraint.GetCullDistance();
+			const FVec3 MeshScale = FVec3(1);	// Scale is built into heightfield
+
 			if (const FImplicitBox3* RawBox = Convex.template GetObject<FImplicitBox3>())
 			{
-				ConstructPlanarConvexTriMeshOneShotManifoldImp(*RawBox, ConvexTransform, TriangleMesh, TriangleMeshTransform, Constraint);
+				ConstructConvexMeshOneShotManifold(*RawBox, ConvexTransform, Mesh, MeshTransform, MeshScale, CullDistance, MeshContacts);
 			}
 			else if (const TImplicitObjectScaled<FImplicitConvex3>* ScaledConvex = Convex.template GetObject<TImplicitObjectScaled<FImplicitConvex3>>())
 			{
-				ConstructPlanarConvexTriMeshOneShotManifoldImp(*ScaledConvex, ConvexTransform, TriangleMesh, TriangleMeshTransform, Constraint);
+				ConstructConvexMeshOneShotManifold(*ScaledConvex, ConvexTransform, Mesh, MeshTransform, MeshScale, CullDistance, MeshContacts);
 			}
 			else if (const TImplicitObjectInstanced<FImplicitConvex3>* InstancedConvex = Convex.template GetObject<TImplicitObjectInstanced<FImplicitConvex3>>())
 			{
-				ConstructPlanarConvexTriMeshOneShotManifoldImp(*InstancedConvex, ConvexTransform, TriangleMesh, TriangleMeshTransform, Constraint);
+				ConstructConvexMeshOneShotManifold(*InstancedConvex, ConvexTransform, Mesh, MeshTransform, MeshScale, CullDistance, MeshContacts);
 			}
 			else if (const FImplicitConvex3* RawConvex = Convex.template GetObject<FImplicitConvex3>())
 			{
-				ConstructPlanarConvexTriMeshOneShotManifoldImp(*RawConvex, ConvexTransform, TriangleMesh, TriangleMeshTransform, Constraint);
+				ConstructConvexMeshOneShotManifold(*RawConvex, ConvexTransform, Mesh, MeshTransform, MeshScale, CullDistance, MeshContacts);
 			}
 			else
 			{
 				check(false);
 			}
+
+			Constraint.SetOneShotManifoldContacts(MeshContacts.GetContactPoints());
 		}
 
 
 		template void ConstructSphereTriangleMeshOneShotManifold<TImplicitObjectScaled<class Chaos::FTriangleMeshImplicitObject, 1>>(const TSphere<FReal, 3>& Sphere, const FRigidTransform3& SphereWorldTransform, const TImplicitObjectScaled<class Chaos::FTriangleMeshImplicitObject, 1>& TriangleMesh, const FRigidTransform3& TriMeshWorldTransform, const FReal Dt, FPBDCollisionConstraint& Constraint);
 		template void ConstructSphereTriangleMeshOneShotManifold<FTriangleMeshImplicitObject>(const TSphere<FReal, 3>& Sphere, const FRigidTransform3& SphereWorldTransform, const FTriangleMeshImplicitObject& TriangleMesh, const FRigidTransform3& TriMeshWorldTransform, const FReal Dt, FPBDCollisionConstraint& Constraint);
-
-		template void ConstructCapsuleTriMeshOneShotManifold<TImplicitObjectScaled<class Chaos::FTriangleMeshImplicitObject, 1>>(const FCapsule& Capsule, const FRigidTransform3& CapsuleWorldTransform, const TImplicitObjectScaled<class Chaos::FTriangleMeshImplicitObject, 1>& TriangleMesh, const FRigidTransform3& TriMeshWorldTransform, const FReal Dt, FPBDCollisionConstraint& Constraint);
-		template void ConstructCapsuleTriMeshOneShotManifold<FTriangleMeshImplicitObject>(const FCapsule& Capsule, const FRigidTransform3& CapsuleWorldTransform, const FTriangleMeshImplicitObject& TriangleMesh, const FRigidTransform3& TriMeshWorldTransform, const FReal Dt, FPBDCollisionConstraint& Constraint);
-
-		template void ConstructConvexHeightFieldOneShotManifold<FConvex>(const FImplicitObject& Convex, const FRigidTransform3& ConvexTransform, const FHeightField& HeightField, const FRigidTransform3& HeightFieldTransform, const FReal Dt, FPBDCollisionConstraint& Constraint);
-		template void ConstructConvexHeightFieldOneShotManifold<FImplicitBox3>(const FImplicitObject& Convex, const FRigidTransform3& ConvexTransform, const FHeightField& HeightField, const FRigidTransform3& HeightFieldTransform, const FReal Dt, FPBDCollisionConstraint& Constraint);
 	}
 }
 
diff --git a/Engine/Source/Runtime/Experimental/Chaos/Private/Chaos/CollisionOneShotManifoldsMiscShapes.h b/Engine/Source/Runtime/Experimental/Chaos/Private/Chaos/CollisionOneShotManifoldsMiscShapes.h
--- a/Engine/Source/Runtime/Experimental/Chaos/Private/Chaos/CollisionOneShotManifoldsMiscShapes.h
+++ b/Engine/Source/Runtime/Experimental/Chaos/Private/Chaos/CollisionOneShotManifoldsMiscShapes.h
@@ -76,11 +76,10 @@ namespace Chaos
 			const FReal Dt,
 			FPBDCollisionConstraint& Constraint);
 
-		template <typename TriMeshType>
 		void ConstructCapsuleTriMeshOneShotManifold(
 			const FCapsule& Capsule, 
 			const FRigidTransform3& CapsuleWorldTransform, 
-			const TriMeshType& TriangleMesh, 
+			const FImplicitObject& TriangleMesh, 
 			const FRigidTransform3& TriMeshWorldTransform, 
 			const FReal Dt, 
 			FPBDCollisionConstraint& Constraint);
@@ -93,15 +92,6 @@ namespace Chaos
 			const FReal Dt, 
 			FPBDCollisionConstraint& Constraint);
 
-		template<typename ConvexType = FConvex>
-		void ConstructConvexHeightFieldOneShotManifold(
-			const FImplicitObject& Convex,
-			const FRigidTransform3& ConvexTransform,
-			const FHeightField& HeightField,
-			const FRigidTransform3& HeightFieldTransform,
-			const FReal Dt,
-			FPBDCollisionConstraint& Constraint);
-
 		/**
 		 * @brief Build the contact manifold between a Planar Convex shape and a TriMesh
 		 * @param Convex an ImplicitObject that must be a box, convex, or a wrapper around it. Asserts if not
@@ -114,5 +104,17 @@ namespace Chaos
 			const FRigidTransform3& TriangleMeshTransform, 
 			FPBDCollisionConstraint& Constraint);
 
+
+		/**
+		 * @brief Build the contact manifold between a Planar Convex shape and a HeightField
+		 * @param Convex an ImplicitObject that must be a box, convex, or a wrapper around it. Asserts if not
+		*/
+		void ConstructPlanarConvexHeightFieldOneShotManifold(
+			const FImplicitObject& Convex, 
+			const FRigidTransform3& ConvexTransform, 
+			const FHeightField& HeightField, 
+			const FRigidTransform3& TriangleMeshTransform, 
+			FPBDCollisionConstraint& Constraint);
+
 	}
 }
\ No newline at end of file
diff --git a/Engine/Source/Runtime/Experimental/Chaos/Private/Chaos/CollisionResolution.cpp b/Engine/Source/Runtime/Experimental/Chaos/Private/Chaos/CollisionResolution.cpp
--- a/Engine/Source/Runtime/Experimental/Chaos/Private/Chaos/CollisionResolution.cpp
+++ b/Engine/Source/Runtime/Experimental/Chaos/Private/Chaos/CollisionResolution.cpp
@@ -466,7 +466,7 @@ namespace Chaos
 		{
 			if(Constraint.GetUseManifold())
 			{
-				ConstructConvexHeightFieldOneShotManifold<FImplicitBox3>(A, ATransform, B, BTransform, Dt, Constraint);
+				ConstructPlanarConvexHeightFieldOneShotManifold(A, ATransform, B, BTransform, Constraint);
 			}
 			else
 			{
@@ -1331,7 +1331,7 @@ namespace Chaos
 			CONDITIONAL_SCOPE_CYCLE_COUNTER(STAT_Collisions_UpdateConvexHeightFieldConstraint, ConstraintsDetailedStats);
 			if (Constraint.GetUseManifold())
 			{
-				ConstructConvexHeightFieldOneShotManifold(A, ATransform, B, BTransform, Dt, Constraint);
+				ConstructPlanarConvexHeightFieldOneShotManifold(A, ATransform, B, BTransform, Constraint);
 			}
 			else
 			{
diff --git a/Engine/Source/Runtime/Experimental/Chaos/Private/Chaos/HeightField.cpp b/Engine/Source/Runtime/Experimental/Chaos/Private/Chaos/HeightField.cpp
--- a/Engine/Source/Runtime/Experimental/Chaos/Private/Chaos/HeightField.cpp
+++ b/Engine/Source/Runtime/Experimental/Chaos/Private/Chaos/HeightField.cpp
@@ -3,6 +3,7 @@
 #include "Chaos/HeightField.h"
 #include "Chaos/Collision/ContactPoint.h"
 #include "Chaos/Collision/ContactPointsMiscShapes.h"
+#include "Chaos/Collision/ContactTriangles.h"
 #include "Chaos/Collision/PBDCollisionConstraint.h"
 #include "Chaos/CollisionOneShotManifolds.h"
 #include "Chaos/Core.h"
@@ -28,14 +29,7 @@ namespace Chaos
 
 	extern FRealSingle Chaos_Collision_EdgePrunePlaneDistance;
 
-	extern bool bChaos_Collision_OneSidedHeightfield;
-
-	int32 bOneSidedHeightField = 1;
-	static FAutoConsoleVariableRef CVarOneSidedHeightField(TEXT("p.Chaos.OneSidedHeightField"), bOneSidedHeightField, TEXT("When enabled, extra steps will ensure that FHeightField::GJKContactPointImp never results in internal-facing contact data."));
-
-	int32 bOneSidedHeightfieldAlwaysSweep = 1;
-	static FAutoConsoleVariableRef CVarOneSidedHeightfieldAlwaysSweep(TEXT("p.Chaos.OneSidedHeightfieldAlwaysSweep"), bOneSidedHeightfieldAlwaysSweep, TEXT("When enabled, always use a sweep to ensure FHeightField::GJKContactPointImp never results \
-	in internal-facing contact data. Else, we only sweep if we detect an inward facing normal. Note that the sweep results can be inaccurate in some cases."));
+	extern bool bChaos_Collision_OneSidedHeightField;
 
 	class FHeightfieldRaycastVisitor
 	{
@@ -604,55 +598,6 @@ namespace Chaos
 		VectorRegister4Float Inflation3DSimd;
 	};
 
-	/**
-	 * Call the triangle visitor on all the triangles that are swept through by the InQueryBounds along the vector InDir. Bounds and direction
-	 * are in heightfield space.
-	 * 
-	 * The visitor signature must be:
-	 * 
-	 *		bool(int32 FaceIndex, const VectorRegister4Float& A, const VectorRegister4Float& B, const VectorRegister4Float& C, FRealSingle& CurrentDistance)
-	 * 
-	 * The visitor must return true if we should continue processing other triangles. It must update the current hit distance if the
-	 * passed in triangle is accepted and closer. The vertex positions passed to the triangle visitor are in heightfield space.
-	*/
-	template<typename TriangleVisitor>
-	class THeightfieldTriangleGenerator
-	{
-	public:
-		THeightfieldTriangleGenerator(const typename FHeightField::FDataType* InData, const TriangleVisitor& InTriangleVisitor)
-			: HfData(InData)
-			, TriangleVisitor(InTriangleVisitor)
-		{
-		}
-
-		bool VisitSweep(int32 CellIndex, FReal& CurrentLength)
-		{
-			const int32 VertexIndex = HfData->CellIndexToVertexIndex(CellIndex);
-			VectorRegister4Float Points[4];
-			FAABBVectorized CellBounds;
-			HfData->GetPointsAndBoundsScaledSimd(VertexIndex, Points, CellBounds);
-
-			const int32 TriIndex0 = CellIndex * 2 + 0;
-			const int32 TriIndex1 = CellIndex * 2 + 1;
-
-			if (!InTriangleVisitor(TriIndex0, Points[0], Points[1], Points[3], CurrentLength))
-			{
-				return false;
-			}
-
-			if (!InTriangleVisitor(TriIndex1, Points[0], Points[3], Points[2], CurrentLength))
-			{
-				return false;
-			}
-
-			return true;
-		}
-
-	private:
-		const typename FHeightField::FDataType* HfData;
-		TriangleVisitor InTriangleVisitor;
-	};
-
 
 	template<typename BufferType>
 	void BuildGeomData(TArrayView<BufferType> BufferView, TArrayView<uint8> MaterialIndexView, int32 NumRows, int32 NumCols, const FVec3& InScale, TUniqueFunction<FReal(const BufferType)> ToRealFunc, typename FHeightField::FDataType& OutData, FAABB3& OutBounds)
@@ -1671,278 +1616,6 @@ namespace Chaos
 		return false;
 	}
 
-	template <typename GeomType>
-	bool FHeightField::ContactManifoldNonPlanarConvexImp(const GeomType& QueryGeom, const FRigidTransform3& QueryTM, const FReal Thickness, TArray<FContactPoint>& ContactPoints) const
-	{
-		auto OverlapTriangle = [&](const FVec3& A, const FVec3& B, const FVec3& C,
-			FPBDCollisionConstraint& Constraint)
-		{
-			const FVec3 AB = B - A;
-			const FVec3 AC = C - A;
-
-			const FVec3 Offset = FVec3::CrossProduct(AB, AC);
-
-			FTriangle TriangleConvex(A, B, C);
-			Collisions::ConstructConvexConvexOneShotManifold(QueryGeom, QueryTM, TriangleConvex, FRigidTransform3::Identity, 0, Constraint);
-		};
-
-		auto InsertSorted = [&](const FContactPoint& ContactPoint)
-		{
-			int32 PointIndex = 0;
-			bool done = false;
-			const FReal SamePointErrorMarginSqr = 0.01f;
-
-			int32 ContactPointsNum = ContactPoints.Num();
-			for (; PointIndex < ContactPointsNum; PointIndex++)
-			{
-				FVec3 DiffVector = ContactPoint.ShapeContactPoints[1] - ContactPoints[PointIndex].ShapeContactPoints[1];
-				// Check if point is the same (or close)
-				if (DiffVector.SizeSquared() < SamePointErrorMarginSqr)
-				{
-					done = true;
-					break;
-				}
-
-				if (ContactPoint.Phi < ContactPoints[PointIndex].Phi)
-				{
-					ContactPoints.Insert(ContactPoint, PointIndex);
-					done = true;
-					break;
-				}
-			}
-
-			if (!done)
-			{
-				ContactPoints.Add(ContactPoint);
-			}
-		};
-
-		bool bResult = false;
-		FAABB3 QueryBounds = QueryGeom.BoundingBox();
-		QueryBounds.Thicken(Thickness);
-		QueryBounds = QueryBounds.TransformedAABB(QueryTM);
-
-		FBounds2D FlatQueryBounds;
-		FlatQueryBounds.Min = FVec2(QueryBounds.Min()[0], QueryBounds.Min()[1]);
-		FlatQueryBounds.Max = FVec2(QueryBounds.Max()[0], QueryBounds.Max()[1]);
-
-		TArray<TVec2<int32>> Intersections;
-		FVec3 Points[4];
-
-		GetGridIntersections(FlatQueryBounds, Intersections);
-
-		FReal MinContactPhi = FLT_MAX;
-		FVec3 LocalContactLocation, LocalContactNormal;
-		for (const TVec2<int32>& Cell : Intersections)
-		{
-			const int32 SingleIndex = Cell[1] * GeomData.NumCols + Cell[0];
-			const int32 CellIndex = Cell[1] * (GeomData.NumCols - 1) + Cell[0];
-
-			// Check for holes and skip checking if we'll never collide
-			if (GeomData.MaterialIndices.IsValidIndex(CellIndex) && GeomData.MaterialIndices[CellIndex] == TNumericLimits<uint8>::Max())
-			{
-				continue;
-			}
-
-			// @todo(chaos): we should not be creating constraints just for collecting contacts...
-			FPBDCollisionConstraint Constraint = FPBDCollisionConstraint::MakeTriangle(&QueryGeom);
-
-			// The triangle is solid so proceed to test it
-			FAABB3 CellBounds;
-			GeomData.GetPointsAndBoundsScaled(SingleIndex, Points, CellBounds);
-			if (CellBounds.Intersects(QueryBounds))
-			{
-				// First Triangle
-				{
-					Constraint.ResetManifold();
-					Constraint.GetGJKWarmStartData().Reset();
-					OverlapTriangle(Points[0], Points[1], Points[3], Constraint);
-					for (FManifoldPoint& ManifoldPoint : Constraint.GetManifoldPoints())
-					{
-						ManifoldPoint.ContactPoint.FaceIndex = CellIndex * 2;
-						InsertSorted(ManifoldPoint.ContactPoint);
-					}
-				}
-				// Second Triangle
-				{
-					Constraint.ResetManifold();
-					Constraint.GetGJKWarmStartData().Reset();
-					OverlapTriangle(Points[0], Points[3], Points[2], Constraint);
-					for (FManifoldPoint& ManifoldPoint : Constraint.GetManifoldPoints())
-					{
-						ManifoldPoint.ContactPoint.FaceIndex = CellIndex * 2 + 1;
-						InsertSorted(ManifoldPoint.ContactPoint);
-					}
-				}
-			}
-		}
-
-		// Remove edge contacts that are "hidden" by face contacts
-		// EdgePruneDistance should be some fraction of the convex margin...
-		const FReal EdgePruneDistance = Chaos_Collision_EdgePrunePlaneDistance;
-		Collisions::PruneEdgeContactPointsOrdered(ContactPoints, EdgePruneDistance);
-
-		// Remove all points (except for the deepest one, and ones with phis similar to it)
-		const FReal CullMargin = 0.1f;
-		int32 NewContactPointCount = ContactPoints.Num() > 0 ? 1 : 0;
-		for (int32 Index = 1; Index < ContactPoints.Num(); Index++)
-		{
-			if (ContactPoints[Index].Phi < 0 || ContactPoints[Index].Phi - ContactPoints[0].Phi < CullMargin)
-			{
-				NewContactPointCount++;
-			}
-			else
-			{
-				break;
-			}
-		}
-		ContactPoints.SetNum(NewContactPointCount, false);
-
-		// Reduce to only 4 contact points from here
-		Collisions::ReduceManifoldContactPointsTriangeMesh(ContactPoints);
-
-		return true;
-	}
-
-	template <typename GeomType>
-	bool FHeightField::ContactManifoldPlanarConvexImp(const GeomType& QueryGeom, const FRigidTransform3& QueryTM, const FReal Thickness, TArray<FContactPoint>& ContactPoints) const
-	{
-		// @todo(chaos): should be an option on the Heightfield?
-		const bool bOneSidedCollision = bChaos_Collision_OneSidedHeightfield;
-
-		auto OverlapTriangle = [&](const FVec3& A, const FVec3& B, const FVec3& C, TCArray<FContactPoint, 4>& TriangleContactPoints)
-		{
-			// Create triangle in query space
-			FTriangle TriangleConvex(
-				QueryTM.InverseTransformPositionNoScale(A), 
-				QueryTM.InverseTransformPositionNoScale(B), 
-				QueryTM.InverseTransformPositionNoScale(C));
-
-			Collisions::ConstructPlanarConvexTriangleOneShotManifold(QueryGeom, TriangleConvex, bOneSidedCollision, Thickness, TriangleContactPoints);
-
-			// Convert back to shape-local space
-			for (FContactPoint& ContactPoint : TriangleContactPoints)
-			{
-				ContactPoint.ShapeContactPoints[1] = QueryTM.TransformPositionNoScale(ContactPoint.ShapeContactPoints[1]);
-				ContactPoint.ShapeContactNormal = QueryTM.TransformVectorNoScale(ContactPoint.ShapeContactNormal);
-			}
-		};
-
-		auto InsertSorted = [&](const FContactPoint& ContactPoint)
-		{
-			int32 PointIndex = 0;
-			bool done = false;
-			const FReal SamePointErrorMarginSqr = 0.01f;
-			
-			int32 ContactPointsNum = ContactPoints.Num();
-			for( ; PointIndex < ContactPointsNum; PointIndex++ )
-			{
-				FVec3 DiffVector = ContactPoint.ShapeContactPoints[1] - ContactPoints[PointIndex].ShapeContactPoints[1];
-				// Check if point is the same (or close)
-				if (DiffVector.SizeSquared() < SamePointErrorMarginSqr)
-				{
-					done = true;
-					break;
-				}
-
-				if (ContactPoint.Phi < ContactPoints[PointIndex].Phi)
-				{
-					ContactPoints.Insert(ContactPoint, PointIndex);
-					done = true;
-					break;
-				}
-			}
-
-			if (!done)
-			{
-				ContactPoints.Add(ContactPoint);
-			}
-		};		
-		
-		bool bResult = false;
-		FAABB3 QueryBounds = QueryGeom.BoundingBox();
-		QueryBounds.Thicken(Thickness);
-		QueryBounds = QueryBounds.TransformedAABB(QueryTM);
-
-		FBounds2D FlatQueryBounds;
-		FlatQueryBounds.Min = FVec2(QueryBounds.Min()[0], QueryBounds.Min()[1]);
-		FlatQueryBounds.Max = FVec2(QueryBounds.Max()[0], QueryBounds.Max()[1]);
-
-		TArray<TVec2<int32>> Intersections;
-		FVec3 Points[4];
-
-		GetGridIntersections(FlatQueryBounds, Intersections);
-
-		FReal MinContactPhi = FLT_MAX;
-		FVec3 LocalContactLocation, LocalContactNormal;
-		TCArray<FContactPoint, 4> TriangleContactPoints;
-
-		for (const TVec2<int32>& Cell : Intersections)
-		{
-			const int32 SingleIndex = Cell[1] * GeomData.NumCols + Cell[0];
-			const int32 CellIndex = Cell[1] * (GeomData.NumCols - 1) + Cell[0];
-
-			// Check for holes and skip checking if we'll never collide
-			if (GeomData.MaterialIndices.IsValidIndex(CellIndex) && GeomData.MaterialIndices[CellIndex] == TNumericLimits<uint8>::Max())
-			{
-				continue;
-			}
-
-			// The triangle is solid so proceed to test it
-			FAABB3 CellBounds;
-			GeomData.GetPointsAndBoundsScaled(SingleIndex, Points, CellBounds);
-			if (CellBounds.Intersects(QueryBounds))
-			{
-				// First Triangle
-				{
-					TriangleContactPoints.Reset();
-					OverlapTriangle(Points[0], Points[1], Points[3], TriangleContactPoints);
-					for (FContactPoint& ContactPoint : TriangleContactPoints)
-					{
-						ContactPoint.FaceIndex = CellIndex * 2;
-						InsertSorted(ContactPoint);
-					}
-				}
-				// Second Triangle
-				{
-					TriangleContactPoints.Reset();
-					OverlapTriangle(Points[0], Points[3], Points[2], TriangleContactPoints);
-					for (FContactPoint& ContactPoint : TriangleContactPoints)
-					{
-						ContactPoint.FaceIndex = CellIndex * 2 + 1;
-						InsertSorted(ContactPoint);
-					}
-				}
-			}
-		}
-
-		// Remove edge contacts that are "hidden" by face contacts
-		// EdgePruneDistance should be some fraction of the convex margin...
-		const FReal EdgePruneDistance = Chaos_Collision_EdgePrunePlaneDistance;
-		Collisions::PruneEdgeContactPointsOrdered(ContactPoints, EdgePruneDistance);
-
-		// Remove all points (except for the deepest one, and ones with phis similar to it)
-		const FReal CullMargin = 0.1f;
-		int32 NewContactPointCount = ContactPoints.Num() > 0 ? 1 : 0;
-		for (int32 Index = 1; Index < ContactPoints.Num(); Index++)
-		{
-			if (ContactPoints[Index].Phi < 0 || ContactPoints[Index].Phi - ContactPoints[0].Phi < CullMargin)
-			{
-				NewContactPointCount++;
-			}
-			else
-			{
-				break;
-			}
-		}
-		ContactPoints.SetNum(NewContactPointCount, false);
-
-		// Reduce to only 4 contact points from here
-		Collisions::ReduceManifoldContactPointsTriangeMesh(ContactPoints);
-
-		return true;
-	}
-
 	template <typename GeomType>
 	bool FHeightField::GJKContactPointImp(const GeomType& QueryGeom, const FRigidTransform3& QueryTM, const FReal Thickness,
 		FVec3& ContactLocation, FVec3& ContactNormal, FReal& ContactPhi) const
@@ -1961,64 +1634,14 @@ namespace Chaos
 				MakeVectorRegisterFloatFromDouble(MakeVectorRegister(B.X, B.Y, B.Z, 0.0f)),
 				MakeVectorRegisterFloatFromDouble(MakeVectorRegister(C.X, C.Y, C.Z, 0.0f)));
 
-
 			FReal Penetration;
 			FVec3 ClosestA, ClosestB, Normal;
-
-			auto SweepAgainstTriangle = [&]() -> bool
-			{
-				//
-				// BUG: This does not detect collisions when we specify a cull distance. It is as if
-				// Thickness is always zero...
-				//
-
-				const FAABB3 Bounds = QueryGeom.BoundingBox();
-				const FReal ApproximateSizeOfObject = Bounds.Extents()[Bounds.LargestAxis()];
-				const FReal ApproximateDistToObject = FVec3::DistSquared(QueryTM.GetLocation(), A);
-				const FReal SweepLength = ApproximateSizeOfObject + ApproximateDistToObject;
-				const FRigidTransform3 QueryStartTM(QueryTM.GetLocation() + TriNormal * SweepLength, QueryTM.GetRotation());
-				if (GJKRaycast2(TriangleConvexReg, QueryGeom, QueryStartTM, -TriNormal, SweepLength, Penetration, ClosestA, Normal, (FReal)0., true))
-				{
-					LocalContactPhi = Penetration - SweepLength;
-					LocalContactLocation = ClosestA + LocalContactPhi * Normal;
-					LocalContactNormal = TriNormal;
-					return true;
-				}
-				return false;
-			};
-
-			if (bOneSidedHeightField)
-			{
-				// HACK:
-				// The regular penetration calculation vs a triangle may result in inward facing normals.
-				// To protect against this, we sweep against the triangle from a distance to ensure an outward
-				// facing normal and MTD.
-
-				if (bOneSidedHeightfieldAlwaysSweep)
-				{
-					return SweepAgainstTriangle();
-				}
-				else
-				{
-					int32 ClosestVertexIndexA, ClosestVertexIndexB;
-					if (GJKPenetration(TriangleConvex, QueryGeom, QueryTM, Penetration, ClosestA, ClosestB, Normal, ClosestVertexIndexA, ClosestVertexIndexB, (FReal)0))
-					{
-						if (FVec3::DotProduct(TriNormal, Normal) < 0)
-						{
-							return SweepAgainstTriangle();
-						}
-						LocalContactLocation = ClosestB;
-						LocalContactNormal = Normal;
-						LocalContactPhi = -Penetration;
-						return true;
-					}
-				}
-				
-			}
-			else
+			int32 ClosestVertexIndexA, ClosestVertexIndexB;
+			if (GJKPenetration(TriangleConvex, QueryGeom, QueryTM, Penetration, ClosestA, ClosestB, Normal, ClosestVertexIndexA, ClosestVertexIndexB, (FReal)0))
 			{
-				int32 ClosestVertexIndexA, ClosestVertexIndexB;
-				if (GJKPenetration(TriangleConvex, QueryGeom, QueryTM, Penetration, ClosestA, ClosestB, Normal, ClosestVertexIndexA, ClosestVertexIndexB, (FReal)0))
+				const bool bOneSidedCollision = bChaos_Collision_OneSidedHeightField;
+				const bool bSkipOneSided = (bOneSidedCollision && (FVec3::DotProduct(TriNormal, Normal) < 0));
+				if (!bSkipOneSided)
 				{
 					LocalContactLocation = ClosestB;
 					LocalContactNormal = Normal; 
@@ -2126,41 +1749,6 @@ namespace Chaos
 		return GJKContactPointImp(QueryGeom, QueryTM, Thickness, ContactLocation, ContactNormal, ContactPhi);
 	}
 
-	bool FHeightField::ContactManifold(const TBox<FReal, 3>& QueryGeom, const FRigidTransform3& QueryTM, const FReal Thickness, TArray<FContactPoint>& ContactPoints) const
-	{
-		return ContactManifoldPlanarConvexImp(QueryGeom, QueryTM, Thickness, ContactPoints);
-	}
-
-	/*bool FHeightField::ContactManifold(const TSphere<FReal, 3>& QueryGeom, const FRigidTransform3& QueryTM, const FReal Thickness, TArray<FContactPoint>& ContactPoints) const
-	{
-		return ContactManifoldImp(QueryGeom, QueryTM, Thickness, ContactPoints);
-	}*/
-
-	bool FHeightField::ContactManifold(const FCapsule& QueryGeom, const FRigidTransform3& QueryTM, const FReal Thickness, TArray<FContactPoint>& ContactPoints) const
-	{
-		return ContactManifoldNonPlanarConvexImp(QueryGeom, QueryTM, Thickness, ContactPoints);
-	}
-
-	bool FHeightField::ContactManifold(const FConvex& QueryGeom, const FRigidTransform3& QueryTM, const FReal Thickness, TArray<FContactPoint>& ContactPoints) const
-	{
-		return ContactManifoldPlanarConvexImp(QueryGeom, QueryTM, Thickness, ContactPoints);
-	}
-
-	bool FHeightField::ContactManifold(const TImplicitObjectScaled<TBox<FReal, 3>>& QueryGeom, const FRigidTransform3& QueryTM, const FReal Thickness, TArray<FContactPoint>& ContactPoints) const
-	{
-		return ContactManifoldPlanarConvexImp(QueryGeom, QueryTM, Thickness, ContactPoints);
-	}
-
-	bool FHeightField::ContactManifold(const TImplicitObjectScaled<FCapsule>& QueryGeom, const FRigidTransform3& QueryTM, const FReal Thickness, TArray<FContactPoint>& ContactPoints) const
-	{
-		return ContactManifoldNonPlanarConvexImp(QueryGeom, QueryTM, Thickness, ContactPoints);
-	}
-
-	bool FHeightField::ContactManifold(const TImplicitObjectScaled<FConvex>& QueryGeom, const FRigidTransform3& QueryTM, const FReal Thickness, TArray<FContactPoint>& ContactPoints) const
-	{
-		return ContactManifoldPlanarConvexImp(QueryGeom, QueryTM, Thickness, ContactPoints);
-	}
-
 	template <typename QueryGeomType>
 	bool FHeightField::OverlapGeomImp(const QueryGeomType& QueryGeom, const FRigidTransform3& QueryTM, const FReal Thickness, FMTDInfo* OutMTD) const
 	{
@@ -2462,28 +2050,6 @@ namespace Chaos
 		return SweepGeomCCDImp(QueryGeom, StartTM, Dir, Length, IgnorePenetration, TargetPenetration, OutTOI, OutPhi, OutPosition, OutNormal, OutFaceIndex);
 	}
 
-	void FHeightField::VisitTriangles(const FAABB3& QueryBounds, const TFunction<void(const FTriangle& Triangle)>& Visitor) const
-	{
-		FBounds2D FlatQueryBounds;
-		FlatQueryBounds.Min = TVector<FReal, 2>(QueryBounds.Min()[0], QueryBounds.Min()[1]);
-		FlatQueryBounds.Max = TVector<FReal, 2>(QueryBounds.Max()[0], QueryBounds.Max()[1]);
-
-		TArray<TVector<int32, 2>> Intersections;
-		FVec3 Points[4];
-
-		GetGridIntersections(FlatQueryBounds, Intersections);
-
-		bool bOverlaps = false;
-		for (const TVector<int32, 2>&Cell : Intersections)
-		{
-			const int32 SingleIndex = Cell[1] * (GeomData.NumCols) + Cell[0];
-			GeomData.GetPointsScaled(SingleIndex, Points);
-
-			Visitor(FTriangle(Points[0], Points[1], Points[3]));
-			Visitor(FTriangle(Points[0], Points[3], Points[2]));
-		}
-	}
-
 	int32 FHeightField::FindMostOpposingFace(const FVec3& Position, const FVec3& UnitDir, int32 HintFaceIndex, FReal SearchDist) const
 	{
 		const FReal SearchDist2 = SearchDist * SearchDist;
diff --git a/Engine/Source/Runtime/Experimental/Chaos/Private/Chaos/TriangleMeshImplicitObject.cpp b/Engine/Source/Runtime/Experimental/Chaos/Private/Chaos/TriangleMeshImplicitObject.cpp
--- a/Engine/Source/Runtime/Experimental/Chaos/Private/Chaos/TriangleMeshImplicitObject.cpp
+++ b/Engine/Source/Runtime/Experimental/Chaos/Private/Chaos/TriangleMeshImplicitObject.cpp
@@ -505,125 +505,6 @@ bool FTriangleMeshImplicitObject::Raycast(const FVec3& StartPoint, const FVec3&
 	}
 }
 
-template <typename GeomType>
-bool FTriangleMeshImplicitObject::ContactManifoldImp(const GeomType& QueryGeom, const FRigidTransform3& QueryTM, const FReal WorldThickness, TArray<FContactPoint>& ContactPoints, FVec3 TriMeshScale) const
-{
-	ensure(TriMeshScale != FVec3(0.0f));
-	bool bResult = false;
-
-	const GeomType& WorldScaleGeom = ScaleGeomIntoWorldHelper(QueryGeom, TriMeshScale);
-	const FVec3 InvTriMeshScale = FReal(1) / TriMeshScale;
-
-	// IMPORTANT QueryTM comes with a invscaled translation so we need a version of the TM with world space translation to properly compute the bounds
-	FRigidTransform3 TriMeshToGeomNoScale{ QueryTM };
-	TriMeshToGeomNoScale.SetTranslation(TriMeshToGeomNoScale.GetTranslation() * TriMeshScale);
-	// NOTE: BVH test is done in tri-mesh local space (whereas collision detection is done in world space becaused you can't non-uniformly scale all shapes)
-	FAABB3 QueryBounds = WorldScaleGeom.CalculateTransformedBounds(TriMeshToGeomNoScale);
-	QueryBounds.ThickenSymmetrically(FVec3(WorldThickness));
-	QueryBounds.ScaleWithNegative(InvTriMeshScale);
-
-	TRigidTransform<FReal, 3> WorldScaleQueryTM;
-	ScaleTransformHelper(TriMeshScale, QueryTM, WorldScaleQueryTM);
-
-	auto InsertSorted = [&](const FContactPoint& ContactPoint)
-	{
-		int32 PointIndex = 0;
-		bool done = false;
-		const FReal ErrorMarginSqr = 0.01f;
-
-		int32 ContactPointsNum = ContactPoints.Num();
-		for (; PointIndex < ContactPointsNum; PointIndex++)
-		{
-			FVec3 DiffVector = ContactPoint.ShapeContactPoints[1] - ContactPoints[PointIndex].ShapeContactPoints[1];
-			// Check if point is the same (or close)
-			if (DiffVector.SizeSquared() < ErrorMarginSqr)
-			{
-				done = true;
-				break;
-			}
-
-			if (ContactPoint.Phi < ContactPoints[PointIndex].Phi)
-			{
-				ContactPoints.Insert(ContactPoint, PointIndex);
-				done = true;
-				break;
-			}
-		}
-
-		if (!done)
-		{
-			ContactPoints.Add(ContactPoint);
-		}
-	};
-
-	auto OverlapTriangle = [&](const FVec3& A, const FVec3& B, const FVec3& C,
-		FPBDCollisionConstraint& Constraint)
-	{
-		FTriangle TriangleConvex(A, B, C);
-		// make sure the constraint does not contain any stale data  ( it is shared between trinagles )
-		// @todo(chaos) : we should eventually not use a constraint here and just get a list of contact points
-		Constraint.ResetManifold();
-		Constraint.GetGJKWarmStartData().Reset();
-		Collisions::ConstructConvexConvexOneShotManifold(WorldScaleGeom, WorldScaleQueryTM, TriangleConvex, FRigidTransform3::Identity, 0, Constraint);
-	};
-
-	auto LambdaHelper = [&](const auto& Elements)
-	{
-		FReal LocalContactPhi = FLT_MAX;
-		FVec3 LocalContactLocation, LocalContactNormal;
-
-		const TArray<int32> PotentialIntersections = FastBVH.FindAllIntersections(QueryBounds);
-
-		// MakeTriangleHelper get rid of the scale wrapper if necessary as MakeTriangle will try to infer properties from it 
-		FPBDCollisionConstraint Constraint = MakeTriangleHelper(WorldScaleGeom);
-
-		for (int32 TriIdx : PotentialIntersections)
-		{
-			FVec3 A, B, C;
-			TriangleMeshTransformVertsHelper(TriMeshScale, TriIdx, MParticles, Elements, A, B, C);
-			OverlapTriangle(A, B, C, Constraint);
-			for(FManifoldPoint& ManifoldPoint : Constraint.GetManifoldPoints())
-			{
-				ManifoldPoint.ContactPoint.FaceIndex = TriIdx;
-				InsertSorted(ManifoldPoint.ContactPoint);
-			}
-
-		}
-
-		// Remove edge contacts that are "hidden" by face contacts
-		// EdgePruneDistance should be some fraction of the convex margin...
-		const FReal EdgePruneDistance = Chaos_Collision_EdgePrunePlaneDistance;
-		Collisions::PruneEdgeContactPointsOrdered(ContactPoints, EdgePruneDistance);
-
-		// Remove all points (except for the deepest one, and ones with phis similar to it)
-		const FReal CullMargin = 0.1f;
-		int32 NewContactPointCount = ContactPoints.Num() > 0 ? 1 : 0;
-		for (int32 Index = 1; Index < ContactPoints.Num(); Index++)
-		{
-			if (ContactPoints[Index].Phi < 0 || ContactPoints[Index].Phi - ContactPoints[0].Phi < CullMargin)
-			{
-				NewContactPointCount++;
-			}
-			else
-			{
-				break;
-			}
-		}
-		ContactPoints.SetNum(NewContactPointCount, false);
-
-		// Reduce to only 4 contact points from here
-		Collisions::ReduceManifoldContactPointsTriangeMesh(ContactPoints);
-
-		return true;
-	};
-
-	if (MElements.RequiresLargeIndices())
-	{
-		return LambdaHelper(MElements.GetLargeIndexBuffer());
-	}
-	return LambdaHelper(MElements.GetSmallIndexBuffer());
-}
-
 template <typename QueryGeomType>
 bool FTriangleMeshImplicitObject::GJKContactPointImp(const QueryGeomType& QueryGeom, const FRigidTransform3& QueryTM, const FReal WorldThickness, FVec3& Location, FVec3& Normal, FReal& OutContactPhi, FVec3 TriMeshScale) const
 {
@@ -738,41 +619,6 @@ bool FTriangleMeshImplicitObject::GJKContactPoint(const TImplicitObjectScaled< F
 	return GJKContactPointImp(QueryGeom, QueryTM, WorldThickness, Location, Normal, ContactPhi, TriMeshScale);
 }
 
-bool FTriangleMeshImplicitObject::ContactManifold(const TBox<FReal, 3>& QueryGeom, const FRigidTransform3& QueryTM, const FReal Thickness, TArray<FContactPoint>& ContactPoints) const
-{
-	return ContactManifoldImp(QueryGeom, QueryTM, Thickness, ContactPoints, FVec3(1));
-}
-
-bool FTriangleMeshImplicitObject::ContactManifold(const FCapsule& QueryGeom, const FRigidTransform3& QueryTM, const FReal Thickness, TArray<FContactPoint>& ContactPoints) const
-{
-	return ContactManifoldImp(QueryGeom, QueryTM, Thickness, ContactPoints, FVec3(1));
-}
-
-bool FTriangleMeshImplicitObject::ContactManifold(const FConvex& QueryGeom, const FRigidTransform3& QueryTM, const FReal Thickness, TArray<FContactPoint>& ContactPoints) const
-{
-	return ContactManifoldImp(QueryGeom, QueryTM, Thickness, ContactPoints, FVec3(1));
-}
-
-bool FTriangleMeshImplicitObject::ContactManifold(const TImplicitObjectScaled<TBox<FReal, 3>>& QueryGeom, const FRigidTransform3& QueryTM, const FReal Thickness, TArray<FContactPoint>& ContactPoints, FVec3 TriMeshScale) const
-{
-	return ContactManifoldImp(QueryGeom, QueryTM, Thickness, ContactPoints, TriMeshScale);
-}
-
-/*bool FTriangleMeshImplicitObject::ContactManifold(const TImplicitObjectScaled<TSphere<FReal, 3>>& QueryGeom, const FRigidTransform3& QueryTM, const FReal Thickness, TArray<FContactPoint>& ContactPoints, FVec3 TriMeshScale) const
-{
-	return ContactManifoldImp(QueryGeom, QueryTM, Thickness, ContactPoints, TriMeshScale);
-}*/
-
-bool FTriangleMeshImplicitObject::ContactManifold(const TImplicitObjectScaled<FCapsule >& QueryGeom, const FRigidTransform3& QueryTM, const FReal Thickness, TArray<FContactPoint>& ContactPoints, FVec3 TriMeshScale) const
-{
-	return ContactManifoldImp(QueryGeom, QueryTM, Thickness, ContactPoints, TriMeshScale);
-}
-
-bool FTriangleMeshImplicitObject::ContactManifold(const TImplicitObjectScaled<FConvex >& QueryGeom, const FRigidTransform3& QueryTM, const FReal Thickness, TArray<FContactPoint>& ContactPoints, FVec3 TriMeshScale) const
-{
-	return ContactManifoldImp(QueryGeom, QueryTM, Thickness, ContactPoints, TriMeshScale);
-}
-
 int32 FTriangleMeshImplicitObject::GetExternalFaceIndexFromInternal(int32 InternalFaceIndex) const
 {
 	if (InternalFaceIndex > -1 && ExternalFaceIndexMap.Get())
@@ -846,51 +692,6 @@ bool FTriangleMeshImplicitObject::Overlap(const FVec3& Point, const FReal Thickn
 	}
 }
 
-void FTriangleMeshImplicitObject::VisitTriangles(const FAABB3& QueryBounds, const TFunction<void(const FTriangle& Triangle)>& Visitor) const
-{
-	const TArray<int32> PotentialIntersections = FastBVH.FindAllIntersections(QueryBounds);
-
-	auto TriangleProducer = [&](const auto& Elements)
-	{
-		for (int32 TriIdx : PotentialIntersections)
-		{
-			TVec3<FReal> A, B, C;
-			TriangleMeshTransformVertsHelper(FVec3(1), TriIdx, MParticles, Elements, A, B, C);
-
-			Visitor(FTriangle(A, B, C));
-		}
-	};
-
-	if (MElements.RequiresLargeIndices())
-	{
-		return TriangleProducer(MElements.GetLargeIndexBuffer());
-	}
-	else
-	{
-		return TriangleProducer(MElements.GetSmallIndexBuffer());
-	}
-}
-
-void FTriangleMeshImplicitObject::VisitTriangle(const int32 TriangleIndex, const TFunction<void(const FTriangle& Triangle)>& Visitor) const
-{
-	const auto TriangleProducer = [&](int32 TriIdx, const auto& Elements)
-	{
-		TVec3<FReal> A, B, C;
-		TriangleMeshTransformVertsHelper(FVec3(1), TriIdx, MParticles, Elements, A, B, C);
-
-		Visitor(FTriangle(A, B, C));
-	};
-
-	if (MElements.RequiresLargeIndices())
-	{
-		return TriangleProducer(TriangleIndex, MElements.GetLargeIndexBuffer());
-	}
-	else
-	{
-		return TriangleProducer(TriangleIndex, MElements.GetSmallIndexBuffer());
-	}
-}
-
 
 template <typename QueryGeomType>
 bool FTriangleMeshImplicitObject::OverlapGeomImp(const QueryGeomType& QueryGeom, const FRigidTransform3& QueryTM, const FReal Thickness, FMTDInfo* OutMTD, FVec3 TriMeshScale) const
diff --git a/Engine/Source/Runtime/Experimental/Chaos/Private/PBDRigidsSolver.cpp b/Engine/Source/Runtime/Experimental/Chaos/Private/PBDRigidsSolver.cpp
--- a/Engine/Source/Runtime/Experimental/Chaos/Private/PBDRigidsSolver.cpp
+++ b/Engine/Source/Runtime/Experimental/Chaos/Private/PBDRigidsSolver.cpp
@@ -83,6 +83,7 @@ namespace Chaos
 		int32 ChaosSolverDebugDrawSpatialAccelerationStructureShowNodes = 0;
 		int32 ChaosSolverDebugDrawSuspensionConstraints = 0;
 		int32 ChaosSolverDrawClusterConstraints = 0;
+		int32 ChaosSolverDebugDrawMeshContacts = 0;
 		DebugDraw::FChaosDebugDrawJointFeatures ChaosSolverDrawJointFeatures = DebugDraw::FChaosDebugDrawJointFeatures::MakeDefault();
 		FAutoConsoleVariableRef CVarChaosSolverDrawShapes(TEXT("p.Chaos.Solver.DebugDrawShapes"), ChaosSolverDebugDrawShapes, TEXT("Draw Shapes (0 = never; 1 = end of frame)."));
 		FAutoConsoleVariableRef CVarChaosSolverDrawCollisions(TEXT("p.Chaos.Solver.DebugDrawCollisions"), ChaosSolverDebugDrawCollisions, TEXT("Draw Collisions (0 = never; 1 = end of frame)."));
@@ -109,6 +110,7 @@ namespace Chaos
 		FAutoConsoleVariableRef CVarChaosSolverDrawJointFeaturesIsland(TEXT("p.Chaos.Solver.DebugDraw.JointFeatures.Island"), ChaosSolverDrawJointFeatures.bIsland, TEXT("Joint features mask (see FDebugDrawJointFeatures)."));
 		FAutoConsoleVariableRef CVarChaosSolverDrawSuspensionConstraints(TEXT("p.Chaos.Solver.DebugDrawSuspension"), ChaosSolverDebugDrawSuspensionConstraints, TEXT("Draw Suspension (0 = never; 1 = end of frame)."));
 		FAutoConsoleVariableRef CVarChaosSolverDrawClusterConstraints(TEXT("p.Chaos.Solver.DebugDraw.Cluster.Constraints"), ChaosSolverDrawClusterConstraints, TEXT("Draw Active Cluster Constraints (0 = never; 1 = end of frame)."));
+		FAutoConsoleVariableRef CVarChaosSolverDrawMeshContacts(TEXT("p.Chaos.Solver.DebugDrawMeshContacts"), ChaosSolverDebugDrawMeshContacts, TEXT("Draw Mesh contacts"));
 
 
 		DebugDraw::FChaosDebugDrawSettings ChaosSolverDebugDebugDrawSettings(
diff --git a/Engine/Source/Runtime/Experimental/Chaos/Public/Chaos/Collision/ContactPoint.h b/Engine/Source/Runtime/Experimental/Chaos/Public/Chaos/Collision/ContactPoint.h
--- a/Engine/Source/Runtime/Experimental/Chaos/Public/Chaos/Collision/ContactPoint.h
+++ b/Engine/Source/Runtime/Experimental/Chaos/Public/Chaos/Collision/ContactPoint.h
@@ -12,13 +12,17 @@ namespace Chaos
 {
 	/**
 	 * @brief Used in FContactPoint to indicate whether the contact is vertex-plane, edge-edge, etc
+	 * 
+	 * @note the order here is the order of preference in the solver. I.e., we like to solve Plane contacts before edge contacts before vertex contacts.
+	 * This is most impotant for collisions against triangle meshes (or any concave shape) where the second shape is always the triangle, and so a PlaneVertex collision 
+	 * counts as a vertex collision.
 	*/
 	enum class EContactPointType : int8
 	{
 		Unknown,
 		VertexPlane,
-		PlaneVertex,
 		EdgeEdge,
+		PlaneVertex,
 		VertexVertex,
 	};
 
diff --git a/Engine/Source/Runtime/Experimental/Chaos/Public/Chaos/Collision/ContactPointsMiscShapes.h b/Engine/Source/Runtime/Experimental/Chaos/Public/Chaos/Collision/ContactPointsMiscShapes.h
--- a/Engine/Source/Runtime/Experimental/Chaos/Public/Chaos/Collision/ContactPointsMiscShapes.h
+++ b/Engine/Source/Runtime/Experimental/Chaos/Public/Chaos/Collision/ContactPointsMiscShapes.h
@@ -40,9 +40,6 @@
 	template <typename GeometryA, typename GeometryB>
 	FContactPoint GJKImplicitContactPoint(const FImplicitObject& A, const FRigidTransform3& ATransform, const GeometryB& B, const FRigidTransform3& BTransform, const FReal CullDistance, const FReal ShapePadding);
 
-	template <typename GeometryA, typename GeometryB>
-	void GJKImplicitManifold(const FImplicitObject& A, const FRigidTransform3& ATransform, const GeometryB& B, const FRigidTransform3& BTransform, const FReal CullDistance, const FReal ShapePadding, TArray<FContactPoint>& ContactPoints);
-
 	FContactPoint SphereSphereContactPoint(const TSphere<FReal, 3>& Sphere1, const FRigidTransform3& Sphere1Transform, const TSphere<FReal, 3>& Sphere2, const FRigidTransform3& Sphere2Transform, const FReal CullDistance, const FReal ShapePadding);
 	FContactPoint SpherePlaneContactPoint(const TSphere<FReal, 3>& Sphere, const FRigidTransform3& SphereTransform, const TPlane<FReal, 3>& Plane, const FRigidTransform3& PlaneTransform, const FReal ShapePadding);
 	FContactPoint SphereBoxContactPoint(const TSphere<FReal, 3>& Sphere, const FRigidTransform3& SphereTransform, const FImplicitBox3& Box, const FRigidTransform3& BoxTransform, const FReal ShapePadding);
diff --git a/Engine/Source/Runtime/Experimental/Chaos/Public/Chaos/Collision/ContactTriangles.h b/Engine/Source/Runtime/Experimental/Chaos/Public/Chaos/Collision/ContactTriangles.h
new file mode 100644
--- /dev/null
+++ b/Engine/Source/Runtime/Experimental/Chaos/Public/Chaos/Collision/ContactTriangles.h
@@ -0,0 +1,302 @@
+// Copyright Epic Games, Inc.All Rights Reserved.
+#pragma once
+
+#include "CoreMinimal.h"
+#include "Templates/UnrealTemplate.h"
+
+#include "Chaos/Core.h"
+#include "Chaos/Triangle.h"
+
+namespace Chaos
+{
+	/**
+	 * @brief Data held alongside contact points when generating contacts against a (likely non-convex) mesh of triangles
+	*/
+	class CHAOS_API FContactTriangle
+	{
+	public:
+		// Triangle data
+		FVec3 Vertices[3];
+		FVec3 FaceNormal;
+
+		// Indices into the mesh (trimesh or heightfield) that generated the triangle
+		int32 VertexIndices[3];
+		int32 TriangleIndex;
+
+		// Does this triangle contains the specified vertex? (VertexIndex is an index into the owning mesh's vertices)
+		inline bool HasVertexIndex(const int32 VertexIndex) const
+		{
+			return (VertexIndices[0] == VertexIndex) || (VertexIndices[1] == VertexIndex) || (VertexIndices[2] == VertexIndex);
+		}
+
+		// Get the centroid of the triangle
+		inline FVec3 GetCentroid() const
+		{
+			return (Vertices[0] + Vertices[1] + Vertices[2]) / FReal(3);
+		}
+
+		// Get the positions of the other two vertices in the triangle. (VertexIndex is an index into the owning mesh's vertices)
+		inline bool GetOtherVertexPositions(const int32 VertexIndex, FVec3& OutVertex0, FVec3& OutVertex1) const
+		{
+			if (VertexIndex == VertexIndices[0])
+			{
+				OutVertex0 = Vertices[1];
+				OutVertex1 = Vertices[2];
+				return true;
+			}
+			else if (VertexIndex == VertexIndices[1])
+			{
+				OutVertex0 = Vertices[2];
+				OutVertex1 = Vertices[0];
+				return true;
+			}
+			else if (VertexIndex == VertexIndices[2])
+			{
+				OutVertex0 = Vertices[0];
+				OutVertex1 = Vertices[1];
+				return true;
+			}
+			return false;
+		}
+
+		// Given a contact point on a triangle, check if it is an edge contact or vertex contact and return vertex indices if so
+		// @todo(chaos): this function should be unnecessary. Collision detection should return the vertex indices for the vertex/edge/face that we collide with
+		inline bool GetEdgeVerticesAtPosition(const FVec3& ContactPosition, FVec3& OutEdgeVertexA, FVec3& OutEdgeVertexB, int32& OutVertexIndexA, int32& OutVertexIndexB, const FReal BaryCentricTolerance = UE_KINDA_SMALL_NUMBER) const
+		{
+			OutVertexIndexA = INDEX_NONE;
+			OutVertexIndexB = INDEX_NONE;
+
+			const FVec3 BaryCentric = ToBarycentric(ContactPosition, Vertices[0], Vertices[1], Vertices[2]);
+
+			// Is it a vertex contact?
+			if (FMath::IsNearlyEqual(BaryCentric.X, FReal(1), BaryCentricTolerance))
+			{
+				OutVertexIndexA = VertexIndices[0];
+				OutEdgeVertexA = Vertices[0];
+				return true;
+			}
+			if (FMath::IsNearlyEqual(BaryCentric.Y, FReal(1), BaryCentricTolerance))
+			{
+				OutVertexIndexA = VertexIndices[1];
+				OutEdgeVertexA = Vertices[1];
+				return true;
+			}
+			if (FMath::IsNearlyEqual(BaryCentric.Z, FReal(1), BaryCentricTolerance))
+			{
+				OutVertexIndexA = VertexIndices[2];
+				OutEdgeVertexA = Vertices[2];
+				return true;
+			}
+
+			// Is it an edge contact?
+			if (FMath::IsNearlyEqual(BaryCentric.X, FReal(0), BaryCentricTolerance))
+			{
+				OutVertexIndexA = VertexIndices[1];
+				OutVertexIndexB = VertexIndices[2];
+				OutEdgeVertexA = Vertices[1];
+				OutEdgeVertexB = Vertices[2];
+				return true;
+			}
+			if (FMath::IsNearlyEqual(BaryCentric.Y, FReal(0), BaryCentricTolerance))
+			{
+				OutVertexIndexA = VertexIndices[2];
+				OutVertexIndexB = VertexIndices[0];
+				OutEdgeVertexA = Vertices[2];
+				OutEdgeVertexB = Vertices[0];
+				return true;
+			}
+			if (FMath::IsNearlyEqual(BaryCentric.Z, FReal(0), BaryCentricTolerance))
+			{
+				OutVertexIndexA = VertexIndices[0];
+				OutVertexIndexB = VertexIndices[1];
+				OutEdgeVertexA = Vertices[0];
+				OutEdgeVertexB = Vertices[1];
+				return true;
+			}
+
+			return false;
+		}
+	};
+
+	/**
+	 * @brief A set of triangles assocuated with some contact points.
+	*/
+	class CHAOS_API FContactTriangles
+	{
+	public:
+
+		inline const int32 Num() const
+		{
+			return Triangles.Num();
+		}
+
+		inline const FContactTriangle& At(const int32 Index) const
+		{
+			return Triangles[Index];
+		}
+
+
+	private:
+		TArray<FContactTriangle> Triangles;
+	};
+
+	using FContactVertexID = int32;
+
+	struct FContactEdgeID
+	{
+		FContactEdgeID()
+			: VertexIDs{ INDEX_NONE < INDEX_NONE }
+		{
+		}
+
+		FContactEdgeID(const FContactVertexID VertextIndexA, const FContactVertexID VertexIndexB)
+			: VertexIDs{ FMath::Min(VertextIndexA, VertexIndexB), FMath::Max(VertextIndexA, VertexIndexB) }
+		{
+		}
+
+		bool IsValid() const
+		{
+			return (VertexIDs[0] != INDEX_NONE) && (VertexIDs[1] != INDEX_NONE);
+		}
+
+		friend bool operator==(const FContactEdgeID& L, const FContactEdgeID& R)
+		{
+			return L.EdgeID == R.EdgeID;
+		}
+
+		friend bool operator!=(const FContactEdgeID& L, const FContactEdgeID& R)
+		{
+			return L.EdgeID != R.EdgeID;
+		}
+
+		friend bool operator<(const FContactEdgeID& L, const FContactEdgeID& R)
+		{
+			return L.EdgeID < R.EdgeID;
+		}
+
+		friend uint32 GetTypeHash(const FContactEdgeID& V)
+		{
+			return ::GetTypeHash(V.EdgeID);
+		}
+
+		union
+		{
+			FContactVertexID VertexIDs[2];
+			uint64 EdgeID;
+		};
+	};
+
+	/**
+	 * @brief Extended data used when processing contactpoints on a triangle mesh. 
+	 * Adds information about which triangle we hit (in the local array of FoOntactTriangles)
+	 * as well as which mesh edge or vertex we hit if appropriate (used for pruning).
+	*/
+	class FTriangleContactPoint : public FContactPoint
+	{
+	public:
+		FTriangleContactPoint()
+			: FContactPoint()
+			, EdgeID(INDEX_NONE, INDEX_NONE)
+			, VertexID(INDEX_NONE)
+			, ContactTriangleIndex(INDEX_NONE)
+		{
+		}
+
+		FTriangleContactPoint(const FContactPoint& InContactPoint)
+			: FContactPoint(InContactPoint)
+			, EdgeID(INDEX_NONE, INDEX_NONE)
+			, VertexID(INDEX_NONE)
+			, ContactTriangleIndex(INDEX_NONE)
+		{
+		}
+
+		FContactEdgeID EdgeID;
+		FContactVertexID VertexID;
+
+		int32 ContactTriangleIndex;
+	};
+
+	/**
+	 * Holds the contacts produced when colliding a convex shape against a non-convex triangular mesh, and provides
+	 * methods to reduce the total set of contact points down to a minimum manifold.
+	*/
+	class FContactTriangleCollector
+	{
+	public:
+		FContactTriangleCollector(const bool bInOneSided, const FRigidTransform3& InConvexTransform)
+			: bOneSidedCollision(bInOneSided)
+			, ConvexTransform(InConvexTransform)
+		{
+		}
+
+		inline TArrayView<const FContactPoint> GetContactPoints() const
+		{ 
+			return MakeArrayView(FinalContactPoints);
+		}
+		
+		// Add contacts between the convex and a single triangle for later processing
+		inline void AddTriangleContacts(const TArrayView<const FContactPoint>& InContactPoints, const FTriangle& Triangle, const int32 TriangleIndex, const int32 VertexIndex0, const int32 VertexIndex1, const int32 VertexIndex2, const FReal CullDistance)
+		{
+			if (InContactPoints.Num() > 0)
+			{
+				const int32 ContactTriangleIndex = AddTriangle(Triangle, TriangleIndex, VertexIndex0, VertexIndex1, VertexIndex2);
+
+				TriangleContactPoints.Reserve(TriangleContactPoints.Num() + InContactPoints.Num());
+				for (const FContactPoint& ContactPoint : InContactPoints)
+				{
+					if (ContactPoint.Phi < CullDistance)
+					{
+						const int32 ContactIndex = TriangleContactPoints.Add(ContactPoint);
+
+						TriangleContactPoints[ContactIndex].FaceIndex = TriangleIndex;
+						TriangleContactPoints[ContactIndex].ContactTriangleIndex = ContactTriangleIndex;
+					}
+				}
+			}
+		}
+
+		// To be called when all contacts have been added. Reduces the set of contacts to a minimum and transforms data into shape space
+		void ProcessContacts(const FRigidTransform3& MeshToConvexTransform);
+
+	private:
+		void SetNumContacts(const int32 Num);
+		void RemoveContact(const int32 ContactIndex);
+
+		void BuildContactFeatureSets();
+		void SortContactPointsForPruning();
+		void SortContactPointsForSolving();
+		void PruneEdgeContactPoints();
+		void FixInvalidNormalContactPoints();
+		void PruneInfacingContactPoints();
+		void PruneUnnecessaryContactPoints(const FReal PhiTolerance, const FReal DistanceTolerance);
+		void ReduceManifoldContactPointsTriangeMesh();
+		void FinalizeContacts(const FRigidTransform3& MeshToConvexTransform);
+
+		int32 AddTriangle(const FTriangle& Triangle, const int32 TriangleIndex, const int32 VertexIndex0, const int32 VertexIndex1, const int32 VertexIndex2);
+		int32 GetOtherTriangleIndexSharingEdge(const int32 KnownTriangleIndex, const int32 EdgeVertexIndexA, const int32 EdgeVertexIndexB) const;
+		int32 GetNextTriangleIndexSharingVertex(const int32 LastContactTriangleIndex, const int32 VertexIndex) const;
+
+		void DebugDrawContactPoints(const FColor& Color, const FReal LineThickness);
+
+		// The output set of contacts
+		// @todo(chaos): would be nice if we didn't need this and we could provide TriangleContactPoints instead
+		TArray<FContactPoint> FinalContactPoints;
+
+		// The collected set of contacts with additional triangle information
+		TArray<FTriangleContactPoint> TriangleContactPoints;
+
+		// All the triangles referenced by the TriangleContactPoints
+		TArray<FContactTriangle> ContactTriangles;
+
+		// A list of all the edges of contacts with faces. Used to prune edge contacts
+		TSet<FContactEdgeID> ContactEdges;
+
+		// A list of all the vertices of contacts with faces and edges. Used to prune vertex contacts
+		TSet<FContactVertexID> ContactVertices;
+
+		bool bOneSidedCollision;
+
+		FRigidTransform3 ConvexTransform;
+	};
+
+}
\ No newline at end of file
diff --git a/Engine/Source/Runtime/Experimental/Chaos/Public/Chaos/Collision/PBDCollisionConstraint.h b/Engine/Source/Runtime/Experimental/Chaos/Public/Chaos/Collision/PBDCollisionConstraint.h
--- a/Engine/Source/Runtime/Experimental/Chaos/Public/Chaos/Collision/PBDCollisionConstraint.h
+++ b/Engine/Source/Runtime/Experimental/Chaos/Public/Chaos/Collision/PBDCollisionConstraint.h
@@ -488,6 +488,7 @@
 
 		void AddIncrementalManifoldContact(const FContactPoint& ContactPoint);
 
+		// @todo(chaos): remove this and use SetOneShotManifoldContacts
 		inline void AddOneshotManifoldContact(const FContactPoint& ContactPoint)
 		{
 			if (ContactPoint.IsSet() && !ManifoldPoints.IsFull())
@@ -504,7 +505,36 @@
 			}
 		}
 
+		/**
+		 * @brief Replace the current manifold points with the input.
+		 * The input array should contain no more than MaxManifoldPoints contacts (any extra will be ignored).
+		 * We assume that all input contacts have been initialized and will not return false from IsSet().
+		 * Ignores contacts deeper than the CullDistance for this constraint.
+		*/
+		inline void SetOneShotManifoldContacts(const TArrayView<const FContactPoint>& ContactPoints)
+		{
+			ResetActiveManifoldContacts();
 
+			// Disable the incremental flag to prevent calling collision detection again
+			Flags.bUseIncrementalManifold = false;
+
+			FReal MinPhi = TNumericLimits<FReal>::Max();
+			const int32 NumContacts = FMath::Min(ContactPoints.Num(), MaxManifoldPoints);
+			for (int32 ContactIndex = 0; ContactIndex < NumContacts; ++ContactIndex)
+			{
+				const FContactPoint& ContactPoint = ContactPoints[ContactIndex];
+				if (ContactPoint.Phi < CullDistance)
+				{
+					int32 ManifoldPointIndex = AddManifoldPoint(ContactPoint);
+					if (ContactPoint.Phi < MinPhi)
+					{
+						ClosestManifoldPointIndex = ManifoldPointIndex;
+						MinPhi = ContactPoint.Phi;
+					}
+				}
+			}
+		}
+
 		void UpdateManifoldContacts();
 
 		// Particle-relative transform of each collision shape in the constraint
diff --git a/Engine/Source/Runtime/Experimental/Chaos/Public/Chaos/Collision/SpatialAccelerationBroadPhase.h b/Engine/Source/Runtime/Experimental/Chaos/Public/Chaos/Collision/SpatialAccelerationBroadPhase.h
--- a/Engine/Source/Runtime/Experimental/Chaos/Public/Chaos/Collision/SpatialAccelerationBroadPhase.h
+++ b/Engine/Source/Runtime/Experimental/Chaos/Public/Chaos/Collision/SpatialAccelerationBroadPhase.h
@@ -168,13 +168,13 @@ namespace Chaos
 				const TParticleView<TPBDRigidParticles<FReal, 3>>& DynamicSleepingView = Particles.GetNonDisabledDynamicView();
 				const TParticleView<TKinematicGeometryParticles<FReal, 3>>& DynamicMovingKinematicView = Particles.GetActiveDynamicMovingKinematicParticlesView();
 
-				if(DynamicSleepingView.Num() > DynamicMovingKinematicView.Num())
+				if(DynamicSleepingView.Num() < DynamicMovingKinematicView.Num())
 				{
-					ComputeParticlesOverlaps<false,false>(DynamicMovingKinematicView, Dt, InSpatialAcceleration, NarrowPhase);
+					ComputeParticlesOverlaps<false, true>(DynamicSleepingView, Dt, InSpatialAcceleration, NarrowPhase);
 				}
 				else
 				{
-					ComputeParticlesOverlaps<false,true>(DynamicSleepingView, Dt, InSpatialAcceleration, NarrowPhase);
+					ComputeParticlesOverlaps<false, false>(DynamicMovingKinematicView, Dt, InSpatialAcceleration, NarrowPhase);
 				}
 			}
 			else
diff --git a/Engine/Source/Runtime/Experimental/Chaos/Public/Chaos/GJK.h b/Engine/Source/Runtime/Experimental/Chaos/Public/Chaos/GJK.h
--- a/Engine/Source/Runtime/Experimental/Chaos/Public/Chaos/GJK.h
+++ b/Engine/Source/Runtime/Experimental/Chaos/Public/Chaos/GJK.h
@@ -730,7 +730,7 @@ namespace Chaos
 	 * @return true always
 	*/
 	template <typename T, typename TGeometryA, typename TGeometryB>
-	bool GJKPenetrationSameSpace(const TGeometryA& A, const TGeometryB& B, T& OutPenetration, TVec3<T>& OutClosestA, TVec3<T>& OutClosestB, TVec3<T>& OutNormal, int32& OutVertexA, int32& OutVertexB, T& OutMaxSupportDelta, const T Epsilon = T(1.e-3), const T EPAEpsilon = T(1.e-2))
+	bool GJKPenetrationSameSpace(const TGeometryA& A, const TGeometryB& B, T& OutPenetration, TVec3<T>& OutClosestA, TVec3<T>& OutClosestB, TVec3<T>& OutNormal, int32& OutVertexA, int32& OutVertexB, T& OutMaxSupportDelta, const TVec3<T>& InitialDir, const T Epsilon = T(1.e-3), const T EPAEpsilon = T(1.e-2))
 	{
 		TGJKSimplexData<T> SimplexData;
 		T SupportDeltaA = 0;
@@ -753,7 +753,7 @@ namespace Chaos
 		};
 
 		// V and Simplex are in A-local space
-		TVec3<T> V = TVec3<T>(-1, 0, 0);
+		TVec3<T> V = InitialDir;
 		TVec3<T> Simplex[4];
 		FSimplex SimplexIDs;
 		T Distance = FLT_MAX;
@@ -882,7 +882,7 @@ namespace Chaos
 	}
 
 	template <typename T, typename TGeometryA, typename TGeometryB>
-	bool GJKPenetrationSameSpace2(const TGeometryA& A, const TGeometryB& B, T& OutPenetration, TVec3<T>& OutClosestA, TVec3<T>& OutClosestB, TVec3<T>& OutNormal, int32& OutVertexA, int32& OutVertexB, T& OutMaxSupportDelta, const T Epsilon = T(1.e-3), const T EPAEpsilon = T(1.e-2))
+	bool GJKPenetrationSameSpace2(const TGeometryA& A, const TGeometryB& B, T& OutPenetration, TVec3<T>& OutClosestA, TVec3<T>& OutClosestB, TVec3<T>& OutNormal, int32& OutVertexA, int32& OutVertexB, T& OutMaxSupportDelta, const TVec3<T>& InitialDir, const T Epsilon = T(1.e-3), const T EPAEpsilon = T(1.e-2))
 	{
 		TGJKSimplexData<T> SimplexData;
 		T SupportDeltaA = 0;
@@ -905,7 +905,7 @@ namespace Chaos
 		};
 
 		// V and Simplex are in A-local space
-		TVec3<T> V = TVec3<T>(-1, 0, 0);
+		TVec3<T> V = InitialDir;
 		TVec3<T> Simplex[4];
 		int32 NumVerts = 0;
 		T Distance = FLT_MAX;
diff --git a/Engine/Source/Runtime/Experimental/Chaos/Public/Chaos/HeightField.h b/Engine/Source/Runtime/Experimental/Chaos/Public/Chaos/HeightField.h
--- a/Engine/Source/Runtime/Experimental/Chaos/Public/Chaos/HeightField.h
+++ b/Engine/Source/Runtime/Experimental/Chaos/Public/Chaos/HeightField.h
@@ -93,19 +93,68 @@
 		bool GJKContactPoint(const TImplicitObjectScaled<FCapsule>& QueryGeom, const FRigidTransform3& QueryTM, const FReal Thickness, FVec3& ContactLocation, FVec3& ContactNormal, FReal& ContactPhi) const;
 		bool GJKContactPoint(const TImplicitObjectScaled<FConvex>& QueryGeom, const FRigidTransform3& QueryTM, const FReal Thickness, FVec3& ContactLocation, FVec3& ContactNormal, FReal& ContactPhi) const;
 
-		bool ContactManifold(const TBox<FReal, 3>& QueryGeom, const FRigidTransform3& QueryTM, const FReal Thickness, TArray<FContactPoint>& ContactPoints) const;
-		bool ContactManifold(const FCapsule& QueryGeom, const FRigidTransform3& QueryTM, const FReal Thickness, TArray<FContactPoint>& ContactPoints) const;
-		bool ContactManifold(const FConvex& QueryGeom, const FRigidTransform3& QueryTM, const FReal Thickness, TArray<FContactPoint>& ContactPoints) const;
-		bool ContactManifold(const TImplicitObjectScaled<TBox<FReal, 3>>& QueryGeom, const FRigidTransform3& QueryTM, const FReal Thickness, TArray<FContactPoint>& ContactPoints) const;
-		bool ContactManifold(const TImplicitObjectScaled<FCapsule>& QueryGeom, const FRigidTransform3& QueryTM, const FReal Thickness, TArray<FContactPoint>& ContactPoints) const;
-		bool ContactManifold(const TImplicitObjectScaled<FConvex>& QueryGeom, const FRigidTransform3& QueryTM, const FReal Thickness, TArray<FContactPoint>& ContactPoints) const;
 
-
 		virtual int32 FindMostOpposingFace(const FVec3& Position, const FVec3& UnitDir, int32 HintFaceIndex, FReal SearchDist) const override;
 		virtual FVec3 FindGeometryOpposingNormal(const FVec3& DenormDir, int32 FaceIndex, const FVec3& OriginalNormal) const override;
 
-		void VisitTriangles(const FAABB3& InQueryBounds, const TFunction<void(const FTriangle& Triangle)>& Visitor) const;
+		/**
+		 * @param QueryBounds Bounding box in which we want to produce triangles, in HeightField space
+		 * @param QueryTransform Transforms from HeightField space to query space (usually the other object's space)
+		 * @param CullDistance The distance at which we can ignore triangles
+		 * @param Visitor void(const FTriangle& Triangle, const int32 TriangleIndex, const int32 VertexIndex0, const int32 VertexIndex1, const int32 OutVertexIndex2)
+		*/
+		template<typename TriangleVisitor>
+		void VisitTriangles(const FAABB3& QueryBounds, const FRigidTransform3& QueryTransform, const TriangleVisitor& Visitor) const
+		{
+			FBounds2D GridQueryBounds;
+			GridQueryBounds.Min = FVec2(QueryBounds.Min()[0], QueryBounds.Min()[1]);
+			GridQueryBounds.Max = FVec2(QueryBounds.Max()[0], QueryBounds.Max()[1]);
 
+			// @todo(chaos): could we do the 3D bounds test here?
+			TArray<TVec2<int32>> Intersections;
+			GetGridIntersections(GridQueryBounds, Intersections);
+
+			FVec3 Points[4];
+
+			for (const TVec2<int32>& Cell : Intersections)
+			{
+				const int32 CellVertexIndex = Cell[1] * GeomData.NumCols + Cell[0];	// First vertex in cell
+				const int32 CellIndex = Cell[1] * (GeomData.NumCols - 1) + Cell[0];
+
+				// Check for holes and skip checking if we'll never collide
+				if (GeomData.MaterialIndices.IsValidIndex(CellIndex) && GeomData.MaterialIndices[CellIndex] == TNumericLimits<uint8>::Max())
+				{
+					continue;
+				}
+
+				// The triangle is solid so proceed to test it
+				FAABB3 CellBounds;
+				GeomData.GetPointsAndBoundsScaled(CellVertexIndex, Points, CellBounds);
+
+				if (CellBounds.Intersects(QueryBounds))
+				{
+					// Transform points into the space of the query
+					// @todo(chaos): duplicate work here when we overlap lots of cells. We could generate the transformed verts for the full query once...
+					Points[0] = QueryTransform.TransformPositionNoScale(Points[0]);
+					Points[1] = QueryTransform.TransformPositionNoScale(Points[1]);
+					Points[2] = QueryTransform.TransformPositionNoScale(Points[2]);
+					Points[3] = QueryTransform.TransformPositionNoScale(Points[3]);
+
+					// @todo(chaos): utility function for this
+					const int32 VertexIndex0 = CellVertexIndex;
+					const int32 VertexIndex1 = CellVertexIndex + 1;
+					const int32 VertexIndex2 = CellVertexIndex + GeomData.NumCols;
+					const int32 VertexIndex3 = CellVertexIndex + GeomData.NumCols + 1;
+
+					// Generate contacts if overlapping
+					const int32 FaceIndex0 = CellIndex * 2 + 0;
+					const int32 FaceIndex1 = CellIndex * 2 + 1;
+					Visitor(FTriangle(Points[0], Points[1], Points[3]), FaceIndex0, VertexIndex0, VertexIndex1, VertexIndex3);
+					Visitor(FTriangle(Points[0], Points[3], Points[2]), FaceIndex1, VertexIndex0, VertexIndex3, VertexIndex2);
+				}
+			}
+		}
+
 		struct FClosestFaceData
 		{
 			int32 FaceIndex = INDEX_NONE;
@@ -773,15 +822,6 @@
 
 		template <typename GeomType>
 		bool GJKContactPointImp(const GeomType& QueryGeom, const FRigidTransform3& QueryTM, const FReal Thickness, FVec3& ContactLocation, FVec3& ContactNormal, FReal& ContactPhi) const;
-
-		// @todo(chaos): remove ContactManifoldNonPlanarConvexImp and ContactManifoldPlanarConvexImp and move collision detection code out of this class.
-		// Instead, provide an API for generating triangles in a query bounds in some space.
-
-		template <typename GeomType>
-		bool ContactManifoldNonPlanarConvexImp(const GeomType& QueryGeom, const FRigidTransform3& QueryTM, const FReal Thickness, TArray<FContactPoint>& ContactPoints) const;
-
-		template <typename GeomType>
-		bool ContactManifoldPlanarConvexImp(const GeomType& QueryGeom, const FRigidTransform3& QueryTM, const FReal Thickness, TArray<FContactPoint>& ContactPoints) const;
-
 	};
+
 }
diff --git a/Engine/Source/Runtime/Experimental/Chaos/Public/Chaos/ImplicitObjectScaled.h b/Engine/Source/Runtime/Experimental/Chaos/Public/Chaos/ImplicitObjectScaled.h
--- a/Engine/Source/Runtime/Experimental/Chaos/Public/Chaos/ImplicitObjectScaled.h
+++ b/Engine/Source/Runtime/Experimental/Chaos/Public/Chaos/ImplicitObjectScaled.h
@@ -236,12 +236,6 @@
 		return MObject->GJKContactPoint(A, AToBTM, Thickness, Location, Normal, Penetration);
 	}
 
-	template <typename QueryGeomType>
-	bool ContactManifold(const QueryGeomType& A, const FRigidTransform3& AToBTM, const FReal Thickness, TArray<FContactPoint>& ContactPoints) const
-	{
-		return MObject->ContactManifold(A, AToBTM, Thickness, ContactPoints);
-	}
-
 	virtual uint16 GetMaterialIndex(uint32 HintIndex) const
 	{
 		return MObject->GetMaterialIndex(HintIndex);
@@ -723,20 +717,6 @@
 		return MObject->GJKContactPoint(ScaledA, AToBTMNoScale, UnscaledThickness, Location, Normal, Penetration, MScale);
 	}
 
-	template <typename QueryGeomType>
-	bool ContactManifold(const QueryGeomType& A, const FRigidTransform3& AToBTM, const FReal Thickness, TArray<FContactPoint>& ContactPoints) const
-	{
-		TRigidTransform<T, d> AToBTMNoScale(AToBTM.GetLocation() * MInvScale, AToBTM.GetRotation());
-
-		// Thickness is a culling distance which cannot be non-uniformly scaled. This gets passed into the ImplicitObject API unscaled so that
-		// if can do the right thing internally if possible, but it makes the API a little confusing. (This is only exists used TriMesh and Heightfield.)
-		// See FTriangleMeshImplicitObject::GJKContactPointImp
-		const FReal UnscaledThickness = Thickness;
-
-		auto ScaledA = MakeScaledHelper(A, MInvScale);
-		return MObject->ContactManifold(ScaledA, AToBTMNoScale, UnscaledThickness, ContactPoints, MScale);
-	}
-
 	/** This is a low level function and assumes the internal object has a OverlapGeom function. Should not be called directly. See GeometryQueries.h : OverlapQuery */
 	template <typename QueryGeomType>
 	bool LowLevelOverlapGeom(const QueryGeomType& B, const TRigidTransform<T, d>& BToATM, T Thickness = 0, FMTDInfo* OutMTD = nullptr) const
diff --git a/Engine/Source/Runtime/Experimental/Chaos/Public/Chaos/Triangle.h b/Engine/Source/Runtime/Experimental/Chaos/Public/Chaos/Triangle.h
--- a/Engine/Source/Runtime/Experimental/Chaos/Public/Chaos/Triangle.h
+++ b/Engine/Source/Runtime/Experimental/Chaos/Public/Chaos/Triangle.h
@@ -10,6 +10,82 @@
 
 namespace Chaos
 {
+	/**
+	 * @brief Convert the cartesian coordinate into a barycentric corrdinate.
+	 * Compute barycentric coordinates/weights of Point inside 3D triangle (V0,V1,V2).
+	 * If point is in triangle plane and inside triangle, coords will be positive and sum to 1.
+	 * ie if result is a, then vPoint = a.x*V0 + a.y*V1 + a.z*V2.
+	 * 
+	 * @note For points outside the triangle it will return negative barycentric coordinates that should still sum to 1. However,
+	 * if the triangle is degenerate and the point is not on the line/vertex, it will return (1,0,0) with no way to detect this condition.
+	 * 
+	 * @see FromBarycentric()
+	 */
+	template <typename RealType>
+	inline TVec3<RealType> ToBarycentric(const TVec3<RealType>& Point, const TVec3<RealType>& V0, const TVec3<RealType>& V1, const TVec3<RealType>& V2)
+	{
+		const TVec3<RealType> V02 = V0 - V2;
+		const TVec3<RealType> V12 = V1 - V2;
+		const TVec3<RealType> PV2 = Point - V2;
+		const RealType M00 = V02.Dot(V02);
+		const RealType M01 = V02.Dot(V12);
+		const RealType M11 = V12.Dot(V12);
+		const RealType R0 = V02.Dot(PV2);
+		const RealType R1 = V12.Dot(PV2);
+		const RealType Det = M00 * M11 - M01 * M01;
+		RealType Bary1, Bary2, Bary3;
+		if (Det > UE_SMALL_NUMBER)
+		{
+			// Non-degenerate triangle
+			const RealType InvDet = RealType(1) / Det;
+			Bary1 = (M11 * R0 - M01 * R1) * InvDet;
+			Bary2 = (M00 * R1 - M01 * R0) * InvDet;
+			Bary3 = RealType(1) - Bary1 - Bary2;
+		}
+		else
+		{
+			const bool bHaveLine02 = (M00 > UE_SMALL_NUMBER);
+			const bool bHaveLine12 = (M11 > UE_SMALL_NUMBER);
+			const bool bInsideLine02 = ((M00 > UE_SMALL_NUMBER) && (R0 >= 0) && (R0 <= M00));
+			const bool bInsideLine12 = ((M11 > UE_SMALL_NUMBER) && (R1 >= 0) && (R1 <= M11));
+			if (bInsideLine02 || (bHaveLine02 && !bInsideLine12))
+			{
+				// Line-degenerate and point on line segment V02
+				const RealType Alpha02 = R0 / M00;
+				Bary1 = Alpha02;
+				Bary2 = RealType(0);
+				Bary3 = RealType(1) - Alpha02;
+			}
+			else if (bHaveLine12)
+			{
+				// Line-degenerate and point on line segment V12
+				const RealType Alpha12 = R1 / M11;
+				Bary1 = RealType(0);
+				Bary2 = Alpha12;
+				Bary3 = RealType(1) - Alpha12;
+			}
+			else
+			{
+				// Point-degenerate, or line-degenerate and point not on line
+				Bary1 = RealType(1);
+				Bary2 = RealType(0);
+				Bary3 = RealType(0);
+			}
+		}
+		return TVec3<RealType>(Bary1, Bary2, Bary3);
+	}
+
+	/**
+	 * @brief Convert the barycentric coordinate into a cartesian corrdinate
+	 * @see ToBarycentric()
+	*/
+	template <typename RealType>
+	inline TVec3<RealType> FromBarycentric(const TVec3<RealType>& Barycentric, const TVec3<RealType>& V0, const TVec3<RealType>& V1, const TVec3<RealType>& V2)
+	{
+		return Barycentric.X * V0 + Barycentric.Y * V1 + Barycentric.Z * V2;
+	}
+
+
 	template<typename T>
 	class TTriangle
 	{
@@ -64,6 +140,11 @@ namespace Chaos
 			OutX = ABC[0];
 		}
 
+		FORCEINLINE TVec3<T> GetCentroid() const
+		{
+			return (GetVertex(0) + GetVertex(1) + GetVertex(2)) / T(3.0);
+		}
+
 		// Get the nearest point on an edge and the edge vertices
 		// Used for manifold generation
 		TVec3<T> GetClosestEdge(int32 PlaneIndexHint, const TVec3<T>& Position, TVec3<T>& OutEdgePos0, TVec3<T>& OutEdgePos1) const
diff --git a/Engine/Source/Runtime/Experimental/Chaos/Public/Chaos/TriangleMeshImplicitObject.h b/Engine/Source/Runtime/Experimental/Chaos/Public/Chaos/TriangleMeshImplicitObject.h
--- a/Engine/Source/Runtime/Experimental/Chaos/Public/Chaos/TriangleMeshImplicitObject.h
+++ b/Engine/Source/Runtime/Experimental/Chaos/Public/Chaos/TriangleMeshImplicitObject.h
@@ -442,12 +442,19 @@
 
 	template <typename IdxType, typename ParticlesType>
 	inline void TriangleMeshTransformVertsHelper(const FRigidTransform3& Transform, int32 TriIdx, const ParticlesType& Particles,
-		const TArray<TVector<IdxType, 3>>& Elements, FVec3& OutA, FVec3& OutB, FVec3& OutC)
+		const TArray<TVector<IdxType, 3>>& Elements, FVec3& OutA, FVec3& OutB, FVec3& OutC,
+		int32& OutVertexIndexA, int32& OutVertexIndexB, int32& OutVertexIndexC)
 	{
-		// Note: deliberately using scaled transform
-		OutA = Transform.TransformPosition(FVector(Particles.X(Elements[TriIdx][0])));
-		OutB = Transform.TransformPosition(FVector(Particles.X(Elements[TriIdx][1])));
-		OutC = Transform.TransformPosition(FVector(Particles.X(Elements[TriIdx][2])));
+		const int32 VertexIndexA = Elements[TriIdx][0];
+		const int32 VertexIndexB = Elements[TriIdx][1];
+		const int32 VertexIndexC = Elements[TriIdx][2];
+		// Note: deliberately using scaled transform here. See VisitTriangles
+		OutA = Transform.TransformPosition(FVector(Particles.X(VertexIndexA)));
+		OutB = Transform.TransformPosition(FVector(Particles.X(VertexIndexB)));
+		OutC = Transform.TransformPosition(FVector(Particles.X(VertexIndexC)));
+		OutVertexIndexA = VertexIndexA;
+		OutVertexIndexB = VertexIndexB;
+		OutVertexIndexC = VertexIndexC;
 	}
 
 
@@ -547,14 +554,6 @@
 		bool GJKContactPoint(const TImplicitObjectScaled < FCapsule >& QueryGeom, const FRigidTransform3& QueryTM, const FReal Thickness, FVec3& Location, FVec3& Normal, FReal& Penetration, FVec3 TriMeshScale = FVec3(1.0f)) const;
 		bool GJKContactPoint(const TImplicitObjectScaled < FConvex >& QueryGeom, const FRigidTransform3& QueryTM, const FReal Thickness, FVec3& Location, FVec3& Normal, FReal& Penetration, FVec3 TriMeshScale = FVec3(1.0f)) const;
 
-
-		bool ContactManifold(const TBox<FReal, 3>& QueryGeom, const FRigidTransform3& QueryTM, const FReal Thickness, TArray<FContactPoint>& ContactPoints) const;
-		bool ContactManifold(const FCapsule& QueryGeom, const FRigidTransform3& QueryTM, const FReal Thickness, TArray<FContactPoint>& ContactPoints) const;
-		bool ContactManifold(const FConvex& QueryGeom, const FRigidTransform3& QueryTM, const FReal Thickness, TArray<FContactPoint>& ContactPoints) const;
-		bool ContactManifold(const TImplicitObjectScaled<TBox<FReal, 3>>& QueryGeom, const FRigidTransform3& QueryTM, const FReal Thickness, TArray<FContactPoint>& ContactPoints, FVec3 TriMeshScale = FVec3(1.0f)) const;
-		bool ContactManifold(const TImplicitObjectScaled<FCapsule>& QueryGeom, const FRigidTransform3& QueryTM, const FReal Thickness, TArray<FContactPoint>& ContactPoints, FVec3 TriMeshScale = FVec3(1.0f)) const;
-		bool ContactManifold(const TImplicitObjectScaled<FConvex>& QueryGeom, const FRigidTransform3& QueryTM, const FReal Thickness, TArray<FContactPoint>& ContactPoints, FVec3 TriMeshScale = FVec3(1.0f)) const;
-
 		// Returns -1 if InternalFaceIndex is not in map, or map is invalid.
 		int32 GetExternalFaceIndexFromInternal(int32 InternalFaceIndex) const;
 
@@ -700,22 +699,18 @@
 
 		void UpdateVertices(const TArray<FVector>& Positions);
 
-		void VisitTriangles(const FAABB3& InQueryBounds, const TFunction<void(const FTriangle& Triangle)>& Visitor) const;
-
-		void VisitTriangle(const int32 TriangleIndex, const TFunction<void(const FTriangle& Triangle)>& Visitor) const;
-
 		/**
 		 * @brief Generate the triangle at the specified index with the specified transform (including scale)
 		*/
-		void GetTransformedTriangle(const int32 TriangleIndex, const FRigidTransform3& Transform, FTriangle& OutTriangle) const
+		void GetTransformedTriangle(const int32 TriangleIndex, const FRigidTransform3& Transform, FTriangle& OutTriangle, int32& OutVertexIndex0, int32& OutVertexIndex1, int32& OutVertexIndex2) const
 		{
 			if (MElements.RequiresLargeIndices())
 			{
-				TriangleMeshTransformVertsHelper(Transform, TriangleIndex, MParticles, MElements.GetLargeIndexBuffer(), OutTriangle[0], OutTriangle[1], OutTriangle[2]);
+				TriangleMeshTransformVertsHelper(Transform, TriangleIndex, MParticles, MElements.GetLargeIndexBuffer(), OutTriangle[0], OutTriangle[1], OutTriangle[2], OutVertexIndex0, OutVertexIndex1, OutVertexIndex2);
 			}
 			else
 			{
-				TriangleMeshTransformVertsHelper(Transform, TriangleIndex, MParticles, MElements.GetSmallIndexBuffer(), OutTriangle[0], OutTriangle[1], OutTriangle[2]);
+				TriangleMeshTransformVertsHelper(Transform, TriangleIndex, MParticles, MElements.GetSmallIndexBuffer(), OutTriangle[0], OutTriangle[1], OutTriangle[2], OutVertexIndex0, OutVertexIndex1, OutVertexIndex2);
 			}
 		}
 
@@ -728,6 +723,28 @@
 			OutTriangleIndices = FastBVH.FindAllIntersections(QueryBounds);
 		}
 
+		/**
+		 * @param QueryBounds Bounding box in which we want to produce triangles, in unscaled TriMesh space
+		 * @param QueryTransform Transforms from TriMesh space to desired space (e.g., could be trimesh world transform including scale to get world-space triangles)
+		 * @param CullDistance The distance at which we can ignore triangles
+		 * @param Visitor void(const FTriangle& Triangle, const int32 TriangleIndex, const int32 VertexIndex0, const int32 VertexIndex1, const int32 OutVertexIndex2)
+		*/
+		template<typename TriangleVisitor>
+		void VisitTriangles(const FAABB3& QueryBounds, const FRigidTransform3& QueryTransform, const TriangleVisitor& Visitor) const
+		{
+			TArray<int32> OverlapIndices;
+			FindOverlappingTriangles(QueryBounds, OverlapIndices);
+
+			for (int32 OverlapIndex = 0; OverlapIndex < OverlapIndices.Num(); ++OverlapIndex)
+			{
+				const int32 TriangleIndex = OverlapIndices[OverlapIndex];
+				FTriangle Triangle;
+				int32 VertexIndex0, VertexIndex1, VertexIndex2;
+				GetTransformedTriangle(TriangleIndex, QueryTransform, Triangle, VertexIndex0, VertexIndex1, VertexIndex2);
+				Visitor(Triangle, TriangleIndex, VertexIndex0, VertexIndex1, VertexIndex2);
+			}
+		}
+
 	private:
 		using BVHType = TAABBTree<int32, TAABBTreeLeafArray<int32, /*bComputeBounds=*/ false, FRealSingle>, /*bMutable=*/false, FRealSingle>;
 		void RebuildFastBVHFromTree(const BVHType& BVH);
@@ -839,9 +856,6 @@
 		template <typename QueryGeomType>
 		bool GJKContactPointImp(const QueryGeomType& QueryGeom, const FRigidTransform3& QueryTM, const FReal Thickness, FVec3& Location, FVec3& Normal, FReal& Penetration, FVec3 TriMeshScale = FVec3(1.0)) const;
 
-		template <typename GeomType>
-		bool ContactManifoldImp(const GeomType& QueryGeom, const FRigidTransform3& QueryTM, const FReal Thickness, TArray<FContactPoint>& ContactPoints, FVec3 TriMeshScale) const;
-
 		template<typename QueryGeomType>
 		bool OverlapGeomImp(const QueryGeomType& QueryGeom, const FRigidTransform3& QueryTM, const FReal Thickness, FMTDInfo* OutMTD = nullptr, FVec3 TriMeshScale = FVec3(1.0f)) const;
 
@@ -867,60 +881,4 @@
 		TUniquePtr<FTriangleMeshImplicitObject> CopySlowImpl(const TArray < TVector<IdxType, 3>>& InElements) const;
 	};
 
-
-	/**
-	 * @brief A helper for iterating over FTriangleMeshImplicitObject triangles in a bounding box
-	*/
-	class FTriangleMeshTriangleProducer
-	{
-	public:
-		FTriangleMeshTriangleProducer()
-			: NextIndex(0)
-			, TriangleIndices()
-		{
-		}
-
-		/**
-		 * @brief Find the set of triangle indices that overlap the query bounds
-		 * @param InQueryBounds The query bounds in triangle mesh space
-		*/
-		inline void Reset(const FTriangleMeshImplicitObject& InTriMesh, const FAABB3& InQueryBounds)
-		{
-			InTriMesh.FindOverlappingTriangles(InQueryBounds, TriangleIndices);
-			NextIndex = 0;
-		}
-
-		/**
-		 * @brief Whether we are done producing triangles
-		 * @return
-		*/
-		inline const bool IsDone() const
-		{
-			return (NextIndex >= TriangleIndices.Num());
-		}
-
-		/**
-		 * @brief Get the next triangle with its index and transform the vertices
-		 * @return true if we produced a triangle, false if we are done with the triangle list
-		*/
-		inline bool NextTriangle(const FTriangleMeshImplicitObject& InTriMesh, const FRigidTransform3& Transform, FTriangle& OutTriangle, int32& OutTriangleIndex)
-		{
-			if (IsDone())
-			{
-				return false;
-			}
-
-			const int32 TriIndex = TriangleIndices[NextIndex++];
-
-			InTriMesh.GetTransformedTriangle(TriIndex, Transform, OutTriangle);
-			OutTriangleIndex = TriIndex;
-
-			return true;
-		}
-
-	private:
-		int32 NextIndex;
-		TArray<int32> TriangleIndices;
-	};
-
 }
