Chaos destruction : Remove on break feature
It is possible to assign each piece of a geo collection with a PostBreakTimer and a removal timer
- Added a RemoveOnBreak tool
- Added a Removal column mode to fracture editor outliner
- Optional RemovalOnBreak attribute array on the rest collection  ( can be added or removed through the UI )

diff --git a/Engine/Plugins/Experimental/ChaosEditor/Source/FractureEditor/Private/FractureEditorCommands.h b/Engine/Plugins/Experimental/ChaosEditor/Source/FractureEditor/Private/FractureEditorCommands.h
--- a/Engine/Plugins/Experimental/ChaosEditor/Source/FractureEditor/Private/FractureEditorCommands.h
+++ b/Engine/Plugins/Experimental/ChaosEditor/Source/FractureEditor/Private/FractureEditorCommands.h
@@ -92,6 +92,7 @@ class FFractureEditorCommands : public TCommands<FFractureEditorCommands>
 
 		// Property Commands
 		TSharedPtr< FUICommandInfo > SetInitialDynamicState;
+		TSharedPtr< FUICommandInfo > SetRemoveOnBreak;
 		
 };
 
diff --git a/Engine/Plugins/Experimental/ChaosEditor/Source/FractureEditor/Private/FractureEditorModeToolkit.cpp b/Engine/Plugins/Experimental/ChaosEditor/Source/FractureEditor/Private/FractureEditorModeToolkit.cpp
--- a/Engine/Plugins/Experimental/ChaosEditor/Source/FractureEditor/Private/FractureEditorModeToolkit.cpp
+++ b/Engine/Plugins/Experimental/ChaosEditor/Source/FractureEditor/Private/FractureEditorModeToolkit.cpp
@@ -613,8 +613,12 @@ TSharedRef<SDockTab> FFractureEditorModeToolkit::CreateHierarchyTab(const FSpawn
 					.AutoHeight()
 					[
 						SNew(SButton)
-						.OnClicked(this, &FFractureEditorModeToolkit::RefreshOutliner)
-						.Text(LOCTEXT("Refresh", "Refresh"))
+						.ForegroundColor(FAppStyle::GetSlateColor("DefaultForeground"))
+						.ContentPadding(FMargin(2, 0))
+						.HAlign(HAlign_Center)
+						.OnClicked(this, &FFractureEditorModeToolkit::OnRefreshOutlinerButtonClicked)
+						.Text(LOCTEXT("GCOUtliner_Refresh_Button_Text", "Refresh"))
+						.ToolTipText(LOCTEXT("GCOUtliner_Refresh_Button_ToolTip", "Refresh the outliner"))
 					]
 					+ SVerticalBox::Slot()
 					.AutoHeight()
@@ -635,13 +639,18 @@ TSharedRef<SDockTab> FFractureEditorModeToolkit::CreateHierarchyTab(const FSpawn
 	return CreatedTab.ToSharedRef();
 }
 
-FReply FFractureEditorModeToolkit::RefreshOutliner()
+FReply FFractureEditorModeToolkit::OnRefreshOutlinerButtonClicked()
+{
+	RefreshOutliner();
+	return FReply::Handled();
+}
+
+void FFractureEditorModeToolkit::RefreshOutliner()
 {
 	if (OutlinerView)
 	{
 		OutlinerView->RegenerateItems();
 	}
-	return FReply::Handled();
 }
 
 TSharedRef<SDockTab> FFractureEditorModeToolkit::CreateStatisticsTab(const FSpawnTabArgs& Args)
@@ -861,6 +870,7 @@ void FFractureEditorModeToolkit::BuildToolPalette(FName PaletteIndex, class FToo
 		ToolbarBuilder.AddToolBarButton(Commands.MakeConvex);
 		ToolbarBuilder.AddToolBarButton(Commands.FixTinyGeo);
 		ToolbarBuilder.AddToolBarButton(Commands.SetInitialDynamicState);
+		ToolbarBuilder.AddToolBarButton(Commands.SetRemoveOnBreak);
 	}
 }
 
diff --git a/Engine/Plugins/Experimental/ChaosEditor/Source/FractureEditor/Private/FractureToolProperties.cpp b/Engine/Plugins/Experimental/ChaosEditor/Source/FractureEditor/Private/FractureToolProperties.cpp
--- a/Engine/Plugins/Experimental/ChaosEditor/Source/FractureEditor/Private/FractureToolProperties.cpp
+++ b/Engine/Plugins/Experimental/ChaosEditor/Source/FractureEditor/Private/FractureToolProperties.cpp
@@ -75,8 +75,122 @@ void UFractureToolSetInitialDynamicState::Execute(TWeakPtr<FFractureEditorModeTo
 				}
 			}
 		}
+		InToolkit.Pin()->RefreshOutliner();
 	}
 }
-			
+
+//==============================================================================================================
+
+UFractureToolSetRemoveOnBreak::UFractureToolSetRemoveOnBreak(const FObjectInitializer& ObjInit)
+	: Super(ObjInit)
+{
+	RemoveOnBreakSettings = NewObject<UFractureRemoveOnBreakSettings>(GetTransientPackage(), UFractureRemoveOnBreakSettings::StaticClass());
+	RemoveOnBreakSettings->OwnerTool = this;
+}
+
+FText UFractureToolSetRemoveOnBreak::GetDisplayText() const
+{
+	return FText(NSLOCTEXT("Fracture", "FractureToolSetRemoveOnBreak", "Set Remove-On-Break"));
+}
+
+FText UFractureToolSetRemoveOnBreak::GetTooltipText() const
+{
+	return FText(NSLOCTEXT("Fracture", "FractureToolSetRemoveOnBreakToolTip", "Set removal-on-break parameters"));
+}
+
+FText UFractureToolSetRemoveOnBreak::GetApplyText() const
+{
+	return FText(NSLOCTEXT("Fracture", "ExecuteSetRemoveOnBreak", "Apply removal-on-break parameters"));
+}
+
+FSlateIcon UFractureToolSetRemoveOnBreak::GetToolIcon() const
+{
+	return FSlateIcon("FractureEditorStyle", "FractureEditor.SetRemoveOnBreak");
+}
+
+TArray<UObject*> UFractureToolSetRemoveOnBreak::GetSettingsObjects() const
+{
+	TArray<UObject*> Settings;
+	Settings.Add(RemoveOnBreakSettings);
+	return Settings;
+}
+
+void UFractureToolSetRemoveOnBreak::RegisterUICommand(FFractureEditorCommands* BindingContext)
+{
+	UI_COMMAND_EXT(BindingContext, UICommandInfo, "SetRemoveOnBreak", "Set Remove On Break", "Set removal-on-break parameters.", EUserInterfaceActionType::ToggleButton, FInputChord());
+	BindingContext->SetRemoveOnBreak = UICommandInfo;
+}
+
+void UFractureToolSetRemoveOnBreak::Execute(TWeakPtr<FFractureEditorModeToolkit> InToolkit)
+{
+	if (InToolkit.IsValid())
+	{
+	 	TSet<UGeometryCollectionComponent*> GeomCompSelection;
+	 	GetSelectedGeometryCollectionComponents(GeomCompSelection);
+	 	for (UGeometryCollectionComponent* GeometryCollectionComponent : GeomCompSelection)
+	 	{
+	 		FGeometryCollectionEdit GCEdit = GeometryCollectionComponent->EditRestCollection(GeometryCollection::EEditUpdate::RestPhysics, true /*bShapeIsUnchanged*/);
+	 		if (UGeometryCollection* GCObject = GCEdit.GetRestCollection())
+	 		{
+	 			TSharedPtr<FGeometryCollection, ESPMode::ThreadSafe> GeometryCollectionPtr = GCObject->GetGeometryCollection();
+	 			if (FGeometryCollection* GeometryCollection = GeometryCollectionPtr.Get())
+	 			{
+	 				TArray<int32> SelectedBones = GeometryCollectionComponent->GetSelectedBones();
+	 				if (SelectedBones.Num())
+	 				{
+	 					const FVector4f DisabledRemoveOnBreakData{ -1, -1, -1, -1 }; 
+	 					if (!GeometryCollection->HasAttribute("RemoveOnBreak", FGeometryCollection::TransformGroup))
+	 					{
+	 						TManagedArray<FVector4f>& NewRemoveOnBreak = GeometryCollection->AddAttribute<FVector4f>("RemoveOnBreak", FGeometryCollection::TransformGroup);
+	 						NewRemoveOnBreak.Fill(DisabledRemoveOnBreakData);
+	 					}
+	 					
+	 					TManagedArray<FVector4f>& RemoveOnBreak = GeometryCollection->ModifyAttribute<FVector4f>("RemoveOnBreak", FGeometryCollection::TransformGroup);
+	 					
+	 					FVector4f RemoveOnBreakParameters{DisabledRemoveOnBreakData};
+	 					if (RemoveOnBreakSettings->Enabled)
+	 					{
+	 						RemoveOnBreakParameters.X = RemoveOnBreakSettings->PostBreakTimer.X;
+	 						RemoveOnBreakParameters.Y = RemoveOnBreakSettings->PostBreakTimer.Y;
+	 						RemoveOnBreakParameters.Z = RemoveOnBreakSettings->RemovalTimer.X;
+	 						RemoveOnBreakParameters.W = RemoveOnBreakSettings->RemovalTimer.Y;
+	 					}
+
+	 					for (int32 Index : SelectedBones)
+	 					{
+	 						RemoveOnBreak[Index] = RemoveOnBreakParameters;
+	 					}
+	 				}
+	 			}
+	 		}
+	 	}
+		InToolkit.Pin()->RefreshOutliner();
+	}
+}
+
+void UFractureToolSetRemoveOnBreak::DeleteRemoveOnBreakData()
+{
+	TSet<UGeometryCollectionComponent*> GeomCompSelection;
+	GetSelectedGeometryCollectionComponents(GeomCompSelection);
+	for (UGeometryCollectionComponent* GeometryCollectionComponent : GeomCompSelection)
+	{
+		FGeometryCollectionEdit GCEdit = GeometryCollectionComponent->EditRestCollection(GeometryCollection::EEditUpdate::RestPhysics, true /*bShapeIsUnchanged*/);
+		if (UGeometryCollection* GCObject = GCEdit.GetRestCollection())
+		{
+			TSharedPtr<FGeometryCollection, ESPMode::ThreadSafe> GeometryCollectionPtr = GCObject->GetGeometryCollection();
+			if (FGeometryCollection* GeometryCollection = GeometryCollectionPtr.Get())
+			{
+				GeometryCollection->RemoveAttribute("RemoveOnBreak", FGeometryCollection::TransformGroup);
+			}
+		}
+	}
+}
+
+
+void UFractureRemoveOnBreakSettings::DeleteRemoveOnBreakData()
+{
+	UFractureToolSetRemoveOnBreak* RemoveOnBreakTool = Cast<UFractureToolSetRemoveOnBreak>(OwnerTool.Get());
+	RemoveOnBreakTool->DeleteRemoveOnBreakData();
+}
 
 #undef LOCTEXT_NAMESPACE
diff --git a/Engine/Plugins/Experimental/ChaosEditor/Source/FractureEditor/Private/FractureToolProperties.h b/Engine/Plugins/Experimental/ChaosEditor/Source/FractureEditor/Private/FractureToolProperties.h
--- a/Engine/Plugins/Experimental/ChaosEditor/Source/FractureEditor/Private/FractureToolProperties.h
+++ b/Engine/Plugins/Experimental/ChaosEditor/Source/FractureEditor/Private/FractureToolProperties.h
@@ -60,3 +60,62 @@ public:
 	TObjectPtr<UFractureInitialDynamicStateSettings> StateSettings;
 
 };
+
+UCLASS(config = EditorPerProjectUserSettings)
+class UFractureRemoveOnBreakSettings : public UFractureToolSettings
+{
+public:
+
+	GENERATED_BODY()
+
+	UFractureRemoveOnBreakSettings(const FObjectInitializer& ObjInit)
+		: Super(ObjInit)
+		, Enabled(false)
+		, PostBreakTimer(0,0)
+		, RemovalTimer(2,3)
+	{}
+
+	/** whether or not the remove on fracture is enabled */
+	UPROPERTY(EditAnywhere, Category = SetRemoveOnBreak, meta = (DisplayName = "Enabled"))
+	bool Enabled;
+
+	/** Min/Max time after break before removal starts */
+	UPROPERTY(EditAnywhere, Category = SetRemoveOnBreak, meta = (DisplayName = "Post Break Timer", EditCondition="Enabled"))
+	FVector2f PostBreakTimer;  
+
+	/** Min/Max time for how long removal lasts */
+	UPROPERTY(EditAnywhere, Category = SetRemoveOnBreak, meta = (DisplayName = "Removal Timer", EditCondition="Enabled"))
+	FVector2f RemovalTimer;
+
+	/** remove the remove on break attribute from the geometry collection, usefull to save memory on the asset if not needed */
+	UFUNCTION(CallInEditor, Category = Actions, meta = (DisplayName = "Delete Remove-On-Break Data"))
+	void DeleteRemoveOnBreakData();
+};
+
+UCLASS(DisplayName = "RemoveOnBreak", Category = "FractureTools")
+class UFractureToolSetRemoveOnBreak : public UFractureModalTool
+{
+public:
+	GENERATED_BODY()
+
+	UFractureToolSetRemoveOnBreak(const FObjectInitializer& ObjInit);
+
+	using FGeometryCollectionPtr = TSharedPtr<FGeometryCollection, ESPMode::ThreadSafe>;
+
+	// UFractureTool Interface
+	virtual FText GetDisplayText() const override;
+	virtual FText GetTooltipText() const override;
+	virtual FText GetApplyText() const override;
+	virtual FSlateIcon GetToolIcon() const override;
+	virtual TArray<UObject*> GetSettingsObjects() const override;
+
+	virtual void RegisterUICommand(FFractureEditorCommands* BindingContext) override;
+
+	virtual void Execute(TWeakPtr<FFractureEditorModeToolkit> InToolkit) override;
+
+	void DeleteRemoveOnBreakData();
+	
+	UPROPERTY(EditAnywhere, Category = InitialDynamicState)
+	TObjectPtr<UFractureRemoveOnBreakSettings> RemoveOnBreakSettings;
+
+};
\ No newline at end of file
diff --git a/Engine/Plugins/Experimental/ChaosEditor/Source/FractureEditor/Private/SGeometryCollectionOutliner.cpp b/Engine/Plugins/Experimental/ChaosEditor/Source/FractureEditor/Private/SGeometryCollectionOutliner.cpp
--- a/Engine/Plugins/Experimental/ChaosEditor/Source/FractureEditor/Private/SGeometryCollectionOutliner.cpp
+++ b/Engine/Plugins/Experimental/ChaosEditor/Source/FractureEditor/Private/SGeometryCollectionOutliner.cpp
@@ -195,6 +195,10 @@ TSharedRef<SWidget> SGeometryCollectionOutlinerRow::GenerateWidgetForColumn(cons
 		return Item->MakeDamageThresholdColumnWidget();
 	if (ColumnName == SGeometryCollectionOutlinerColumnID::Broken)
 		return Item->MakeBrokenColumnWidget();
+	if (ColumnName == SGeometryCollectionOutlinerColumnID::PostBreakTime)
+		return Item->MakePostBreakTimecolumnWidget();
+	if (ColumnName == SGeometryCollectionOutlinerColumnID::RemovalTime)
+		return Item->MakeRemovalTimeColumnWidget();
 	return Item->MakeEmptyColumnWidget();
 }
 
@@ -245,13 +249,15 @@ void SGeometryCollectionOutliner::RegenerateHeader()
 		HeaderRowWidget->AddColumn(
 			SHeaderRow::Column(SGeometryCollectionOutlinerColumnID::RelativeSize)
 				.DefaultLabel(LOCTEXT("GCOutliner_Column_RelativeSize", "Relative Size"))
-				.HAlignHeader(EHorizontalAlignment::HAlign_Right)
+				.DefaultTooltip(LOCTEXT("GCOutliner_Column_RelativeSize_ToolTip", "Relative size ( Used for size specific data )"))
+				.HAlignHeader(EHorizontalAlignment::HAlign_Center)
 				.FillWidth(CustomFillWidth)
 		);
 		HeaderRowWidget->AddColumn(
 			SHeaderRow::Column(SGeometryCollectionOutlinerColumnID::InitialState)
 				.DefaultLabel(LOCTEXT("GCOutliner_Column_InitialState", "Initial State"))
-				.HAlignHeader(EHorizontalAlignment::HAlign_Left)
+				.DefaultTooltip(LOCTEXT("GCOutliner_Column_InitialState_ToolTip", "Initial state override"))
+				.HAlignHeader(EHorizontalAlignment::HAlign_Center)
 				.FillWidth(CustomFillWidth)
 		);
 		break;
@@ -259,23 +265,43 @@ void SGeometryCollectionOutliner::RegenerateHeader()
 	case EOutlinerColumnMode::Damage:
 		HeaderRowWidget->AddColumn(
 			SHeaderRow::Column(SGeometryCollectionOutlinerColumnID::Damage)
-				.DefaultLabel(LOCTEXT("GCOutliner_Column_Damage", "Damage"))
-	 			.HAlignHeader(EHorizontalAlignment::HAlign_Right)
+				.DefaultLabel(LOCTEXT("GCOutliner_Column_Damage", "Max damage"))
+				.DefaultTooltip(LOCTEXT("GCOutliner_Column_Damage_ToolTip", "Maximum amount of damage recorded ( through collision )"))
+	 			.HAlignHeader(EHorizontalAlignment::HAlign_Center)
 	 			.FillWidth(CustomFillWidth)
 		);
 		HeaderRowWidget->AddColumn(
 			SHeaderRow::Column(SGeometryCollectionOutlinerColumnID::DamageThreshold)
 				.DefaultLabel(LOCTEXT("GCOutliner_Column_DamageThreshold", "Damage Threshold"))
-				 .HAlignHeader(EHorizontalAlignment::HAlign_Right)
-				 .FillWidth(CustomFillWidth)
+				.DefaultTooltip(LOCTEXT("GCOutliner_Column_DamageThreshold_ToolTip", "Current damage threshold ( equivalent to internal strain )"))
+				.HAlignHeader(EHorizontalAlignment::HAlign_Center)
+				.FillWidth(CustomFillWidth)
 		);
 		HeaderRowWidget->AddColumn(
 			SHeaderRow::Column(SGeometryCollectionOutlinerColumnID::Broken)
 				.DefaultLabel(LOCTEXT("GCOutliner_Column_Broken", "Broken"))
-				 .HAlignHeader(EHorizontalAlignment::HAlign_Center)
-				 .FillWidth(1.0)
+				.DefaultTooltip(LOCTEXT("GCOutliner_Column_Broken_ToolTip", "Whether the piece has broken off"))
+				.HAlignHeader(EHorizontalAlignment::HAlign_Center)
+				.FillWidth(1.0)
 		);
 		break;
+	case EOutlinerColumnMode::Removal:
+		HeaderRowWidget->AddColumn(
+			SHeaderRow::Column(SGeometryCollectionOutlinerColumnID::PostBreakTime)
+				.DefaultLabel(LOCTEXT("GCOutliner_Column_PostBreakTime", "Post Break Time"))
+				.DefaultTooltip(LOCTEXT("GCOutliner_Column_PostBreakTime_ToolTip", "Min/Max time after break until removal starts"))
+				.HAlignHeader(EHorizontalAlignment::HAlign_Center)
+				.FillWidth(CustomFillWidth)
+		);
+		HeaderRowWidget->AddColumn(
+			SHeaderRow::Column(SGeometryCollectionOutlinerColumnID::RemovalTime)
+				.DefaultLabel(LOCTEXT("GCOutliner_Column_RemovalTime", "Removal Time"))
+				.DefaultTooltip(LOCTEXT("GCOutliner_Column_RemovalTime_ToolTip", "Min/Max time for removal"))
+				.HAlignHeader(EHorizontalAlignment::HAlign_Center)
+				.FillWidth(CustomFillWidth)
+		);
+		break;
+		
 	}
 }
 
@@ -721,6 +747,8 @@ void FGeometryCollectionTreeItemBone::UpdateItemFromCollection()
 	DamageThreshold = 0;
 	Broken = false;
 	InitialState = INDEX_NONE;
+	RemoveOnBreakAvailable = false;
+	RemoveOnBreak = FVector4f{-1};
 	
 	// Name / ItemText
 	UOutlinerSettings* OutlinerSettings = GetMutableDefault<UOutlinerSettings>();
@@ -739,6 +767,7 @@ void FGeometryCollectionTreeItemBone::UpdateItemFromCollection()
 		TSharedPtr<FGeometryCollection, ESPMode::ThreadSafe> GeometryCollectionPtr = RestCollection->GetGeometryCollection();
 		const TManagedArray<int32>& SimulationType = GeometryCollectionPtr->SimulationType;
 		const TManagedArray<float>* RelativeSizes = GeometryCollectionPtr->FindAttribute<float>("Size", FTransformCollection::TransformGroup);
+		const TManagedArray<FVector4f>* RemoveOnBreakArray = GeometryCollectionPtr->FindAttribute<FVector4f>("RemoveOnBreak", FTransformCollection::TransformGroup);
 		
 		if (ensure(ItemBoneIndex >= 0 && ItemBoneIndex < SimulationType.Num()))
 		{
@@ -790,6 +819,12 @@ void FGeometryCollectionTreeItemBone::UpdateItemFromCollection()
 				DamageThreshold = DamageData.DamageThreshold;
 				Broken = DamageData.bIsBroken;
 			}
+
+			RemoveOnBreakAvailable = (RemoveOnBreakArray != nullptr); 
+			if (RemoveOnBreakAvailable)
+			{
+				RemoveOnBreak = (*RemoveOnBreakArray)[ItemBoneIndex];
+			}
 		}
 	}
 }
@@ -888,6 +923,45 @@ TSharedRef<SWidget> FGeometryCollectionTreeItemBone::MakeBrokenColumnWidget() co
 			];
 }
 
+static FText FormatRemoveOnBreakTimeData(bool bDataAvailable, float MinTime, float MaxTime)
+{
+	if (bDataAvailable)
+	{
+		if (MinTime>=0)
+		{
+			return FText::Format(LOCTEXT("RemoveOnBreakFormat", "[{0}s, {1}s]"),MinTime, MaxTime);
+		}
+		return FText(LOCTEXT("RemoveOnBreakEmpty", "[-, -]"));
+	}
+	return FText();
+}
+
+TSharedRef<SWidget> FGeometryCollectionTreeItemBone::MakePostBreakTimecolumnWidget() const
+{
+	return SNew(SHorizontalBox)
+		+ SHorizontalBox::Slot()
+			.Padding(12.f, 0.f)
+			.HAlign(HAlign_Center)
+			[
+				SNew(STextBlock)
+				.Text(FormatRemoveOnBreakTimeData(RemoveOnBreakAvailable, RemoveOnBreak.X, RemoveOnBreak.Y))
+				.ColorAndOpacity(ItemColor)
+			];
+}
+
+TSharedRef<SWidget> FGeometryCollectionTreeItemBone::MakeRemovalTimeColumnWidget() const
+{
+	return SNew(SHorizontalBox)
+	+ SHorizontalBox::Slot()
+		.Padding(12.f, 0.f)
+		.HAlign(HAlign_Center)
+		[
+			SNew(STextBlock)
+			.Text(FormatRemoveOnBreakTimeData(RemoveOnBreakAvailable, RemoveOnBreak.Z, RemoveOnBreak.W))
+			.ColorAndOpacity(ItemColor)
+		];
+}
+
 TSharedRef<SWidget> FGeometryCollectionTreeItemBone::MakeInitialStateColumnWidget() const
 {
 	return SNew(SHorizontalBox)
diff --git a/Engine/Plugins/Experimental/ChaosEditor/Source/FractureEditor/Private/SGeometryCollectionOutliner.h b/Engine/Plugins/Experimental/ChaosEditor/Source/FractureEditor/Private/SGeometryCollectionOutliner.h
--- a/Engine/Plugins/Experimental/ChaosEditor/Source/FractureEditor/Private/SGeometryCollectionOutliner.h
+++ b/Engine/Plugins/Experimental/ChaosEditor/Source/FractureEditor/Private/SGeometryCollectionOutliner.h
@@ -39,7 +39,8 @@ UENUM(BlueprintType)
 enum class EOutlinerColumnMode : uint8
 {
 	StateAndSize = 0		UMETA(DisplayName = "State And Size"),
-	Damage = 2				UMETA(DisplayName = "Damage"),
+	Damage = 1				UMETA(DisplayName = "Damage"),
+	Removal = 2				UMETA(DisplayName = "Removal"),
 };
 
 /** Settings for Outliner configuration. **/
@@ -147,7 +148,9 @@ public:
 		, InitialState(INDEX_NONE)
 		, Damage(0)
 		, DamageThreshold(0)
-		, Broken(false) 
+		, Broken(false)
+		, RemoveOnBreakAvailable(false)
+		, RemoveOnBreak(-1)
 	{}
 
 	/** FGeometryCollectionTreeItem interface */
@@ -158,6 +161,8 @@ public:
 	TSharedRef<SWidget> MakeDamageThresholdColumnWidget() const;
 	TSharedRef<SWidget> MakeBrokenColumnWidget() const;
 	TSharedRef<SWidget> MakeInitialStateColumnWidget() const;
+	TSharedRef<SWidget> MakePostBreakTimecolumnWidget() const;
+	TSharedRef<SWidget> MakeRemovalTimeColumnWidget() const;
 	TSharedRef<SWidget> MakeEmptyColumnWidget() const;
 	virtual void GetChildren(FGeometryCollectionTreeItemList& OutChildren) override;
 	virtual int32 GetBoneIndex() const override { return BoneIndex; }
@@ -185,6 +190,8 @@ private:
 	float Damage;
 	float DamageThreshold;
 	bool Broken;
+	bool RemoveOnBreakAvailable;
+	FVector4f RemoveOnBreak;
 };
 
 typedef TSharedPtr<class FGeometryCollectionTreeItemBone> FGeometryCollectionTreeItemBonePtr;
@@ -199,6 +206,9 @@ namespace SGeometryCollectionOutlinerColumnID
 	const FName Damage("Damage");
 	const FName DamageThreshold("DamageThreshold");
 	const FName Broken("Broken");
+	// Removal Column Mode
+	const FName PostBreakTime("PostBreakTime");
+	const FName RemovalTime("RemovalTime");
 }
 
 class SGeometryCollectionOutlinerRow : public SMultiColumnTableRow<FGeometryCollectionTreeItemBonePtr>
diff --git a/Engine/Plugins/Experimental/ChaosEditor/Source/FractureEditor/Public/FractureEditorModeToolkit.h b/Engine/Plugins/Experimental/ChaosEditor/Source/FractureEditor/Public/FractureEditorModeToolkit.h
--- a/Engine/Plugins/Experimental/ChaosEditor/Source/FractureEditor/Public/FractureEditorModeToolkit.h
+++ b/Engine/Plugins/Experimental/ChaosEditor/Source/FractureEditor/Public/FractureEditorModeToolkit.h
@@ -179,6 +179,8 @@ public:
 
 	void UpdateExplodedVectors(UGeometryCollectionComponent* GeometryCollectionComponent) const;
 
+	void RefreshOutliner();
+	
 	void RegenerateOutliner();
 	void RegenerateHistogram();
 
@@ -213,7 +215,7 @@ private:
 	/** Callback for map changes. */
 	void HandleMapChanged(UWorld* NewWorld, EMapChangeType MapChangeType);
 
-	FReply RefreshOutliner();
+	FReply OnRefreshOutlinerButtonClicked();
 	
 private:
 	UFractureModalTool* ActiveTool;
diff --git a/Engine/Source/Runtime/Experimental/Chaos/Public/GeometryCollection/ManagedArrayTypeValues.inl b/Engine/Source/Runtime/Experimental/Chaos/Public/GeometryCollection/ManagedArrayTypeValues.inl
--- a/Engine/Source/Runtime/Experimental/Chaos/Public/GeometryCollection/ManagedArrayTypeValues.inl
+++ b/Engine/Source/Runtime/Experimental/Chaos/Public/GeometryCollection/ManagedArrayTypeValues.inl
@@ -48,7 +48,7 @@ MANAGED_ARRAY_TYPE(FIntVector2, IntVector2)
 MANAGED_ARRAY_TYPE(TArray<FIntVector2>, IntVector2Array)
 MANAGED_ARRAY_TYPE(TArray<int32>, Int32Array)
 MANAGED_ARRAY_TYPE(TArray<float>, FloatArray)
-
+MANAGED_ARRAY_TYPE(FVector4f, Vector4f)
 
 // NOTE: new types must be added at the bottom to keep serialization from breaking
 
diff --git a/Engine/Source/Runtime/Experimental/GeometryCollectionEngine/Private/GeometryCollection/GeometryCollectionComponent.cpp b/Engine/Source/Runtime/Experimental/GeometryCollectionEngine/Private/GeometryCollection/GeometryCollectionComponent.cpp
--- a/Engine/Source/Runtime/Experimental/GeometryCollectionEngine/Private/GeometryCollection/GeometryCollectionComponent.cpp
+++ b/Engine/Source/Runtime/Experimental/GeometryCollectionEngine/Private/GeometryCollection/GeometryCollectionComponent.cpp
@@ -1916,6 +1916,7 @@ void UGeometryCollectionComponent::TickComponent(float DeltaTime, enum ELevelTic
 			{ 
 				IncrementSleepTimer(DeltaTime);
 			}
+			IncrementBreakTimer(DeltaTime);
 
 			if (RestCollection->HasVisibleGeometry() || DynamicCollection->IsDirty())
 			{
@@ -2048,14 +2049,62 @@ void UGeometryCollectionComponent::ResetDynamicCollection()
 				}
 			}
 
-			if (!DynamicCollection->HasAttribute("RemovalDuration", FGeometryCollection::TransformGroup))
+			if (!DynamicCollection->HasAttribute("SleepRemovalDuration", FGeometryCollection::TransformGroup))
 			{
 				float MinTime = FMath::Max(0.0f, RestCollection->RemovalDuration.X);
 				float MaxTime = FMath::Max(MinTime, RestCollection->RemovalDuration.Y);
-				TManagedArray<float>& RemovalDuration = DynamicCollection->AddAttribute<float>("RemovalDuration", FGeometryCollection::TransformGroup);
-				for (int32 Idx = 0; Idx < RemovalDuration.Num(); ++Idx)
+				TManagedArray<float>& SleepRemovalDuration = DynamicCollection->AddAttribute<float>("SleepRemovalDuration", FGeometryCollection::TransformGroup);
+				for (int32 Idx = 0; Idx < SleepRemovalDuration.Num(); ++Idx)
 				{
-					RemovalDuration[Idx] = FMath::RandRange(MinTime, MaxTime);
+					SleepRemovalDuration[Idx] = FMath::RandRange(MinTime, MaxTime);
+				}
+			}
+		}
+		
+		// Remove on break feature related dynamic attribute arrays
+		if (const TManagedArray<FVector4f>* RemoveOnBreak = RestCollection->GetGeometryCollection()->FindAttribute<FVector4f>("RemoveOnBreak", FGeometryCollection::TransformGroup))
+		{
+			// Needs to be checked here too because remove on sleep and remove on break features are orthogonal
+			if (!DynamicCollection->HasAttribute("Decay", FGeometryCollection::TransformGroup))
+			{
+				TManagedArray<float>& Decay = DynamicCollection->AddAttribute<float>("Decay", FGeometryCollection::TransformGroup);
+				Decay.Fill(0);
+			}
+
+			if (!DynamicCollection->HasAttribute("UniformScale", FGeometryCollection::TransformGroup))
+			{
+				TManagedArray<FTransform>& UniformScale = DynamicCollection->AddAttribute<FTransform>("UniformScale", FGeometryCollection::TransformGroup);
+				UniformScale.Fill(FTransform::Identity);
+			}
+		
+			if (!DynamicCollection->HasAttribute("BreakTimer", FGeometryCollection::TransformGroup))
+			{
+				TManagedArray<float>& BreakTimer = DynamicCollection->AddAttribute<float>("BreakTimer", FGeometryCollection::TransformGroup);
+				BreakTimer.Fill(-1); // disabled timer , waiting for the break event
+			}
+			
+			if (!DynamicCollection->HasAttribute("PostBreakDuration", FGeometryCollection::TransformGroup))
+			{
+				TManagedArray<float>& PostBreakDuration = DynamicCollection->AddAttribute<float>("PostBreakDuration", FGeometryCollection::TransformGroup);
+				for (int32 Idx = 0; Idx < PostBreakDuration.Num(); ++Idx)
+				{
+					const FVector4f RemoveOnBreakData = (*RemoveOnBreak)[Idx];
+					const bool bRemoveOnBreakEnabled = RemoveOnBreakData.X >= 0;
+					const float MinTime = FMath::Max(0.0f, RemoveOnBreakData.X);
+					const float MaxTime = FMath::Max(MinTime, RemoveOnBreakData.Y);
+					PostBreakDuration[Idx] = bRemoveOnBreakEnabled? FMath::RandRange(MinTime, MaxTime): -1;
+				}
+			}
+
+			if (!DynamicCollection->HasAttribute("BreakRemovalDuration", FGeometryCollection::TransformGroup))
+			{
+				TManagedArray<float>& BreakRemovalDuration = DynamicCollection->AddAttribute<float>("BreakRemovalDuration", FGeometryCollection::TransformGroup);
+				for (int32 Idx = 0; Idx < BreakRemovalDuration.Num(); ++Idx)
+				{
+					const FVector4f RemoveOnBreakData = (*RemoveOnBreak)[Idx];
+					const float MinTime = FMath::Max(0.0f, RemoveOnBreakData.Z);
+					const float MaxTime = FMath::Max(MinTime, RemoveOnBreakData.W);
+					BreakRemovalDuration[Idx] = FMath::RandRange(MinTime, MaxTime);
 				}
 			}
 		}
@@ -3185,14 +3234,11 @@ void UGeometryCollectionComponent::CalculateGlobalMatrices()
 		{
 			// Have to fully rebuild
 			if (DynamicCollection 
-				&& RestCollection->bRemoveOnMaxSleep 
 				&& DynamicCollection->HasAttribute("UniformScale", FGeometryCollection::TransformGroup)
 				&& DynamicCollection->HasAttribute("Decay", FGeometryCollection::TransformGroup))
 			{
 				const TManagedArray<float>& Decay = DynamicCollection->GetAttribute<float>("Decay", FGeometryCollection::TransformGroup);
 				TManagedArray<FTransform>& UniformScale = DynamicCollection->ModifyAttribute<FTransform>("UniformScale", FGeometryCollection::TransformGroup);
-				const TManagedArray<FBox>& BoundingBoxes = RestCollection->GetGeometryCollection()->BoundingBox;
-				const TManagedArray<int32>& TransformToGeom = RestCollection->GetGeometryCollection()->TransformToGeometryIndex;
 
 				const FTransform InverseComponentTransform = GetComponentTransform().Inverse();
 				const FTransform ZeroScaleTransform(FQuat::Identity, FVector::Zero(), FVector(0, 0, 0));
@@ -3220,10 +3266,6 @@ void UGeometryCollectionComponent::CalculateGlobalMatrices()
 							const FVector LocalDown = LocalRotation.RotateVector(FVector(0.f, 0.f, ShrinkRadius));
 							const FVector CenterOfMass = DynamicCollection->MassToLocal[Idx].GetTranslation();
 							const FVector ScaleCenter = LocalDown + CenterOfMass;
-							//const FTransform ToScaleCenter(ScaleCenter);
-							//const FTransform ToScaleCenterInverse(-ScaleCenter);
-							//UniformScale[Idx] = ToCOM.Inverse() * LocalDown.Inverse() * FTransform(FQuat::Identity, FVector(0.f, 0.f, 0.f), FVector(Scale)) * LocalDown * ToCOM;
-							//UniformScale[Idx] = ToScaleCenterInverse * FTransform(FQuat::Identity, FVector(0.f, 0.f, 0.f), FVector(Scale)) * ToScaleCenter;
 							UniformScale[Idx] = FTransform(FQuat::Identity, ScaleCenter * (FVector::FReal)(1.f - Scale), FVector(Scale));
 						}
 					}
@@ -3457,17 +3499,16 @@ void UGeometryCollectionComponent::IncrementSleepTimer(float DeltaTime)
 	if (DynamicCollection && PhysicsProxy 
 		&& DynamicCollection->HasAttribute("SleepTimer", FGeometryCollection::TransformGroup) 
 		&& DynamicCollection->HasAttribute("MaxSleepTime", FGeometryCollection::TransformGroup)
-		&& DynamicCollection->HasAttribute("RemovalDuration", FGeometryCollection::TransformGroup)
+		&& DynamicCollection->HasAttribute("SleepRemovalDuration", FGeometryCollection::TransformGroup)
 		&& DynamicCollection->HasAttribute("Decay", FGeometryCollection::TransformGroup)
 		&& DynamicCollection->HasAttribute("LastPosition", FGeometryCollection::TransformGroup))
 	{
 		TManagedArray<float>& SleepTimer = DynamicCollection->ModifyAttribute<float>("SleepTimer", FGeometryCollection::TransformGroup);
-		const TManagedArray<float>& RemovalDuration = DynamicCollection->GetAttribute<float>("RemovalDuration", FGeometryCollection::TransformGroup);
+		const TManagedArray<float>& SleepRemovalDuration = DynamicCollection->GetAttribute<float>("SleepRemovalDuration", FGeometryCollection::TransformGroup);
 		const TManagedArray<float>& MaxSleepTime = DynamicCollection->GetAttribute<float>("MaxSleepTime", FGeometryCollection::TransformGroup);
 		TManagedArray<float>& Decay = DynamicCollection->ModifyAttribute<float>("Decay", FGeometryCollection::TransformGroup);
 		TManagedArray<FVector>& LastPosition = DynamicCollection->ModifyAttribute<FVector>("LastPosition", FGeometryCollection::TransformGroup);
 		const TManagedArray<uint8>& InternalClusterParentType = DynamicCollection->GetAttribute<uint8>("InternalClusterParentTypeArray", FGeometryCollection::TransformGroup);
-		const TManagedArray<bool>& Active = DynamicCollection->Active;
 		const TManagedArray<int32>& Parents = DynamicCollection->Parent;
 		
 		TArray<int32> ToDisable;
@@ -3495,12 +3536,10 @@ void UGeometryCollectionComponent::IncrementSleepTimer(float DeltaTime)
 				}
 			
 				// update the decay and disable the particle when decay has completed 
-				const bool bZeroRemovalDuration = (RemovalDuration[TransformIdx] < UE_SMALL_NUMBER);
-				const float UpdatedDecay = bZeroRemovalDuration? 1.f: FMath::Clamp<float>((SleepTimer[TransformIdx] - MaxSleepTime[TransformIdx]) / RemovalDuration[TransformIdx], 0.f, 1.f);
-				if (UpdatedDecay != Decay[TransformIdx])
+				const bool bZeroRemovalDuration = (SleepRemovalDuration[TransformIdx] < UE_SMALL_NUMBER);
+				const float UpdatedDecay = bZeroRemovalDuration? 1.f: FMath::Clamp<float>((SleepTimer[TransformIdx] - MaxSleepTime[TransformIdx]) / SleepRemovalDuration[TransformIdx], 0.f, 1.f);
+				if (UpdatedDecay > Decay[TransformIdx])
 				{
-					ensure(UpdatedDecay > Decay[TransformIdx]);
-
 					// if parent is internal and dynamic we need to break it before decaying 
 					if (InternalClusterParentType[TransformIdx] == (uint8)Chaos::EInternalClusterType::Dynamic)
 					{
@@ -3531,6 +3570,62 @@ void UGeometryCollectionComponent::IncrementSleepTimer(float DeltaTime)
 	}
 }
 
+void UGeometryCollectionComponent::IncrementBreakTimer(float DeltaTime)
+{
+	if (DeltaTime <= 0)
+	{
+		return;
+	}
+	
+	if (DynamicCollection && PhysicsProxy 
+		&& DynamicCollection->HasAttribute("BreakTimer", FGeometryCollection::TransformGroup) 
+		&& DynamicCollection->HasAttribute("PostBreakDuration", FGeometryCollection::TransformGroup)
+		&& DynamicCollection->HasAttribute("BreakRemovalDuration", FGeometryCollection::TransformGroup)
+		&& DynamicCollection->HasAttribute("Decay", FGeometryCollection::TransformGroup)
+		)
+	{
+		TManagedArray<float>& BreakTimer = DynamicCollection->ModifyAttribute<float>("BreakTimer", FGeometryCollection::TransformGroup);
+		const TManagedArray<float>& PostBreakDuration = DynamicCollection->GetAttribute<float>("PostBreakDuration", FGeometryCollection::TransformGroup);
+		const TManagedArray<float>& BreakRemovalDuration = DynamicCollection->GetAttribute<float>("BreakRemovalDuration", FGeometryCollection::TransformGroup);
+		TManagedArray<float>& Decay = DynamicCollection->ModifyAttribute<float>("Decay", FGeometryCollection::TransformGroup);
+		
+		TArray<int32> ToDisable;
+		for (int32 TransformIdx = 0; TransformIdx < BreakTimer.Num(); ++TransformIdx)
+		{
+			const bool bIsBroken = DynamicCollection->Active[TransformIdx];
+			const bool bIsRemovedFromBreakEnabled = (PostBreakDuration[TransformIdx] >= 0);
+			
+			const int32 State = DynamicCollection->DynamicState[TransformIdx];
+			const bool bIsDynamicOrSleeping = (State == (int)EObjectStateTypeEnum::Chaos_Object_Sleeping) || (State == (int)EObjectStateTypeEnum::Chaos_Object_Dynamic);
+			
+			if (bIsBroken && bIsRemovedFromBreakEnabled && bIsDynamicOrSleeping)
+			{
+				BreakTimer[TransformIdx] += DeltaTime;
+				const bool bZeroRemovalDuration = (BreakRemovalDuration[TransformIdx] < UE_SMALL_NUMBER);
+				const float UpdatedBreakDecay = bZeroRemovalDuration? 1.f: FMath::Clamp<float>((BreakTimer[TransformIdx] - PostBreakDuration[TransformIdx]) / BreakRemovalDuration[TransformIdx], 0.f, 1.f);
+				if (UpdatedBreakDecay > Decay[TransformIdx])
+				{
+					DynamicCollection->MakeDirty();
+					Decay[TransformIdx] = UpdatedBreakDecay;
+
+					if (Decay[TransformIdx] >= 1.0f)
+					{
+						// Disable the particle if it has decayed  the requisite time
+						Decay[TransformIdx] = 1.0f;
+						ToDisable.Add(TransformIdx);
+					}
+				}
+			}
+		}
+
+		if (ToDisable.Num())
+		{
+			PhysicsProxy->DisableParticles(MoveTemp(ToDisable));
+		}
+	}
+}
+
+
 bool UGeometryCollectionComponent::CalculateInnerSphere(int32 TransformIndex, FSphere& SphereOut) const
 {
 	// Approximates the inscribed sphere. Returns false if no such sphere exists, if for instance the index is to an embedded geometry. 
diff --git a/Engine/Source/Runtime/Experimental/GeometryCollectionEngine/Public/GeometryCollection/GeometryCollectionComponent.h b/Engine/Source/Runtime/Experimental/GeometryCollectionEngine/Public/GeometryCollection/GeometryCollectionComponent.h
--- a/Engine/Source/Runtime/Experimental/GeometryCollectionEngine/Public/GeometryCollection/GeometryCollectionComponent.h
+++ b/Engine/Source/Runtime/Experimental/GeometryCollectionEngine/Public/GeometryCollection/GeometryCollectionComponent.h
@@ -893,6 +893,7 @@ private:
 	void ClearEmbeddedGeometry();
 
 	void IncrementSleepTimer(float DeltaTime);
+	void IncrementBreakTimer(float DeltaTime);
 	bool CalculateInnerSphere(int32 TransformIndex, FSphere& SphereOut) const;
 	void ProcessRepData();
 
