Dataflow : Node Updates
-  Pass through
- Static Mesh
- Skeletal mesh updates.

diff --git a/Engine/Plugins/Experimental/Dataflow/Source/DataflowEditor/Private/Dataflow/DataflowEditorActions.cpp b/Engine/Plugins/Experimental/Dataflow/Source/DataflowEditor/Private/Dataflow/DataflowEditorActions.cpp
--- a/Engine/Plugins/Experimental/Dataflow/Source/DataflowEditor/Private/Dataflow/DataflowEditorActions.cpp
+++ b/Engine/Plugins/Experimental/Dataflow/Source/DataflowEditor/Private/Dataflow/DataflowEditorActions.cpp
@@ -61,16 +61,19 @@ void FDataflowEditorCommands::EvaluateNodes(const FGraphPanelSelectionSet& Selec
 			{
 				if (const TSharedPtr<FDataflowNode> DataflowNode = DataflowGraph->FindBaseNode(EdNode->GetDataflowNodeGuid()))
 				{
-					if (DataflowNode->GetOutputs().Num())
+					if (DataflowNode->bActive)
 					{
-						for (Dataflow::FConnection* NodeOutput : DataflowNode->GetOutputs())
+						if (DataflowNode->GetOutputs().Num())
 						{
-							Evaluate(DataflowNode.Get(),NodeOutput);
+							for (Dataflow::FConnection* NodeOutput : DataflowNode->GetOutputs())
+							{
+								Evaluate(DataflowNode.Get(), NodeOutput);
+							}
+						}
+						else
+						{
+							Evaluate(DataflowNode.Get(), nullptr);
 						}
-					}
-					else
-					{
-						Evaluate(DataflowNode.Get(), nullptr);
 					}
 				}
 			}
diff --git a/Engine/Plugins/Experimental/Dataflow/Source/DataflowNodes/Private/Dataflow/DataflowSkeletalMeshNodes.cpp b/Engine/Plugins/Experimental/Dataflow/Source/DataflowNodes/Private/Dataflow/DataflowSkeletalMeshNodes.cpp
--- a/Engine/Plugins/Experimental/Dataflow/Source/DataflowNodes/Private/Dataflow/DataflowSkeletalMeshNodes.cpp
+++ b/Engine/Plugins/Experimental/Dataflow/Source/DataflowNodes/Private/Dataflow/DataflowSkeletalMeshNodes.cpp
@@ -20,12 +20,16 @@ namespace Dataflow
 void FGetSkeletalMeshDataflowNode::Evaluate(const Dataflow::FContext& Context, Dataflow::FConnection* Out) const
 {
 	SkeletalMeshOut->SetValue(nullptr, Context);
-	if (const Dataflow::FEngineContext* EngineContext = (const Dataflow::FEngineContext*)(&Context))
+	if (SkeletalMesh)
 	{
-		if (const USkeletalMesh* SkeletalMesh = Dataflow::Reflection::FindObjectPtrProperty<USkeletalMesh>(
-			EngineContext->Owner, SkeletalMeshAttributeName))
+		SkeletalMeshOut->SetValue(SkeletalMesh, Context);
+	}
+	else if (const Dataflow::FEngineContext* EngineContext = Context.AsType<Dataflow::FEngineContext>(FName("UFleshAsset")))
+	{
+		if (const USkeletalMesh* SkeletalMeshFromOwner = Dataflow::Reflection::FindObjectPtrProperty<USkeletalMesh>(
+			EngineContext->Owner, PropertyName))
 		{
-			SkeletalMeshOut->SetValue(SkeletalMesh, Context);
+			SkeletalMeshOut->SetValue(SkeletalMeshFromOwner, Context);
 		}
 	}
 }
diff --git a/Engine/Plugins/Experimental/Dataflow/Source/DataflowNodes/Private/Dataflow/DataflowStaticMeshNodes.cpp b/Engine/Plugins/Experimental/Dataflow/Source/DataflowNodes/Private/Dataflow/DataflowStaticMeshNodes.cpp
--- a/Engine/Plugins/Experimental/Dataflow/Source/DataflowNodes/Private/Dataflow/DataflowStaticMeshNodes.cpp
+++ b/Engine/Plugins/Experimental/Dataflow/Source/DataflowNodes/Private/Dataflow/DataflowStaticMeshNodes.cpp
@@ -1,12 +1,30 @@
 // Copyright Epic Games, Inc. All Rights Reserved.
 
 #include "Dataflow/DataflowStaticMeshNodes.h"
+#include "Engine/StaticMesh.h"
 
 namespace Dataflow
 {
 	void RegisterStaticMeshNodes()
 	{
-		//DATAFLOW_NODE_REGISTER_CREATION_FACTORY(ExampleNode);
+		DATAFLOW_NODE_REGISTER_CREATION_FACTORY(FGetStaticMeshDataflowNode);
+	}
+}
+
+void FGetStaticMeshDataflowNode::Evaluate(const Dataflow::FContext& Context, Dataflow::FConnection* Out) const
+{
+	StaticMeshOut->SetValue(nullptr, Context);
+	if (StaticMesh)
+	{
+		StaticMeshOut->SetValue(StaticMesh, Context);
+	}
+	else if (const Dataflow::FEngineContext* EngineContext = Context.AsType<Dataflow::FEngineContext>(FName("UFleshAsset")))
+	{
+		if (const UStaticMesh* StaticMeshFromOwner = Dataflow::Reflection::FindObjectPtrProperty<UStaticMesh>(
+			EngineContext->Owner, PropertyName))
+		{
+			StaticMeshOut->SetValue(StaticMeshFromOwner, Context);
+		}
 	}
 }
 
diff --git a/Engine/Plugins/Experimental/Dataflow/Source/DataflowNodes/Public/Dataflow/DataflowNodesConnectionTypes.h b/Engine/Plugins/Experimental/Dataflow/Source/DataflowNodes/Public/Dataflow/DataflowNodesConnectionTypes.h
--- a/Engine/Plugins/Experimental/Dataflow/Source/DataflowNodes/Public/Dataflow/DataflowNodesConnectionTypes.h
+++ b/Engine/Plugins/Experimental/Dataflow/Source/DataflowNodes/Public/Dataflow/DataflowNodesConnectionTypes.h
@@ -5,12 +5,15 @@
 #include "CoreMinimal.h"
 #include "Dataflow/DataflowInputOutput.h"
 
+class UStaticMesh;
 class USkeletalMesh;
 
 namespace Dataflow
 {
-	typedef const USkeletalMesh* USkeletalMeshPtr;	
+	typedef const UStaticMesh* UStaticMeshPtr;
+	typedef const USkeletalMesh* USkeletalMeshPtr;
 	//typedef TSharedPtr<FManagedArrayCollection> FManagedArrayCollectionSharedPtr;
+	DATAFLOW_CONNECTION_TYPE_PRIMITIVE(DATAFLOWNODES_API, UStaticMeshPtr, StaticMeshPtr);
 	DATAFLOW_CONNECTION_TYPE_PRIMITIVE(DATAFLOWNODES_API, USkeletalMeshPtr, SkeletalMeshPtr);
 	DATAFLOW_CONNECTION_TYPE_PRIMITIVE(DATAFLOWNODES_API, FString, String);
 	DATAFLOW_CONNECTION_TYPE_PRIMITIVE(DATAFLOWNODES_API, FVector3f, Vector3f);
diff --git a/Engine/Plugins/Experimental/Dataflow/Source/DataflowNodes/Public/Dataflow/DataflowSkeletalMeshNodes.h b/Engine/Plugins/Experimental/Dataflow/Source/DataflowNodes/Public/Dataflow/DataflowSkeletalMeshNodes.h
--- a/Engine/Plugins/Experimental/Dataflow/Source/DataflowNodes/Public/Dataflow/DataflowSkeletalMeshNodes.h
+++ b/Engine/Plugins/Experimental/Dataflow/Source/DataflowNodes/Public/Dataflow/DataflowSkeletalMeshNodes.h
@@ -23,10 +23,13 @@ struct DATAFLOWNODES_API FGetSkeletalMeshDataflowNode: public FDataflowNode
 	DATAFLOW_NODE_DEFINE_INTERNAL(FGetSkeletalMeshDataflowNode, "SkeletalMesh", "Dataflow", "Skeletal Mesh")
 
 public:
+	
+	UPROPERTY(EditAnywhere, Category = "Dataflow")
+	TObjectPtr<USkeletalMesh> SkeletalMesh = nullptr;
 
 
-	UPROPERTY(EditAnywhere, Category = "Dataflow")
-	FName SkeletalMeshAttributeName;
+	UPROPERTY(EditAnywhere, Category = "Dataflow" )
+	FName PropertyName = "SkeletalMesh";
 
 	TSharedPtr< class Dataflow::TOutput<Dataflow::USkeletalMeshPtr> > SkeletalMeshOut;
 
@@ -55,7 +58,7 @@ public:
 	FSkeletalMeshBoneDataflowNode(const Dataflow::FNodeParameters& InParam, FGuid InGuid = FGuid::NewGuid())
 		: FDataflowNode(InParam, InGuid)
 		, SkeletalMeshIn(new Dataflow::TInput<Dataflow::USkeletalMeshPtr>(Dataflow::TInputParameters<Dataflow::USkeletalMeshPtr>({ FName("SkeletalMeshIn"), this })))
-		, BoneIndexOut(new Dataflow::TOutput<int>(Dataflow::TOutputParameters<int>({ FName("BoneIndexOut"), this })))
+		, BoneIndexOut(new Dataflow::TOutput<int>(Dataflow::TOutputParameters<int>({ FName("BoneIndexOut"), this, INDEX_NONE })))
 	{}
 
 
diff --git a/Engine/Plugins/Experimental/Dataflow/Source/DataflowNodes/Public/Dataflow/DataflowStaticMeshNodes.h b/Engine/Plugins/Experimental/Dataflow/Source/DataflowNodes/Public/Dataflow/DataflowStaticMeshNodes.h
--- a/Engine/Plugins/Experimental/Dataflow/Source/DataflowNodes/Public/Dataflow/DataflowStaticMeshNodes.h
+++ b/Engine/Plugins/Experimental/Dataflow/Source/DataflowNodes/Public/Dataflow/DataflowStaticMeshNodes.h
@@ -4,32 +4,48 @@
 
 #include "CoreMinimal.h"
 
-namespace Dataflow
+
+#include "Dataflow/DataflowNodesConnectionTypes.h"
+#include "Dataflow/DataflowCore.h"
+#include "Dataflow/DataflowEngine.h"
+
+
+#include "DataflowStaticMeshNodes.generated.h"
+
+
+DEFINE_LOG_CATEGORY_STATIC(LogDataflowStaticMeshNodes, Log, All);
+
+class UStaticMesh;
+
+USTRUCT()
+struct DATAFLOWNODES_API FGetStaticMeshDataflowNode : public FDataflowNode
 {
+	GENERATED_USTRUCT_BODY()
+		DATAFLOW_NODE_DEFINE_INTERNAL(FGetStaticMeshDataflowNode, "StaticMesh", "Dataflow", "Static Mesh")
 
-	/*
-	class CHAOSFLESHENGINENODES_API ExampleNode : public FNode
-	{
-		DATAFLOW_NODE_DEFINE_INTERNAL(ExampleNode)
+public:
 
-	public:
-		typede ExampleDataSharedPtr DataType;
+	UPROPERTY(EditAnywhere, Category = "Dataflow")
+		TObjectPtr<UStaticMesh> StaticMesh = nullptr;
 
-		TSharedPtr< class TOutput<DataType> > Output;
 
-		ExampleNode(const FNodeParameters& InParam, FGuid InGuid = FGuid::NewGuid())
-			: FNode(InParam, InGuid)
-			, Output(new TOutput<DataType>(TOutputParameters<DataType>({ FName("ExampleOut"), this })))
-		{}
+	UPROPERTY(EditAnywhere, Category = "Dataflow")
+		FName PropertyName = "StaticMesh";
 
-		virtual void Evaluate(const FContext& Context, FConnection* Out) const override
-		{
-			Output->(<set example data>, Context);
-		}
-	};
-	*/
+	TSharedPtr< class Dataflow::TOutput<Dataflow::UStaticMeshPtr> > StaticMeshOut;
+
+	FGetStaticMeshDataflowNode(const Dataflow::FNodeParameters& InParam, FGuid InGuid = FGuid::NewGuid())
+		: FDataflowNode(InParam, InGuid)
+		, StaticMeshOut(new Dataflow::TOutput<Dataflow::UStaticMeshPtr>(Dataflow::TOutputParameters<Dataflow::UStaticMeshPtr>({ FName("StaticMeshOut"), this })))
+	{}
+
+
+	virtual void Evaluate(const Dataflow::FContext& Context, Dataflow::FConnection* Out) const override;
+};
 
-	void RegisterStaticMeshNodes();
 
+namespace Dataflow
+{
+	void RegisterStaticMeshNodes();
 }
 
diff --git a/Engine/Source/Runtime/Experimental/Dataflow/Core/Public/Dataflow/DataflowInputOutput.h b/Engine/Source/Runtime/Experimental/Dataflow/Core/Public/Dataflow/DataflowInputOutput.h
--- a/Engine/Source/Runtime/Experimental/Dataflow/Core/Public/Dataflow/DataflowInputOutput.h
+++ b/Engine/Source/Runtime/Experimental/Dataflow/Core/Public/Dataflow/DataflowInputOutput.h
@@ -6,6 +6,7 @@
 #include "Dataflow/DataflowNodeParameters.h"
 #include "Dataflow/DataflowNode.h"
 #include "Dataflow/DataflowConnection.h"
+#include "Templates/Function.h"
 
 namespace Dataflow
 {
@@ -25,7 +26,7 @@ namespace Dataflow
 		FName Type;
 		FName Name;
 		FDataflowNode* Owner = nullptr;
-		T Default;
+		T Default = T();
 	};
 
 	template<class T>
@@ -34,7 +35,7 @@ namespace Dataflow
 		typedef FConnection Super;
 		friend class FConnection;
 
-		T Default;
+		T Default = T();
 		TOutput<T>* Connection;
 	public:
 		TInput(const TInputParameters<T>& Param, FGuid InGuid = FGuid::NewGuid())
@@ -112,29 +113,47 @@ namespace Dataflow
 	//
 
 	template<class T>
-	struct DATAFLOWCORE_API TOutputParameters {
-		TOutputParameters(FName InName, FDataflowNode * InOwner) 
+	struct DATAFLOWCORE_API TOutputParameters 
+	{
+		typedef TFunction<void(const  Dataflow::FContext& Context, Dataflow::FConnection*)> PassthroughCallback;
+
+		TOutputParameters(FName InName, FDataflowNode* InOwner, T InDefault = T())
 			: Type(GraphConnectionTypeName<T>())
 			, Name(InName)
-			, Owner(InOwner) {}
+			, Owner(InOwner)
+			, Default(InDefault){}
+
+		TOutputParameters(FName InName, FDataflowNode* InOwner, T InDefault, PassthroughCallback InPassthrough)
+			: Type(GraphConnectionTypeName<T>())
+			, Name(InName)
+			, Owner(InOwner)
+			, Default(InDefault)
+			, Passthrough(InPassthrough) {}
 		FName Type;
 		FName Name;
 		FDataflowNode* Owner = nullptr;
+		T Default = T();
+		PassthroughCallback Passthrough;
 	};
 
 	template<class T>
 	class DATAFLOWCORE_API TOutput : public FConnection
 	{
+		typedef TFunction<void(const  Dataflow::FContext& Context, Dataflow::FConnection*)> PassthroughCallback;
 		typedef FConnection Super;
 		friend class FConnection;
 
 		uint32 CacheKey = UINT_MAX;
 		TCacheValue<T> Cache;
 		TArray< TInput<T>* > Connections;
-	
+		PassthroughCallback Passthrough;
+		T Default = T();
+
 	public:
 		TOutput(const TOutputParameters<T>& Param, FGuid InGuid = FGuid::NewGuid())
 			: FConnection(FPin::EDirection::OUTPUT, Param.Type, Param.Name, Param.Owner, InGuid)
+			, Passthrough(Param.Passthrough)
+			, Default(Param.Default)
 		{
 			Super::BindOutput(Param.Owner, this);
 		}
@@ -174,7 +193,18 @@ namespace Dataflow
 		{
 			if (CacheKey != Context.GetTypeHash())
 			{
-				OwningNode->Evaluate(Context, this);
+				if (OwningNode->bActive)
+				{
+					OwningNode->Evaluate(Context, this);
+				}
+				else
+				{
+					SetValue(Default, Context);
+					if (Passthrough)
+					{
+						Passthrough(Context, this);
+					}
+				}
 			}
 			ensure(CacheKey == Context.GetTypeHash());
 			return DeepCopy<T>(Cache.Data);
diff --git a/Engine/Source/Runtime/Experimental/Dataflow/Core/Public/Dataflow/DataflowNode.h b/Engine/Source/Runtime/Experimental/Dataflow/Core/Public/Dataflow/DataflowNode.h
--- a/Engine/Source/Runtime/Experimental/Dataflow/Core/Public/Dataflow/DataflowNode.h
+++ b/Engine/Source/Runtime/Experimental/Dataflow/Core/Public/Dataflow/DataflowNode.h
@@ -34,6 +34,9 @@ struct DATAFLOWCORE_API FDataflowNode
 	TArray< Dataflow::FConnection* > Inputs;
 	TArray< Dataflow::FConnection* > Outputs;
 
+	UPROPERTY(EditAnywhere, Category = "Dataflow")
+	bool bActive = true;
+
 	FDataflowNode()
 		: Guid(FGuid())
 		, Name("Invalid")
