Performance Improvement and Hitch Removal: AABBTree remove unnecessary memory frees for LeafArrays

diff --git a/Engine/Source/Runtime/Experimental/Chaos/Public/Chaos/AABBTree.h b/Engine/Source/Runtime/Experimental/Chaos/Public/Chaos/AABBTree.h
--- a/Engine/Source/Runtime/Experimental/Chaos/Public/Chaos/AABBTree.h
+++ b/Engine/Source/Runtime/Experimental/Chaos/Public/Chaos/AABBTree.h
@@ -430,6 +430,88 @@ FChaosArchive& operator<<(FChaosArchive& Ar, TAABBTreeLeafArray<TPayloadType, bC
 	return Ar;
 }
 
+
+// Default container behaviour is a that of a TArray
+template<typename LeafType>
+class TLeafContainer  : public TArray<LeafType>
+{
+public:
+	void Serialize(FChaosArchive& Ar)
+	{
+		Ar << *static_cast<TArray<LeafType>*>(this);
+	}
+};
+
+
+// Here we are specializing the behaviour of our leaf container to minimize memory allocations/deallocations when the container is reset. 
+// This is accomplished by only resetting the leafArrays and not the whole container
+// This is only specialized for one specific type, but can expanded to more types if necessary
+template<>
+class TLeafContainer<TAABBTreeLeafArray<FAccelerationStructureHandle, true, FReal >> : private TArray<TAABBTreeLeafArray<FAccelerationStructureHandle, true, FReal>>
+{
+	
+private:
+	typedef TAABBTreeLeafArray<FAccelerationStructureHandle, true, FReal> FLeafType;
+	using FParent = TArray<FLeafType>;
+public:
+
+	using ElementType = FParent::ElementType;
+	
+
+	FLeafType& operator[](SizeType Index)
+	{
+		return FParent::operator[](Index);
+	}
+	const FLeafType& operator[](SizeType Index) const
+	{
+		return FParent::operator[](Index);
+	}
+	SizeType Num() const
+	{
+		return NumOfValidElements;
+	}
+	void Reserve(SizeType Number)
+	{
+		FParent::Reserve(Number);
+	}
+	void Reset()
+	{
+		for (int32 ElementIndex = 0; ElementIndex < NumOfValidElements; ElementIndex++)
+		{
+			(*this)[ElementIndex].Reset();
+		}
+		NumOfValidElements = 0;		
+	}
+	SizeType Add(const FLeafType& Item)
+	{
+		NumOfValidElements++;
+		if (NumOfValidElements > FParent::Num())
+		{
+			FParent::Add(Item);
+		}
+		else
+		{
+			(*this)[NumOfValidElements - 1] = Item;
+		}		
+		return NumOfValidElements - 1;
+	}
+
+	void Serialize(FChaosArchive& Ar)
+	{
+		ensure(false); // This type does not get serialized
+	}
+private:
+	int32 NumOfValidElements = 0;
+};
+
+template <typename LeafType>
+FChaosArchive& operator<<(FChaosArchive& Ar, TLeafContainer<LeafType>& LeafArray)
+{
+	LeafArray.Serialize(Ar);
+
+	return Ar;
+}
+
 template <typename T>
 struct TAABBTreeNode
 {
@@ -958,8 +1040,9 @@ public:
 	void DynamicTreeDebugStats()
 	{
 		TArray<FElement> AllElements;
-		for (auto& Leaf : Leaves)
+		for (int LeafIndex = 0; LeafIndex < Leaves.Num(); LeafIndex++)
 		{
+			const TLeafType& Leaf = Leaves[LeafIndex];
 			Leaf.GatherElements(AllElements);
 		}
 
@@ -1630,8 +1713,10 @@ public:
 		TreeExpensiveStats.StatMaxDirtyElements = DirtyElements.Num();
 		TreeExpensiveStats.StatMaxNumLeaves = Leaves.Num();
 		int32 StatMaxLeafSize = 0;
-		for(const TLeafType& Leaf : Leaves)
+		for (int LeafIndex = 0; LeafIndex < Leaves.Num(); LeafIndex++)
 		{
+			const TLeafType& Leaf = Leaves[LeafIndex];
+
 			StatMaxLeafSize = FMath::Max(StatMaxLeafSize, (int32)Leaf.GetElementCount());
 		}
 		TreeExpensiveStats.StatMaxLeafSize = StatMaxLeafSize;
@@ -2140,8 +2225,9 @@ private:
 		TArray<FElement> AllElements;
 
 		int32 ReserveCount = DirtyElements.Num() + GlobalPayloads.Num();
-		for (const auto& Leaf : Leaves)
+		for (int LeafIndex = 0; LeafIndex < Leaves.Num(); LeafIndex++)
 		{
+			const TLeafType& Leaf = Leaves[LeafIndex];
 			ReserveCount += static_cast<int32>(Leaf.GetReserveCount());
 		}
 
@@ -2150,8 +2236,9 @@ private:
 		AllElements.Append(DirtyElements);
 		AllElements.Append(GlobalPayloads);
 
-		for (auto& Leaf : Leaves)
+		for (int LeafIndex = 0; LeafIndex < Leaves.Num(); LeafIndex++)
 		{
+			TLeafType& Leaf = Leaves[LeafIndex];
 			Leaf.GatherElements(AllElements);
 		}
 
@@ -3064,8 +3151,9 @@ private:
 			{
 				Nodes[0].DebugDraw(*InInterface, Nodes, { 1.f, 1.f, 1.f }, 5.f);
 			}
-			for (const TLeafType& Leaf : Leaves)
+			for (int LeafIndex = 0; LeafIndex < Leaves.Num(); LeafIndex++)
 			{
+				const TLeafType& Leaf = Leaves[LeafIndex];
 				Leaf.DebugDrawLeaf(*InInterface, FLinearColor::MakeRandomColor(), 10.f);
 			}
 		}
@@ -3163,7 +3251,7 @@ private:
 	}
 
 	TArray<FNode> Nodes;
-	TArray<TLeafType> Leaves;
+	TLeafContainer<TLeafType> Leaves;
 	TArray<FElement> DirtyElements;
 
 	// DynamicTree members
