Remove on break
- Added cluster crumbling feature
- Fix localization warnings
- Make sure the delete attribute button refrehes the outliner
- Helper class to encode and decode remval on break packed data

diff --git a/Engine/Plugins/Experimental/ChaosEditor/Source/FractureEditor/Private/FractureToolProperties.cpp b/Engine/Plugins/Experimental/ChaosEditor/Source/FractureEditor/Private/FractureToolProperties.cpp
--- a/Engine/Plugins/Experimental/ChaosEditor/Source/FractureEditor/Private/FractureToolProperties.cpp
+++ b/Engine/Plugins/Experimental/ChaosEditor/Source/FractureEditor/Private/FractureToolProperties.cpp
@@ -123,6 +123,7 @@ void UFractureToolSetRemoveOnBreak::RegisterUICommand(FFractureEditorCommands* B
 
 void UFractureToolSetRemoveOnBreak::Execute(TWeakPtr<FFractureEditorModeToolkit> InToolkit)
 {
+	Toolkit = InToolkit;
 	if (InToolkit.IsValid())
 	{
 	 	TSet<UGeometryCollectionComponent*> GeomCompSelection;
@@ -138,27 +139,21 @@ void UFractureToolSetRemoveOnBreak::Execute(TWeakPtr<FFractureEditorModeToolkit>
 	 				TArray<int32> SelectedBones = GeometryCollectionComponent->GetSelectedBones();
 	 				if (SelectedBones.Num())
 	 				{
-	 					const FVector4f DisabledRemoveOnBreakData{ -1, -1, -1, -1 }; 
 	 					if (!GeometryCollection->HasAttribute("RemoveOnBreak", FGeometryCollection::TransformGroup))
 	 					{
 	 						TManagedArray<FVector4f>& NewRemoveOnBreak = GeometryCollection->AddAttribute<FVector4f>("RemoveOnBreak", FGeometryCollection::TransformGroup);
-	 						NewRemoveOnBreak.Fill(DisabledRemoveOnBreakData);
+	 						NewRemoveOnBreak.Fill(FRemoveOnBreakData::DisabledPackedData);
 	 					}
 	 					
 	 					TManagedArray<FVector4f>& RemoveOnBreak = GeometryCollection->ModifyAttribute<FVector4f>("RemoveOnBreak", FGeometryCollection::TransformGroup);
-	 					
-	 					FVector4f RemoveOnBreakParameters{DisabledRemoveOnBreakData};
-	 					if (RemoveOnBreakSettings->Enabled)
-	 					{
-	 						RemoveOnBreakParameters.X = RemoveOnBreakSettings->PostBreakTimer.X;
-	 						RemoveOnBreakParameters.Y = RemoveOnBreakSettings->PostBreakTimer.Y;
-	 						RemoveOnBreakParameters.Z = RemoveOnBreakSettings->RemovalTimer.X;
-	 						RemoveOnBreakParameters.W = RemoveOnBreakSettings->RemovalTimer.Y;
-	 					}
+
+	 					const FVector2f BreakTimer{ RemoveOnBreakSettings->PostBreakTimer.X, RemoveOnBreakSettings->PostBreakTimer.Y };
+	 					const FVector2f RemovalTimer{ RemoveOnBreakSettings->RemovalTimer.X, RemoveOnBreakSettings->RemovalTimer.Y };
+	 					const FRemoveOnBreakData RemoveOnBreakData(RemoveOnBreakSettings->Enabled,  BreakTimer, RemoveOnBreakSettings->ClusterCrumbling, RemovalTimer);
 
 	 					for (int32 Index : SelectedBones)
 	 					{
-	 						RemoveOnBreak[Index] = RemoveOnBreakParameters;
+	 						RemoveOnBreak[Index] = RemoveOnBreakData.GetPackedData();
 	 					}
 	 				}
 	 			}
@@ -184,6 +179,10 @@ void UFractureToolSetRemoveOnBreak::DeleteRemoveOnBreakData()
 			}
 		}
 	}
+	if (Toolkit.IsValid())
+	{
+		Toolkit.Pin()->RefreshOutliner();
+	}
 }
 
 
diff --git a/Engine/Plugins/Experimental/ChaosEditor/Source/FractureEditor/Private/FractureToolProperties.h b/Engine/Plugins/Experimental/ChaosEditor/Source/FractureEditor/Private/FractureToolProperties.h
--- a/Engine/Plugins/Experimental/ChaosEditor/Source/FractureEditor/Private/FractureToolProperties.h
+++ b/Engine/Plugins/Experimental/ChaosEditor/Source/FractureEditor/Private/FractureToolProperties.h
@@ -70,9 +70,10 @@ public:
 
 	UFractureRemoveOnBreakSettings(const FObjectInitializer& ObjInit)
 		: Super(ObjInit)
-		, Enabled(false)
+		, Enabled(true)
 		, PostBreakTimer(0,0)
 		, RemovalTimer(2,3)
+		, ClusterCrumbling(false)
 	{}
 
 	/** whether or not the remove on fracture is enabled */
@@ -86,6 +87,10 @@ public:
 	/** Min/Max time for how long removal lasts */
 	UPROPERTY(EditAnywhere, Category = SetRemoveOnBreak, meta = (DisplayName = "Removal Timer", EditCondition="Enabled"))
 	FVector2f RemovalTimer;
+	
+	/** When set, clusters will crumble when post break timer expires, non clusters will simply use the removal timer */
+	UPROPERTY(EditAnywhere, Category = SetRemoveOnBreak, meta = (DisplayName = "Cluster Crumbling"))
+	bool ClusterCrumbling;
 
 	/** remove the remove on break attribute from the geometry collection, usefull to save memory on the asset if not needed */
 	UFUNCTION(CallInEditor, Category = Actions, meta = (DisplayName = "Delete Remove-On-Break Data"))
@@ -117,5 +122,7 @@ public:
 	
 	UPROPERTY(EditAnywhere, Category = InitialDynamicState)
 	TObjectPtr<UFractureRemoveOnBreakSettings> RemoveOnBreakSettings;
-
+	
+private:
+	TWeakPtr<FFractureEditorModeToolkit> Toolkit;
 };
\ No newline at end of file
diff --git a/Engine/Plugins/Experimental/ChaosEditor/Source/FractureEditor/Private/SGeometryCollectionOutliner.cpp b/Engine/Plugins/Experimental/ChaosEditor/Source/FractureEditor/Private/SGeometryCollectionOutliner.cpp
--- a/Engine/Plugins/Experimental/ChaosEditor/Source/FractureEditor/Private/SGeometryCollectionOutliner.cpp
+++ b/Engine/Plugins/Experimental/ChaosEditor/Source/FractureEditor/Private/SGeometryCollectionOutliner.cpp
@@ -25,10 +25,10 @@ static FText GettextFromInitialDynamicState(int32 InitialDynamicState)
 {
 	switch (InitialDynamicState)
 	{
-	case 0: return NSLOCTEXT("Fracture", "FractureToolInitialDynamicStateNoOverride", "-");
-	case 1: return NSLOCTEXT("Fracture", "FractureToolInitialDynamicStateSleeping", "Sleeping");
-	case 2: return NSLOCTEXT("Fracture", "FractureToolInitialDynamicStateKinematic", "Kinematic");
-	case 3: return NSLOCTEXT("Fracture", "FractureToolInitialDynamicStateStatic", "Static");
+	case 0: return LOCTEXT("GCOutliner_InitialState_NoOverride_Text", "-");
+	case 1: return LOCTEXT("GCOutliner_InitialState_Sleeping_Text", "Sleeping");
+	case 2: return LOCTEXT("GCOutliner_InitialState_Kinematic_Text", "Kinematic");
+	case 3: return LOCTEXT("GCOutliner_InitialState_Static_Text", "Static");
 	default:
 		return FText();
 	}
@@ -847,6 +847,8 @@ void FGeometryCollectionTreeItemBone::UpdateItemFromCollection()
 			{
 				ImportedCollisionsAvailable = (*ExternalCollisions)[ItemBoneIndex] != nullptr;
 			}
+
+			IsCluster = (GeometryCollectionPtr->Children[ItemBoneIndex].Num() > 0);
 		}
 	}
 }
@@ -945,13 +947,13 @@ TSharedRef<SWidget> FGeometryCollectionTreeItemBone::MakeBrokenColumnWidget() co
 			];
 }
 
-static FText FormatRemoveOnBreakTimeData(bool bDataAvailable, float MinTime, float MaxTime)
+static FText FormatRemoveOnBreakTimeData(bool bDataAvailable, bool bEnabled, const FVector2f& Timer)
 {
 	if (bDataAvailable)
 	{
-		if (MinTime>=0)
+		if (bEnabled)
 		{
-			return FText::Format(LOCTEXT("RemoveOnBreakFormat", "[{0}s, {1}s]"),MinTime, MaxTime);
+			return FText::Format(LOCTEXT("RemoveOnBreakFormat", "[{0}s, {1}s]"),Timer.X, Timer.Y);
 		}
 		return FText(LOCTEXT("RemoveOnBreakEmpty", "[-, -]"));
 	}
@@ -966,20 +968,21 @@ TSharedRef<SWidget> FGeometryCollectionTreeItemBone::MakePostBreakTimecolumnWidg
 			.HAlign(HAlign_Center)
 			[
 				SNew(STextBlock)
-				.Text(FormatRemoveOnBreakTimeData(RemoveOnBreakAvailable, RemoveOnBreak.X, RemoveOnBreak.Y))
+				.Text(FormatRemoveOnBreakTimeData(RemoveOnBreakAvailable, RemoveOnBreak.IsEnabled(), RemoveOnBreak.GetBreakTimer()))
 				.ColorAndOpacity(ItemColor)
 			];
 }
 
 TSharedRef<SWidget> FGeometryCollectionTreeItemBone::MakeRemovalTimeColumnWidget() const
 {
+	const bool EnableRemovalTimer = RemoveOnBreak.IsEnabled() && !(IsCluster && RemoveOnBreak.GetClusterCrumbling());  
 	return SNew(SHorizontalBox)
 	+ SHorizontalBox::Slot()
 		.Padding(12.f, 0.f)
 		.HAlign(HAlign_Center)
 		[
 			SNew(STextBlock)
-			.Text(FormatRemoveOnBreakTimeData(RemoveOnBreakAvailable, RemoveOnBreak.Z, RemoveOnBreak.W))
+			.Text(FormatRemoveOnBreakTimeData(RemoveOnBreakAvailable, EnableRemovalTimer,  RemoveOnBreak.GetRemovalTimer()))
 			.ColorAndOpacity(ItemColor)
 		];
 }
diff --git a/Engine/Plugins/Experimental/ChaosEditor/Source/FractureEditor/Private/SGeometryCollectionOutliner.h b/Engine/Plugins/Experimental/ChaosEditor/Source/FractureEditor/Private/SGeometryCollectionOutliner.h
--- a/Engine/Plugins/Experimental/ChaosEditor/Source/FractureEditor/Private/SGeometryCollectionOutliner.h
+++ b/Engine/Plugins/Experimental/ChaosEditor/Source/FractureEditor/Private/SGeometryCollectionOutliner.h
@@ -151,7 +151,7 @@ public:
 		, DamageThreshold(0)
 		, Broken(false)
 		, RemoveOnBreakAvailable(false)
-		, RemoveOnBreak(-1)
+		, RemoveOnBreak(FRemoveOnBreakData::DisabledPackedData)
 		, ImportedCollisionsAvailable(false)
 		, ImportedCollisionsUsed(false)
 	{}
@@ -195,7 +195,8 @@ private:
 	float DamageThreshold;
 	bool Broken;
 	bool RemoveOnBreakAvailable;
-	FVector4f RemoveOnBreak;
+	bool IsCluster;
+	FRemoveOnBreakData RemoveOnBreak;
 	bool ImportedCollisionsAvailable;
 	bool ImportedCollisionsUsed;
 };
diff --git a/Engine/Source/Runtime/Experimental/Chaos/Private/PhysicsProxy/GeometryCollectionPhysicsProxy.cpp b/Engine/Source/Runtime/Experimental/Chaos/Private/PhysicsProxy/GeometryCollectionPhysicsProxy.cpp
--- a/Engine/Source/Runtime/Experimental/Chaos/Private/PhysicsProxy/GeometryCollectionPhysicsProxy.cpp
+++ b/Engine/Source/Runtime/Experimental/Chaos/Private/PhysicsProxy/GeometryCollectionPhysicsProxy.cpp
@@ -1527,7 +1527,7 @@ void FGeometryCollectionPhysicsProxy::DisableParticles(TArray<int32>&& Transform
 	}
 }
 
-void FGeometryCollectionPhysicsProxy::BreakInternalClusterParent(TArray<int32>&& TransformGroupIndices)
+void FGeometryCollectionPhysicsProxy::BreakInternalClusterParents(TArray<int32>&& TransformGroupIndices)
 {
 	check(IsInGameThread());
 
@@ -1560,7 +1560,29 @@ void FGeometryCollectionPhysicsProxy::BreakInternalClusterParent(TArray<int32>&&
 			}
 		});
 	}
-	
+}
+
+void FGeometryCollectionPhysicsProxy::BreakClusters(TArray<int32>&& TransformGroupIndices)
+{
+	check(IsInGameThread());
+
+	if (Chaos::FPhysicsSolver* RBDSolver = GetSolver<Chaos::FPhysicsSolver>())
+	{
+		RBDSolver->EnqueueCommandImmediate([this, RBDSolver, IndicesToBreakParent = MoveTemp(TransformGroupIndices)]()
+		{
+			Chaos::FRigidClustering& Clustering = RBDSolver->GetEvolution()->GetRigidClustering();
+			for (int32 TransformIdx : IndicesToBreakParent)
+			{
+				if (Chaos::FPBDRigidClusteredParticleHandle* ClusterHandle = SolverParticleHandles[TransformIdx])
+				{
+					if (!ClusterHandle->Disabled() && Clustering.GetChildrenMap().Contains(ClusterHandle))
+					{
+						Clustering.ReleaseClusterParticles(ClusterHandle, true /* bForceRelease */);
+					}
+				}
+			}
+		});
+	}
 }
 
 int32 FGeometryCollectionPhysicsProxy::CalculateHierarchyLevel(const FGeometryDynamicCollection& DynamicCollection, int32 TransformIndex)
diff --git a/Engine/Source/Runtime/Experimental/Chaos/Public/PhysicsProxy/GeometryCollectionPhysicsProxy.h b/Engine/Source/Runtime/Experimental/Chaos/Public/PhysicsProxy/GeometryCollectionPhysicsProxy.h
--- a/Engine/Source/Runtime/Experimental/Chaos/Public/PhysicsProxy/GeometryCollectionPhysicsProxy.h
+++ b/Engine/Source/Runtime/Experimental/Chaos/Public/PhysicsProxy/GeometryCollectionPhysicsProxy.h
@@ -242,7 +242,8 @@ public:
 	bool GetIsObjectDynamic() const { return IsObjectDynamic; }
 
 	void DisableParticles(TArray<int32>&& TransformGroupIndices);
-	void BreakInternalClusterParent(TArray<int32>&& TransformGroupIndices);
+	void BreakInternalClusterParents(TArray<int32>&& TransformGroupIndices);
+	void BreakClusters(TArray<int32>&& TransformGroupIndices);
 
 	FProxyInterpolationData& GetInterpolationData() { return InterpolationData; }
 	const FProxyInterpolationData& GetInterpolationData() const { return InterpolationData; }
@@ -295,6 +296,7 @@ protected:
 	void PrepareBufferData(Chaos::FDirtyGeometryCollectionData& BufferData, const FGeometryDynamicCollection& ThreadCollection,  Chaos::FReal SolverLastDt = 0.0);
 
 	void CreateNonClusteredParticles(Chaos::FPBDRigidsSolver* RigidsSolver,	const FGeometryCollection& RestCollection, const FGeometryDynamicCollection& DynamicCollection);
+
 private:
 
 	FSimulationParameters Parameters;
diff --git a/Engine/Source/Runtime/Experimental/GeometryCollectionEngine/Private/GeometryCollection/GeometryCollectionComponent.cpp b/Engine/Source/Runtime/Experimental/GeometryCollectionEngine/Private/GeometryCollection/GeometryCollectionComponent.cpp
--- a/Engine/Source/Runtime/Experimental/GeometryCollectionEngine/Private/GeometryCollection/GeometryCollectionComponent.cpp
+++ b/Engine/Source/Runtime/Experimental/GeometryCollectionEngine/Private/GeometryCollection/GeometryCollectionComponent.cpp
@@ -2088,23 +2088,27 @@ void UGeometryCollectionComponent::ResetDynamicCollection()
 				TManagedArray<float>& PostBreakDuration = DynamicCollection->AddAttribute<float>("PostBreakDuration", FGeometryCollection::TransformGroup);
 				for (int32 Idx = 0; Idx < PostBreakDuration.Num(); ++Idx)
 				{
-					const FVector4f RemoveOnBreakData = (*RemoveOnBreak)[Idx];
-					const bool bRemoveOnBreakEnabled = RemoveOnBreakData.X >= 0;
-					const float MinTime = FMath::Max(0.0f, RemoveOnBreakData.X);
-					const float MaxTime = FMath::Max(MinTime, RemoveOnBreakData.Y);
-					PostBreakDuration[Idx] = bRemoveOnBreakEnabled? FMath::RandRange(MinTime, MaxTime): -1;
+					const FRemoveOnBreakData RemoveOnBreakData{ (*RemoveOnBreak)[Idx] };
+					const FVector2f PostBreakTimer = RemoveOnBreakData.GetBreakTimer();
+					const float MinTime = FMath::Max(0.0f, PostBreakTimer.X);
+					const float MaxTime = FMath::Max(MinTime, PostBreakTimer.Y);
+					PostBreakDuration[Idx] = RemoveOnBreakData.IsEnabled()? FMath::RandRange(MinTime, MaxTime): -1;
 				}
 			}
 
 			if (!DynamicCollection->HasAttribute("BreakRemovalDuration", FGeometryCollection::TransformGroup))
 			{
+				const TManagedArray<TSet<int32>>& Children = DynamicCollection->Children;
 				TManagedArray<float>& BreakRemovalDuration = DynamicCollection->AddAttribute<float>("BreakRemovalDuration", FGeometryCollection::TransformGroup);
 				for (int32 Idx = 0; Idx < BreakRemovalDuration.Num(); ++Idx)
 				{
-					const FVector4f RemoveOnBreakData = (*RemoveOnBreak)[Idx];
-					const float MinTime = FMath::Max(0.0f, RemoveOnBreakData.Z);
-					const float MaxTime = FMath::Max(MinTime, RemoveOnBreakData.W);
-					BreakRemovalDuration[Idx] = FMath::RandRange(MinTime, MaxTime);
+					const FRemoveOnBreakData RemoveOnBreakData{ (*RemoveOnBreak)[Idx] };
+					const FVector2f RemovalTimer = RemoveOnBreakData.GetRemovalTimer();
+					const float MinTime = FMath::Max(0.0f, RemovalTimer.X);
+					const float MaxTime = FMath::Max(MinTime, RemovalTimer.Y);
+					const bool bIsCluster = (Children[Idx].Num() > 0);
+					const bool bUseClusterCrumbling = (bIsCluster && RemoveOnBreakData.GetClusterCrumbling());
+					BreakRemovalDuration[Idx] = bUseClusterCrumbling? -1: FMath::RandRange(MinTime, MaxTime);
 				}
 			}
 		}
@@ -3561,7 +3565,7 @@ void UGeometryCollectionComponent::IncrementSleepTimer(float DeltaTime)
 
 		if (ToBreakParent.Num())
 		{
-			PhysicsProxy->BreakInternalClusterParent(MoveTemp(ToBreakParent));
+			PhysicsProxy->BreakInternalClusterParents(MoveTemp(ToBreakParent));
 		}
 		if (ToDisable.Num())
 		{
@@ -3590,6 +3594,7 @@ void UGeometryCollectionComponent::IncrementBreakTimer(float DeltaTime)
 		TManagedArray<float>& Decay = DynamicCollection->ModifyAttribute<float>("Decay", FGeometryCollection::TransformGroup);
 		
 		TArray<int32> ToDisable;
+		TArray<int32> ToCrumble;
 		for (int32 TransformIdx = 0; TransformIdx < BreakTimer.Num(); ++TransformIdx)
 		{
 			const bool bIsBroken = DynamicCollection->Active[TransformIdx];
@@ -3601,23 +3606,41 @@ void UGeometryCollectionComponent::IncrementBreakTimer(float DeltaTime)
 			if (bIsBroken && bIsRemovedFromBreakEnabled && bIsDynamicOrSleeping)
 			{
 				BreakTimer[TransformIdx] += DeltaTime;
+				const bool bPostBreakTimeExpired = (BreakTimer[TransformIdx] >= PostBreakDuration[TransformIdx]);
 				const bool bZeroRemovalDuration = (BreakRemovalDuration[TransformIdx] < UE_SMALL_NUMBER);
-				const float UpdatedBreakDecay = bZeroRemovalDuration? 1.f: FMath::Clamp<float>((BreakTimer[TransformIdx] - PostBreakDuration[TransformIdx]) / BreakRemovalDuration[TransformIdx], 0.f, 1.f);
+				float UpdatedBreakDecay = 0;
+				if (bPostBreakTimeExpired)
+				{
+					UpdatedBreakDecay = bZeroRemovalDuration? 1.f: FMath::Clamp<float>((BreakTimer[TransformIdx] - PostBreakDuration[TransformIdx]) / BreakRemovalDuration[TransformIdx], 0.f, 1.f);
+				}
 				if (UpdatedBreakDecay > Decay[TransformIdx])
 				{
 					DynamicCollection->MakeDirty();
 					Decay[TransformIdx] = UpdatedBreakDecay;
 
-					if (Decay[TransformIdx] >= 1.0f)
+					const bool bUseClusterCrumbling = (BreakRemovalDuration[TransformIdx] < 0);
+					if (bUseClusterCrumbling)
 					{
-						// Disable the particle if it has decayed  the requisite time
-						Decay[TransformIdx] = 1.0f;
-						ToDisable.Add(TransformIdx);
+						ToCrumble.AddUnique(TransformIdx);
+						Decay[TransformIdx] = 0;
+					}
+					else
+					{
+						if (Decay[TransformIdx] >= 1.0f)
+						{
+							// Disable the particle if it has decayed  the requisite time
+							Decay[TransformIdx] = 1.0f;
+							ToDisable.Add(TransformIdx);
+						}
 					}
 				}
 			}
 		}
 
+		if (ToCrumble.Num())
+		{
+			PhysicsProxy->BreakClusters(MoveTemp(ToCrumble));
+		}
 		if (ToDisable.Num())
 		{
 			PhysicsProxy->DisableParticles(MoveTemp(ToDisable));
diff --git a/Engine/Source/Runtime/Experimental/GeometryCollectionEngine/Public/GeometryCollection/GeometryCollectionComponent.h b/Engine/Source/Runtime/Experimental/GeometryCollectionEngine/Public/GeometryCollection/GeometryCollectionComponent.h
--- a/Engine/Source/Runtime/Experimental/GeometryCollectionEngine/Public/GeometryCollection/GeometryCollectionComponent.h
+++ b/Engine/Source/Runtime/Experimental/GeometryCollectionEngine/Public/GeometryCollection/GeometryCollectionComponent.h
@@ -165,6 +165,42 @@ namespace GeometryCollection
 	ENUM_CLASS_FLAGS(EEditUpdate);
 }
 
+// this class helps reading and writing the packed data used in the managed array
+struct FRemoveOnBreakData
+{
+public:
+	inline static const FVector4f DisabledPackedData{ -1, 0, 0, 0 };   
+	
+	FRemoveOnBreakData()
+		: PackedData(FRemoveOnBreakData::DisabledPackedData)
+	{}
+
+	FRemoveOnBreakData(const FVector4f& InPackedData)
+		: PackedData(InPackedData)
+	{}
+
+	FRemoveOnBreakData(bool bEnable, const FVector2f& BreakTimer, bool bClusterCrumbling, const FVector2f& RemovalTimer)
+	{
+		PackedData.X = FMath::Min(FMath::Abs(BreakTimer.X), FMath::Abs(BreakTimer.Y)); // Min break timer
+		PackedData.Y = FMath::Max(FMath::Abs(BreakTimer.X), FMath::Abs(BreakTimer.Y)); // Min break timer
+		PackedData.Z = FMath::Min(FMath::Abs(RemovalTimer.X), FMath::Abs(RemovalTimer.Y)); // Min removal timer
+		PackedData.W = FMath::Max(FMath::Abs(RemovalTimer.X), FMath::Abs(RemovalTimer.Y)); // Max removal timer
+
+		PackedData.X *= bEnable? +1: -1;
+		PackedData.Z *= bClusterCrumbling? -1: +1;
+	}
+	
+	bool IsEnabled() const { return PackedData.X >= 0; }
+	bool GetClusterCrumbling() const { return IsEnabled() && (PackedData.Z < 0); } 
+	FVector2f GetBreakTimer() const	{ return FVector2f(FMath::Abs(PackedData.X), FMath::Abs(PackedData.Y)); }
+	FVector2f GetRemovalTimer() const { return FVector2f(FMath::Abs(PackedData.Z), FMath::Abs(PackedData.W)); }
+
+	const FVector4f& GetPackedData() const { return PackedData; }
+
+private:
+	FVector4f PackedData;
+};
+
 /**
 *	FGeometryCollectionEdit
 *     Structured RestCollection access where the scope
