Chaos - workaround for invalid convexes in manifold generation

diff --git a/Engine/Source/Runtime/Experimental/Chaos/Private/Chaos/CollisionOneShotManifolds.cpp b/Engine/Source/Runtime/Experimental/Chaos/Private/Chaos/CollisionOneShotManifolds.cpp
--- a/Engine/Source/Runtime/Experimental/Chaos/Private/Chaos/CollisionOneShotManifolds.cpp
+++ b/Engine/Source/Runtime/Experimental/Chaos/Private/Chaos/CollisionOneShotManifolds.cpp
@@ -1046,6 +1046,13 @@ namespace Chaos
 			const TGJKShape<ConvexType> GJKConvex(Convex);
 			FContactPoint GJKContactPoint = GJKContactPointSameSpace(GJKConvex, GJKSegment, UnusedMaxMarginDelta, VertexIndexA, VertexIndexB);
 
+			// We are seeing very rare cases where the VertexIndex is not set from an FConvex, even though the convex has vertices
+			// @todo(chaos): track down this problem - there is no path through GJKContactPointMargin (for FConvex) which does not set the VertexIndex
+			if (!CheckVertexIndex(Convex, VertexIndexA))
+			{
+				return;
+			}
+
 			// Map contact back onto capsule
 			GJKContactPoint.ShapeContactPoints[1] = GJKContactPoint.ShapeContactPoints[1] + CapsuleRadius * GJKContactPoint.ShapeContactNormal;
 			GJKContactPoint.Phi = GJKContactPoint.Phi - CapsuleRadius;
@@ -1230,6 +1237,13 @@ namespace Chaos
 			const FVec3 InitialGJKDir = -Triangle.GetCentroid();
 			FContactPoint GJKContactPoint = GJKContactPointSameSpace(GJKConvex, GJKTriangle, MaxMarginDelta, VertexIndexA, VertexIndexB, InitialGJKDir);
 
+			// We are seeing very rare cases where the VertexIndex is not set from an FConvex, even though the convex has vertices
+			// @todo(chaos): track down this problem - there is no path through GJKContactPointMargin (for FConvex) which does not set the VertexIndex
+			if (!CheckVertexIndex(Convex, VertexIndexA))
+			{
+				return;
+			}
+
 			PHYSICS_CSV_CUSTOM_EXPENSIVE(PhysicsCounters, NumManifoldsGJKCalled, 1, ECsvCustomStatOp::Accumulate);
 
 			// GJK is using margins and rounded corner, so if we have a corner-to-corner contact it will under-report the actual distance by an amount that depends on how
