Fix nopch compilation that cannot find definition of FPhysicsSolverBase

diff --git a/Engine/Source/Runtime/Experimental/Chaos/Private/Chaos/Utilities.cpp b/Engine/Source/Runtime/Experimental/Chaos/Private/Chaos/Utilities.cpp
--- a/Engine/Source/Runtime/Experimental/Chaos/Private/Chaos/Utilities.cpp
+++ b/Engine/Source/Runtime/Experimental/Chaos/Private/Chaos/Utilities.cpp
@@ -2,6 +2,7 @@
 #include "Chaos/Utilities.h"
 #include "Chaos/Matrix.h"
 #include "Chaos/Rotation.h"
+#include "Chaos/Framework/PhysicsSolverBase.h"
 
 namespace Chaos
 {
@@ -145,4 +146,12 @@ namespace Chaos
 		TRotation<FReal, 3> R1 = InR0 + (TRotation<FReal, 3>::FromElements(InW.X, InW.Y, InW.Z, 0.f) * InR0) * decltype(InR0.X)(InDt * 0.5f);
 		return R1.GetNormalized();
 	}
+
+namespace Utilities
+{
+	FReal GetSolverPhysicsResultsTime(FPhysicsSolverBase* Solver)
+	{
+		return Solver->GetPhysicsResultsTime_External();
+	}
+}
 }
diff --git a/Engine/Source/Runtime/Experimental/Chaos/Public/Chaos/SimCallbackObject.h b/Engine/Source/Runtime/Experimental/Chaos/Public/Chaos/SimCallbackObject.h
--- a/Engine/Source/Runtime/Experimental/Chaos/Public/Chaos/SimCallbackObject.h
+++ b/Engine/Source/Runtime/Experimental/Chaos/Public/Chaos/SimCallbackObject.h
@@ -11,6 +11,11 @@ namespace Chaos
 class FPhysicsSolverBase;
 class FCollisionContactModifier;
 
+namespace Utilities
+{
+	CHAOS_API FReal GetSolverPhysicsResultsTime(FPhysicsSolverBase*);
+}
+
 /**
  * Callback API used for executing code at different points in the simulation.
  * The external thread pushes input data at its own rate (typically once per game thread tick)
@@ -256,7 +261,7 @@ public:
 	*/
 	TSimCallbackOutputHandle<TOutputType> PopOutputData_External()
 	{
-		const FReal ResultsTime = GetSolver()->GetPhysicsResultsTime_External();
+		const FReal ResultsTime = Utilities::GetSolverPhysicsResultsTime(GetSolver());
 		if(!CurrentOutput_External)
 		{
 			OutputQueue.Dequeue(CurrentOutput_External);
