Workaround for C4499 under C++20: an explicit specialization cannot have a storage class (ignored) static - yet in C++17 it is required.

diff --git a/Engine/Source/Runtime/Experimental/Chaos/Public/Chaos/AABBTree.h b/Engine/Source/Runtime/Experimental/Chaos/Public/Chaos/AABBTree.h
--- a/Engine/Source/Runtime/Experimental/Chaos/Public/Chaos/AABBTree.h
+++ b/Engine/Source/Runtime/Experimental/Chaos/Public/Chaos/AABBTree.h
@@ -14,6 +14,7 @@
 #include "ProfilingDebugging/CsvProfiler.h"
 #include "ChaosStats.h"
 #include "Math/VectorRegister.h"
+#include <type_traits>
 
 CSV_DECLARE_CATEGORY_EXTERN(ChaosPhysicsTimers);
 
@@ -3086,31 +3087,31 @@ private:
 	template<typename ContainerType>
 	static void AddToContainerHelper(const ContainerType& ContainerFrom, ContainerType& ContainerTo, int32 Index)
 	{
-		ContainerTo.Add(ContainerFrom[Index]);
-	}
-
-	template<>
-	static void AddToContainerHelper(const TArrayAsMap<TPayloadType, FAABBTreePayloadInfo>& ContainerFrom, TArrayAsMap<TPayloadType, FAABBTreePayloadInfo>& ContainerTo, int32 Index)
-	{
-		ContainerTo.AddFrom(ContainerFrom, Index);
+		if constexpr (std::is_same_v<ContainerType, TArrayAsMap<TPayloadType, FAABBTreePayloadInfo>>)
+		{
+			ContainerTo.AddFrom(ContainerFrom, Index);
+		}
+		else
+		{
+			ContainerTo.Add(ContainerFrom[Index]);
+		}
 	}
 
 	template<typename ContainerType>
 	static int32 ContainerElementSizeHelper(const ContainerType& Container, int32 Index)
 	{
-		return sizeof(typename ContainerType::ElementType);
-	}
-
-	template<>
-	static int32 ContainerElementSizeHelper(const TArray<TAABBTreeLeafArray<TPayloadType, true>>& Container, int32 Index)
-	{
-		return sizeof(typename TArray<TAABBTreeLeafArray<TPayloadType, true>>::ElementType) + sizeof(typename decltype(Container[Index].Elems)::ElementType) * Container[Index].GetElementCount();
-	}
-
-	template<>
-	static int32 ContainerElementSizeHelper(const TArray<TAABBTreeLeafArray<TPayloadType, false>>& Container, int32 Index)
-	{
-		return sizeof(typename TArray<TAABBTreeLeafArray<TPayloadType, false>>::ElementType) + sizeof(typename decltype(Container[Index].Elems)::ElementType) * Container[Index].GetElementCount();
+		if constexpr (std::is_same_v<ContainerType, TArray<TAABBTreeLeafArray<TPayloadType, true>>>)
+		{
+			return sizeof(typename TArray<TAABBTreeLeafArray<TPayloadType, true>>::ElementType) + sizeof(typename decltype(Container[Index].Elems)::ElementType) * Container[Index].GetElementCount();
+		}
+		else if constexpr (std::is_same_v<ContainerType, TArray<TAABBTreeLeafArray<TPayloadType, false>>>)
+		{
+			return sizeof(typename TArray<TAABBTreeLeafArray<TPayloadType, false>>::ElementType) + sizeof(typename decltype(Container[Index].Elems)::ElementType) * Container[Index].GetElementCount();
+		}
+		else
+		{
+			return sizeof(typename ContainerType::ElementType);
+		}
 	}
 
 	template<typename ContainerType>
