add option to not remove overlaps on geometry collection convex hulls

diff --git a/Engine/Plugins/Experimental/ChaosEditor/Source/FractureEditor/Private/FractureToolConvex.cpp b/Engine/Plugins/Experimental/ChaosEditor/Source/FractureEditor/Private/FractureToolConvex.cpp
--- a/Engine/Plugins/Experimental/ChaosEditor/Source/FractureEditor/Private/FractureToolConvex.cpp
+++ b/Engine/Plugins/Experimental/ChaosEditor/Source/FractureEditor/Private/FractureToolConvex.cpp
@@ -293,6 +293,7 @@ int32 UFractureToolConvex::ExecuteFracture(const FFractureToolContext& FractureC
 		Properties.FractionRemove = ConvexSettings->FractionAllowRemove;
 		Properties.SimplificationThreshold = ConvexSettings->SimplificationDistanceThreshold;
 		Properties.CanExceedFraction = ConvexSettings->CanExceedFraction;
+		Properties.bRemoveOverlaps = ConvexSettings->bRemoveOverlaps;
 		FractureContext.GetGeometryCollection()->SetConvexProperties(Properties);
 	}
 
diff --git a/Engine/Plugins/Experimental/ChaosEditor/Source/FractureEditor/Private/FractureToolConvex.h b/Engine/Plugins/Experimental/ChaosEditor/Source/FractureEditor/Private/FractureToolConvex.h
--- a/Engine/Plugins/Experimental/ChaosEditor/Source/FractureEditor/Private/FractureToolConvex.h
+++ b/Engine/Plugins/Experimental/ChaosEditor/Source/FractureEditor/Private/FractureToolConvex.h
@@ -32,6 +32,10 @@ public:
 	UPROPERTY(EditAnywhere, Category = Automatic, meta = (ClampMin = "0"))
 	double SimplificationDistanceThreshold = 10.0;
 
+	/** Whether to automatically cut away overlapping parts of the convex hulls, to avoid the simulation 'popping' to fix the overlaps */
+	UPROPERTY(EditAnywhere, Category = Automatic)
+	bool bRemoveOverlaps = true;
+
 	/** Delete convex hulls from selected clusters.  Does not affect hulls on leaves. */
 	UFUNCTION(CallInEditor, Category = Custom, meta = (DisplayName = "Delete From Selected"))
 	void DeleteFromSelected();
diff --git a/Engine/Source/Runtime/Experimental/Chaos/Private/GeometryCollection/GeometryCollectionConvexPropertiesInterface.cpp b/Engine/Source/Runtime/Experimental/Chaos/Private/GeometryCollection/GeometryCollectionConvexPropertiesInterface.cpp
--- a/Engine/Source/Runtime/Experimental/Chaos/Private/GeometryCollection/GeometryCollectionConvexPropertiesInterface.cpp
+++ b/Engine/Source/Runtime/Experimental/Chaos/Private/GeometryCollection/GeometryCollectionConvexPropertiesInterface.cpp
@@ -9,6 +9,7 @@ const FName FGeometryCollectionConvexPropertiesInterface::ConvexEnable = "Enable
 const FName FGeometryCollectionConvexPropertiesInterface::ConvexFractionRemoveAttribute = "FractionRemove";
 const FName FGeometryCollectionConvexPropertiesInterface::ConvexSimplificationThresholdAttribute = "SimplificationThreshold";
 const FName FGeometryCollectionConvexPropertiesInterface::ConvexCanExceedFractionAttribute = "CanExceedFraction";
+const FName FGeometryCollectionConvexPropertiesInterface::ConvexRemoveOverlapsAttribute = "RemoveOverlaps";
 
 FGeometryCollectionConvexPropertiesInterface::FGeometryCollectionConvexPropertiesInterface(FGeometryCollection* InGeometryCollection)
 	: FManagedArrayInterface(InGeometryCollection)
@@ -48,6 +49,11 @@ FGeometryCollectionConvexPropertiesInterface::InitializeInterface()
 		ManagedCollection->AddAttribute<float>(ConvexCanExceedFractionAttribute, ConvexPropertiesGroup);
 	}
 
+	if (!ManagedCollection->HasAttribute(ConvexRemoveOverlapsAttribute, ConvexPropertiesGroup))
+	{
+		ManagedCollection->AddAttribute<bool>(ConvexRemoveOverlapsAttribute, ConvexPropertiesGroup);
+	}
+
 	SetDefaultProperty();
 }
 
@@ -74,6 +80,7 @@ FGeometryCollectionConvexPropertiesInterface::GetDefaultProperty() const
 		const TManagedArray<float>& FractionRemove = ManagedCollection->GetAttribute<float>(ConvexFractionRemoveAttribute, ConvexPropertiesGroup);
 		const TManagedArray<float>& SimplificationThreshold = ManagedCollection->GetAttribute<float>(ConvexSimplificationThresholdAttribute, ConvexPropertiesGroup);
 		const TManagedArray<float>& CanExceedFraction = ManagedCollection->GetAttribute<float>(ConvexCanExceedFractionAttribute, ConvexPropertiesGroup);
+		const TManagedArray<bool>& RemoveOverlaps = ManagedCollection->GetAttribute<bool>(ConvexRemoveOverlapsAttribute, ConvexPropertiesGroup);
 
 		int32 DefaultIndex = Index.Find(INDEX_NONE);
 		if (0 <= DefaultIndex && DefaultIndex < Index.Num())
@@ -82,6 +89,7 @@ FGeometryCollectionConvexPropertiesInterface::GetDefaultProperty() const
 			DefaultProperty.FractionRemove = FractionRemove[DefaultIndex];
 			DefaultProperty.SimplificationThreshold = SimplificationThreshold[DefaultIndex];
 			DefaultProperty.CanExceedFraction = CanExceedFraction[DefaultIndex];
+			DefaultProperty.bRemoveOverlaps = RemoveOverlaps[DefaultIndex];
 		}
 	}
 	return DefaultProperty;
@@ -112,11 +120,13 @@ FGeometryCollectionConvexPropertiesInterface::GetConvexProperties(int TransformG
 		const TManagedArray<float>& FractionRemove = ManagedCollection->GetAttribute<float>(ConvexFractionRemoveAttribute, ConvexPropertiesGroup);
 		const TManagedArray<float>& SimplificationThreshold = ManagedCollection->GetAttribute<float>(ConvexSimplificationThresholdAttribute, ConvexPropertiesGroup);
 		const TManagedArray<float>& CanExceedFraction = ManagedCollection->GetAttribute<float>(ConvexCanExceedFractionAttribute, ConvexPropertiesGroup);
+		const TManagedArray<bool>& RemoveOverlaps = ManagedCollection->GetAttribute<bool>(ConvexRemoveOverlapsAttribute, ConvexPropertiesGroup);
 
 		ConvexProperty.Enable = Enable[PropIndex];
 		ConvexProperty.FractionRemove = FractionRemove[PropIndex];
 		ConvexProperty.SimplificationThreshold = SimplificationThreshold[PropIndex];
 		ConvexProperty.CanExceedFraction = CanExceedFraction[PropIndex];
+		ConvexProperty.bRemoveOverlaps = RemoveOverlaps[PropIndex];
 	}
 	return ConvexProperty;
 }
@@ -130,6 +140,7 @@ FGeometryCollectionConvexPropertiesInterface::SetConvexProperties(const FConvexC
 	TManagedArray<float>& FractionRemove = ManagedCollection->ModifyAttribute<float>(ConvexFractionRemoveAttribute, ConvexPropertiesGroup);
 	TManagedArray<float>& SimplificationThreshold = ManagedCollection->ModifyAttribute<float>(ConvexSimplificationThresholdAttribute, ConvexPropertiesGroup);
 	TManagedArray<float>& CanExceedFraction = ManagedCollection->ModifyAttribute<float>(ConvexCanExceedFractionAttribute, ConvexPropertiesGroup);
+	TManagedArray<bool>& RemoveOverlaps = ManagedCollection->ModifyAttribute<bool>(ConvexRemoveOverlapsAttribute, ConvexPropertiesGroup);
 
 	int32 AttributeIndex = INDEX_NONE;
 	if (!Index.Contains(TransformGroupIndex))
@@ -148,5 +159,6 @@ FGeometryCollectionConvexPropertiesInterface::SetConvexProperties(const FConvexC
 		FractionRemove[AttributeIndex] = InConvexAttributes.FractionRemove;
 		SimplificationThreshold[AttributeIndex] = InConvexAttributes.SimplificationThreshold;
 		CanExceedFraction[AttributeIndex] = InConvexAttributes.CanExceedFraction;
+		RemoveOverlaps[AttributeIndex] = InConvexAttributes.bRemoveOverlaps;
 	}
 }
diff --git a/Engine/Source/Runtime/Experimental/Chaos/Private/GeometryCollection/GeometryCollectionConvexUtility.cpp b/Engine/Source/Runtime/Experimental/Chaos/Private/GeometryCollection/GeometryCollectionConvexUtility.cpp
--- a/Engine/Source/Runtime/Experimental/Chaos/Private/GeometryCollection/GeometryCollectionConvexUtility.cpp
+++ b/Engine/Source/Runtime/Experimental/Chaos/Private/GeometryCollection/GeometryCollectionConvexUtility.cpp
@@ -320,7 +320,8 @@ void CreateNonoverlappingConvexHulls(
 	const TManagedArray<float>* Volume,
 	double FracAllowRemove,
 	double SimplificationDistanceThreshold,
-	double CanExceedFraction
+	double CanExceedFraction,
+	bool bRemoveOverlaps
 )
 {
 	int32 NumBones = TransformToConvexIndices.Num();
@@ -566,7 +567,7 @@ void CreateNonoverlappingConvexHulls(
 		Depths[HullIdx] = Depth;
 		MaxDepth = FMath::Max(Depth, MaxDepth);
 
-		if (!HasCustomConvexFn(HullIdx) && TransformToConvexIndices[HullIdx].Num() > 0)
+		if (bRemoveOverlaps && !HasCustomConvexFn(HullIdx) && TransformToConvexIndices[HullIdx].Num() > 0)
 		{
 			const TSet<int32>& Neighbors = LeafProximity[HullIdx];
 			for (int32 NbrIdx : Neighbors)
@@ -632,76 +633,79 @@ void CreateNonoverlappingConvexHulls(
 	// Use initial leaf proximity to compute which cluster bones are in proximity to which neighbors at the same level of the bone hierarchy
 
 	TArray<TSet<int32>> SameDepthClusterProximity;
-	SameDepthClusterProximity.SetNum(NumBones);
-	for (int32 Bone = 0; Bone < NumBones; Bone++)
+	if (bRemoveOverlaps)
 	{
-		for (int32 NbrBone : LeafProximity[Bone])
+		SameDepthClusterProximity.SetNum(NumBones);
+		for (int32 Bone = 0; Bone < NumBones; Bone++)
 		{
-			if (Bone > NbrBone)
-			{
-				continue;
-			}
-
-			auto FindCommonParent = [&Parents](int32 BoneA, int32 BoneB)
+			for (int32 NbrBone : LeafProximity[Bone])
 			{
-				int AParent = Parents[BoneA];
-				int BParent = Parents[BoneB];
-				if (AParent == BParent) // early out if they're in the same cluster
+				if (Bone > NbrBone)
 				{
-					return AParent;
+					continue;
 				}
 
-				TSet<int32> AParents;
-				while (AParent != INDEX_NONE)
+				auto FindCommonParent = [&Parents](int32 BoneA, int32 BoneB)
 				{
-					AParents.Add(AParent);
-					AParent = Parents[AParent];
-				}
+					int AParent = Parents[BoneA];
+					int BParent = Parents[BoneB];
+					if (AParent == BParent) // early out if they're in the same cluster
+					{
+						return AParent;
+					}
 
-				while (BParent != INDEX_NONE && !AParents.Contains(BParent))
-				{
-					BParent = Parents[BParent];
-				}
-				return BParent;
-			};
-			int32 CommonParent = FindCommonParent(Bone, NbrBone);
+					TSet<int32> AParents;
+					while (AParent != INDEX_NONE)
+					{
+						AParents.Add(AParent);
+						AParent = Parents[AParent];
+					}
 
-			auto ConnectAtMatchingDepth = [&SameDepthClusterProximity, &Parents, CommonParent, &Depths](int32 BoneA, int32 BoneToWalkParents)
-			{
-				if (BoneA == INDEX_NONE || BoneToWalkParents == INDEX_NONE)
-				{
-					return;
-				}
+					while (BParent != INDEX_NONE && !AParents.Contains(BParent))
+					{
+						BParent = Parents[BParent];
+					}
+					return BParent;
+				};
+				int32 CommonParent = FindCommonParent(Bone, NbrBone);
 
-				int32 DepthA = Depths[BoneA];
-				int32 DepthB = Depths[BoneToWalkParents];
-				while (BoneToWalkParents != INDEX_NONE && BoneToWalkParents != CommonParent && DepthB >= DepthA)
+				auto ConnectAtMatchingDepth = [&SameDepthClusterProximity, &Parents, CommonParent, &Depths](int32 BoneA, int32 BoneToWalkParents)
 				{
-					if (DepthB == DepthA)
+					if (BoneA == INDEX_NONE || BoneToWalkParents == INDEX_NONE)
 					{
-						bool bWasInSet = false;
-						SameDepthClusterProximity[BoneToWalkParents].Add(BoneA, &bWasInSet);
-						if (bWasInSet)
+						return;
+					}
+
+					int32 DepthA = Depths[BoneA];
+					int32 DepthB = Depths[BoneToWalkParents];
+					while (BoneToWalkParents != INDEX_NONE && BoneToWalkParents != CommonParent && DepthB >= DepthA)
+					{
+						if (DepthB == DepthA)
 						{
-							break;
+							bool bWasInSet = false;
+							SameDepthClusterProximity[BoneToWalkParents].Add(BoneA, &bWasInSet);
+							if (bWasInSet)
+							{
+								break;
+							}
+							SameDepthClusterProximity[BoneA].Add(BoneToWalkParents);
 						}
-						SameDepthClusterProximity[BoneA].Add(BoneToWalkParents);
-					}
 
-					BoneToWalkParents = Parents[BoneToWalkParents];
-					DepthB--;
-				}
-			};
+						BoneToWalkParents = Parents[BoneToWalkParents];
+						DepthB--;
+					}
+				};
 
-			// connect bone to neighboring bone's clusters
-			ConnectAtMatchingDepth(Bone, Parents[NbrBone]);
+				// connect bone to neighboring bone's clusters
+				ConnectAtMatchingDepth(Bone, Parents[NbrBone]);
 
-			// walk chain of parents of bone and connect each to the chain of neighboring bone + its parents
-			int32 BoneParent = Parents[Bone];
-			while (BoneParent != INDEX_NONE && BoneParent != CommonParent)
-			{
-				ConnectAtMatchingDepth(BoneParent, NbrBone);
-				BoneParent = Parents[BoneParent];
+				// walk chain of parents of bone and connect each to the chain of neighboring bone + its parents
+				int32 BoneParent = Parents[Bone];
+				while (BoneParent != INDEX_NONE && BoneParent != CommonParent)
+				{
+					ConnectAtMatchingDepth(BoneParent, NbrBone);
+					BoneParent = Parents[BoneParent];
+				}
 			}
 		}
 	}
@@ -709,7 +713,10 @@ void CreateNonoverlappingConvexHulls(
 	// Compute initial hulls at all levels and filter out any that are too large relative to the geometry they contain
 
 	TArray<TSet<int32>> ClusterProximity;
-	ClusterProximity.SetNum(NumBones);
+	if (bRemoveOverlaps)
+	{
+		ClusterProximity.SetNum(NumBones);
+	}
 	TArrayView<int32> DepthOrderView(ByDepthOrder);
 
 	for (int32 Idx = 0, SliceEndIdx = Idx + 1; Idx < ByDepthOrder.Num(); Idx = SliceEndIdx)
@@ -771,6 +778,11 @@ void CreateNonoverlappingConvexHulls(
 			}
 		});
 
+		if (!bRemoveOverlaps)
+		{
+			continue;
+		}
+
 		// Compute cluster proximity
 		for (int32 BoneA : SameDepthSlice)
 		{
@@ -808,6 +820,10 @@ void CreateNonoverlappingConvexHulls(
 		}
 	}
 
+	if (!bRemoveOverlaps)
+	{
+		return; // rest of function is just for removing overlaps
+	}
 
 
 	// Compute all initial non-leaf hull volumes
@@ -1145,7 +1161,7 @@ double ComputeGeometryVolume(
 
 }
 
-FGeometryCollectionConvexUtility::FGeometryCollectionConvexData FGeometryCollectionConvexUtility::CreateNonOverlappingConvexHullData(FGeometryCollection* GeometryCollection, double FracAllowRemove, double SimplificationDistanceThreshold, double CanExceedFraction)
+FGeometryCollectionConvexUtility::FGeometryCollectionConvexData FGeometryCollectionConvexUtility::CreateNonOverlappingConvexHullData(FGeometryCollection* GeometryCollection, double FracAllowRemove, double SimplificationDistanceThreshold, double CanExceedFraction, bool bRemoveOverlaps)
 {
 	check(GeometryCollection);
 
@@ -1170,7 +1186,7 @@ FGeometryCollectionConvexUtility::FGeometryCollectionConvexData FGeometryCollect
 	GeometryCollectionAlgo::GlobalMatrices(GeometryCollection->Transform, GeometryCollection->Parent, GlobalTransformArray);
 	HullsFromGeometry(*GeometryCollection, GlobalTransformArray, HasCustomConvexFn, Convexes, TransformToConvexIndexArr, GeometryCollection->SimulationType, FGeometryCollection::ESimulationTypes::FST_Rigid, SimplificationDistanceThreshold);
 
-	CreateNonoverlappingConvexHulls(Convexes, TransformToConvexIndexArr, HasCustomConvexFn, GeometryCollection->SimulationType, FGeometryCollection::ESimulationTypes::FST_Rigid, FGeometryCollection::ESimulationTypes::FST_None, GeometryCollection->Parent, GCProximity, GeometryCollection->TransformIndex, Volume, FracAllowRemove, SimplificationDistanceThreshold, CanExceedFraction);
+	CreateNonoverlappingConvexHulls(Convexes, TransformToConvexIndexArr, HasCustomConvexFn, GeometryCollection->SimulationType, FGeometryCollection::ESimulationTypes::FST_Rigid, FGeometryCollection::ESimulationTypes::FST_None, GeometryCollection->Parent, GCProximity, GeometryCollection->TransformIndex, Volume, FracAllowRemove, SimplificationDistanceThreshold, CanExceedFraction, bRemoveOverlaps);
 
 	TransformHullsToLocal(GlobalTransformArray, Convexes, TransformToConvexIndexArr);
 
diff --git a/Engine/Source/Runtime/Experimental/Chaos/Public/GeometryCollection/GeometryCollectionConvexPropertiesInterface.h b/Engine/Source/Runtime/Experimental/Chaos/Public/GeometryCollection/GeometryCollectionConvexPropertiesInterface.h
--- a/Engine/Source/Runtime/Experimental/Chaos/Public/GeometryCollection/GeometryCollectionConvexPropertiesInterface.h
+++ b/Engine/Source/Runtime/Experimental/Chaos/Public/GeometryCollection/GeometryCollectionConvexPropertiesInterface.h
@@ -25,6 +25,8 @@ public :
 	static const FName ConvexSimplificationThresholdAttribute;
 	// Attribute
 	static const FName ConvexCanExceedFractionAttribute;
+	// Attribute
+	static const FName ConvexRemoveOverlapsAttribute;
 
 
 	struct FConvexCreationProperties {
@@ -32,6 +34,7 @@ public :
 		float FractionRemove = 0.5f;
 		float SimplificationThreshold = 10.0f;
 		float CanExceedFraction = 0.5f;
+		bool bRemoveOverlaps = true;
 	};
 
 	FGeometryCollectionConvexPropertiesInterface(FGeometryCollection* InGeometryCollection);
diff --git a/Engine/Source/Runtime/Experimental/Chaos/Public/GeometryCollection/GeometryCollectionConvexUtility.h b/Engine/Source/Runtime/Experimental/Chaos/Public/GeometryCollection/GeometryCollectionConvexUtility.h
--- a/Engine/Source/Runtime/Experimental/Chaos/Public/GeometryCollection/GeometryCollectionConvexUtility.h
+++ b/Engine/Source/Runtime/Experimental/Chaos/Public/GeometryCollection/GeometryCollectionConvexUtility.h
@@ -38,8 +38,9 @@ public:
 	 @param FractionAllowRemove					The fraction of a convex body we can cut away to remove overlaps with neighbors, before we fall back to using the hulls of the children directly.  (Does not affect leaves of hierarchy)
 	 @param SimplificationDistanceThreshold		Approximate minimum distance between vertices, below which we remove vertices to generate a simpler convex shape.  If 0.0, no simplification will occur.
 	 @param CanExceedFraction					The fraction by which the convex body volume on a cluster can exceed the volume of the geometry under that cluster (a value of 1 == exceed by 100% == convex hull has 2x the volume of the geometry)
+	 @param bRemoveOverlaps						Whether to cut away overlaps between adjacent convex hulls
 	 */
-	static FGeometryCollectionConvexData CreateNonOverlappingConvexHullData(FGeometryCollection* GeometryCollection, double FractionAllowRemove = .3, double SimplificationDistanceThreshold = 0.0, double CanExceedFraction = .5);
+	static FGeometryCollectionConvexData CreateNonOverlappingConvexHullData(FGeometryCollection* GeometryCollection, double FractionAllowRemove = .3, double SimplificationDistanceThreshold = 0.0, double CanExceedFraction = .5, bool bRemoveOverlaps = true);
 
 	/** Returns the convex hull of the vertices contained in the specified geometry. */
 	static TUniquePtr<Chaos::FConvex> FindConvexHull(const FGeometryCollection* GeometryCollection, int32 GeometryIndex);
diff --git a/Engine/Source/Runtime/Experimental/GeometryCollectionEngine/Private/GeometryCollection/GeometryCollectionObject.cpp b/Engine/Source/Runtime/Experimental/GeometryCollectionEngine/Private/GeometryCollection/GeometryCollectionObject.cpp
--- a/Engine/Source/Runtime/Experimental/GeometryCollectionEngine/Private/GeometryCollection/GeometryCollectionObject.cpp
+++ b/Engine/Source/Runtime/Experimental/GeometryCollectionEngine/Private/GeometryCollection/GeometryCollectionObject.cpp
@@ -284,7 +284,8 @@ void UGeometryCollection::UpdateConvexGeometry()
 	if (GeometryCollection)
 	{
 		FGeometryCollectionConvexPropertiesInterface::FConvexCreationProperties ConvexProperties = GeometryCollection->GetConvexProperties();
-		FGeometryCollectionConvexUtility::CreateNonOverlappingConvexHullData(GeometryCollection.Get(), ConvexProperties.FractionRemove, ConvexProperties.SimplificationThreshold, ConvexProperties.CanExceedFraction);
+		FGeometryCollectionConvexUtility::CreateNonOverlappingConvexHullData(GeometryCollection.Get(), ConvexProperties.FractionRemove, 
+			ConvexProperties.SimplificationThreshold, ConvexProperties.CanExceedFraction, ConvexProperties.bRemoveOverlaps);
 		InvalidateCollection();
 	}
 #endif
