Chaos - friction stability improvement
- use the average normal impulse for clipping to the friction cone
- remove some dead code

diff --git a/Engine/Source/Runtime/Experimental/Chaos/Private/Chaos/Collision/PBDCollisionSolver.cpp b/Engine/Source/Runtime/Experimental/Chaos/Private/Chaos/Collision/PBDCollisionSolver.cpp
--- a/Engine/Source/Runtime/Experimental/Chaos/Private/Chaos/Collision/PBDCollisionSolver.cpp
+++ b/Engine/Source/Runtime/Experimental/Chaos/Private/Chaos/Collision/PBDCollisionSolver.cpp
@@ -23,9 +23,6 @@ namespace Chaos
 	{
 		extern int32 Chaos_Collision_UseShockPropagation;
 
-		bool bChaos_PBDCollisionSolver_VectorRegister = false;
-		FAutoConsoleVariableRef CVarChaos_PBDCollisionSolver_VectorRegister(TEXT("p.Chaos.PBDCollisionSolver.VectorRegister"), bChaos_PBDCollisionSolver_VectorRegister, TEXT(""));
-
 		bool bChaos_PBDCollisionSolver_Position_SolveEnabled = true;
 		float Chaos_PBDCollisionSolver_Position_MinInvMassScale = 0.77f;
 		float Chaos_PBDCollisionSolver_Position_StaticFrictionStiffness = 0.5f;
diff --git a/Engine/Source/Runtime/Experimental/Chaos/Private/Chaos/Collision/PBDCollisionSolver.h b/Engine/Source/Runtime/Experimental/Chaos/Private/Chaos/Collision/PBDCollisionSolver.h
--- a/Engine/Source/Runtime/Experimental/Chaos/Private/Chaos/Collision/PBDCollisionSolver.h
+++ b/Engine/Source/Runtime/Experimental/Chaos/Private/Chaos/Collision/PBDCollisionSolver.h
@@ -102,6 +102,7 @@ namespace Chaos
 		 * @param MaxPushOut a limit on the position error for this iteration to prevent initial-penetration explosion (a common PBD problem)
 		*/
 		void CalculateContactPositionErrorNormal(const FConstraintSolverBody& Body0, const FConstraintSolverBody& Body1, const FSolverReal MaxPushOut, FSolverReal& OutContactDeltaNormal) const;
+		void CalculateContactPositionErrorTangential(const FConstraintSolverBody& Body0, const FConstraintSolverBody& Body1, FSolverReal& OutContactDeltaTanget0, FSolverReal& OutContactDeltaTangent1) const;
 		void CalculateContactPositionError(const FConstraintSolverBody& Body0, const FConstraintSolverBody& Body1, const FSolverReal MaxPushOut, FSolverReal& OutContactDeltaNormal, FSolverReal& OutContactDeltaTanget0, FSolverReal& OutContactDeltaTangent1) const;
 
 		/**
@@ -408,7 +409,7 @@ namespace Chaos
 	FORCEINLINE_DEBUGGABLE void ApplyFrictionCone(
 		const FSolverReal StaticFriction,
 		const FSolverReal DynamicFriction,
-		const FSolverReal NetPushOutNormal,
+		const FSolverReal MaxFrictionPushOut,
 		FSolverReal& InOutPushOutTangentU,
 		FSolverReal& InOutPushOutTangentV,
 		FSolverReal& InOutNetPushOutTangentU,
@@ -418,10 +419,11 @@ namespace Chaos
 		// Assume we stay in the friction cone...
 		OutStaticFrictionRatio = FSolverReal(1);
 
-		if (NetPushOutNormal < FSolverReal(UE_KINDA_SMALL_NUMBER))
+		if (MaxFrictionPushOut < FSolverReal(UE_KINDA_SMALL_NUMBER))
 		{
 			// Note: we have already added the current iteration's PushOut to the NetPushOut but it has not been applied to the body
-			// so we must subtract it again to calculate the actual pushout we want to undo (i.e., the net pushout that has been applied to the body so far)
+			// so we must subtract it again to calculate the actual pushout we want to undo (i.e., the net pushout that has been applied 
+			// to the body so far from previous iterations)
 			InOutPushOutTangentU = -(InOutNetPushOutTangentU - InOutPushOutTangentU);
 			InOutPushOutTangentV = -(InOutNetPushOutTangentV - InOutPushOutTangentV);
 			InOutNetPushOutTangentU = FSolverReal(0);
@@ -431,11 +433,11 @@ namespace Chaos
 		else
 		{
 			// If we exceed the static friction cone, clip to the dynamic friction cone
-			const FSolverReal MaxStaticPushOutTangentSq = FMath::Square(StaticFriction * NetPushOutNormal);
+			const FSolverReal MaxStaticPushOutTangentSq = FMath::Square(StaticFriction * MaxFrictionPushOut);
 			const FSolverReal NetPushOutTangentSq = FMath::Square(InOutNetPushOutTangentU) + FMath::Square(InOutNetPushOutTangentV);
 			if (NetPushOutTangentSq > MaxStaticPushOutTangentSq)
 			{
-				const FSolverReal MaxDynamicPushOutTangent = DynamicFriction * NetPushOutNormal;
+				const FSolverReal MaxDynamicPushOutTangent = DynamicFriction * MaxFrictionPushOut;
 				const FSolverReal FrictionMultiplier = MaxDynamicPushOutTangent * FMath::InvSqrt(NetPushOutTangentSq);
 				const FSolverReal NetPushOutTangentU = FrictionMultiplier * InOutNetPushOutTangentU;
 				const FSolverReal NetPushOutTangentV = FrictionMultiplier * InOutNetPushOutTangentV;
@@ -449,73 +451,64 @@ namespace Chaos
 	}
 
 
-	FORCEINLINE_DEBUGGABLE void ApplyPositionCorrectionWithFriction(
+	FORCEINLINE_DEBUGGABLE void ApplyPositionCorrectionTangential(
 		const FSolverReal Stiffness,
 		const FSolverReal StaticFriction,
 		const FSolverReal DynamicFriction,
-		const FSolverReal ContactDeltaNormal,
-		const FSolverReal ContactDeltaTangent0,
-		const FSolverReal ContactDeltaTangent1,
+		const FSolverReal MaxFrictionPushOut,
+		const FSolverReal ContactDeltaTangentU,
+		const FSolverReal ContactDeltaTangentV,
 		FPBDCollisionSolverManifoldPoint& ManifoldPoint,
 		FConstraintSolverBody& Body0,
 		FConstraintSolverBody& Body1)
 	{
-		FSolverReal PushOutNormal = FSolverReal(0);
 		FSolverReal PushOutTangentU = FSolverReal(0);
 		FSolverReal PushOutTangentV = FSolverReal(0);
 
-		CalculatePositionCorrectionNormal(
-			Stiffness,
-			ContactDeltaNormal,
-			ManifoldPoint.ContactMassNormal,
-			ManifoldPoint.NetPushOutNormal,
-			PushOutNormal);						// Out
-
 		CalculatePositionCorrectionTangent(
 			Stiffness,
-			ContactDeltaTangent0,
+			ContactDeltaTangentU,
 			ManifoldPoint.ContactMassTangentU,
 			PushOutTangentU);					// Out
 
 		CalculatePositionCorrectionTangent(
 			Stiffness,
-			ContactDeltaTangent1,
+			ContactDeltaTangentV,
 			ManifoldPoint.ContactMassTangentV,
 			PushOutTangentV);					// Out
 
-		ManifoldPoint.NetPushOutNormal += PushOutNormal;
 		ManifoldPoint.NetPushOutTangentU += PushOutTangentU;
 		ManifoldPoint.NetPushOutTangentV += PushOutTangentV;
 
 		ApplyFrictionCone(
 			StaticFriction,
 			DynamicFriction,
-			ManifoldPoint.NetPushOutNormal,
-			PushOutTangentU,									// Out
-			PushOutTangentV,									// Out
-			ManifoldPoint.NetPushOutTangentU,					// Out
-			ManifoldPoint.NetPushOutTangentV,					// Out
+			MaxFrictionPushOut,
+			PushOutTangentU,									// InOut
+			PushOutTangentV,									// InOut
+			ManifoldPoint.NetPushOutTangentU,					// InOut
+			ManifoldPoint.NetPushOutTangentV,					// InOut
 			ManifoldPoint.StaticFrictionRatio);					// Out
 
 		// Update the particle state based on the pushout
-		const FSolverVec3 PushOut = PushOutNormal * ManifoldPoint.WorldContactNormal + PushOutTangentU * ManifoldPoint.WorldContactTangentU + PushOutTangentV * ManifoldPoint.WorldContactTangentV;
+		const FVec3 PushOut = PushOutTangentU * ManifoldPoint.WorldContactTangentU + PushOutTangentV * ManifoldPoint.WorldContactTangentV;
 		if (Body0.IsDynamic())
 		{
 			const FSolverVec3 DX0 = Body0.InvM() * PushOut;
-			const FSolverVec3 DR0 = ManifoldPoint.WorldContactNormalAngular0 * PushOutNormal + ManifoldPoint.WorldContactTangentUAngular0 * PushOutTangentU + ManifoldPoint.WorldContactTangentVAngular0 * PushOutTangentV;
+			const FSolverVec3 DR0 = ManifoldPoint.WorldContactTangentUAngular0 * PushOutTangentU + ManifoldPoint.WorldContactTangentVAngular0 * PushOutTangentV;
 			Body0.ApplyPositionDelta(DX0);
 			Body0.ApplyRotationDelta(DR0);
 		}
 		if (Body1.IsDynamic())
 		{
 			const FSolverVec3 DX1 = -Body1.InvM() * PushOut;
-			const FSolverVec3 DR1 = ManifoldPoint.WorldContactNormalAngular1 * -PushOutNormal + ManifoldPoint.WorldContactTangentUAngular1 * -PushOutTangentU + ManifoldPoint.WorldContactTangentVAngular1 * -PushOutTangentV;
+			const FSolverVec3 DR1 = ManifoldPoint.WorldContactTangentUAngular1 * -PushOutTangentU + ManifoldPoint.WorldContactTangentVAngular1 * -PushOutTangentV;
 			Body1.ApplyPositionDelta(DX1);
 			Body1.ApplyRotationDelta(DR1);
 		}
 	}
 
-	FORCEINLINE_DEBUGGABLE void ApplyPositionCorrectionNoFriction(
+	FORCEINLINE_DEBUGGABLE void ApplyPositionCorrectionNormal(
 		const FSolverReal Stiffness,
 		const FSolverReal ContactDeltaNormal,
 		FPBDCollisionSolverManifoldPoint& ManifoldPoint,
@@ -892,11 +885,21 @@ namespace Chaos
 		// NOTE: MaxPushOut == 0 disables the pushout limits
 		if ((MaxPushOut > 0) && (OutContactDeltaNormal < -MaxPushOut))
 		{
-			const FSolverReal ClampedContactDeltaNormal = -MaxPushOut;
-			OutContactDeltaNormal = ClampedContactDeltaNormal;
+			OutContactDeltaNormal = -MaxPushOut;
 		}
 	}
 
+
+	FORCEINLINE_DEBUGGABLE void FPBDCollisionSolverManifoldPoint::CalculateContactPositionErrorTangential(const FConstraintSolverBody& Body0, const FConstraintSolverBody& Body1, FSolverReal& OutContactDeltaTangentU, FSolverReal& OutContactDeltaTangentV) const
+	{
+		// Linear version: calculate the contact delta assuming linear motion after applying a positional impulse at the contact point. There will be an error that depends on the size of the rotation.
+		const FSolverVec3 ContactDelta0 = Body0.DP() + FSolverVec3::CrossProduct(Body0.DQ(), RelativeContactPosition0);
+		const FSolverVec3 ContactDelta1 = Body1.DP() + FSolverVec3::CrossProduct(Body1.DQ(), RelativeContactPosition1);
+		const FSolverVec3 ContactDelta = ContactDelta0 - ContactDelta1;
+		OutContactDeltaTangentU = WorldContactDeltaTangentU + FSolverVec3::DotProduct(ContactDelta, WorldContactTangentU);
+		OutContactDeltaTangentV = WorldContactDeltaTangentV + FSolverVec3::DotProduct(ContactDelta, WorldContactTangentV);
+	}
+
 	FORCEINLINE_DEBUGGABLE void FPBDCollisionSolverManifoldPoint::CalculateContactPositionError(const FConstraintSolverBody& Body0, const FConstraintSolverBody& Body1, const FSolverReal MaxPushOut, FSolverReal& OutContactDeltaNormal, FSolverReal& OutContactDeltaTangentU, FSolverReal& OutContactDeltaTangentV) const
 	{
 #if CHAOS_NONLINEAR_COLLISIONS_ENABLED
@@ -1044,27 +1047,71 @@ namespace Chaos
 		FConstraintSolverBody& Body0 = SolverBody0();
 		FConstraintSolverBody& Body1 = SolverBody1();
 
-		// Apply the position correction so that all contacts have zero separation
+		// Accumulate net pushout for friction limits below
+		bool bApplyFriction[MaxPointsPerConstraint] = { false, };
+		int32 NumFrictionContacts = 0;
+		FSolverReal TotalPushOutNormal = FSolverReal(0);
+
+		// Apply the position correction along the normal and determine if we want to run friction on each point
 		for (int32 PointIndex = 0; PointIndex < State.NumManifoldPoints; ++PointIndex)
 		{
 			FPBDCollisionSolverManifoldPoint& SolverManifoldPoint = State.ManifoldPoints[PointIndex];
 
-			FSolverReal ContactDeltaNormal, ContactDeltaTangent0, ContactDeltaTangent1;
-			SolverManifoldPoint.CalculateContactPositionError(Body0.SolverBody(), Body1.SolverBody(), MaxPushOut, ContactDeltaNormal, ContactDeltaTangent0, ContactDeltaTangent1);
+			FSolverReal ContactDeltaNormal;
+			SolverManifoldPoint.CalculateContactPositionErrorNormal(Body0.SolverBody(), Body1.SolverBody(), MaxPushOut, ContactDeltaNormal);
 
+			// Apply a normal correction if we still have penetration or if we are now separated but have previously applied a correction that we may want to undo
 			const bool bProcessManifoldPoint = (ContactDeltaNormal < FSolverReal(0)) || (SolverManifoldPoint.NetPushOutNormal > FSolverReal(UE_SMALL_NUMBER));
 			if (bProcessManifoldPoint)
 			{
-				ApplyPositionCorrectionWithFriction(
+				ApplyPositionCorrectionNormal(
 					State.Stiffness,
-					State.StaticFriction,
-					State.DynamicFriction,
 					ContactDeltaNormal,
-					ContactDeltaTangent0,
-					ContactDeltaTangent1,
 					SolverManifoldPoint,
 					Body0,
 					Body1);
+
+				TotalPushOutNormal += SolverManifoldPoint.NetPushOutNormal;
+			}
+
+			// Friction gets updated for any point with a net normal correction or where we have previously had a normal correction and 
+			// already applied friction (in which case we may need to zero it)
+			if ((SolverManifoldPoint.NetPushOutNormal != 0) || (SolverManifoldPoint.NetPushOutTangentU != 0) || (SolverManifoldPoint.NetPushOutTangentV != 0))
+			{
+				bApplyFriction[PointIndex] = true;
+				++NumFrictionContacts;
+			}
+		}
+
+		// Apply the tangential position correction if required
+		if (NumFrictionContacts > 0)
+		{
+			// We clip the tangential correction at each contact to the friction cone, but we use to average impulse
+			// among all contacts as the clipping limit. This is not really correct but it is much more stable to 
+			// differences in contacts from tick to tick
+			// @todo(chaos): try a decaying maximum per contact point rather than an average (again - we had that once!)
+			const FSolverReal FrictionMaxPushOut = TotalPushOutNormal / FSolverReal(NumFrictionContacts);
+
+			for (int32 PointIndex = 0; PointIndex < State.NumManifoldPoints; ++PointIndex)
+			{
+				if (bApplyFriction[PointIndex])
+				{
+					FPBDCollisionSolverManifoldPoint& SolverManifoldPoint = State.ManifoldPoints[PointIndex];
+
+					FSolverReal ContactDeltaTangentU, ContactDeltaTangentV;
+					SolverManifoldPoint.CalculateContactPositionErrorTangential(Body0.SolverBody(), Body1.SolverBody(), ContactDeltaTangentU, ContactDeltaTangentV);
+
+					ApplyPositionCorrectionTangential(
+						State.Stiffness,
+						State.StaticFriction,
+						State.DynamicFriction,
+						FrictionMaxPushOut,
+						ContactDeltaTangentU,
+						ContactDeltaTangentV,
+						SolverManifoldPoint,
+						Body0,
+						Body1);
+				}
 			}
 		}
 
@@ -1088,7 +1135,7 @@ namespace Chaos
 			const bool bProcessManifoldPoint = (ContactDeltaNormal < FSolverReal(0)) || (SolverManifoldPoint.NetPushOutNormal > FSolverReal(UE_SMALL_NUMBER));
 			if (bProcessManifoldPoint)
 			{
-				ApplyPositionCorrectionNoFriction(
+				ApplyPositionCorrectionNormal(
 					State.Stiffness,
 					ContactDeltaNormal,
 					SolverManifoldPoint,
diff --git a/Engine/Source/Runtime/Experimental/Chaos/Private/Chaos/Collision/SolverCollisionContainer.cpp b/Engine/Source/Runtime/Experimental/Chaos/Private/Chaos/Collision/SolverCollisionContainer.cpp
--- a/Engine/Source/Runtime/Experimental/Chaos/Private/Chaos/Collision/SolverCollisionContainer.cpp
+++ b/Engine/Source/Runtime/Experimental/Chaos/Private/Chaos/Collision/SolverCollisionContainer.cpp
@@ -3,6 +3,7 @@
 
 #include "Chaos/Collision/CollisionApplyType.h"
 #include "Chaos/Collision/PBDCollisionConstraint.h"
+#include "Chaos/DebugDrawQueue.h"
 #include "Chaos/Evolution/SolverBodyContainer.h"
 #include "Chaos/Particle/ParticleUtilities.h"
 #include "Chaos/PBDCollisionConstraints.h"
@@ -20,8 +21,6 @@ namespace Chaos
 {
 	namespace CVars
 	{
-		extern bool bChaos_PBDCollisionSolver_VectorRegister;
-
 		extern bool bChaos_PBDCollisionSolver_Position_SolveEnabled;
 		extern bool bChaos_PBDCollisionSolver_Velocity_SolveEnabled;
 
@@ -30,8 +29,8 @@ namespace Chaos
 		// Solver stiffness will be equal to 0 when the mass ratio is MassRatio2.
 		FRealSingle Chaos_PBDCollisionSolver_AutoStiffness_MassRatio1 = 0;
 		FRealSingle Chaos_PBDCollisionSolver_AutoStiffness_MassRatio2 = 0;
-		FAutoConsoleVariableRef CVarChaosPBDCollisionSolverAutoStiffnessMassRatio1(TEXT("p.Chaos.PBDCollision.AutoStiffness.MassRatio1"), Chaos_PBDCollisionSolver_AutoStiffness_MassRatio1, TEXT(""));
-		FAutoConsoleVariableRef CVarChaosPBDCollisionSolverAutoStiffnessMassRatio2(TEXT("p.Chaos.PBDCollision.AutoStiffness.MassRatio2"), Chaos_PBDCollisionSolver_AutoStiffness_MassRatio2, TEXT(""));
+		FAutoConsoleVariableRef CVarChaosPBDCollisionSolverAutoStiffnessMassRatio1(TEXT("p.Chaos.PBDCollisionSolver.AutoStiffness.MassRatio1"), Chaos_PBDCollisionSolver_AutoStiffness_MassRatio1, TEXT(""));
+		FAutoConsoleVariableRef CVarChaosPBDCollisionSolverAutoStiffnessMassRatio2(TEXT("p.Chaos.PBDCollisionSolver.AutoStiffness.MassRatio2"), Chaos_PBDCollisionSolver_AutoStiffness_MassRatio2, TEXT(""));
 
 	}
 	using namespace CVars;
@@ -196,14 +195,7 @@ namespace Chaos
 			Solver.SolverBody1().SetInvMScale(Constraint->GetInvMassScale1());
 			Solver.SolverBody1().SetInvIScale(Constraint->GetInvInertiaScale1());
 
-			if (!bChaos_PBDCollisionSolver_VectorRegister)
-			{
-				GatherManifoldPoints(Dt, Body0, Body1);
-			}
-			else
-			{
-				GatherManifoldPointsVectorRegister(Dt, Body0, Body1);
-			}
+			GatherManifoldPoints(Dt, Body0, Body1);
 
 			// We should try to remove this - the Constraint should not need to know about solver objects
 			Constraint->SetSolverBodies(Body0, Body1);
@@ -299,128 +291,19 @@ namespace Chaos
 					WorldContactDeltaNormal,
 					WorldContactDeltaTangentU,
 					WorldContactDeltaTangentV);
-			}
-		}
-
-		void GatherManifoldPointsVectorRegister(
-			const FReal InDt,
-			const FSolverBody* Body0,
-			const FSolverBody* Body1)
-		{
-			// Attempt to make this LWC independent. Actually this isn't easy because we explicitly use MakeVectorRegisterFloatFromDouble
-
-			// We handle incremental manifolds by just collecting any new contacts
-			const int32 BeginPointIndex = Solver.NumManifoldPoints();
-			const int32 EndPointIndex = Solver.SetNumManifoldPoints(Constraint->GetManifoldPoints().Num());
-			TArrayView<FManifoldPoint> ManifoldPoints = Constraint->GetManifoldPoints();
-
-			const FSolverReal Dt = FSolverReal(InDt);
-			const FSolverReal RestitutionVelocityThreshold = FSolverReal(Constraint->GetRestitutionThreshold()) * Dt;
-			const FSolverReal Restitution = FSolverReal(Constraint->GetRestitution());
-
-			// World space positions - must be LWC compatible
-			const VectorRegister ShapeWorldPos0 = Constraint->GetShapeWorldTransform0().GetTranslationRegister();
-			const VectorRegister ShapeWorldRot0 = Constraint->GetShapeWorldTransform0().GetRotationRegister();
-			const VectorRegister ShapeWorldPos1 = Constraint->GetShapeWorldTransform1().GetTranslationRegister();
-			const VectorRegister ShapeWorldRot1 = Constraint->GetShapeWorldTransform1().GetRotationRegister();
-			const VectorRegister BodyPos0 = VectorLoadFloat3(&Body0->P());
-			const VectorRegister BodyPos1 = VectorLoadFloat3(&Body1->P());
-
-			for (int32 ManifoldPointIndex = BeginPointIndex; ManifoldPointIndex < EndPointIndex; ++ManifoldPointIndex)
-			{
-				FManifoldPoint& ManifoldPoint = ManifoldPoints[ManifoldPointIndex];
-				if (ManifoldPoint.Flags.bDisabled)
-				{
-					continue;
-				}
-
-				const FSavedManifoldPoint* SavedManifoldPoint = Constraint->FindSavedManifoldPoint(ManifoldPoint);
-
-				const VectorRegister ShapeContactPos0 = VectorLoadFloat3(&ManifoldPoint.ContactPoint.ShapeContactPoints[0]);
-				const VectorRegister ShapeContactPos1 = VectorLoadFloat3(&ManifoldPoint.ContactPoint.ShapeContactPoints[1]);
-				const VectorRegister ShapeContactNormal1 = VectorLoadFloat3(&ManifoldPoint.ContactPoint.ShapeContactNormal);
-
-				const VectorRegister WorldContactPoint0 = VectorAdd(ShapeWorldPos0, VectorQuaternionRotateVector(ShapeWorldRot0, ShapeContactPos0));
-				const VectorRegister WorldContactPoint1 = VectorAdd(ShapeWorldPos1, VectorQuaternionRotateVector(ShapeWorldRot1, ShapeContactPos1));
-				const VectorRegister WorldContactPoint = VectorMultiply(VectorSetFloat1(FSolverReal(0.5)), VectorAdd(WorldContactPoint0, WorldContactPoint1));
-
-				// NOTE: low precision for relative coordinates
-				const SolverVectorRegister WorldContactNormal = MakeVectorRegisterFloatFromDouble(VectorQuaternionRotateVector(ShapeWorldRot1, ShapeContactNormal1));
-				const SolverVectorRegister RelativeContactPosition0 = MakeVectorRegisterFloatFromDouble(VectorSubtract(WorldContactPoint, BodyPos0));
-				const SolverVectorRegister RelativeContactPosition1 = MakeVectorRegisterFloatFromDouble(VectorSubtract(WorldContactPoint, BodyPos1));
-				const FSolverReal TargetPhi = FSolverReal(ManifoldPoint.TargetPhi);
-
-				// If we have contact data from a previous tick, use it to calculate the lateral position delta we need
-				// to apply to move the contacts back to their original relative locations (i.e., to enforce static friction)
-				// @todo(chaos): we should not be writing back to the constraint here - find a better way to update the friction anchor. See FPBDCollisionConstraint::SetSolverResults
-				// NOTE: low precision for relative coordinates
-				SolverVectorRegister WorldFrictionDelta = VectorZeroFloat();
-				if (SavedManifoldPoint != nullptr)
-				{
-					const VectorRegister SavedShapeContactPos0 = VectorLoadFloat3(&SavedManifoldPoint->ShapeContactPoints[0]);
-					const VectorRegister SavedShapeContactPos1 = VectorLoadFloat3(&SavedManifoldPoint->ShapeContactPoints[1]);
-					const VectorRegister FrictionDelta0 = VectorSubtract(SavedShapeContactPos0, ShapeContactPos0);
-					const VectorRegister FrictionDelta1 = VectorSubtract(SavedShapeContactPos1, ShapeContactPos1);
-					WorldFrictionDelta = MakeVectorRegisterFloatFromDouble(VectorSubtract(VectorQuaternionRotateVector(ShapeWorldRot0, FrictionDelta0), VectorQuaternionRotateVector(ShapeWorldRot1, FrictionDelta1)));
-
-					ManifoldPoint.ShapeAnchorPoints[0] = SavedManifoldPoint->ShapeContactPoints[0];
-					ManifoldPoint.ShapeAnchorPoints[1] = SavedManifoldPoint->ShapeContactPoints[1];
-				}
-				else
-				{
-					const VectorRegister BodyV0 = VectorLoadFloat3(&Body0->V());
-					const VectorRegister BodyW0 = VectorLoadFloat3(&Body0->W());
-					const VectorRegister BodyV1 = VectorLoadFloat3(&Body1->V());
-					const VectorRegister BodyW1 = VectorLoadFloat3(&Body1->W());
-					const VectorRegister ContactVel0 = VectorAdd(BodyV0, VectorCross(BodyW0, RelativeContactPosition0));
-					const VectorRegister ContactVel1 = VectorAdd(BodyV1, VectorCross(BodyW1, RelativeContactPosition1));
-					const SolverVectorRegister ContactVel = MakeVectorRegisterFloatFromDouble(VectorSubtract(ContactVel0, ContactVel1));
-					WorldFrictionDelta = VectorMultiply(ContactVel, VectorSetFloat1(Dt));
 
-					ManifoldPoint.ShapeAnchorPoints[0] = ManifoldPoint.ContactPoint.ShapeContactPoints[0];
-					ManifoldPoint.ShapeAnchorPoints[1] = ManifoldPoint.ContactPoint.ShapeContactPoints[1];
-				}
-
-				// World-space contact tangents. We are treating the normal as the constraint-space Z axis
-				// and the Tangent U and V as the constraint-space X and Y axes respectively
-				SolverVectorRegister WorldContactTangentU = VectorCross(MakeVectorRegisterFloatConstant(0, 1, 0, 0), WorldContactNormal);
-				SolverVectorRegister WorldContactTangentULenSq = VectorDot3(WorldContactTangentU, WorldContactTangentU);
-				SolverVectorRegister WorldContactTangentUCompareMask = VectorCompareGT(WorldContactTangentULenSq, VectorSetFloat1(UE_KINDA_SMALL_NUMBER));
-				if (VectorMaskBits(WorldContactTangentUCompareMask))
-				{
-					WorldContactTangentU = VectorMultiply(WorldContactTangentU, VectorReciprocalSqrt(WorldContactTangentULenSq));
-				}
-				else
+#if CHAOS_DEBUG_DRAW && 0
+				if (FDebugDrawQueue::GetInstance().IsDebugDrawingEnabled())
 				{
-					WorldContactTangentU = VectorCross(MakeVectorRegisterFloatConstant(1, 0, 0, 0), WorldContactNormal);
-					WorldContactTangentU = VectorNormalize(WorldContactTangentU);
+					FDebugDrawQueue::GetInstance().DrawDebugLine(WorldContactPoint1, WorldContactPoint1 - 5.0f * WorldContactDeltaNormal * WorldContactNormal, FColor::Red, false, 0, 10, 0.2f);
+					FDebugDrawQueue::GetInstance().DrawDebugLine(WorldContactPoint1, WorldContactPoint1 - 5.0f * WorldContactDeltaTangentU * WorldContactTangentU, FColor::Blue, false, 0, 10, 0.2f);
+					FDebugDrawQueue::GetInstance().DrawDebugLine(WorldContactPoint1, WorldContactPoint1 - 5.0f * WorldContactDeltaTangentV * WorldContactTangentV, FColor::Blue, false, 0, 10, 0.2f);
+					FDebugDrawQueue::GetInstance().DrawDebugLine(WorldContactPoint0, WorldContactPoint1, FColor::Black, false, 0, 10, 0.2f);
 				}
-				const SolverVectorRegister WorldContactTangentV = VectorCross(WorldContactNormal, WorldContactTangentU);
-
-				// The contact point error we are trying to correct in this solver
-				const SolverVectorRegister WorldContactDelta = MakeVectorRegisterFloatFromDouble(VectorSubtract(WorldContactPoint0, WorldContactPoint1));
-				const FSolverReal WorldContactDeltaNormal = VectorGetComponent(VectorDot3(WorldContactDelta, WorldContactNormal), 0) - TargetPhi;
-				const FSolverReal WorldContactDeltaTangentU = VectorGetComponent(VectorDot3(VectorAdd(WorldContactDelta, WorldFrictionDelta), WorldContactTangentU), 0);
-				const FSolverReal WorldContactDeltaTangentV = VectorGetComponent(VectorDot3(VectorAdd(WorldContactDelta, WorldFrictionDelta), WorldContactTangentV), 0);
-
-				// Copy all the properties into the solver
-				Solver.SetManifoldPoint(
-					ManifoldPointIndex,
-					Dt,
-					Restitution,
-					RestitutionVelocityThreshold,
-					RelativeContactPosition0,
-					RelativeContactPosition1,
-					WorldContactNormal,
-					WorldContactTangentU,
-					WorldContactTangentV,
-					WorldContactDeltaNormal,
-					WorldContactDeltaTangentU,
-					WorldContactDeltaTangentV);
+#endif
 			}
 		}
 
-
 		/**
 		 * @brief Send all solver results to the constraint
 		*/
diff --git a/Engine/Source/Runtime/Experimental/Chaos/Public/Chaos/Collision/SolverCollisionContainer.h b/Engine/Source/Runtime/Experimental/Chaos/Public/Chaos/Collision/SolverCollisionContainer.h
--- a/Engine/Source/Runtime/Experimental/Chaos/Public/Chaos/Collision/SolverCollisionContainer.h
+++ b/Engine/Source/Runtime/Experimental/Chaos/Public/Chaos/Collision/SolverCollisionContainer.h
@@ -60,7 +60,8 @@ namespace Chaos
 
 		virtual void Reset(const int32 InMaxCollisions) override;
 
-		/** Set the number of collision solvers in the container
+		/** 
+		 * Set the number of collision solvers in the container
 		 * @param MaxCollisions Max number of collision solvers
 		 */
 		void SetNum(const int32 MaxCollisions);
