Dataflow : Passthrough callback
- removed self-reference in initialization.

diff --git a/Engine/Source/Runtime/Experimental/Dataflow/Core/Public/Dataflow/DataflowInputOutput.h b/Engine/Source/Runtime/Experimental/Dataflow/Core/Public/Dataflow/DataflowInputOutput.h
--- a/Engine/Source/Runtime/Experimental/Dataflow/Core/Public/Dataflow/DataflowInputOutput.h
+++ b/Engine/Source/Runtime/Experimental/Dataflow/Core/Public/Dataflow/DataflowInputOutput.h
@@ -115,7 +115,7 @@ namespace Dataflow
 	template<class T>
 	struct DATAFLOWCORE_API TOutputParameters 
 	{
-		typedef TFunction<void(const  Dataflow::FContext& Context, Dataflow::FConnection*)> PassthroughCallback;
+		typedef TFunction<T(const  Dataflow::FContext& Context, Dataflow::FConnection*)> PassthroughCallback;
 
 		TOutputParameters(FName InName, FDataflowNode* InOwner, T InDefault = T())
 			: Type(GraphConnectionTypeName<T>())
@@ -139,7 +139,7 @@ namespace Dataflow
 	template<class T>
 	class DATAFLOWCORE_API TOutput : public FConnection
 	{
-		typedef TFunction<void(const  Dataflow::FContext& Context, Dataflow::FConnection*)> PassthroughCallback;
+		typedef TFunction<T(const  Dataflow::FContext& Context, Dataflow::FConnection*)> PassthroughCallback;
 		typedef FConnection Super;
 		friend class FConnection;
 
@@ -158,7 +158,6 @@ namespace Dataflow
 			Super::BindOutput(Param.Owner, this);
 		}
 
-
 		const TArray<TInput<T>*>& GetConnections() const { return Connections; }
 		TArray<TInput<T>* >& GetConnections() { return Connections; }
 
@@ -202,7 +201,7 @@ namespace Dataflow
 					SetValue(Default, Context);
 					if (Passthrough)
 					{
-						Passthrough(Context, this);
+						SetValue(Passthrough(Context, this), Context);
 					}
 				}
 			}
