Interchange: Fix ensure when re-importing staticmesh. The CreatePackage is now always call even if we do a re-import. This allow factories to cleanup the resource in the game thread before re-importing the asset.

#rnx

diff --git a/Engine/Plugins/Experimental/Interchange/Runtime/Source/Import/Private/Mesh/InterchangeStaticMeshFactory.cpp b/Engine/Plugins/Experimental/Interchange/Runtime/Source/Import/Private/Mesh/InterchangeStaticMeshFactory.cpp
--- a/Engine/Plugins/Experimental/Interchange/Runtime/Source/Import/Private/Mesh/InterchangeStaticMeshFactory.cpp
+++ b/Engine/Plugins/Experimental/Interchange/Runtime/Source/Import/Private/Mesh/InterchangeStaticMeshFactory.cpp
@@ -62,7 +62,7 @@ UObject* UInterchangeStaticMeshFactory::CreateEmptyAsset(const FCreateAssetParam
 	}
 
 	// create an asset if it doesn't exist
-	UObject* ExistingAsset = StaticFindObject(nullptr, Arguments.Parent, *Arguments.AssetName);
+	UObject* ExistingAsset = (Arguments.ReimportObject != nullptr) ? Arguments.ReimportObject : StaticFindObject(nullptr, Arguments.Parent, *Arguments.AssetName);
 
 	// create a new static mesh or overwrite existing asset, if possible
 	if (!ExistingAsset)
diff --git a/Engine/Source/Runtime/Experimental/Interchange/Engine/Private/Tasks/InterchangeTaskCreateAsset.cpp b/Engine/Source/Runtime/Experimental/Interchange/Engine/Private/Tasks/InterchangeTaskCreateAsset.cpp
--- a/Engine/Source/Runtime/Experimental/Interchange/Engine/Private/Tasks/InterchangeTaskCreateAsset.cpp
+++ b/Engine/Source/Runtime/Experimental/Interchange/Engine/Private/Tasks/InterchangeTaskCreateAsset.cpp
@@ -147,6 +147,20 @@ void UE::Interchange::FTaskCreatePackage::DoTask(ENamedThreads::Type CurrentThre
 		{
 			Pkg = ExistingObject->GetPackage();
 			PackageName = Pkg->GetPathName();
+			//Import Asset describe by the node
+			UInterchangeFactoryBase::FCreateAssetParams CreateAssetParams;
+			CreateAssetParams.AssetName = AssetName;
+			CreateAssetParams.AssetNode = FactoryNode;
+			CreateAssetParams.Parent = Pkg;
+			CreateAssetParams.SourceData = AsyncHelper->SourceDatas[SourceIndex];
+			CreateAssetParams.Translator = AsyncHelper->Translators[SourceIndex];
+			if (AsyncHelper->BaseNodeContainers.IsValidIndex(SourceIndex))
+			{
+				CreateAssetParams.NodeContainer = AsyncHelper->BaseNodeContainers[SourceIndex].Get();
+			}
+			CreateAssetParams.ReimportObject = AsyncHelper->TaskData.ReimportObject;
+			//We call CreateEmptyAsset to ensure any resource use by an existing UObject is release on the game thread
+			Factory->CreateEmptyAsset(CreateAssetParams);
 		}
 		else
 		{
