Chaos destruction : add slow moving detection fro Remove on sleep feature

diff --git a/Engine/Source/Runtime/Experimental/GeometryCollectionEngine/Private/GeometryCollection/GeometryCollectionComponent.cpp b/Engine/Source/Runtime/Experimental/GeometryCollectionEngine/Private/GeometryCollection/GeometryCollectionComponent.cpp
--- a/Engine/Source/Runtime/Experimental/GeometryCollectionEngine/Private/GeometryCollection/GeometryCollectionComponent.cpp
+++ b/Engine/Source/Runtime/Experimental/GeometryCollectionEngine/Private/GeometryCollection/GeometryCollectionComponent.cpp
@@ -77,12 +77,6 @@ static bool bChaos_BoxCalcBounds_ISPC_Enabled = CHAOS_BOX_CALC_BOUNDS_ISPC_ENABL
 static FAutoConsoleVariableRef CVarChaosBoxCalcBoundsISPCEnabled(TEXT("p.Chaos.BoxCalcBounds.ISPC"), bChaos_BoxCalcBounds_ISPC_Enabled, TEXT("Whether to use ISPC optimizations in calculating box bounds in geometry collections"));
 #endif
 
-static bool ChaosCheckIfStillForRemoval  = false;
-static FAutoConsoleVariableRef CVarChaosCheckIfStillForRemoval(TEXT("p.Chaos.CheckIfStillForRemoval"), ChaosCheckIfStillForRemoval, TEXT("if true, non moving GC ( even non sleeping ) will be considered for removal if the feature is enabled"));
-
-static float ChaosStillCheckVelocityThreshold  = 1; // 1 cm/s
-static FAutoConsoleVariableRef CVarChaosStillCheckDistanceThreshold(TEXT("p.Chaos.ChaosStillCheckVelocityThreshold"), ChaosStillCheckVelocityThreshold, TEXT("When using ChaosStillCheckVelocityThreshold, this set the velocity threshold in cm/s "));
-
 DEFINE_LOG_CATEGORY_STATIC(UGCC_LOG, Error, All);
 
 extern FGeometryCollectionDynamicDataPool GDynamicDataPool;
@@ -3527,14 +3521,19 @@ void UGeometryCollectionComponent::IncrementSleepTimer(float DeltaTime)
 			const bool IsRemovalActive = (MaxSleepTime[TransformIdx] >= 0) && bIsDynamicOrSleeping && bNoParent && !bHasKinematicOrStaticInternalClusterParent;
 			if (IsRemovalActive)
 			{
-				const FVector CurrentPosition = DynamicCollection->Transform[TransformIdx].GetTranslation();
-				const FVector::FReal InstantVelocity = (CurrentPosition-LastPosition[TransformIdx]).Size() / (FVector::FReal)DeltaTime; 
-				const bool bIsStill = ChaosCheckIfStillForRemoval && (InstantVelocity < ChaosStillCheckVelocityThreshold);
-				LastPosition[TransformIdx] = CurrentPosition;
+				bool bIsSlowMoving = false;
+				if (RestCollection->bSlowMovingAsSleeping)
+				{
+					const FVector CurrentPosition = DynamicCollection->Transform[TransformIdx].GetTranslation();
+					LastPosition[TransformIdx] = CurrentPosition;
+					
+					const FVector::FReal InstantVelocity = (CurrentPosition-LastPosition[TransformIdx]).Size() / (FVector::FReal)DeltaTime; 
+					bIsSlowMoving = (InstantVelocity < RestCollection->SlowMovingVelocityThreshold);
+				}
 				
 				const bool IsSleeping = (DynamicCollection->DynamicState[TransformIdx] == (int)EObjectStateTypeEnum::Chaos_Object_Sleeping);
 				const bool HasStartedDecaying = (Decay[TransformIdx] > 0); 
-				if (IsSleeping || bIsStill || HasStartedDecaying)
+				if (IsSleeping || bIsSlowMoving || HasStartedDecaying)
 				{
 					SleepTimer[TransformIdx] += DeltaTime;
 				}
diff --git a/Engine/Source/Runtime/Experimental/GeometryCollectionEngine/Private/GeometryCollection/GeometryCollectionObject.cpp b/Engine/Source/Runtime/Experimental/GeometryCollectionEngine/Private/GeometryCollection/GeometryCollectionObject.cpp
--- a/Engine/Source/Runtime/Experimental/GeometryCollectionEngine/Private/GeometryCollection/GeometryCollectionObject.cpp
+++ b/Engine/Source/Runtime/Experimental/GeometryCollectionEngine/Private/GeometryCollection/GeometryCollectionObject.cpp
@@ -95,6 +95,8 @@ UGeometryCollection::UGeometryCollection(const FObjectInitializer& ObjectInitial
 	, bRemoveOnMaxSleep(false)
 	, MaximumSleepTime(5.0, 10.0)
 	, RemovalDuration(2.5, 5.0)
+	, bSlowMovingAsSleeping(true)
+	, SlowMovingVelocityThreshold(1)
 	, EnableRemovePiecesOnFracture(false)
 	, GeometryCollection(new FGeometryCollection())
 {
diff --git a/Engine/Source/Runtime/Experimental/GeometryCollectionEngine/Public/GeometryCollection/GeometryCollectionObject.h b/Engine/Source/Runtime/Experimental/GeometryCollectionEngine/Public/GeometryCollection/GeometryCollectionObject.h
--- a/Engine/Source/Runtime/Experimental/GeometryCollectionEngine/Public/GeometryCollection/GeometryCollectionObject.h
+++ b/Engine/Source/Runtime/Experimental/GeometryCollectionEngine/Public/GeometryCollection/GeometryCollectionObject.h
@@ -555,17 +555,25 @@ public:
 #endif
 
 	/** Remove particle from simulation and dissolve rendered geometry once sleep threshold has been exceeded. */
-	UPROPERTY(EditAnywhere, BlueprintReadOnly, Category = Removal, meta = (DisplayName = "RemoveOnMaxSleep"))
+	UPROPERTY(EditAnywhere, BlueprintReadOnly, Category = Removal, meta = (DisplayName = "Remove on Sleep"))
 	bool bRemoveOnMaxSleep;
 	
 	/** How long may the particle sleep before initiating removal (in seconds). */
-	UPROPERTY(EditAnywhere, BlueprintReadOnly, Category = Removal, DisplayName = "Sleep Min Max")
+	UPROPERTY(EditAnywhere, BlueprintReadOnly, Category = Removal, meta = (DisplayName = "Sleep Min Max", EditCondition="bRemoveOnMaxSleep"))
 	FVector2D MaximumSleepTime;
 
 	/** How long does the removal process take (in seconds). */
-	UPROPERTY(EditAnywhere, BlueprintReadOnly, Category = Removal)
+	UPROPERTY(EditAnywhere, BlueprintReadOnly, Category = Removal, meta = (DisplayName = "Removal Duration", EditCondition="bRemoveOnMaxSleep"))
 	FVector2D RemovalDuration;
 
+	/** when on non-sleeping, slow moving pieces will be considered as sleeping, this help removal of jittery but not really moving objects. */
+	UPROPERTY(EditAnywhere, BlueprintReadOnly, Category = Removal, meta = (DisplayName = "Slow*-Moving as sleeping", EditCondition="bRemoveOnMaxSleep"))
+	bool bSlowMovingAsSleeping;
+
+	/** When slow moving detection is on, this defines the linear velocity thresholds in cm/s to consider the object as sleeping . */
+	UPROPERTY(EditAnywhere, BlueprintReadOnly, Category = Removal, meta = (DisplayName = "Slow*-Moving Velocity Threshold", EditCondition="bRemoveOnMaxSleep && bSlowMovingAsSleeping"))
+	float SlowMovingVelocityThreshold;
+	
 	/*
 	* Size Specfic Data reflects the default geometry to bind to rigid bodies smaller
 	* than the max size volume. This can also be empty to reflect no collision geometry
