Improve fracture tool selectors:
 - Add leaf and cluster selection tools
 - Make children selection not deselect nodes w/ no children
 - Use TSet::Add instead of TArray::AddUnique to create selections

diff --git a/Engine/Plugins/Experimental/ChaosEditor/Source/FractureEditor/Private/FractureEditorCommands.h b/Engine/Plugins/Experimental/ChaosEditor/Source/FractureEditor/Private/FractureEditorCommands.h
--- a/Engine/Plugins/Experimental/ChaosEditor/Source/FractureEditor/Private/FractureEditorCommands.h
+++ b/Engine/Plugins/Experimental/ChaosEditor/Source/FractureEditor/Private/FractureEditorCommands.h
@@ -33,6 +33,8 @@ class FFractureEditorCommands : public TCommands<FFractureEditorCommands>
 		TSharedPtr< FUICommandInfo > SelectSiblings;
 		TSharedPtr< FUICommandInfo > SelectAllInLevel;
 		TSharedPtr< FUICommandInfo > SelectInvert;
+		TSharedPtr< FUICommandInfo > SelectLeaves;
+		TSharedPtr< FUICommandInfo > SelectClusters;
 		TSharedPtr< FUICommandInfo > SelectCustom;
 
 		// View Settings
diff --git a/Engine/Plugins/Experimental/ChaosEditor/Source/FractureEditor/Private/FractureEditorModeToolkit.cpp b/Engine/Plugins/Experimental/ChaosEditor/Source/FractureEditor/Private/FractureEditorModeToolkit.cpp
--- a/Engine/Plugins/Experimental/ChaosEditor/Source/FractureEditor/Private/FractureEditorModeToolkit.cpp
+++ b/Engine/Plugins/Experimental/ChaosEditor/Source/FractureEditor/Private/FractureEditorModeToolkit.cpp
@@ -825,6 +825,8 @@ void FFractureEditorModeToolkit::BuildToolPalette(FName PaletteIndex, class FToo
 		ToolbarBuilder.AddToolBarButton(Commands.SelectSiblings);
 		ToolbarBuilder.AddToolBarButton(Commands.SelectAllInLevel);
 		ToolbarBuilder.AddToolBarButton(Commands.SelectNeighbors);
+		ToolbarBuilder.AddToolBarButton(Commands.SelectLeaves);
+		ToolbarBuilder.AddToolBarButton(Commands.SelectClusters);
 		ToolbarBuilder.AddToolBarButton(Commands.SelectCustom);
 	}
 	else if (PaletteIndex == TEXT("Fracture"))
diff --git a/Engine/Plugins/Experimental/ChaosEditor/Source/FractureEditor/Private/FractureToolSelectors.cpp b/Engine/Plugins/Experimental/ChaosEditor/Source/FractureEditor/Private/FractureToolSelectors.cpp
--- a/Engine/Plugins/Experimental/ChaosEditor/Source/FractureEditor/Private/FractureToolSelectors.cpp
+++ b/Engine/Plugins/Experimental/ChaosEditor/Source/FractureEditor/Private/FractureToolSelectors.cpp
@@ -303,4 +303,68 @@ void UFractureToolSelectInvert::Execute(TWeakPtr<FFractureEditorModeToolkit> InT
 	}
 }
 
+
+
+
+FText UFractureToolSelectLeaf::GetDisplayText() const
+{
+	return FText(NSLOCTEXT("Fracture", "FractureToolSelectLeaf", "Select Leaves"));
+}
+
+FText UFractureToolSelectLeaf::GetTooltipText() const
+{
+	return FText(NSLOCTEXT("Fracture", "FractureToolSelectLeafTooltip", "Select leaf bones that represent rigid bodies"));
+}
+
+FSlateIcon UFractureToolSelectLeaf::GetToolIcon() const
+{
+	return FSlateIcon("FractureEditorStyle", "FractureEditor.SelectLeaf");
+}
+
+void UFractureToolSelectLeaf::RegisterUICommand(FFractureEditorCommands* BindingContext)
+{
+	UI_COMMAND_EXT(BindingContext, UICommandInfo, "SelectLeaf", "Leaf", "Select (rigid) leaf bones.", EUserInterfaceActionType::Button, FInputChord());
+	BindingContext->SelectLeaves = UICommandInfo;
+}
+
+void UFractureToolSelectLeaf::Execute(TWeakPtr<FFractureEditorModeToolkit> InToolkit)
+{
+	if (InToolkit.IsValid())
+	{
+		SelectByMode(InToolkit.Pin().Get(), GeometryCollection::ESelectionMode::Leaves);
+	}
+}
+
+
+
+
+FText UFractureToolSelectCluster::GetDisplayText() const
+{
+	return FText(NSLOCTEXT("Fracture", "FractureToolSelectCluster", "Select Clusters"));
+}
+
+FText UFractureToolSelectCluster::GetTooltipText() const
+{
+	return FText(NSLOCTEXT("Fracture", "FractureToolSelectLeafTooltip", "Select cluster bones"));
+}
+
+FSlateIcon UFractureToolSelectCluster::GetToolIcon() const
+{
+	return FSlateIcon("FractureEditorStyle", "FractureEditor.SelectCluster");
+}
+
+void UFractureToolSelectCluster::RegisterUICommand(FFractureEditorCommands* BindingContext)
+{
+	UI_COMMAND_EXT(BindingContext, UICommandInfo, "SelectCluster", "Cluster", "Select cluster bones.", EUserInterfaceActionType::Button, FInputChord());
+	BindingContext->SelectClusters = UICommandInfo;
+}
+
+void UFractureToolSelectCluster::Execute(TWeakPtr<FFractureEditorModeToolkit> InToolkit)
+{
+	if (InToolkit.IsValid())
+	{
+		SelectByMode(InToolkit.Pin().Get(), GeometryCollection::ESelectionMode::Clusters);
+	}
+}
+
 #undef LOCTEXT_NAMESPACE
\ No newline at end of file
diff --git a/Engine/Plugins/Experimental/ChaosEditor/Source/FractureEditor/Private/FractureToolSelectors.h b/Engine/Plugins/Experimental/ChaosEditor/Source/FractureEditor/Private/FractureToolSelectors.h
--- a/Engine/Plugins/Experimental/ChaosEditor/Source/FractureEditor/Private/FractureToolSelectors.h
+++ b/Engine/Plugins/Experimental/ChaosEditor/Source/FractureEditor/Private/FractureToolSelectors.h
@@ -151,3 +151,37 @@ public:
 	virtual void RegisterUICommand(FFractureEditorCommands* BindingContext) override;
 	virtual void Execute(TWeakPtr<FFractureEditorModeToolkit> InToolkit) override;
 };
+
+
+UCLASS(DisplayName = "Select Leaves", Category = "FractureTools")
+class UFractureToolSelectLeaf : public UFractureToolSelectAll
+{
+public:
+	GENERATED_BODY()
+
+	UFractureToolSelectLeaf(const FObjectInitializer& ObjInit) : Super(ObjInit) {}
+
+	// UFractureActionTool Interface
+	virtual FText GetDisplayText() const override;
+	virtual FText GetTooltipText() const override;
+	virtual FSlateIcon GetToolIcon() const override;
+	virtual void RegisterUICommand(FFractureEditorCommands* BindingContext) override;
+	virtual void Execute(TWeakPtr<FFractureEditorModeToolkit> InToolkit) override;
+};
+
+
+UCLASS(DisplayName = "Select Clusters", Category = "FractureTools")
+class UFractureToolSelectCluster : public UFractureToolSelectAll
+{
+public:
+	GENERATED_BODY()
+
+	UFractureToolSelectCluster(const FObjectInitializer& ObjInit) : Super(ObjInit) {}
+
+	// UFractureActionTool Interface
+	virtual FText GetDisplayText() const override;
+	virtual FText GetTooltipText() const override;
+	virtual FSlateIcon GetToolIcon() const override;
+	virtual void RegisterUICommand(FFractureEditorCommands* BindingContext) override;
+	virtual void Execute(TWeakPtr<FFractureEditorModeToolkit> InToolkit) override;
+};
diff --git a/Engine/Source/Runtime/Experimental/GeometryCollectionEngine/Private/GeometryCollection/GeometryCollectionComponent.cpp b/Engine/Source/Runtime/Experimental/GeometryCollectionEngine/Private/GeometryCollection/GeometryCollectionComponent.cpp
--- a/Engine/Source/Runtime/Experimental/GeometryCollectionEngine/Private/GeometryCollection/GeometryCollectionComponent.cpp
+++ b/Engine/Source/Runtime/Experimental/GeometryCollectionEngine/Private/GeometryCollection/GeometryCollectionComponent.cpp
@@ -31,6 +31,8 @@
 #include "PhysicsField/PhysicsFieldComponent.h"
 #include "Engine/InstancedStaticMesh.h"
 
+#include "Algo/RemoveIf.h"
+
 #if WITH_EDITOR
 #include "AssetToolsModule.h"
 #include "Editor.h"
@@ -2757,6 +2759,34 @@ void FScopedColorEdit::SelectBones(GeometryCollection::ESelectionMode SelectionM
 		}
 		break;
 
+		case GeometryCollection::ESelectionMode::Leaves:
+		{
+			ResetBoneSelection();
+			TArray<int32> BonesToSelect;
+			FGeometryCollectionClusteringUtility::GetBonesToLevel(GeometryCollectionPtr.Get(), GetViewLevel(), BonesToSelect, true, true);
+			const TManagedArray<int32>& SimType = GeometryCollectionPtr->SimulationType;
+			BonesToSelect.SetNum(Algo::RemoveIf(BonesToSelect, [&](int32 BoneIdx)
+				{
+					return SimType[BoneIdx] != FGeometryCollection::ESimulationTypes::FST_Rigid;
+				}));
+			AppendSelectedBones(BonesToSelect);
+		}
+		break;
+
+		case GeometryCollection::ESelectionMode::Clusters:
+		{
+			ResetBoneSelection();
+			TArray<int32> BonesToSelect;
+			FGeometryCollectionClusteringUtility::GetBonesToLevel(GeometryCollectionPtr.Get(), GetViewLevel(), BonesToSelect, true, true);
+			const TManagedArray<int32>& SimType = GeometryCollectionPtr->SimulationType;
+			BonesToSelect.SetNum(Algo::RemoveIf(BonesToSelect, [&](int32 BoneIdx)
+			{
+				return SimType[BoneIdx] != FGeometryCollection::ESimulationTypes::FST_Clustered;
+			}));
+			AppendSelectedBones(BonesToSelect);
+		}
+		break;
+
 		case GeometryCollection::ESelectionMode::InverseGeometry:
 		{
 			TArray<int32> Roots;
@@ -2809,23 +2839,23 @@ void FScopedColorEdit::SelectBones(GeometryCollection::ESelectionMode SelectionM
 
 			const TArray<int32> SelectedBones = GetSelectedBones();
 
-			TArray<int32> NewSelection;
+			TSet<int32> NewSelection;
 			for (int32 Bone : SelectedBones)
 			{
-				NewSelection.AddUnique(Bone);
+				NewSelection.Add(Bone);
 				int32 GeometryIdx = TransformToGeometryIndex[Bone];
 				if (GeometryIdx != INDEX_NONE)
 				{
 					const TSet<int32>& Neighbors = Proximity[GeometryIdx];
 					for (int32 NeighborGeometryIndex : Neighbors)
 					{
-						NewSelection.AddUnique(TransformIndex[NeighborGeometryIndex]);
+						NewSelection.Add(TransformIndex[NeighborGeometryIndex]);
 					}
 				}
 			}
 
 			ResetBoneSelection();
-			AppendSelectedBones(NewSelection);
+			AppendSelectedBones(NewSelection.Array());
 		}
 		break;
 
@@ -2835,18 +2865,18 @@ void FScopedColorEdit::SelectBones(GeometryCollection::ESelectionMode SelectionM
 			
 			const TArray<int32> SelectedBones = GetSelectedBones();
 
-			TArray<int32> NewSelection;
+			TSet<int32> NewSelection;
 			for (int32 Bone : SelectedBones)
 			{
 				int32 ParentBone = Parents[Bone];
 				if (ParentBone != FGeometryCollection::Invalid)
 				{
-					NewSelection.AddUnique(ParentBone);
+					NewSelection.Add(ParentBone);
 				}
 			}
 
 			ResetBoneSelection();
-			AppendSelectedBones(NewSelection);
+			AppendSelectedBones(NewSelection.Array());
 		}
 		break;
 
@@ -2856,17 +2886,22 @@ void FScopedColorEdit::SelectBones(GeometryCollection::ESelectionMode SelectionM
 
 			const TArray<int32> SelectedBones = GetSelectedBones();
 
-			TArray<int32> NewSelection;
+			TSet<int32> NewSelection;
 			for (int32 Bone : SelectedBones)
 			{
+				if (Children[Bone].IsEmpty())
+				{
+					NewSelection.Add(Bone);
+					continue;
+				}
 				for (int32 Child : Children[Bone])
 				{
-					NewSelection.AddUnique(Child);
+					NewSelection.Add(Child);
 				}
 			}
 
 			ResetBoneSelection();
-			AppendSelectedBones(NewSelection);
+			AppendSelectedBones(NewSelection.Array());
 		}
 		break;
 
@@ -2877,7 +2912,7 @@ void FScopedColorEdit::SelectBones(GeometryCollection::ESelectionMode SelectionM
 
 			const TArray<int32> SelectedBones = GetSelectedBones();
 
-			TArray<int32> NewSelection;
+			TSet<int32> NewSelection;
 			for (int32 Bone : SelectedBones)
 			{
 				int32 ParentBone = Parents[Bone];
@@ -2885,14 +2920,14 @@ void FScopedColorEdit::SelectBones(GeometryCollection::ESelectionMode SelectionM
 				{
 					for (int32 Child : Children[ParentBone])
 					{
-						NewSelection.AddUnique(Child);
+						NewSelection.Add(Child);
 					}
 				}
 
 			}
 
 			ResetBoneSelection();
-			AppendSelectedBones(NewSelection);
+			AppendSelectedBones(NewSelection.Array());
 		}
 		break;
 
@@ -2904,7 +2939,7 @@ void FScopedColorEdit::SelectBones(GeometryCollection::ESelectionMode SelectionM
 
 				const TArray<int32> SelectedBones = GetSelectedBones();
 
-				TArray<int32> NewSelection;
+				TSet<int32> NewSelection;
 				for (int32 Bone : SelectedBones)
 				{
 					int32 Level = Levels[Bone];
@@ -2912,13 +2947,13 @@ void FScopedColorEdit::SelectBones(GeometryCollection::ESelectionMode SelectionM
 					{
 						if (Levels[TransformIdx] == Level)
 						{
-							NewSelection.AddUnique(TransformIdx);
+							NewSelection.Add(TransformIdx);
 						}
 					}
 				}
 
 				ResetBoneSelection();
-				AppendSelectedBones(NewSelection);
+				AppendSelectedBones(NewSelection.Array());
 			}	
 		}
 		break;
diff --git a/Engine/Source/Runtime/Experimental/GeometryCollectionEngine/Public/GeometryCollection/GeometryCollectionComponent.h b/Engine/Source/Runtime/Experimental/GeometryCollectionEngine/Public/GeometryCollection/GeometryCollectionComponent.h
--- a/Engine/Source/Runtime/Experimental/GeometryCollectionEngine/Public/GeometryCollection/GeometryCollectionComponent.h
+++ b/Engine/Source/Runtime/Experimental/GeometryCollectionEngine/Public/GeometryCollection/GeometryCollectionComponent.h
@@ -52,7 +52,9 @@ namespace GeometryCollection
 		Neighbors,
 		Parent,
 		Children,
-		Level
+		Level,
+		Leaves,
+		Clusters
 	};
 }
 
