Fix external strain field to set the maximum of the applied strain
Changed the logic to take the maximum of the collision and external strain

diff --git a/Engine/Source/Programs/HeadlessChaos/Private/GeometryCollection/GeometryCollectionTestClustering.cpp b/Engine/Source/Programs/HeadlessChaos/Private/GeometryCollection/GeometryCollectionTestClustering.cpp
--- a/Engine/Source/Programs/HeadlessChaos/Private/GeometryCollection/GeometryCollectionTestClustering.cpp
+++ b/Engine/Source/Programs/HeadlessChaos/Private/GeometryCollection/GeometryCollectionTestClustering.cpp
@@ -307,7 +307,7 @@ namespace GeometryCollectionTest
 
 			if (Frame == 2)
 			{
-				ParticleHandles[0]->AddExternalStrain(50.0f);
+				ParticleHandles[0]->SetExternalStrain(50.0f);
 				Clustering.BreakingModel();
 			}
 
diff --git a/Engine/Source/Runtime/Experimental/Chaos/Private/Chaos/PBDRigidClustering.cpp b/Engine/Source/Runtime/Experimental/Chaos/Private/Chaos/PBDRigidClustering.cpp
--- a/Engine/Source/Runtime/Experimental/Chaos/Private/Chaos/PBDRigidClustering.cpp
+++ b/Engine/Source/Runtime/Experimental/Chaos/Private/Chaos/PBDRigidClustering.cpp
@@ -592,15 +592,15 @@ namespace Chaos
 				continue;
 			}
 
-			const FReal ChildStrain = Child->CollisionImpulses() + Child->GetExternalStrains();
-			if (ChildStrain >= Child->Strain() || bForceRelease)
+			const FReal MaxAppliedStrain = FMath::Max(Child->CollisionImpulses(), Child->GetExternalStrain());
+			if ((MaxAppliedStrain >= Child->Strain()) || bForceRelease)
 			{
 				//UE_LOG(LogTemp, Warning, TEXT("Releasing child %d from parent %p due to strain %.5f Exceeding internal strain %.5f (Source: %s)"), ChildIdx, ClusteredParticle, ChildStrain, Child->Strain(), bForceRelease ? TEXT("Forced by caller") : ExternalStrainMap ? TEXT("External") : TEXT("Collision"));
 
 				// if necessary propagate through the graph
 				if (GraphPropagationBasedCollisionImpulseProcessing && !bForceRelease)
 				{
-					const FReal RemainingStrain = Child->Strain() - ChildStrain;
+					const FReal RemainingStrain = (MaxAppliedStrain - Child->Strain());
 					if (RemainingStrain > 0)
 					{
 						// todo(chaos) : could do better and have something weighted on distance with a falloff maybe ?  
@@ -642,7 +642,7 @@ namespace Chaos
 
 			// no need anymore of the external strain let's clear it
 			// @todo(chaos) impulse collision are right now cleared somewhere else, but we should consolidate that eventually
-			Child->ClearExternalStrains();
+			Child->ClearExternalStrain();
 		}
 
 		if (ActivatedChildren.Num() > 0)
@@ -721,8 +721,7 @@ namespace Chaos
 					{
 						if (ensure(!ClusterHandle || ClusteredChildHandle->ClusterIds().Id == ClusterHandle))
 						{
-							ClusteredChildHandle->ClearExternalStrains();
-							ClusteredChildHandle->AddExternalStrain(TNumericLimits<FReal>::Max());
+							ClusteredChildHandle->SetExternalStrain(TNumericLimits<FReal>::Max());
 							ClusterHandle = ClusteredChildHandle->ClusterIds().Id;
 						}
 						else
@@ -809,7 +808,7 @@ namespace Chaos
 										ClusteredChild->CollisionImpulse() = FLT_MAX;
 										MCollisionImpulseArrayDirty = true;
 									}
-									else if (!bPotentialBreak && ClusteredChild->GetExternalStrains() > 0)
+									else if (!bPotentialBreak && ClusteredChild->GetExternalStrain() > 0)
 									{
 										bPotentialBreak = true;
 									}
diff --git a/Engine/Source/Runtime/Experimental/Chaos/Public/Chaos/ParticleHandle.h b/Engine/Source/Runtime/Experimental/Chaos/Public/Chaos/ParticleHandle.h
--- a/Engine/Source/Runtime/Experimental/Chaos/Public/Chaos/ParticleHandle.h
+++ b/Engine/Source/Runtime/Experimental/Chaos/Public/Chaos/ParticleHandle.h
@@ -1208,9 +1208,9 @@ public:
 	T& CollisionImpulses() { return PBDRigidClusteredParticles->CollisionImpulses(ParticleIdx); }
 	void SetCollisionImpulses(const T Value) { PBDRigidClusteredParticles->CollisionImpulses(ParticleIdx) = Value; }
 
-	T GetExternalStrains() const { return PBDRigidClusteredParticles->ExternalStrains(ParticleIdx); }
-	void AddExternalStrain(const T Value) { PBDRigidClusteredParticles->ExternalStrains(ParticleIdx) += Value; }
-	void ClearExternalStrains() { PBDRigidClusteredParticles->ExternalStrains(ParticleIdx) = static_cast<T>(0); }
+	T GetExternalStrain() const { return PBDRigidClusteredParticles->ExternalStrains(ParticleIdx); }
+	void SetExternalStrain(const T Value) { PBDRigidClusteredParticles->ExternalStrains(ParticleIdx) = Value; }
+	void ClearExternalStrain() { PBDRigidClusteredParticles->ExternalStrains(ParticleIdx) = static_cast<T>(0); }
 	
 	const T& Strain() const { return PBDRigidClusteredParticles->Strains(ParticleIdx); }
 	T& Strain() { return PBDRigidClusteredParticles->Strains(ParticleIdx); }
diff --git a/Engine/Source/Runtime/Experimental/Chaos/Public/PhysicsProxy/FieldSystemProxyHelper.h b/Engine/Source/Runtime/Experimental/Chaos/Public/PhysicsProxy/FieldSystemProxyHelper.h
--- a/Engine/Source/Runtime/Experimental/Chaos/Public/PhysicsProxy/FieldSystemProxyHelper.h
+++ b/Engine/Source/Runtime/Experimental/Chaos/Public/PhysicsProxy/FieldSystemProxyHelper.h
@@ -494,7 +494,8 @@ namespace Chaos
 					{
 						if (Chaos::FPBDRigidClusteredParticleHandle* ClusteredParticle = ParticleHandles[Index.Sample]->CastToClustered())
 						{
-							ClusteredParticle->AddExternalStrain(ExternalStrainValue);
+							const FReal CurrentExternalStrains = ClusteredParticle->GetExternalStrain();
+							ClusteredParticle->SetExternalStrain(FMath::Max(CurrentExternalStrains, ExternalStrainValue));
 						}
 					}
 				}
