Chaos - temp fix for kinematic body velocity in joint solver

diff --git a/Engine/Source/Runtime/Experimental/Chaos/Private/Chaos/Joint/PBDJointCachedSolverGaussSeidel.cpp b/Engine/Source/Runtime/Experimental/Chaos/Private/Chaos/Joint/PBDJointCachedSolverGaussSeidel.cpp
--- a/Engine/Source/Runtime/Experimental/Chaos/Private/Chaos/Joint/PBDJointCachedSolverGaussSeidel.cpp
+++ b/Engine/Source/Runtime/Experimental/Chaos/Private/Chaos/Joint/PBDJointCachedSolverGaussSeidel.cpp
@@ -205,6 +205,28 @@ void FPBDJointCachedSolver::Init(
 
 	InitDerivedState();
 
+	// Temporary fix for incorrect implicit velocity calculation for kinematic bodies
+	// @todo(chaos): this should be fixed where we gather the SolverBody data
+	bool bUpdatedKinematicXR = false;
+	if (!Body0().IsDynamic() && (!V(0).IsNearlyZero() || !W(0).IsNearlyZero()))
+	{
+		InitConnectorXs[0] = InitConnectorXs[0] - V(0) * Dt;
+		InitConnectorRs[0] = FRotation3::IntegrateRotationWithAngularVelocity(InitConnectorRs[0], -W(0), Dt);
+		bUpdatedKinematicXR = true;
+	}
+	if (!Body1().IsDynamic() && (!V(1).IsNearlyZero() || !W(1).IsNearlyZero()))
+	{
+		InitConnectorXs[1] = InitConnectorXs[1] - V(1) * Dt;
+		InitConnectorRs[1] = FRotation3::IntegrateRotationWithAngularVelocity(InitConnectorRs[1], -W(1), Dt);
+		bUpdatedKinematicXR = true;
+	}
+	if (bUpdatedKinematicXR)
+	{
+		InitConnectorRs[1].EnforceShortestArcWith(InitConnectorRs[0]);
+		ConnectorWDts[0] = FRotation3::CalculateAngularVelocity(InitConnectorRs[0], ConnectorRs[0], 1.0f);
+		ConnectorWDts[1] = FRotation3::CalculateAngularVelocity(InitConnectorRs[1], ConnectorRs[1], 1.0f);
+	}
+
 	UpdateMass0();
 	UpdateMass1();
 
