Chaos - better fix for joint drive damping
- initialize kinematic solver body X correctly

diff --git a/Engine/Source/Runtime/Experimental/Chaos/Private/Chaos/Collision/PBDCollisionConstraintHandle.cpp b/Engine/Source/Runtime/Experimental/Chaos/Private/Chaos/Collision/PBDCollisionConstraintHandle.cpp
--- a/Engine/Source/Runtime/Experimental/Chaos/Private/Chaos/Collision/PBDCollisionConstraintHandle.cpp
+++ b/Engine/Source/Runtime/Experimental/Chaos/Private/Chaos/Collision/PBDCollisionConstraintHandle.cpp
@@ -64,7 +64,7 @@ namespace Chaos
 
 	void FPBDCollisionConstraintHandle::PreGatherInput(const FReal Dt, FPBDIslandSolverData& SolverData)
 	{
-		ConcreteContainer()->PreGatherInput(GetContact(), SolverData);
+		ConcreteContainer()->PreGatherInput(Dt, GetContact(), SolverData);
 	}
 
 	void FPBDCollisionConstraintHandle::GatherInput(FReal Dt, const int32 Particle0Level, const int32 Particle1Level, FPBDIslandSolverData& SolverData)
diff --git a/Engine/Source/Runtime/Experimental/Chaos/Private/Chaos/Collision/SolverCollisionContainer.cpp b/Engine/Source/Runtime/Experimental/Chaos/Private/Chaos/Collision/SolverCollisionContainer.cpp
--- a/Engine/Source/Runtime/Experimental/Chaos/Private/Chaos/Collision/SolverCollisionContainer.cpp
+++ b/Engine/Source/Runtime/Experimental/Chaos/Private/Chaos/Collision/SolverCollisionContainer.cpp
@@ -76,12 +76,13 @@ namespace Chaos
 		bool IsIncrementalManifold() const { return bIsIncremental; }
 
 		void PreGatherInput(
+			const FReal Dt,
 			FPBDCollisionConstraint& InConstraint,
 			FSolverBodyContainer& SolverBodyContainer)
 		{
 			Constraint = &InConstraint;
-			SolverBodyContainer.FindOrAdd(Constraint->GetParticle0());
-			SolverBodyContainer.FindOrAdd(Constraint->GetParticle1());
+			SolverBodyContainer.FindOrAdd(Constraint->GetParticle0(), Dt);
+			SolverBodyContainer.FindOrAdd(Constraint->GetParticle1(), Dt);
 		}
 
 		/**
@@ -152,8 +153,8 @@ namespace Chaos
 
 			// Find the solver bodies for the particles we constrain. This will add them to the container
 			// if they aren't there already, and ensure that they are populated with the latest data.
-			FSolverBody* Body0 = SolverBodyContainer.FindOrAdd(Constraint->GetParticle0());
-			FSolverBody* Body1 = SolverBodyContainer.FindOrAdd(Constraint->GetParticle1());
+			FSolverBody* Body0 = SolverBodyContainer.FindOrAdd(Constraint->GetParticle0(), Dt);
+			FSolverBody* Body1 = SolverBodyContainer.FindOrAdd(Constraint->GetParticle1(), Dt);
 
 			Body0->SetLevel(Particle0Level);
 			Body1->SetLevel(Particle1Level);
@@ -389,7 +390,7 @@ namespace Chaos
 		CollisionSolvers.SetNum(MaxCollisions, false);
 	}
 
-	void FPBDCollisionSolverContainer::PreAddConstraintSolver(FPBDCollisionConstraint& Constraint, FSolverBodyContainer& SolverBodyContainer, int32& SolverIndex)
+	void FPBDCollisionSolverContainer::PreAddConstraintSolver(const FReal Dt, FPBDCollisionConstraint& Constraint, FSolverBodyContainer& SolverBodyContainer, int32& SolverIndex)
 	{
 		// This container is required to allocate pointers that are valid for the whole tick,
 		// so we cannot allow the container to resize during the tick. See Reset()
@@ -399,12 +400,12 @@ namespace Chaos
 		FPBDCollisionSolverAdapter& CollisionSolver = CollisionSolvers[SolverIndex];
 		CollisionSolver.GetSolver().Reset();
 
-		CollisionSolver.PreGatherInput(Constraint, SolverBodyContainer);
+		CollisionSolver.PreGatherInput(Dt, Constraint, SolverBodyContainer);
 
 		++SolverIndex;
 	}
 
-	void FPBDCollisionSolverContainer::AddConstraintSolver(FReal Dt, FPBDCollisionConstraint& Constraint, const int32 Particle0Level, const int32 Particle1Level, FSolverBodyContainer& SolverBodyContainer, const FPBDCollisionSolverSettings& SolverSettings)
+	void FPBDCollisionSolverContainer::AddConstraintSolver(const FReal Dt, FPBDCollisionConstraint& Constraint, const int32 Particle0Level, const int32 Particle1Level, FSolverBodyContainer& SolverBodyContainer, const FPBDCollisionSolverSettings& SolverSettings)
 	{
 		// This container is required to allocate pointers that are valid for the whole tick,
 		// so we cannot allow the container to resize during the tick. See Reset()
diff --git a/Engine/Source/Runtime/Experimental/Chaos/Private/Chaos/Evolution/SolverBodyContainer.cpp b/Engine/Source/Runtime/Experimental/Chaos/Private/Chaos/Evolution/SolverBodyContainer.cpp
--- a/Engine/Source/Runtime/Experimental/Chaos/Private/Chaos/Evolution/SolverBodyContainer.cpp
+++ b/Engine/Source/Runtime/Experimental/Chaos/Private/Chaos/Evolution/SolverBodyContainer.cpp
@@ -24,14 +24,14 @@ namespace Chaos
 		}
 	}
 
-	int32 FSolverBodyContainer::AddParticle(FGenericParticleHandle InParticle)
+	int32 FSolverBodyContainer::AddParticle(FGenericParticleHandle InParticle, const FReal Dt)
 	{
 		// No array resizing allowed (we want fixed pointers)
 		check(NumItems() < MaxItems());
-		return SolverBodies.Emplace(InParticle);
+		return SolverBodies.Emplace(FSolverBodyAdapter(InParticle, Dt));
 	}
 
-	FSolverBody* FSolverBodyContainer::FindOrAdd(FGenericParticleHandle InParticle)
+	FSolverBody* FSolverBodyContainer::FindOrAdd(FGenericParticleHandle InParticle, const FReal Dt)
 	{
 		// For dynamic bodies, we store a cookie on the Particle that holds the solver body index
 		// For kinematics we cannot do this because the kinematic may be in multiple islands and 
@@ -43,7 +43,7 @@ namespace Chaos
 			if (InParticle->IsDynamic())
 			{
 				// First time we have seen this particle, so add it
-				ItemIndex = AddParticle(InParticle);
+				ItemIndex = AddParticle(InParticle, Dt);
 				InParticle->SetSolverBodyIndex(ItemIndex);
 			}
 			else // Not Dynamic
@@ -56,7 +56,7 @@ namespace Chaos
 				else
 				{
 					// First time we have seen this particle, so add it
-					ItemIndex = AddParticle(InParticle);
+					ItemIndex = AddParticle(InParticle, Dt);
 					ParticleToIndexMap.Add(InParticle, ItemIndex);
 				}
 			}			
diff --git a/Engine/Source/Runtime/Experimental/Chaos/Private/Chaos/Joint/PBDJointCachedSolverGaussSeidel.cpp b/Engine/Source/Runtime/Experimental/Chaos/Private/Chaos/Joint/PBDJointCachedSolverGaussSeidel.cpp
--- a/Engine/Source/Runtime/Experimental/Chaos/Private/Chaos/Joint/PBDJointCachedSolverGaussSeidel.cpp
+++ b/Engine/Source/Runtime/Experimental/Chaos/Private/Chaos/Joint/PBDJointCachedSolverGaussSeidel.cpp
@@ -205,28 +205,6 @@ void FPBDJointCachedSolver::Init(
 
 	InitDerivedState();
 
-	// Temporary fix for incorrect implicit velocity calculation for kinematic bodies
-	// @todo(chaos): this should be fixed where we gather the SolverBody data
-	bool bUpdatedKinematicXR = false;
-	if (!Body0().IsDynamic() && (!V(0).IsNearlyZero() || !W(0).IsNearlyZero()))
-	{
-		InitConnectorXs[0] = InitConnectorXs[0] - V(0) * Dt;
-		InitConnectorRs[0] = FRotation3::IntegrateRotationWithAngularVelocity(InitConnectorRs[0], -W(0), Dt);
-		bUpdatedKinematicXR = true;
-	}
-	if (!Body1().IsDynamic() && (!V(1).IsNearlyZero() || !W(1).IsNearlyZero()))
-	{
-		InitConnectorXs[1] = InitConnectorXs[1] - V(1) * Dt;
-		InitConnectorRs[1] = FRotation3::IntegrateRotationWithAngularVelocity(InitConnectorRs[1], -W(1), Dt);
-		bUpdatedKinematicXR = true;
-	}
-	if (bUpdatedKinematicXR)
-	{
-		InitConnectorRs[1].EnforceShortestArcWith(InitConnectorRs[0]);
-		ConnectorWDts[0] = FRotation3::CalculateAngularVelocity(InitConnectorRs[0], ConnectorRs[0], 1.0f);
-		ConnectorWDts[1] = FRotation3::CalculateAngularVelocity(InitConnectorRs[1], ConnectorRs[1], 1.0f);
-	}
-
 	UpdateMass0();
 	UpdateMass1();
 
diff --git a/Engine/Source/Runtime/Experimental/Chaos/Private/Chaos/PBDCollisionConstraints.cpp b/Engine/Source/Runtime/Experimental/Chaos/Private/Chaos/PBDCollisionConstraints.cpp
--- a/Engine/Source/Runtime/Experimental/Chaos/Private/Chaos/PBDCollisionConstraints.cpp
+++ b/Engine/Source/Runtime/Experimental/Chaos/Private/Chaos/PBDCollisionConstraints.cpp
@@ -405,12 +405,12 @@ namespace Chaos
 		return SolverData.GetConstraintContainer<FPBDCollisionSolverContainer>(ContainerId);
 	}
 
-	void FPBDCollisionConstraints::PreGatherInput(FPBDCollisionConstraint& Constraint, FPBDIslandSolverData& SolverData)
+	void FPBDCollisionConstraints::PreGatherInput(const FReal Dt, FPBDCollisionConstraint& Constraint, FPBDIslandSolverData& SolverData)
 	{
 		if (SolverType == EConstraintSolverType::QuasiPbd)
 		{
 			FPBDCollisionSolverContainer& SolverContainer = GetConstraintSolverContainer(SolverData);
-			SolverContainer.PreAddConstraintSolver(Constraint, SolverData.GetBodyContainer(), SolverData.GetConstraintIndex(ContainerId));
+			SolverContainer.PreAddConstraintSolver(Dt, Constraint, SolverData.GetBodyContainer(), SolverData.GetConstraintIndex(ContainerId));
 		}
 	}
 
@@ -438,7 +438,7 @@ namespace Chaos
 			{
 				if (Constraint->IsEnabled())
 				{
-					PreGatherInput(*Constraint, SolverData);
+					PreGatherInput(Dt, *Constraint, SolverData);
 				}
 			}
 		}
@@ -617,8 +617,8 @@ namespace Chaos
 	{
 		SolverData.GetConstraintHandles(ContainerId).Add(&Constraint);
 
-		FSolverBody* SolverBody0 = SolverData.GetBodyContainer().FindOrAdd(Constraint.Particle[0]);
-		FSolverBody* SolverBody1 = SolverData.GetBodyContainer().FindOrAdd(Constraint.Particle[1]);
+		FSolverBody* SolverBody0 = SolverData.GetBodyContainer().FindOrAdd(Constraint.Particle[0], Dt);
+		FSolverBody* SolverBody1 = SolverData.GetBodyContainer().FindOrAdd(Constraint.Particle[1], Dt);
 
 		SolverBody0->SetLevel(Particle0Level);
 		SolverBody1->SetLevel(Particle1Level);
diff --git a/Engine/Source/Runtime/Experimental/Chaos/Private/Chaos/PBDJointConstraints.cpp b/Engine/Source/Runtime/Experimental/Chaos/Private/Chaos/PBDJointConstraints.cpp
--- a/Engine/Source/Runtime/Experimental/Chaos/Private/Chaos/PBDJointConstraints.cpp
+++ b/Engine/Source/Runtime/Experimental/Chaos/Private/Chaos/PBDJointConstraints.cpp
@@ -886,8 +886,8 @@ namespace Chaos
 
 		// Find the solver bodies for the particles we constrain. This will add them to the container
 		// if they aren't there already, and ensure that they are populated with the latest data.
-		SolverData.GetBodyContainer().FindOrAdd(Particle0);
-		SolverData.GetBodyContainer().FindOrAdd(Particle1);
+		SolverData.GetBodyContainer().FindOrAdd(Particle0, Dt);
+		SolverData.GetBodyContainer().FindOrAdd(Particle1, Dt);
 	}
 
 	void FPBDJointConstraints::GatherInput(const FReal Dt, const int32 ConstraintIndex, const int32 Particle0Level, const int32 Particle1Level, FPBDIslandSolverData& SolverData)
@@ -905,8 +905,8 @@ namespace Chaos
 
 		// Find the solver bodies for the particles we constrain. This will add them to the container
 		// if they aren't there already, and ensure that they are populated with the latest data.
-		FSolverBody* Body0 = SolverData.GetBodyContainer().FindOrAdd(Particle0);
-		FSolverBody* Body1 = SolverData.GetBodyContainer().FindOrAdd(Particle1);
+		FSolverBody* Body0 = SolverData.GetBodyContainer().FindOrAdd(Particle0, Dt);
+		FSolverBody* Body1 = SolverData.GetBodyContainer().FindOrAdd(Particle1, Dt);
 
 		if (Settings.bUseLinearSolver)
 		{
diff --git a/Engine/Source/Runtime/Experimental/Chaos/Private/Chaos/PBDPositionConstraints.cpp b/Engine/Source/Runtime/Experimental/Chaos/Private/Chaos/PBDPositionConstraints.cpp
--- a/Engine/Source/Runtime/Experimental/Chaos/Private/Chaos/PBDPositionConstraints.cpp
+++ b/Engine/Source/Runtime/Experimental/Chaos/Private/Chaos/PBDPositionConstraints.cpp
@@ -29,7 +29,7 @@ namespace Chaos
 	{
 		SolverData.GetConstraintIndices(ContainerId).Add(ConstraintIndex);
 
-		ConstraintSolverBodies[ConstraintIndex] = SolverData.GetBodyContainer().FindOrAdd(ConstrainedParticles[ConstraintIndex]);
+		ConstraintSolverBodies[ConstraintIndex] = SolverData.GetBodyContainer().FindOrAdd(ConstrainedParticles[ConstraintIndex], Dt);
 	}
 
 	void FPBDPositionConstraints::GatherInput(const FReal Dt, const int32 ConstraintIndex, const int32 Particle0Level, const int32 Particle1Level, FPBDIslandSolverData& SolverData)
diff --git a/Engine/Source/Runtime/Experimental/Chaos/Private/Chaos/PBDRigidDynamicSpringConstraints.cpp b/Engine/Source/Runtime/Experimental/Chaos/Private/Chaos/PBDRigidDynamicSpringConstraints.cpp
--- a/Engine/Source/Runtime/Experimental/Chaos/Private/Chaos/PBDRigidDynamicSpringConstraints.cpp
+++ b/Engine/Source/Runtime/Experimental/Chaos/Private/Chaos/PBDRigidDynamicSpringConstraints.cpp
@@ -118,8 +118,8 @@ void FPBDRigidDynamicSpringConstraints::PreGatherInput(const FReal Dt, const int
 
 	ConstraintSolverBodies[ConstraintIndex] =
 	{
-		SolverData.GetBodyContainer().FindOrAdd(Constraints[ConstraintIndex][0]),
-		SolverData.GetBodyContainer().FindOrAdd(Constraints[ConstraintIndex][1])
+		SolverData.GetBodyContainer().FindOrAdd(Constraints[ConstraintIndex][0], Dt),
+		SolverData.GetBodyContainer().FindOrAdd(Constraints[ConstraintIndex][1], Dt)
 	};
 }
 
diff --git a/Engine/Source/Runtime/Experimental/Chaos/Private/Chaos/PBDRigidSpringConstraints.cpp b/Engine/Source/Runtime/Experimental/Chaos/Private/Chaos/PBDRigidSpringConstraints.cpp
--- a/Engine/Source/Runtime/Experimental/Chaos/Private/Chaos/PBDRigidSpringConstraints.cpp
+++ b/Engine/Source/Runtime/Experimental/Chaos/Private/Chaos/PBDRigidSpringConstraints.cpp
@@ -167,8 +167,8 @@ namespace Chaos
 	
 		ConstraintSolverBodies[ConstraintIndex] =
 		{
-			SolverData.GetBodyContainer().FindOrAdd(Constraints[ConstraintIndex][0]),
-			SolverData.GetBodyContainer().FindOrAdd(Constraints[ConstraintIndex][1])
+			SolverData.GetBodyContainer().FindOrAdd(Constraints[ConstraintIndex][0], Dt),
+			SolverData.GetBodyContainer().FindOrAdd(Constraints[ConstraintIndex][1], Dt)
 		};
 	}
 
diff --git a/Engine/Source/Runtime/Experimental/Chaos/Private/Chaos/PBDRigidsEvolutionGBF.cpp b/Engine/Source/Runtime/Experimental/Chaos/Private/Chaos/PBDRigidsEvolutionGBF.cpp
--- a/Engine/Source/Runtime/Experimental/Chaos/Private/Chaos/PBDRigidsEvolutionGBF.cpp
+++ b/Engine/Source/Runtime/Experimental/Chaos/Private/Chaos/PBDRigidsEvolutionGBF.cpp
@@ -705,7 +705,7 @@ void FPBDRigidsEvolutionGBF::GatherSolverInput(FReal Dt, int32 GroupIndex)
 								{
 									// Only the non sleeping dynamic rigid particles involved in any constraints will be added
 									// by default to the solver bodies
-									const int32 BodyIndex = IslandGroup->GetBodyContainer().GetBodies().Emplace(GenericHandle);
+									const int32 BodyIndex = IslandGroup->GetBodyContainer().GetBodies().Emplace(GenericHandle, Dt);
 									GenericHandle->SetSolverBodyIndex(BodyIndex);
 								}
 							}
diff --git a/Engine/Source/Runtime/Experimental/Chaos/Private/Chaos/PBDSuspensionConstraints.cpp b/Engine/Source/Runtime/Experimental/Chaos/Private/Chaos/PBDSuspensionConstraints.cpp
--- a/Engine/Source/Runtime/Experimental/Chaos/Private/Chaos/PBDSuspensionConstraints.cpp
+++ b/Engine/Source/Runtime/Experimental/Chaos/Private/Chaos/PBDSuspensionConstraints.cpp
@@ -127,7 +127,7 @@ namespace Chaos
 	{
 		SolverData.GetConstraintIndices(ContainerId).Add(ConstraintIndex);
 
-		ConstraintSolverBodies[ConstraintIndex] = SolverData.GetBodyContainer().FindOrAdd(ConstrainedParticles[ConstraintIndex]);
+		ConstraintSolverBodies[ConstraintIndex] = SolverData.GetBodyContainer().FindOrAdd(ConstrainedParticles[ConstraintIndex], Dt);
 
 		ConstraintResults[ConstraintIndex].Reset();
 
diff --git a/Engine/Source/Runtime/Experimental/Chaos/Public/Chaos/Collision/SolverCollisionContainer.h b/Engine/Source/Runtime/Experimental/Chaos/Public/Chaos/Collision/SolverCollisionContainer.h
--- a/Engine/Source/Runtime/Experimental/Chaos/Public/Chaos/Collision/SolverCollisionContainer.h
+++ b/Engine/Source/Runtime/Experimental/Chaos/Public/Chaos/Collision/SolverCollisionContainer.h
@@ -67,8 +67,8 @@ namespace Chaos
 		void SetNum(const int32 MaxCollisions);
 
 		// Add a solver constraint for the specified constraint and gather the required solver data
-		void PreAddConstraintSolver(FPBDCollisionConstraint& Constraint, FSolverBodyContainer& SolverBodyContainer, int32& SolverIndex);
-		void AddConstraintSolver(FReal Dt, FPBDCollisionConstraint& Constraint, const int32 Particle0Level, const int32 Particle1Level, FSolverBodyContainer& SolverBodyContainer, const FPBDCollisionSolverSettings& SolverSettings);
+		void PreAddConstraintSolver(const FReal Dt, FPBDCollisionConstraint& Constraint, FSolverBodyContainer& SolverBodyContainer, int32& SolverIndex);
+		void AddConstraintSolver(const FReal Dt, FPBDCollisionConstraint& Constraint, const int32 Particle0Level, const int32 Particle1Level, FSolverBodyContainer& SolverBodyContainer, const FPBDCollisionSolverSettings& SolverSettings);
 
 		bool SolvePositionSerial(const FReal Dt, const int32 It, const int32 NumIts, const int32 BeginIndex, const int32 EndIndex, const FPBDCollisionSolverSettings& SolverSettings);
 		bool SolvePositionParallel(const FReal Dt, const int32 It, const int32 NumIts, const int32 BeginIndex, const int32 EndIndex, const FPBDCollisionSolverSettings& SolverSettings);
diff --git a/Engine/Source/Runtime/Experimental/Chaos/Public/Chaos/Evolution/SolverBodyContainer.h b/Engine/Source/Runtime/Experimental/Chaos/Public/Chaos/Evolution/SolverBodyContainer.h
--- a/Engine/Source/Runtime/Experimental/Chaos/Public/Chaos/Evolution/SolverBodyContainer.h
+++ b/Engine/Source/Runtime/Experimental/Chaos/Public/Chaos/Evolution/SolverBodyContainer.h
@@ -22,16 +22,16 @@ namespace Chaos
 	class FSolverBodyAdapter
 	{
 	public:
-		FSolverBodyAdapter(const FGenericParticleHandle& InParticle)
+		FSolverBodyAdapter(const FGenericParticleHandle& InParticle, const FReal Dt)
 			: Particle(InParticle)
 		{
-			GatherInput();
+			GatherInput(Dt);
 		}
 		
 		FSolverBody& GetSolverBody() { return SolverBody; }
 		const FGenericParticleHandle& GetParticle() const { return Particle; }
 
-		FORCEINLINE void GatherInput()
+		FORCEINLINE void GatherInput(const FReal Dt)
 		{
 			if (Particle.IsValid())
 			{
@@ -54,8 +54,16 @@ namespace Chaos
 				}
 				else
 				{
-					SolverBody.SetX(SolverBody.P());
-					SolverBody.SetR(SolverBody.Q());
+					// @todo(chaos): we really need Kinematics to store their initial positions...this keeps coming up
+					SolverBody.SetX(SolverBody.P() - SolverBody.V() * Dt);
+					if (SolverBody.W().IsNearlyZero())
+					{
+						SolverBody.SetR(SolverBody.Q());
+					}
+					else
+					{
+						SolverBody.SetR(FRotation3::IntegrateRotationWithAngularVelocity(SolverBody.Q(), -SolverBody.W(), Dt));
+					}
 				}
 				// No need to update the UpdateRotationDependentState since this function is only
 				// valid for dynamic particle for which the SetInvILocal is already doing the job
@@ -162,7 +170,7 @@ namespace Chaos
 		// Add a solver body to represent the solver state of the particle
 		// This should ideally be called in the order in which the bodies will be accessed
 		// (or as close as we can get, given most constraints access 2 bodies so there is no perfect order)
-		FSolverBody* FindOrAdd(FGenericParticleHandle InParticle);
+		FSolverBody* FindOrAdd(FGenericParticleHandle InParticle, const FReal Dt);
 
 		// Collect all the data we need from the particles represented by our SolverBodies
 		//void GatherInput();
@@ -180,7 +188,7 @@ namespace Chaos
 		void UpdateRotationDependentState();
 
 	private:
-		int32 AddParticle(FGenericParticleHandle InParticle);
+		int32 AddParticle(FGenericParticleHandle InParticle, const FReal Dt);
 
 		// Solver bodies, usually collected in the order in which they are accessed
 		TArray<FSolverBodyAdapter> SolverBodies;
diff --git a/Engine/Source/Runtime/Experimental/Chaos/Public/Chaos/PBDCollisionConstraints.h b/Engine/Source/Runtime/Experimental/Chaos/Public/Chaos/PBDCollisionConstraints.h
--- a/Engine/Source/Runtime/Experimental/Chaos/Public/Chaos/PBDCollisionConstraints.h
+++ b/Engine/Source/Runtime/Experimental/Chaos/Public/Chaos/PBDCollisionConstraints.h
@@ -172,7 +172,7 @@ public:
 	//
 
 	void SetNumIslandConstraints(const int32 NumIslandConstraints, FPBDIslandSolverData& SolverData);
-	void PreGatherInput(FPBDCollisionConstraint& Constraint, FPBDIslandSolverData& SolverData);
+	void PreGatherInput(const FReal Dt, FPBDCollisionConstraint& Constraint, FPBDIslandSolverData& SolverData);
 	void GatherInput(const FReal Dt, FPBDCollisionConstraint& Constraint, const int32 Particle0Level, const int32 Particle1Level, FPBDIslandSolverData& SolverData);
 
 	bool ApplyPhase1Serial(const FReal Dt, const int32 It, const int32 NumIts, FPBDIslandSolverData& SolverData);
