Chaos destruction : make sure velocities don't accumulate for geometry collection disabled children with active parents

diff --git a/Engine/Source/Runtime/Experimental/Chaos/Private/PhysicsProxy/GeometryCollectionPhysicsProxy.cpp b/Engine/Source/Runtime/Experimental/Chaos/Private/PhysicsProxy/GeometryCollectionPhysicsProxy.cpp
--- a/Engine/Source/Runtime/Experimental/Chaos/Private/PhysicsProxy/GeometryCollectionPhysicsProxy.cpp
+++ b/Engine/Source/Runtime/Experimental/Chaos/Private/PhysicsProxy/GeometryCollectionPhysicsProxy.cpp
@@ -2124,6 +2124,9 @@ void FGeometryCollectionPhysicsProxy::BufferPhysicsResults_Internal(Chaos::FPBDR
 							{
 								const FRigidTransform3 ChildToWorld = Handle->ChildToParent() * FRigidTransform3(ClusterParent->X(), ClusterParent->R());
 								UpdateParticleHandleTransformIfNeeded(*CurrentSolver, *Handle, ChildToWorld);
+								// fields may have applied velocities, we need to make sure to clear that up, so that we don't accumulate
+								Handle->SetV(FVec3::Zero());
+								Handle->SetW(FVec3::Zero());
 							}
 						}
 					}
