Chaos: Group sibling choice and cost for AABBTree leaf insertion to avoid double cache misses. Refactored node to be smaller so more nodes fit in a cache line. Altered node access pattern during FindSibling to avoid repeated indirection through the Node array.

diff --git a/Engine/Source/Runtime/Experimental/Chaos/Public/Chaos/AABBTree.h b/Engine/Source/Runtime/Experimental/Chaos/Public/Chaos/AABBTree.h
--- a/Engine/Source/Runtime/Experimental/Chaos/Public/Chaos/AABBTree.h
+++ b/Engine/Source/Runtime/Experimental/Chaos/Public/Chaos/AABBTree.h
@@ -517,15 +517,18 @@ template <typename T>
 struct TAABBTreeNode
 {
 	TAABBTreeNode()
-	{
-		ChildrenBounds[0] = TAABB<T, 3>();
-		ChildrenBounds[1] = TAABB<T, 3>();
-	}
+		: ChildrenBounds{TAABB<T, 3>() , TAABB<T, 3>()}
+		, ChildrenNodes{INDEX_NONE, INDEX_NONE}
+		, ParentNode(INDEX_NONE)
+		, bLeaf(false)
+		, bDirtyNode(false)
+	{}
+
 	TAABB<T, 3> ChildrenBounds[2];
-	int32 ChildrenNodes[2] = { INDEX_NONE, INDEX_NONE };
-	int32 ParentNode = INDEX_NONE;
-	bool bLeaf = false;
-	bool bDirtyNode = false;
+	int32 ChildrenNodes[2];
+	int32 ParentNode;
+	bool bLeaf : 1;
+	bool bDirtyNode : 1;
 
 #if !UE_BUILD_SHIPPING
 	void DebugDraw(ISpacialDebugDrawInterface<T>& InInterface, const TArray<TAABBTreeNode<T>>& Nodes, const FVec3& InLinearColor, float InThickness) const
@@ -563,12 +566,15 @@ struct TAABBTreeNode
 			Ar << Node;
 		}
 
-		Ar << bLeaf;
+		// Making a copy here because bLeaf is a bitfield and the archive can't handle it
+		bool bLeafCopy = bLeaf;
+		Ar << bLeafCopy;
 
 		// Dynamic trees are not serialized
 		if (Ar.IsLoading())
 		{
 			ParentNode = INDEX_NONE;
+			bLeaf = bLeafCopy;
 		}
 	}
 };
@@ -1172,92 +1178,91 @@ public:
 		return(WhichChildAmI(NodeIdx) ^ 1);
 	}
 
+private:
+	using FElement = TPayloadBoundsElement<TPayloadType, T>;
+	using FNode = TAABBTreeNode<T>;
+
+public:
 	int32 FindBestSibling(const TAABB<T, 3>& InNewBounds, bool& bOutAddToLeaf)
 	{
 		bOutAddToLeaf = false;
-		if (RootNode == INDEX_NONE)
-		{
-			return INDEX_NONE;
-		}
 		
 		TAABB<T, 3> NewBounds = InNewBounds;
 		NewBounds.Thicken(FAABBTreeCVars::DynamicTreeBoundingBoxPadding);
+		const FReal NewBoundsArea = NewBounds.GetArea();
 
 		//Priority Q of indices to explore
 		PriorityQ.Reset();
-		SumDeltaCostQ.Reset();
 
 		int32 QIndex = 0;
 		
 		// Initializing
-		
+		FNode& RNode = Nodes[RootNode];
 		TAABB<T, 3> WorkingAABB{ NewBounds };
-		WorkingAABB.GrowToInclude(Nodes[RootNode].ChildrenBounds[0]);
-		if (!Nodes[RootNode].bLeaf)
+		WorkingAABB.GrowToInclude(RNode.ChildrenBounds[0]);
+		if(!RNode.bLeaf)
 		{
-			WorkingAABB.GrowToInclude(Nodes[RootNode].ChildrenBounds[1]);
+			WorkingAABB.GrowToInclude(RNode.ChildrenBounds[1]);
 		}
 
 		int32 BestSiblingIdx = RootNode;
 		FReal BestCost = WorkingAABB.GetArea();
-		PriorityQ.Add(RootNode);
-		SumDeltaCostQ.Add(0.0f);
+		PriorityQ.Emplace(RNode, RootNode, 0.0f);
 
 		while (PriorityQ.Num() - QIndex)
 		{
 			// Pop from queue
-			uint32 TestSibling = PriorityQ[QIndex];
-			FReal SumDeltaCost = SumDeltaCostQ[QIndex];
+			const FNodeIndexAndCost NodeAndCost = PriorityQ[QIndex];
+			FNode& TestNode = NodeAndCost.template Get<0>();
+			const FReal SumDeltaCost = NodeAndCost.template Get<2>();
 			QIndex++;
 
-			// Alternative is a stack (this is not very optimal so don't use it)
-			//uint32 TestSibling = PriorityQ[PriorityQ.Num() - 1];
-			//FReal SumDeltaCost = SumDeltaCostQ[SumDeltaCostQ.Num() - 1];
-			//PriorityQ.SetNumUninitialized(PriorityQ.Num() - 1);
-			//SumDeltaCostQ.SetNumUninitialized(SumDeltaCostQ.Num() - 1);
-
 			// TestSibling bounds union with new bounds
 			bool bAddToLeaf = false;
-			WorkingAABB = Nodes[TestSibling].ChildrenBounds[0];
-			if (!Nodes[TestSibling].bLeaf)
+			WorkingAABB = TestNode.ChildrenBounds[0];
+			if (!TestNode.bLeaf)
 			{
-				WorkingAABB.GrowToInclude(Nodes[TestSibling].ChildrenBounds[1]);
+				WorkingAABB.GrowToInclude(TestNode.ChildrenBounds[1]);
 			}
 			else
 			{
-				int32 LeafIdx = Nodes[TestSibling].ChildrenNodes[0];
+				int32 LeafIdx = TestNode.ChildrenNodes[0];
 				bAddToLeaf = Leaves[LeafIdx].GetElementCount() < FAABBTreeCVars::DynamicTreeLeafCapacity;
 			}
-			FReal TestSiblingArea = WorkingAABB.GetArea();
+
+			const FReal TestSiblingArea = WorkingAABB.GetArea();
+
 			WorkingAABB.GrowToInclude(NewBounds);
 
-			FReal NewPotentialNodeArea = WorkingAABB.GetArea();
-			FReal CostForChoosingNode = NewPotentialNodeArea + SumDeltaCost;
+			const FReal NewPotentialNodeArea = WorkingAABB.GetArea();
+			const FReal CostForChoosingNode = NewPotentialNodeArea + SumDeltaCost;
+			
 			if (bAddToLeaf)
 			{
 				// No new node is added (we can experiment with this cost function
 				// It is faster overall if we don't subtract here
 				//CostForChoosingNode -= TestSiblingArea;
 			}
-			FReal NewDeltaCost = NewPotentialNodeArea - TestSiblingArea;
+			
+			const FReal NewDeltaCost = NewPotentialNodeArea - TestSiblingArea;
 			// Did we get a better cost?
 			if (CostForChoosingNode < BestCost)
 			{
 				BestCost = CostForChoosingNode;
-				BestSiblingIdx = TestSibling;
+				BestSiblingIdx = NodeAndCost.template Get<1>();
 				bOutAddToLeaf = bAddToLeaf;
 			}
 
 			// Lower bound of Children costs
-			FReal ChildCostLowerBound = NewBounds.GetArea() + NewDeltaCost + SumDeltaCost;
+			const FReal NewCost = NewDeltaCost + SumDeltaCost;
+			const FReal ChildCostLowerBound = NewBoundsArea + NewCost;
 
-			if (!Nodes[TestSibling].bLeaf && ChildCostLowerBound < BestCost)
+			if (!TestNode.bLeaf && ChildCostLowerBound < BestCost)
 			{
 				// Now we will push the children
-				PriorityQ.Add(Nodes[TestSibling].ChildrenNodes[0]);
-				PriorityQ.Add(Nodes[TestSibling].ChildrenNodes[1]);
-				SumDeltaCostQ.Add(NewDeltaCost + SumDeltaCost);
-				SumDeltaCostQ.Add(NewDeltaCost + SumDeltaCost);
+				PriorityQ.Reserve(PriorityQ.Num() + 2);
+				PriorityQ.Emplace(Nodes[TestNode.ChildrenNodes[0]], TestNode.ChildrenNodes[0], NewCost);
+				PriorityQ.Emplace(Nodes[TestNode.ChildrenNodes[1]], TestNode.ChildrenNodes[1], NewCost);
 			}
 
 		}
@@ -1342,6 +1347,14 @@ public:
 		//	DynamicTreeDebugStats();
 		//}
 
+		// New tree?
+		if(RootNode == INDEX_NONE)
+		{
+			int32 NewLeafNode = AllocateLeafNodeAndLeaf(Payload, NewBounds);
+			RootNode = NewLeafNode;
+			return;
+		}
+
 		// Find the best sibling
 		bool bAddToLeaf;
 		int32 BestSibling = FindBestSibling(NewBounds, bAddToLeaf);
@@ -1363,41 +1376,34 @@ public:
 
 		int32 NewLeafNode = AllocateLeafNodeAndLeaf(Payload, NewBounds);
 
-		// New tree?
-		if (RootNode == INDEX_NONE)
-		{
-			RootNode = NewLeafNode;
-			return;
-		}
-
 		// New internal parent node
-		int32 oldParent = Nodes[BestSibling].ParentNode;
-		int32 newParent = AllocateInternalNode();
-		Nodes[newParent].ParentNode = oldParent;
-		Nodes[newParent].ChildrenNodes[0] = BestSibling;
-		Nodes[newParent].ChildrenNodes[1] = NewLeafNode;
-		Nodes[newParent].ChildrenBounds[0] = Nodes[BestSibling].ChildrenBounds[0];
+		int32 OldParent = Nodes[BestSibling].ParentNode;
+		int32 NewParent = AllocateInternalNode();
+		FNode& NewParentNode = Nodes[NewParent];
+		NewParentNode.ParentNode = OldParent;
+		NewParentNode.ChildrenNodes[0] = BestSibling;
+		NewParentNode.ChildrenNodes[1] = NewLeafNode;
+		NewParentNode.ChildrenBounds[0] = Nodes[BestSibling].ChildrenBounds[0];
 		if (!Nodes[BestSibling].bLeaf)
 		{
-			Nodes[newParent].ChildrenBounds[0].GrowToInclude(Nodes[BestSibling].ChildrenBounds[1]);
+			NewParentNode.ChildrenBounds[0].GrowToInclude(Nodes[BestSibling].ChildrenBounds[1]);
 		}
-		Nodes[newParent].ChildrenBounds[1] = Nodes[NewLeafNode].ChildrenBounds[0];
+		NewParentNode.ChildrenBounds[1] = Nodes[NewLeafNode].ChildrenBounds[0];
 
-		if (oldParent != INDEX_NONE)
+		if (OldParent != INDEX_NONE)
 		{
 			int32 ChildIdx = WhichChildAmI(BestSibling);
-			Nodes[oldParent].ChildrenNodes[ChildIdx] = newParent;
+			Nodes[OldParent].ChildrenNodes[ChildIdx] = NewParent;
 		}
 		else
 		{
-			RootNode = newParent;
+			RootNode = NewParent;
 		}
 
-		Nodes[BestSibling].ParentNode = newParent;
-		Nodes[NewLeafNode].ParentNode = newParent;
-
-		UpdateAncestorBounds(newParent, true);
+		Nodes[BestSibling].ParentNode = NewParent;
+		Nodes[NewLeafNode].ParentNode = NewParent;
 
+		UpdateAncestorBounds(NewParent, true);
 	}	
 
 	void  UpdateAncestorBounds(int32 NodeIdx, bool bDoRotation = false)
@@ -2217,9 +2223,6 @@ public:
 
 private:
 
-	using FElement = TPayloadBoundsElement<TPayloadType, T>;
-	using FNode = TAABBTreeNode<T>;
-
 	void ReoptimizeTree()
 	{
 		TRACE_CPUPROFILER_EVENT_SCOPE(TAABBTree::ReoptimizeTree);
@@ -3195,9 +3198,7 @@ private:
 		, OverlappingCounts(Other.OverlappingCounts)
 	{
 		ensure(bDynamicTree == Other.IsTreeDynamic());
-
 		PriorityQ.Reserve(32);
-		SumDeltaCostQ.Reserve(32);
 	}
 
 	virtual ISpatialAcceleration<TPayloadType, T, 3>& operator=(const ISpatialAcceleration<TPayloadType, T, 3>& Other) override
@@ -3303,9 +3304,12 @@ private:
 	/** Number of overlapping leaves per leaf */
 	TArray<int32> OverlappingCounts;
 
-	/** Buffers for calculating the best candidate for tree insertion, reset per-calculation to avoid over allocation */
-	TArray<int32> PriorityQ;
-	TArray<FReal> SumDeltaCostQ;
+	/** 
+	 * Buffers for calculating the best candidate for tree insertion, reset per-calculation to avoid over allocation 
+	 * Tuple of Node/Index/Cost here avoids cache miss when accessing the node and its cost.
+	 */
+	using FNodeIndexAndCost = TTuple<FNode&, int32, FReal>;
+	TArray<FNodeIndexAndCost> PriorityQ;
 };
 
 template<typename TPayloadType, typename TLeafType, bool bMutable, typename T>
