Interchange: display the source description at the top of the import dialog settings

#rnx

diff --git a/Engine/Plugins/Experimental/Interchange/Editor/Source/Pipelines/Private/InterchangePipelineConfigurationGeneric.cpp b/Engine/Plugins/Experimental/Interchange/Editor/Source/Pipelines/Private/InterchangePipelineConfigurationGeneric.cpp
--- a/Engine/Plugins/Experimental/Interchange/Editor/Source/Pipelines/Private/InterchangePipelineConfigurationGeneric.cpp
+++ b/Engine/Plugins/Experimental/Interchange/Editor/Source/Pipelines/Private/InterchangePipelineConfigurationGeneric.cpp
@@ -3,6 +3,7 @@
 
 #include "CoreMinimal.h"
 #include "Framework/Application/SlateApplication.h"
+#include "InterchangeSourceData.h"
 #include "Interfaces/IMainFrameModule.h"
 #include "Misc/Paths.h"
 #include "SInterchangePipelineConfigurationDialog.h"
@@ -10,7 +11,7 @@
 #include "UObject/ObjectMacros.h"
 #include "Widgets/SWindow.h"
 
-EInterchangePipelineConfigurationDialogResult UInterchangePipelineConfigurationGeneric::ShowPipelineConfigurationDialog()
+EInterchangePipelineConfigurationDialogResult UInterchangePipelineConfigurationGeneric::ShowPipelineConfigurationDialog(TWeakObjectPtr<UInterchangeSourceData> SourceData)
 {
 	//Create and show the graph inspector UI dialog
 	TSharedPtr<SWindow> ParentWindow;
@@ -27,6 +28,7 @@ EInterchangePipelineConfigurationDialogResult UInterchangePipelineConfigurationG
 	(
 		SAssignNew(InterchangePipelineConfigurationDialog, SInterchangePipelineConfigurationDialog)
 		.OwnerWindow(Window)
+		.SourceData(SourceData)
 		.bReimport(false)
 	);
 
@@ -45,7 +47,7 @@ EInterchangePipelineConfigurationDialogResult UInterchangePipelineConfigurationG
 	return EInterchangePipelineConfigurationDialogResult::Import;
 }
 
-EInterchangePipelineConfigurationDialogResult UInterchangePipelineConfigurationGeneric::ShowReimportPipelineConfigurationDialog(TArray<UInterchangePipelineBase*>& PipelineStack)
+EInterchangePipelineConfigurationDialogResult UInterchangePipelineConfigurationGeneric::ShowReimportPipelineConfigurationDialog(TArray<UInterchangePipelineBase*>& PipelineStack, TWeakObjectPtr<UInterchangeSourceData> SourceData)
 {
 	//Create and show the graph inspector UI dialog
 	TSharedPtr<SWindow> ParentWindow;
@@ -62,6 +64,7 @@ EInterchangePipelineConfigurationDialogResult UInterchangePipelineConfigurationG
 	(
 		SAssignNew(InterchangePipelineConfigurationDialog, SInterchangePipelineConfigurationDialog)
 		.OwnerWindow(Window)
+		.SourceData(SourceData)
 		.bReimport(true)
 		.PipelineStack(PipelineStack)
 	);
diff --git a/Engine/Plugins/Experimental/Interchange/Editor/Source/Pipelines/Private/SInterchangePipelineConfigurationDialog.cpp b/Engine/Plugins/Experimental/Interchange/Editor/Source/Pipelines/Private/SInterchangePipelineConfigurationDialog.cpp
--- a/Engine/Plugins/Experimental/Interchange/Editor/Source/Pipelines/Private/SInterchangePipelineConfigurationDialog.cpp
+++ b/Engine/Plugins/Experimental/Interchange/Editor/Source/Pipelines/Private/SInterchangePipelineConfigurationDialog.cpp
@@ -8,6 +8,7 @@
 #include "IDocumentation.h"
 #include "InterchangePipelineBase.h"
 #include "InterchangeProjectSettings.h"
+#include "InterchangeSourceData.h"
 #include "Misc/App.h"
 #include "Modules/ModuleManager.h"
 #include "PropertyEditorModule.h"
@@ -406,6 +407,7 @@ void SInterchangePipelineConfigurationDialog::Construct(const FArguments& InArgs
 	//Make sure there is a valid default value
 
 	OwnerWindow = InArgs._OwnerWindow;
+	SourceData = InArgs._SourceData;
 	bReimport = InArgs._bReimport;
 	PipelineStack = InArgs._PipelineStack;
 
@@ -419,6 +421,14 @@ void SInterchangePipelineConfigurationDialog::Construct(const FArguments& InArgs
 		[
 			SNew(SVerticalBox)
 			+ SVerticalBox::Slot()
+			.AutoHeight()
+			.HAlign(HAlign_Left)
+			.Padding(2)
+			[
+				SNew(STextBlock)
+				.Text(this, &SInterchangePipelineConfigurationDialog::GetSourceDescription)
+			]
+			+ SVerticalBox::Slot()
 			.FillHeight(1.0f)
 			.Padding(2)
 			[
@@ -470,6 +480,24 @@ void SInterchangePipelineConfigurationDialog::OnSelectionChanged(TSharedPtr<FInt
 	PipelineConfigurationDetailsView->SetObject(Pipeline);
 }
 
+FText SInterchangePipelineConfigurationDialog::GetSourceDescription() const
+{
+	FText ActionDescription;
+	if (bReimport)
+	{
+		ActionDescription = LOCTEXT("GetSourceDescription_Reimport", "Reimport");
+	}
+	else
+	{
+		ActionDescription = LOCTEXT("GetSourceDescription_Import", "Import");
+	}
+	if (SourceData.IsValid())
+	{
+		ActionDescription = FText::Format(LOCTEXT("GetSourceDescription", "{0} source {1}"), ActionDescription, FText::FromString(SourceData->GetFilename()));
+	}
+	return ActionDescription;
+}
+
 bool SInterchangePipelineConfigurationDialog::RecursiveValidatePipelineSettings(const TSharedPtr<FInterchangePipelineStacksTreeNodeItem>& ParentNode) const
 {
 	if (ParentNode->Pipeline)
diff --git a/Engine/Plugins/Experimental/Interchange/Editor/Source/Pipelines/Private/SInterchangePipelineConfigurationDialog.h b/Engine/Plugins/Experimental/Interchange/Editor/Source/Pipelines/Private/SInterchangePipelineConfigurationDialog.h
--- a/Engine/Plugins/Experimental/Interchange/Editor/Source/Pipelines/Private/SInterchangePipelineConfigurationDialog.h
+++ b/Engine/Plugins/Experimental/Interchange/Editor/Source/Pipelines/Private/SInterchangePipelineConfigurationDialog.h
@@ -6,6 +6,7 @@
 #include "InputCoreTypes.h"
 #include "Input/Reply.h"
 #include "InterchangePipelineBase.h"
+#include "InterchangeSourceData.h"
 #include "Styling/SlateColor.h"
 #include "UObject/GCObject.h"
 #include "Widgets/DeclarativeSyntaxSupport.h"
@@ -95,6 +96,7 @@ public:
 	{}
 
 		SLATE_ARGUMENT(TWeakPtr<SWindow>, OwnerWindow)
+		SLATE_ARGUMENT(TWeakObjectPtr<UInterchangeSourceData>, SourceData)
 		SLATE_ARGUMENT(bool, bReimport)
 		SLATE_ARGUMENT(TArray<UInterchangePipelineBase*>, PipelineStack)
 	SLATE_END_ARGS()
@@ -121,12 +123,15 @@ public:
 	bool IsImportAll() { return bImportAll; }
 private:
 	TWeakPtr< SWindow > OwnerWindow;
+	TWeakObjectPtr<UInterchangeSourceData> SourceData;
 	bool bReimport = false;
 	TArray<UInterchangePipelineBase*> PipelineStack;
 
 	TSharedRef<SBox> SpawnPipelineConfiguration();
 	void OnSelectionChanged(TSharedPtr<FInterchangePipelineStacksTreeNodeItem> Item, ESelectInfo::Type SelectionType);
 
+	FText GetSourceDescription() const;
+
 	bool RecursiveValidatePipelineSettings(const TSharedPtr<FInterchangePipelineStacksTreeNodeItem>& ParentNode) const;
 	bool IsImportButtonEnabled() const;
 
diff --git a/Engine/Plugins/Experimental/Interchange/Editor/Source/Pipelines/Public/InterchangePipelineConfigurationGeneric.h b/Engine/Plugins/Experimental/Interchange/Editor/Source/Pipelines/Public/InterchangePipelineConfigurationGeneric.h
--- a/Engine/Plugins/Experimental/Interchange/Editor/Source/Pipelines/Public/InterchangePipelineConfigurationGeneric.h
+++ b/Engine/Plugins/Experimental/Interchange/Editor/Source/Pipelines/Public/InterchangePipelineConfigurationGeneric.h
@@ -9,6 +9,8 @@
 
 #include "InterchangePipelineConfigurationGeneric.generated.h"
 
+class UInterchangeSourceData;
+
 UCLASS(BlueprintType, Blueprintable)
 class INTERCHANGEEDITORPIPELINES_API UInterchangePipelineConfigurationGeneric : public UInterchangePipelineConfigurationBase
 {
@@ -18,6 +20,6 @@ public:
 
 protected:
 
-	virtual EInterchangePipelineConfigurationDialogResult ShowPipelineConfigurationDialog() override;
-	virtual EInterchangePipelineConfigurationDialogResult ShowReimportPipelineConfigurationDialog(TArray<UInterchangePipelineBase*>& PipelineStack) override;
+	virtual EInterchangePipelineConfigurationDialogResult ShowPipelineConfigurationDialog(TWeakObjectPtr<UInterchangeSourceData> SourceData) override;
+	virtual EInterchangePipelineConfigurationDialogResult ShowReimportPipelineConfigurationDialog(TArray<UInterchangePipelineBase*>& PipelineStack, TWeakObjectPtr<UInterchangeSourceData> SourceData) override;
 };
diff --git a/Engine/Source/Runtime/Experimental/Interchange/Engine/Private/InterchangeManager.cpp b/Engine/Source/Runtime/Experimental/Interchange/Engine/Private/InterchangeManager.cpp
--- a/Engine/Source/Runtime/Experimental/Interchange/Engine/Private/InterchangeManager.cpp
+++ b/Engine/Source/Runtime/Experimental/Interchange/Engine/Private/InterchangeManager.cpp
@@ -777,7 +777,7 @@ UInterchangeManager::ImportInternal(const FString& ContentPath, const UInterchan
 			if (RegisteredPipelineConfiguration && bShowPipelineStacksConfigurationDialog && !bIsUnattended)
 			{
 				//Show the dialog, a plugin should have registered this dialog. We use a plugin to be able to use editor code when doing UI
-				EInterchangePipelineConfigurationDialogResult DialogResult = RegisteredPipelineConfiguration->ScriptedShowReimportPipelineConfigurationDialog(PipelineStack);
+				EInterchangePipelineConfigurationDialogResult DialogResult = RegisteredPipelineConfiguration->ScriptedShowReimportPipelineConfigurationDialog(PipelineStack, DuplicateSourceData);
 				if (DialogResult == EInterchangePipelineConfigurationDialogResult::Cancel)
 				{
 					bImportCancel = true;
@@ -798,7 +798,7 @@ UInterchangeManager::ImportInternal(const FString& ContentPath, const UInterchan
 			if (RegisteredPipelineConfiguration && (bShowPipelineStacksConfigurationDialog || (!DefaultPipelineStacks.Contains(PipelineStackName) && !bIsUnattended)))
 			{
 				//Show the dialog, a plugin should have register this dialog. We use a plugin to be able to use editor code when doing UI
-				EInterchangePipelineConfigurationDialogResult DialogResult = RegisteredPipelineConfiguration->ScriptedShowPipelineConfigurationDialog();
+				EInterchangePipelineConfigurationDialogResult DialogResult = RegisteredPipelineConfiguration->ScriptedShowPipelineConfigurationDialog(DuplicateSourceData);
 				if (DialogResult == EInterchangePipelineConfigurationDialogResult::Cancel)
 				{
 					bImportCancel = true;
diff --git a/Engine/Source/Runtime/Experimental/Interchange/Engine/Public/InterchangePipelineConfigurationBase.h b/Engine/Source/Runtime/Experimental/Interchange/Engine/Public/InterchangePipelineConfigurationBase.h
--- a/Engine/Source/Runtime/Experimental/Interchange/Engine/Public/InterchangePipelineConfigurationBase.h
+++ b/Engine/Source/Runtime/Experimental/Interchange/Engine/Public/InterchangePipelineConfigurationBase.h
@@ -8,6 +8,8 @@
 
 #include "InterchangePipelineConfigurationBase.generated.h"
 
+class UInterchangeSourceData;
+
 UENUM(BlueprintType, Experimental)
 enum class EInterchangePipelineConfigurationDialogResult : uint8
 {
@@ -27,26 +29,26 @@ public:
 	 * Non virtual helper to allow blueprint to implement event base function to implement the ShowPipelineConfigurationDialog,
 	 */
 	UFUNCTION(BlueprintNativeEvent, BlueprintCallable, Category = "Interchange | Translator")
-	EInterchangePipelineConfigurationDialogResult ScriptedShowPipelineConfigurationDialog();
+	EInterchangePipelineConfigurationDialogResult ScriptedShowPipelineConfigurationDialog(UInterchangeSourceData* SourceData);
 
 	/** The default implementation (call if the blueprint do not have any implementation) will call the virtual ExecuteImportPipeline */
-	EInterchangePipelineConfigurationDialogResult ScriptedShowPipelineConfigurationDialog_Implementation()
+	EInterchangePipelineConfigurationDialogResult ScriptedShowPipelineConfigurationDialog_Implementation(UInterchangeSourceData* SourceData)
 	{
 		//By default we call the virtual import pipeline execution
-		return ShowPipelineConfigurationDialog();
+		return ShowPipelineConfigurationDialog(SourceData);
 	}
 
 	/**
 	 * Non virtual helper to allow blueprint to implement event base function to implement the ShowPipelineConfigurationDialog,
 	 */
 	UFUNCTION(BlueprintNativeEvent, BlueprintCallable, Category = "Interchange | Translator")
-	EInterchangePipelineConfigurationDialogResult ScriptedShowReimportPipelineConfigurationDialog(TArray<UInterchangePipelineBase*>& PipelineStack);
+	EInterchangePipelineConfigurationDialogResult ScriptedShowReimportPipelineConfigurationDialog(TArray<UInterchangePipelineBase*>& PipelineStack, UInterchangeSourceData* SourceData);
 
 	/** The default implementation (call if the blueprint do not have any implementation) will call the virtual ExecuteImportPipeline */
-	EInterchangePipelineConfigurationDialogResult ScriptedShowReimportPipelineConfigurationDialog_Implementation(TArray<UInterchangePipelineBase*>& PipelineStack)
+	EInterchangePipelineConfigurationDialogResult ScriptedShowReimportPipelineConfigurationDialog_Implementation(TArray<UInterchangePipelineBase*>& PipelineStack, UInterchangeSourceData* SourceData)
 	{
 		//By default we call the virtual import pipeline execution
-		return ShowReimportPipelineConfigurationDialog(PipelineStack);
+		return ShowReimportPipelineConfigurationDialog(PipelineStack, SourceData);
 	}
 
 protected:
@@ -54,7 +56,7 @@ protected:
 	/**
 	 * This function show a dialog use to configure pipeline stacks and return a stack name that tell the caller the user choice.
 	 */
-	virtual EInterchangePipelineConfigurationDialogResult ShowPipelineConfigurationDialog()
+	virtual EInterchangePipelineConfigurationDialogResult ShowPipelineConfigurationDialog(TWeakObjectPtr<UInterchangeSourceData> SourceData)
 	{ 
 		//Not implemented
 		return EInterchangePipelineConfigurationDialogResult::Cancel;
@@ -63,7 +65,7 @@ protected:
 	/**
 	 * This function show a dialog use to configure pipeline stacks and return a stack name that tell the caller the user choice.
 	 */
-	virtual EInterchangePipelineConfigurationDialogResult ShowReimportPipelineConfigurationDialog(TArray<UInterchangePipelineBase*>& PipelineStack)
+	virtual EInterchangePipelineConfigurationDialogResult ShowReimportPipelineConfigurationDialog(TArray<UInterchangePipelineBase*>& PipelineStack, TWeakObjectPtr<UInterchangeSourceData> SourceData)
 	{
 		//Not implemented
 		return EInterchangePipelineConfigurationDialogResult::Cancel;
