Add LanczosCG solver.

diff --git a/Engine/Source/Programs/HeadlessChaos/Private/HeadlessChaosTestMath.cpp b/Engine/Source/Programs/HeadlessChaos/Private/HeadlessChaosTestMath.cpp
--- a/Engine/Source/Programs/HeadlessChaos/Private/HeadlessChaosTestMath.cpp
+++ b/Engine/Source/Programs/HeadlessChaos/Private/HeadlessChaosTestMath.cpp
@@ -2,8 +2,11 @@
 
 #include "HeadlessChaos.h"
 #include "HeadlessChaosTestUtility.h"
+#include "Chaos/Math/Krylov.h"
 #include "Chaos/Matrix.h"
 #include "Chaos/Utilities.h"
+#include "Chaos/Vector.h"
+#include "GenericPlatform/GenericPlatformMath.h"
 
 namespace ChaosTest
 {
@@ -25,5 +28,143 @@ namespace ChaosTest
 		}
 	}
 
+	Chaos::TVector<double, 4> 
+	ToVec4(const TArray<double>& x)
+	{
+		return Chaos::TVector<double, 4>(x[0], x[1], x[2], x[3]);
+	}
+
+	TArray<double>
+	ToArray4(const TVector<double, 4>& x)
+	{
+		return TArray<double>({x[0], x[1], x[2], x[3]});
+	}
+
+	TEST(MathTests, TestLanczosCGSolver) 
+	{
+		Chaos::PMatrix<double, 4, 4> A(0, 0, 0, 0, 0, 0, 0, 0, 0, 0, 0, 0, 0, 0, 0, 0);
+		A.M[0][0] = double(2);
+		A.M[0][1] = -double(1);
+		A.M[1][0] = -double(1);
+		A.M[1][1] = double(2);
+		A.M[1][2] = -double(1);
+		A.M[2][1] = -double(1);
+		A.M[2][2] = double(2);
+		A.M[2][3] = -double(1);
+		A.M[3][2] = -double(1);
+		A.M[3][3] = double(2);
+		A = A.GetTransposed(); // UE is col maj, not row
+
+		TArray<double> x({ 1, 1, 1, 1 });
+		TArray<double> x_cg;
+		TArray<double> b = ToArray4(A.TransformFVector4(ToVec4(x)));
+
+		Chaos::LanczosCG<double>(
+			[&A](TArray<double>& y, const TArray<double>& x) { y = ToArray4(A.TransformFVector4(ToVec4(x))); },
+			[](const TArray<double>& x, const TArray<double>& y) { return Dot4(ToVec4(x), ToVec4(y)); },
+			[](TArray<double>& y, double a, const TArray<double>& x) { y = ToArray4(ToVec4(y) + a * ToVec4(x)); },
+			[](TArray<double>& y, double a) { y = ToArray4(a * ToVec4(y)); },
+			[](TArray<double>& y, const TArray<double>& x) { y = x; },
+			x_cg,
+			b,
+			5);
+
+		double error = 0.;
+		for (int i = 0; i < 4; i++)
+			error += (x[i] - x_cg[i]) * (x[i] - x_cg[i]);
+		error = FGenericPlatformMath::Sqrt(error);
+		double tol = 1.e-10;
+		EXPECT_NEAR(error, double(0), tol);
+	}
+
+#if 0
+	void ICRes(const Chaos::PMatrix<double, 4, 4>& L, const Chaos::TVector<double, 4>& d, Chaos::TVector<double, 4>& y, const Chaos::TVector<double, 4>& x)
+	{
+		Chaos::TVector<double, 4> u;
+		for (int k = 0; k < 4; k++)
+		{
+			if (L.M[k][k])
+			{
+				double rly = L.M[k][k];
+				for (int j = 0; j < k; j++)
+				{
+					rly -= L.M[k][j] * u[j];
+				}
+				u[k] = rly / L.M[k][k];
+			}
+		}
+		for (int k = 4 - 1; k >= 0; k--)
+		{
+			double lu = 0.;
+			for (int i = k + 1; i < 4; i++)
+			{
+				lu += L.M[i][k] * y[i];
+			}
+			y[k] = u[k] - d[k] * lu;
+		}
+	}
 
+	TEST(MathTests, LanczosPCGSolver)
+	{
+		Chaos::PMatrix<double, 4, 4> A(2, 0, 0, 0, 0, 0, 0, -2.6, 0, 0, 4, 3.6, 0, -2.6, 3.6, 5);
+		Chaos::PMatrix<double, 4, 4> M;
+		Chaos::PMatrix<double, 3, 3> LMat;
+		Chaos::TVector<double, 4> DVec;
+		TVector<double, 4> x(1, 1, 1, 1);
+		TVector<double, 4> x_pcg;
+		TVector<double, 4> b;
+		//b = A * x;
+		b = A.TransformFVector4(x);
+
+		// Calculate M
+		auto ICU = [](Chaos::PMatrix<double, 4, 4>& AMat, Chaos::PMatrix<double, 4, 4>& L, Chaos::TVector<double, 4>& d)
+		{
+			double epsilon = .1;
+			L.M[0][0] = AMat.M[0][0];
+			
+			d[0] = 1. / L.M[0][0];
+			for (int k = 0; k < 4; k++) 
+			{
+				double value = 0.;
+				for (int j =0; j < 4; j++) 
+				{
+					double aij = AMat.M[k][j];
+					double ldl = 0.;
+					for(int i=0; i < 4; i++)
+					{
+						ldl -= L.M[k][i] * L.M[j][i] * d[i];
+					}
+					value = aij + ldl;
+					L.M[k][j] = value;
+				}
+				if (abs(L.M[k][k]) < epsilon)
+					L.M[k][k] = epsilon;
+				d[k] = 1. / L.M[k][k];
+			}
+		};
+
+		auto multiplyA = [&A](Chaos::TVector<double, 4>& y, const Chaos::TVector<double, 4>& x) { y = A.TransformFVector4(x); };
+		auto multiplyPrecond = [&LMat, &DVec, this](Chaos::TVector<double, 4>& y, const Chaos::TVector<double, 4>& x) { ICRes(LMat, DVec, y, x); };
+		auto dotProduct = [](const Chaos::TVector<double, 4>& x, const Chaos::TVector<double, 4>& y) { return Dot4(x, y); };
+		auto set = [](Chaos::TVector<double, 4>& y, const Chaos::TVector<double, 4>& x) { y = x; };
+		auto setScaled = [](Chaos::TVector<double, 4>& a, const Chaos::TVector<double, 4>& b, const double s) { a = s * b; };        // setScaled(a,b,s): a <- s * b
+		auto scaleAndAdd = [](Chaos::TVector<double, 4>& a, const double s, const Chaos::TVector<double, 4>& b) { a = s * a + b; };  // scaleAndAdd(a,s,b): a <- s * a + b
+		auto addScaled = [](Chaos::TVector<double, 4>& a, const Chaos::TVector<double, 4>& b, const double s) { a += s * b; };       // addScaled(a,b,s): a <- a + s * b
+		auto addScaled2 = [](Chaos::TVector<double, 4>& a, const Chaos::TVector<double, 4>& b1, const double s1, const Chaos::TVector<double, 4>& b2, const double s2) {
+			a += s1 * b1 + s2 * b2;
+		};  // addScaled2(a,b1,s1,b2,s2); a <- s1 * b1 + s2 * b2
+		auto residual = [&](double& r, const Chaos::TVector<double, 4>& x, const Chaos::TVector<double, 4>& b) {
+			Chaos::TVector<double, 4> res, Mres;
+			multiplyA(res, x);
+			addScaled(res, b, double(-1));
+			multiplyPrecond(Mres, res);
+			r = FGenericPlatformMath::Sqrt(dotProduct(Mres, res));
+		};
+
+		Chaos::LanczosPCG<T>(multiplyA, multiplyPrecond, dotProduct, set, setScaled, scaleAndAdd, addScaled, addScaled2, residual, x_pcg, b, 4);
+		double r;
+		residual(r, x_pcg, b);
+		EXPECT_NEAR(r, 0., 1.e-10);
+	}
+#endif // 0
 }
diff --git a/Engine/Source/Runtime/Experimental/Chaos/Public/Chaos/Math/Krylov.h b/Engine/Source/Runtime/Experimental/Chaos/Public/Chaos/Math/Krylov.h
new file mode 100644
--- /dev/null
+++ b/Engine/Source/Runtime/Experimental/Chaos/Public/Chaos/Math/Krylov.h
@@ -0,0 +1,285 @@
+// Copyright Epic Games, Inc. All Rights Reserved.
+#pragma once
+
+#include "ChaosCheck.h"
+#include "ChaosLog.h"
+#include "Chaos/Framework/Parallel.h"
+#include "GenericPlatform/GenericPlatformMath.h"
+#include "Misc/AssertionMacros.h"
+
+namespace Chaos {
+
+template <typename Func, class T>
+void LanczosCG(
+	Func multiplyA, 
+	TArray<T>& x,
+	const TArray<T>& b,
+	const int max_it, 
+	const T res = 1e-4, 
+	bool check_residual = false) 
+{
+	auto dotProduct = [](const TArray<T>& x, const TArray<T>& y) 
+	{
+		T result = T(0);
+		checkfSlow(x.Num() == y.Num(), TEXT("LanczosCG: trying to take the dot product with vectors of different size."));
+		PhysicsParallelFor(x.Num(), [&](const int32 i)
+		{
+			result += x[i] * y[i];
+		},
+		x.Num() < 1000);
+		return result;
+	};
+
+	auto AXPY = [](TArray<T>& y, T a, const TArray<T>& x) 
+	{
+		checkfSlow(x.Num() == y.Num(), TEXT("LanczosCG: trying to take a linear combination of vectors with different sizes."));
+		PhysicsParallelFor(x.Num(), [&](const int32 i)
+		{
+			y[i] += a * x[i];
+	    },
+		x.Num() < 1000); // Single threaded for less than 1k operations.
+	};
+
+	auto scale = [](TArray<T>& y, T a) 
+	{
+		PhysicsParallelFor(y.Num(), [&](const int32 i)
+		{
+			y[i] *= a;
+	    },
+		y.Num() < 1000); // Single threaded for less than 1k operations.
+	};
+
+	auto set = [](TArray<T>& y, const TArray<T>& x) 
+	{
+		checkfSlow(x.Num() == y.Num(), TEXT("LanczosCG: trying to set to vectors with different sizes."));
+		PhysicsParallelFor(x.Num(), [&](const int32 i)
+		{
+			y[i] = x[i];
+		},
+		x.Num() < 1000); // Single threaded for less than 1k operations.
+	};
+
+	LanczosCG<T>(multiplyA, dotProduct, AXPY, scale, set, x, b, max_it, res, check_residual);
+}
+
+template <class T, class TV, typename Func, int32 d=3>
+void LanczosCG(
+	Func multiplyA, 
+	TArray<TV>& x,
+	TArray<TV>& b,
+	const int max_it, 
+	const T res = 1e-4, 
+	const TArray<int32>* use_list = nullptr)
+{
+	auto dot_product = [&use_list](const TArray<TV>& x, const TArray<TV>& y) 
+	{
+		checkfSlow(x.Num() == y.Num(), TEXT("Chaos::LanczosCG()::dot_product[]: x and y not sized consistently."));
+		T result = T(0);
+		if (use_list != nullptr) 
+		{
+			PhysicsParallelFor(use_list->Num(), [&](const int32 i)
+			{
+				for (int32 i = 0; i < use_list->Num(); ++i) 
+				{
+					T dot = T(0);
+					int32 j = (*use_list)[i];
+					for (int32 alpha = 0; alpha < d; alpha++) 
+					{
+						dot += x[j][alpha] * y[j][alpha];
+					}
+					result += dot;
+				}
+			},
+			use_list->Num() < 1000);
+		}
+		else 
+		{
+			PhysicsParallelFor(x.Num(), [&](const int32 i)
+			{
+				T dot = T(0);
+				for (size_t alpha = 0; alpha < d; alpha++) 
+				{
+					dot += x[i][alpha] * y[i][alpha];
+				}
+				result += dot;
+			},
+			x.Num() < 1000);
+		}
+		return result;
+	};
+	auto AXPY = [&use_list](TArray<TV>& y, T a, const TArray<TV>& x) 
+	{
+		checkfSlow(x.Num() == y.Num(), TEXT("Chaos::LanczosCG()::AXPY[]: x and y not sized consistently."));
+		if (use_list != nullptr) 
+		{
+			PhysicsParallelFor(use_list->Num(), [&](const int32 i) 
+			{
+				const int32 j = (*use_list)[i];
+				for (int32 alpha = 0; alpha < d; alpha++) 
+				{
+					y[j][alpha] += a * x[j][alpha];
+				}
+			},
+			use_list->Num() < 1000);
+		}
+		else 
+		{
+			PhysicsParallelFor(x.Num(), [&](const int32 i)
+			{
+				for (int32 alpha = 0; alpha < d; alpha++) 
+				{
+					y[i][alpha] += a * x[i][alpha];
+				}
+			},
+			x.Num() < 1000);
+		}
+	};
+	auto scale = [&use_list](TArray<TV>& y, T a) 
+	{
+		if (use_list != nullptr) 
+		{
+			PhysicsParallelFor(use_list->Num(), [&](const int32 i)
+			{
+				const int32 j = (*use_list)[i];
+				for (int32 alpha = 0; alpha < d; alpha++) 
+				{
+					y[j][alpha] = a * y[j][alpha];
+				}
+			});
+		}
+		else 
+		{
+			PhysicsParallelFor(y.Num(), [&](const int32 i)
+			{
+				for (size_t alpha = 0; alpha < d; alpha++) 
+				{
+					y[i][alpha] = a * y[i][alpha];
+				}
+			});
+		}
+	};
+	auto set = [&use_list](TArray<TV>& y, const TArray<TV>& x) 
+	{
+		if (use_list != nullptr) 
+		{
+			PhysicsParallelFor(use_list->Num(), [&](const int32 i)
+			{
+				const int32 j = (*use_list)[i];
+				for (size_t alpha = 0; alpha < d; alpha++) 
+				{
+					y[j][alpha] = x[j][alpha];
+				}
+			});
+		}
+		else 
+		{
+			PhysicsParallelFor(y.Num(), [&](const int32 i)
+			{
+				for (int32 alpha = 0; alpha < d; alpha++) 
+				{
+					y[i][alpha] = x[i][alpha];
+				}
+			});
+		}
+	};
+
+	LanczosCG<T>(multiplyA, dot_product, AXPY, scale, set, x, b, max_it, res);
+}
+
+template <typename T, typename TV, typename Func1, typename Func2, typename Func3, typename Func4, typename Func5>
+void LanczosCG(
+	Func1 multiplyA, 
+	Func2 dotProduct, 
+	Func3 AXPY, 
+	Func4 scale, 
+	Func5 set, 
+	TArray<TV>& x,
+	const TArray<TV>& b,
+	const int max_it, 
+	const T res = T(1e-4), 
+	bool check_residual = false) 
+{
+	// multiplyA = [&A](TV &y, const TV& x): y = A*x;
+	// dotProduct = [](const TV& x, const TV& y): return x^Ty
+	// AXPY = [](Vector& y, T a, const TV& x): y += a*x;
+	// scale = [](Vector& y, T a): y <- a*y
+	// set = [](TV& y, const TV& x): y <- x
+	// Printing out: use
+	// (*((TV*)&b))[i]
+	// ((Vector*)&b)->operator[](i)
+
+	TArray<T> v; v.SetNum(x.Num());
+	TArray<T> q; q.SetNum(x.Num());
+	TArray<T> q_old; q_old.SetNum(x.Num());
+	TArray<T> c; c.SetNum(x.Num());
+
+	T beta = FGenericPlatformMath::Sqrt(T(dotProduct(b, b)));
+	if (beta < res) 
+	{
+		UE_LOG(LogChaos, VeryVerbose, TEXT("Lanczos residual = %d. Lanczos converged in 0 iterations."), beta)
+		scale(x, T(0));
+		return;
+	}
+	set(q, b);
+	scale(q, T(1) / beta);
+	multiplyA(v, q);
+	T alpha = dotProduct(v, q);
+	T d = alpha;
+	T p = beta / d;
+	set(c, q);
+	set(x, c);
+	scale(x, p);
+	AXPY(v, -alpha, q);
+	T residual;
+	TArray<T> y; y.SetNum(x.Num());
+	for (int it = 1; it < max_it; it++) 
+	{
+		beta = FGenericPlatformMath::Sqrt(T(dotProduct(v, v)));
+		if (check_residual == true)
+		{
+			multiplyA(y, x);
+			AXPY(y, T(-1), b);
+			residual = FGenericPlatformMath::Sqrt(T(dotProduct(y, y)));
+			if (residual < res) 
+			{
+				UE_LOG(LogChaos, VeryVerbose, TEXT("Ax - b residual = %g. Lanczos converged in %d iterations."), residual, it);
+				return;
+			}
+		}
+		else
+		{
+			if (beta < res) 
+			{
+				UE_LOG(LogChaos, VeryVerbose, TEXT("Lanczos residual = %g. Lanczos converged in %d iterations."), beta, it);
+				return;
+			}
+		}
+		T mu = beta / d;
+		set(q_old, q);
+		set(q, v);
+		scale(q, T(1) / beta);
+		multiplyA(v, q);
+		alpha = dotProduct(v, q);
+		T d_new = alpha - mu * mu * d;
+		p = -p * d * mu / d_new;
+		scale(c, -mu);
+		AXPY(c, T(1), q);
+		AXPY(x, p, c);
+		AXPY(v, -alpha, q);
+		AXPY(v, -beta, q_old);
+		d = d_new;
+	}
+
+	if (check_residual)
+	{
+		UE_LOG(LogChaos, Warning, TEXT("Lanczos used max iterations (%d). Residual = %g."), max_it, residual);
+	}
+	else
+	{
+		UE_LOG(LogChaos, Warning, TEXT("Lanczos used max iterations (%d)."), max_it);
+	}
+}
+
+
+} // namespace Chaos
+
