Chaos : GC optimization
- Cache local (GC component space ) bounds reducing game thread cost  when having a lot of non moving GC in a scene
- Make sure the dynamic collection is invalidated only when necessary

diff --git a/Engine/Source/Runtime/Experimental/Chaos/Private/PhysicsProxy/GeometryCollectionPhysicsProxy.cpp b/Engine/Source/Runtime/Experimental/Chaos/Private/PhysicsProxy/GeometryCollectionPhysicsProxy.cpp
--- a/Engine/Source/Runtime/Experimental/Chaos/Private/PhysicsProxy/GeometryCollectionPhysicsProxy.cpp
+++ b/Engine/Source/Runtime/Experimental/Chaos/Private/PhysicsProxy/GeometryCollectionPhysicsProxy.cpp
@@ -2191,6 +2191,8 @@ bool FGeometryCollectionPhysicsProxy::PullFromPhysicsState(const Chaos::FDirtyGe
 	TManagedArray<FVector3f>* AngularVelocities = GameThreadCollection.FindAttributeTyped<FVector3f>("AngularVelocity", FTransformCollection::TransformGroup);
 	TManagedArray<uint8>* InternalClusterParentTypeArray = GameThreadCollection.FindAttributeTyped<uint8>("InternalClusterParentTypeArray", FTransformCollection::TransformGroup);
 
+	bool bIsCollectionDirty = false;
+	
 	// We should never be changing the number of entries, this would break other 
 	// attributes in the transform group.
 	const int32 NumTransforms = GameThreadCollection.Transform.Num();
@@ -2199,21 +2201,47 @@ bool FGeometryCollectionPhysicsProxy::PullFromPhysicsState(const Chaos::FDirtyGe
 		// first : copy the non interpolate-able values
 		for (int32 TransformGroupIndex = 0; TransformGroupIndex < NumTransforms; ++TransformGroupIndex)
 		{
-			if (!TargetResults.States[TransformGroupIndex].DisabledState)
+			const bool bIsActive = !TargetResults.States[TransformGroupIndex].DisabledState;
+			if (GameThreadCollection.Active[TransformGroupIndex] != bIsActive)
+			{
+				GameThreadCollection.Active[TransformGroupIndex] = bIsActive;
+				bIsCollectionDirty = true;
+			}
+			
+			const int32 NewState = TargetResults.States[TransformGroupIndex].DynamicState;
+			if (GameThreadCollection.DynamicState[TransformGroupIndex] != NewState)
 			{
-				GameThreadCollection.Parent[TransformGroupIndex] = TargetResults.Parent[TransformGroupIndex];
+				GameThreadCollection.DynamicState[TransformGroupIndex] = NewState;
+				bIsCollectionDirty = true;
+			}
+
+			if (bIsActive)
+			{
+				const int32 NewParent = TargetResults.Parent[TransformGroupIndex];
+				if (GameThreadCollection.Parent[TransformGroupIndex] != NewParent)
+				{
+					GameThreadCollection.Parent[TransformGroupIndex] = NewParent;
+					bIsCollectionDirty = true;
+				}
+				// if (NewState != (int32)EObjectStateTypeEnum::Chaos_Object_Static && NewState != (int32)EObjectStateTypeEnum::Chaos_Object_Sleeping)
+				// {
+				// 	bIsCollectionDirty = true;
+				// }
 			}
-			GameThreadCollection.DynamicState[TransformGroupIndex] = TargetResults.States[TransformGroupIndex].DynamicState;
-			GameThreadCollection.Active[TransformGroupIndex] = !TargetResults.States[TransformGroupIndex].DisabledState;
 
 			if (InternalClusterParentTypeArray)
 			{
-				Chaos::EInternalClusterType Parenttype = Chaos::EInternalClusterType::None;
+				Chaos::EInternalClusterType ParentType = Chaos::EInternalClusterType::None;
 				if (TargetResults.States[TransformGroupIndex].HasInternalClusterParent != 0)
 				{
-					Parenttype = (TargetResults.States[TransformGroupIndex].DynamicInternalClusterParent != 0)? Chaos::EInternalClusterType::Dynamic: Chaos::EInternalClusterType::KinematicOrStatic; 
+					ParentType = (TargetResults.States[TransformGroupIndex].DynamicInternalClusterParent != 0)? Chaos::EInternalClusterType::Dynamic: Chaos::EInternalClusterType::KinematicOrStatic; 
+				}
+				const uint8 ParentTypeUInt8 = static_cast<uint8>(ParentType);
+				if ((*InternalClusterParentTypeArray)[TransformGroupIndex] != ParentTypeUInt8)
+				{
+					(*InternalClusterParentTypeArray)[TransformGroupIndex] = ParentTypeUInt8;
+					bIsCollectionDirty = true;
 				}
-				(*InternalClusterParentTypeArray)[TransformGroupIndex] =static_cast<uint8>(Parenttype);
 			}
 		}
 
@@ -2239,9 +2267,28 @@ bool FGeometryCollectionPhysicsProxy::PullFromPhysicsState(const Chaos::FDirtyGe
 				if (!TargetResults.States[TransformGroupIndex].DisabledState)
 				{
 					Chaos::FGeometryParticle& GTParticle = *GTParticles[TransformGroupIndex];
-					GTParticle.SetX(FMath::Lerp(Prev.ParticleXs[TransformGroupIndex], Next.ParticleXs[TransformGroupIndex], *Alpha), false);
-					GTParticle.SetR(FMath::Lerp(Prev.ParticleRs[TransformGroupIndex], Next.ParticleRs[TransformGroupIndex], *Alpha), false);
-					GTParticle.UpdateShapeBounds();
+
+					const Chaos::FVec3 OldX = GTParticle.X();
+					const Chaos::FVec3 NewX = FMath::Lerp(Prev.ParticleXs[TransformGroupIndex], Next.ParticleXs[TransformGroupIndex], *Alpha);
+					const bool bNeedUpdateX = (NewX != OldX);
+					if (bNeedUpdateX)
+					{
+						GTParticle.SetX(NewX, false);
+					}
+					
+					const Chaos::FRotation3 OldR = GTParticle.R();
+					const Chaos::FRotation3 NewR = FMath::Lerp(Prev.ParticleRs[TransformGroupIndex], Next.ParticleRs[TransformGroupIndex], *Alpha);
+					const bool bNeedUpdateR = (NewR != OldR); 
+					if (bNeedUpdateR)
+					{
+						GTParticle.SetR(NewR, false);
+					}
+
+					if (bNeedUpdateX || bNeedUpdateR)
+					{
+						GTParticle.UpdateShapeBounds();
+						bIsCollectionDirty = true;
+					}
 
 					GameThreadCollection.Transform[TransformGroupIndex].Blend(Prev.Transforms[TransformGroupIndex], Next.Transforms[TransformGroupIndex], *Alpha);
 
@@ -2261,9 +2308,28 @@ bool FGeometryCollectionPhysicsProxy::PullFromPhysicsState(const Chaos::FDirtyGe
 				if (!TargetResults.States[TransformGroupIndex].DisabledState)
 				{
 					Chaos::FGeometryParticle& GTParticle = *GTParticles[TransformGroupIndex];
-					GTParticle.SetX(TargetResults.ParticleXs[TransformGroupIndex], false);
-					GTParticle.SetR(TargetResults.ParticleRs[TransformGroupIndex], false);
-					GTParticle.UpdateShapeBounds();
+
+					const Chaos::FVec3 OldX = GTParticle.X();
+					const Chaos::FVec3& NewX = TargetResults.ParticleXs[TransformGroupIndex];
+					const bool bNeedUpdateX = (NewX != OldX);
+					if (bNeedUpdateX)
+					{
+						GTParticle.SetX(NewX, false);
+					}
+					
+					const Chaos::FRotation3 OldR = GTParticle.R();
+					const Chaos::FRotation3& NewR = TargetResults.ParticleRs[TransformGroupIndex];
+					const bool bNeedUpdateR = (NewR != OldR);
+					if (bNeedUpdateR)
+					{
+						GTParticle.SetR(NewR, false);
+					}
+
+					if (bNeedUpdateX || bNeedUpdateR)
+					{
+						GTParticle.UpdateShapeBounds();
+						bIsCollectionDirty = true;
+					}
 
 					GameThreadCollection.Transform[TransformGroupIndex] = TargetResults.Transforms[TransformGroupIndex];
 
@@ -2278,7 +2344,10 @@ bool FGeometryCollectionPhysicsProxy::PullFromPhysicsState(const Chaos::FDirtyGe
 		}
 
 		//question: why do we need this? Sleeping objects will always have to update GPU
-		GameThreadCollection.MakeDirty();
+		if (bIsCollectionDirty)
+		{
+			GameThreadCollection.MakeDirty();
+		}
 	}
 
 	return true;
diff --git a/Engine/Source/Runtime/Experimental/GeometryCollectionEngine/Private/GeometryCollection/GeometryCollectionComponent.cpp b/Engine/Source/Runtime/Experimental/GeometryCollectionEngine/Private/GeometryCollection/GeometryCollectionComponent.cpp
--- a/Engine/Source/Runtime/Experimental/GeometryCollectionEngine/Private/GeometryCollection/GeometryCollectionComponent.cpp
+++ b/Engine/Source/Runtime/Experimental/GeometryCollectionEngine/Private/GeometryCollection/GeometryCollectionComponent.cpp
@@ -79,6 +79,9 @@ static bool bChaos_BoxCalcBounds_ISPC_Enabled = CHAOS_BOX_CALC_BOUNDS_ISPC_ENABL
 static FAutoConsoleVariableRef CVarChaosBoxCalcBoundsISPCEnabled(TEXT("p.Chaos.BoxCalcBounds.ISPC"), bChaos_BoxCalcBounds_ISPC_Enabled, TEXT("Whether to use ISPC optimizations in calculating box bounds in geometry collections"));
 #endif
 
+bool bChaos_GC_CacheComponentSpaceBounds = true;
+FAutoConsoleVariableRef CVarChaosGCCacheComponentSpaceBounds(TEXT("p.Chaos.GC.CacheComponentSpaceBounds"), bChaos_GC_CacheComponentSpaceBounds, TEXT("Cache component space bounds for performance"));
+
 DEFINE_LOG_CATEGORY_STATIC(UGCC_LOG, Error, All);
 
 extern FGeometryCollectionDynamicDataPool GDynamicDataPool;
@@ -393,27 +396,11 @@ void UGeometryCollectionComponent::GetLifetimeReplicatedProps(TArray<FLifetimePr
 	DOREPLIFETIME(UGeometryCollectionComponent, RepData);
 }
 
-FBoxSphereBounds UGeometryCollectionComponent::CalcBounds(const FTransform& LocalToWorldIn) const
-{	
-	SCOPE_CYCLE_COUNTER(STAT_GCCUpdateBounds);
-
-	// #todo(dmp): hack to make bounds calculation work when we don't have valid physics proxy data.  This will
-	// force bounds calculation.
-
-	const FGeometryCollectionResults* Results = PhysicsProxy ? PhysicsProxy->GetConsumerResultsGT() : nullptr;
-
-	const int32 NumTransforms = Results ? Results->GlobalTransforms.Num() : 0;
-
-	if (!CachePlayback && WorldBounds.GetSphere().W > 1e-5 && NumTransforms > 0)
+FBox UGeometryCollectionComponent::ComputeBounds(const FMatrix& LocalToWorldWithScale) const
+{
+	FBox BoundingBox(ForceInit);
+	if (RestCollection)
 	{
-		return WorldBounds;
-	} 
-	else if (RestCollection)
-	{			
-		const FMatrix LocalToWorldWithScale = LocalToWorldIn.ToMatrixWithScale();
-
-		FBox BoundingBox(ForceInit);
-
 		//Hold on to reference so it doesn't get GC'ed
 		auto HackGeometryCollectionPtr = RestCollection->GetGeometryCollection();
 
@@ -450,17 +437,20 @@ FBoxSphereBounds UGeometryCollectionComponent::CalcBounds(const FTransform& Loca
 			GeometryCollectionAlgo::GlobalMatrices(Transforms, ParentIndices, TmpGlobalMatrices);
 			if (TmpGlobalMatrices.Num() == 0)
 			{
-				return FBoxSphereBounds(ForceInitToZero);
+				BoundingBox = FBox(ForceInitToZero);
 			}
-
-			for (int32 BoxIdx = 0; BoxIdx < NumBoxes; ++BoxIdx)
+			else
 			{
-				const int32 TransformIndex = TransformIndices[BoxIdx];
-
-				if(RestCollection->GetGeometryCollection()->IsGeometry(TransformIndex))
+				for (int32 BoxIdx = 0; BoxIdx < NumBoxes; ++BoxIdx)
 				{
-					BoundingBox += BoundingBoxes[BoxIdx].TransformBy(TmpGlobalMatrices[TransformIndex] * LocalToWorldWithScale);
+					const int32 TransformIndex = TransformIndices[BoxIdx];
+
+					if(RestCollection->GetGeometryCollection()->IsGeometry(TransformIndex))
+					{
+						BoundingBox += BoundingBoxes[BoxIdx].TransformBy(TmpGlobalMatrices[TransformIndex] * LocalToWorldWithScale);
+					}
 				}
+
 			}
 		}
 		else if (bGeometryCollectionSingleThreadedBoundsCalculation)
@@ -504,10 +494,48 @@ FBoxSphereBounds UGeometryCollectionComponent::CalcBounds(const FTransform& Loca
 				}
 			}
 		}
+	}
+	return BoundingBox;
+}
+
+FBoxSphereBounds UGeometryCollectionComponent::CalcBounds(const FTransform& LocalToWorldIn) const
+{	
+	SCOPE_CYCLE_COUNTER(STAT_GCCUpdateBounds);
+
+	// #todo(dmp): hack to make bounds calculation work when we don't have valid physics proxy data.  This will
+	// force bounds calculation.
+
+	const FGeometryCollectionResults* Results = PhysicsProxy ? PhysicsProxy->GetConsumerResultsGT() : nullptr;
+	const int32 NumTransforms = Results ? Results->GlobalTransforms.Num() : 0;
+
+	if (bChaos_GC_CacheComponentSpaceBounds)
+	{
+		bool NeedBoundsUpdate = false;
+		NeedBoundsUpdate |= (ComponentSpaceBounds.GetSphere().W < 1e-5);
+		NeedBoundsUpdate |= CachePlayback;
+		NeedBoundsUpdate |= (NumTransforms > 0);
+		NeedBoundsUpdate |= (DynamicCollection && DynamicCollection->IsDirty());
+		
+		if (NeedBoundsUpdate)
+		{
+			ComponentSpaceBounds = ComputeBounds(FMatrix::Identity);
+		}
+		else
+		{
+			NeedBoundsUpdate = false;
+		}
 
-		return FBoxSphereBounds(BoundingBox);
+		return ComponentSpaceBounds.TransformBy(LocalToWorldIn);
 	}
-	return FBoxSphereBounds(ForceInitToZero);
+
+	// non cached bounds path 
+	if (!CachePlayback && WorldBounds.GetSphere().W > 1e-5 && NumTransforms > 0)
+	{
+		return WorldBounds;
+	} 
+
+	const FMatrix LocalToWorldWithScale = LocalToWorldIn.ToMatrixWithScale();
+	return FBoxSphereBounds(ComputeBounds(LocalToWorldWithScale));
 }
 
 void UGeometryCollectionComponent::CreateRenderState_Concurrent(FRegisterComponentContext* Context)
@@ -3556,19 +3584,19 @@ void UGeometryCollectionComponent::IncrementSleepTimer(float DeltaTime)
 			const bool IsRemovalActive = (MaxSleepTime[TransformIdx] >= 0) && bIsDynamicOrSleeping && bNoParent && !bHasKinematicOrStaticInternalClusterParent;
 			if (IsRemovalActive)
 			{
-				bool bIsSlowMoving = false;
+				bool bIsNotMoving = false;
 				if (RestCollection->bSlowMovingAsSleeping)
 				{
 					const FVector CurrentPosition = DynamicCollection->Transform[TransformIdx].GetTranslation();
 					LastPosition[TransformIdx] = CurrentPosition;
 					
 					const FVector::FReal InstantVelocity = (CurrentPosition-LastPosition[TransformIdx]).Size() / (FVector::FReal)DeltaTime; 
-					bIsSlowMoving = (InstantVelocity < RestCollection->SlowMovingVelocityThreshold);
+					bIsNotMoving = (InstantVelocity < RestCollection->SlowMovingVelocityThreshold);
 				}
 				
 				const bool IsSleeping = (DynamicCollection->DynamicState[TransformIdx] == (int)EObjectStateTypeEnum::Chaos_Object_Sleeping);
 				const bool HasStartedDecaying = (Decay[TransformIdx] > 0); 
-				if (IsSleeping || bIsSlowMoving || HasStartedDecaying)
+				if (IsSleeping || bIsNotMoving || HasStartedDecaying)
 				{
 					SleepTimer[TransformIdx] += DeltaTime;
 				}
diff --git a/Engine/Source/Runtime/Experimental/GeometryCollectionEngine/Public/GeometryCollection/GeometryCollectionComponent.h b/Engine/Source/Runtime/Experimental/GeometryCollectionEngine/Public/GeometryCollection/GeometryCollectionComponent.h
--- a/Engine/Source/Runtime/Experimental/GeometryCollectionEngine/Public/GeometryCollection/GeometryCollectionComponent.h
+++ b/Engine/Source/Runtime/Experimental/GeometryCollectionEngine/Public/GeometryCollection/GeometryCollectionComponent.h
@@ -828,6 +828,7 @@ protected:
 
 	void CalculateLocalBounds();
 	void CalculateGlobalMatrices();
+	FBox ComputeBounds(const FMatrix& LocalToWorldWithScale) const;
 
 	void RegisterForEvents();
 	void UpdateRBCollisionEventRegistration();
@@ -886,7 +887,8 @@ private:
 
 	TArray<FMatrix> GlobalMatrices;
 	FBox LocalBounds;
-	
+
+	mutable FBoxSphereBounds ComponentSpaceBounds;
 	FBoxSphereBounds WorldBounds;
 
 	float CurrentCacheTime;
