Chaos - change GJK iteration limit ensure to a warning and log some shape details to help track it down

diff --git a/Engine/Source/Runtime/Experimental/Chaos/Public/Chaos/Box.h b/Engine/Source/Runtime/Experimental/Chaos/Public/Chaos/Box.h
--- a/Engine/Source/Runtime/Experimental/Chaos/Public/Chaos/Box.h
+++ b/Engine/Source/Runtime/Experimental/Chaos/Public/Chaos/Box.h
@@ -447,7 +447,10 @@ namespace Chaos
 			return TAABB<T, d>::GetRotationOfMass();
 		}
 
-		virtual FString ToString() const { return FString::Printf(TEXT("TAABB Min:%s, Max:%s, Margin:%f"), *Min().ToString(), *Max().ToString(), GetMargin()); }
+		virtual FString ToString() const override
+		{ 
+			return FString::Printf(TEXT("Box: Min: [%s], Max: [%s], Margin: %f"), *Min().ToString(), *Max().ToString(), GetMargin());
+		}
 
 		FORCEINLINE void SerializeImp(FArchive& Ar)
 		{
diff --git a/Engine/Source/Runtime/Experimental/Chaos/Public/Chaos/Capsule.h b/Engine/Source/Runtime/Experimental/Chaos/Public/Chaos/Capsule.h
--- a/Engine/Source/Runtime/Experimental/Chaos/Public/Chaos/Capsule.h
+++ b/Engine/Source/Runtime/Experimental/Chaos/Public/Chaos/Capsule.h
@@ -418,6 +418,11 @@ namespace Chaos
 			}
 		}
 
+		virtual FString ToString() const override
+		{
+			return FString::Printf(TEXT("Capsule: Height: %f Radius: %f"), GetHeight(), GetRadius());
+		}
+
 		virtual TUniquePtr<FImplicitObject> Copy() const override
 		{
 			return TUniquePtr<FImplicitObject>(new FCapsule(*this));
diff --git a/Engine/Source/Runtime/Experimental/Chaos/Public/Chaos/Convex.h b/Engine/Source/Runtime/Experimental/Chaos/Public/Chaos/Convex.h
--- a/Engine/Source/Runtime/Experimental/Chaos/Public/Chaos/Convex.h
+++ b/Engine/Source/Runtime/Experimental/Chaos/Public/Chaos/Convex.h
@@ -720,9 +720,9 @@ namespace Chaos
 			return SupportPoint;
 		}
 
-		virtual FString ToString() const
+		virtual FString ToString() const override
 		{
-			return FString::Printf(TEXT("Convex"));
+			return FString::Printf(TEXT("Convex: Verts: %d, Planes: %d, Margin %f"), NumVertices(), NumPlanes(), GetMargin());
 		}
 
 		const TArray<FVec3Type>& GetVertices() const
diff --git a/Engine/Source/Runtime/Experimental/Chaos/Public/Chaos/GJK.h b/Engine/Source/Runtime/Experimental/Chaos/Public/Chaos/GJK.h
--- a/Engine/Source/Runtime/Experimental/Chaos/Public/Chaos/GJK.h
+++ b/Engine/Source/Runtime/Experimental/Chaos/Public/Chaos/GJK.h
@@ -20,6 +20,30 @@
 
 namespace Chaos
 {
+	// Check the GJK iteration count against a limit to prevent infinite loops. We should never really hit the limit but 
+	// we are seeing it happen in some cases and more often on some platforms than others so for now we have a warning when it hits.
+	// @todo(chaos): track this issue down (see UE-156361)
+	template<typename ConvexTypeA, typename ConvexTypeB>
+	inline bool CheckGJKIterationLimit(const int32 NumIterations, const ConvexTypeA& A, const ConvexTypeB& B)
+	{
+		const int32 MaxIterations = 32;
+		const bool bLimitExceeded = (NumIterations >= MaxIterations);
+
+#if !UE_BUILD_TEST && !UE_BUILD_SHIPPING
+		if (bLimitExceeded)
+		{
+			static bool bLogged = false;
+			if (!bLogged)
+			{
+				UE_LOG(LogChaos, Warning, TEXT("GJK hit iteration limit with shapes:\n    A: %s\n    B: %s"), *A.ToString(), *B.ToString());
+				bLogged = true;
+			}
+		}
+#endif
+
+		return bLimitExceeded;
+	}
+
 
 	// Calculate the margins used for queries based on shape radius, shape margins and shape types
 	template <typename TGeometryA, typename TGeometryB, typename T>
@@ -416,9 +440,11 @@ namespace Chaos
 		const T SeparatedDistance = ThicknessA + ThicknessB + Epsilon;
 		while (!bIsContact && !bIsDegenerate)
 		{
-			if (!ensure(NumIterations++ < 32))	//todo: take this out
+			// If taking too long just stop
+			++NumIterations;
+			if (CheckGJKIterationLimit(NumIterations, A, B))
 			{
-				break;	//if taking too long just stop. This should never happen
+				break;	
 			}
 
 			const TVec3<T> NegV = -V;
@@ -584,9 +610,11 @@ namespace Chaos
 		const T ThicknessB = B.GetMargin();
 		while (!bIsContact && !bIsResult)
 		{
-			if (!ensure(NumIterations++ < 32))	//todo: take this out
+			// If taking too long just stop
+			++NumIterations;
+			if (CheckGJKIterationLimit(NumIterations, A, B))
 			{
-				break;	//if taking too long just stop. This should never happen
+				break;
 			}
 
 			const TVec3<T> NegV = -V;
@@ -767,9 +795,11 @@ namespace Chaos
 		const T SeparatedDistance = ThicknessA + ThicknessB + Epsilon;
 		while (!bIsContact && !bIsDegenerate)
 		{
-			if (!ensure(NumIterations++ < 32))	//todo: take this out
+			// If taking too long just stop
+			++NumIterations;
+			if (CheckGJKIterationLimit(NumIterations, A, B))
 			{
-				break;	//if taking too long just stop. This should never happen
+				break;
 			}
 
 			const TVec3<T> NegV = -V;
@@ -919,9 +949,11 @@ namespace Chaos
 		const T SeparatedDistance = ThicknessA + ThicknessB + Epsilon;
 		while (!bIsContact && !bIsResult)
 		{
-			if (!ensure(NumIterations++ < 32))	//todo: take this out
+			// If taking too long just stop
+			++NumIterations;
+			if (CheckGJKIterationLimit(NumIterations, A, B))
 			{
-				break;	//if taking too long just stop. This should never happen
+				break;
 			}
 
 			const TVec3<T> NegV = -V;
diff --git a/Engine/Source/Runtime/Experimental/Chaos/Public/Chaos/GJKShape.h b/Engine/Source/Runtime/Experimental/Chaos/Public/Chaos/GJKShape.h
--- a/Engine/Source/Runtime/Experimental/Chaos/Public/Chaos/GJKShape.h
+++ b/Engine/Source/Runtime/Experimental/Chaos/Public/Chaos/GJKShape.h
@@ -40,6 +40,11 @@ namespace Chaos
 			return Shape.IsConvex();
 		}
 
+		FString ToString() const
+		{
+			return FString::Printf(TEXT("TGJKShape: %s"), *Shape.ToString());
+		}
+
 		const FImplicitObjectType& Shape;
 	};
 
@@ -84,6 +89,11 @@ namespace Chaos
 			return Shape.IsConvex();
 		}
 
+		FString ToString() const
+		{
+			return FString::Printf(TEXT("TGJKCoreShape: %s, Margin: %f"), *Shape.ToString(), GetMargin());
+		}
+
 		const FImplicitObjectType& Shape;
 		const FReal Margin;
 	};
@@ -111,6 +121,11 @@ namespace Chaos
 			return Radius;
 		}
 
+		FString ToString() const
+		{
+			return FString::Printf(TEXT("FGJKSphereShape: Pos: [%f, %f, %f], Radius: %f"), Pos.X, Pos.Y, Pos.Z, Radius);
+		}
+
 	private:
 		FVec3 Pos;
 		FReal Radius;
diff --git a/Engine/Source/Runtime/Experimental/Chaos/Public/Chaos/ImplicitObjectScaled.h b/Engine/Source/Runtime/Experimental/Chaos/Public/Chaos/ImplicitObjectScaled.h
--- a/Engine/Source/Runtime/Experimental/Chaos/Public/Chaos/ImplicitObjectScaled.h
+++ b/Engine/Source/Runtime/Experimental/Chaos/Public/Chaos/ImplicitObjectScaled.h
@@ -203,6 +203,11 @@ public:
 		return TUniquePtr<FImplicitObject>(CopyHelper(this));
 	}
 
+	virtual FString ToString() const override
+	{
+		return FString::Printf(TEXT("Instanced %s, Margin %f"), *MObject->ToString(), GetMargin());
+	}
+
 	static const TImplicitObjectInstanced<TConcrete>& AsInstancedChecked(const FImplicitObject& Obj)
 	{
 		if(TIsSame<TConcrete,FImplicitObject>::Value)
@@ -1002,6 +1007,11 @@ public:
 		return TUniquePtr<FImplicitObject>(CopyHelper(this));
 	}
 
+	virtual FString ToString() const override
+	{
+		return FString::Printf(TEXT("Scaled %s, Scale: [%f, %f, %f], Margin: %f"), *MObject->ToString(), MScale.X, MScale.Y, MScale.Z, GetMargin());
+	}
+
 private:
 	ObjectType MObject;
 	TSharedPtr<TConcrete, ESPMode::ThreadSafe> MSharedPtrForRefCount; // Temporary solution to force ref counting on trianglemesh from body setup.
diff --git a/Engine/Source/Runtime/Experimental/Chaos/Public/Chaos/Segment.h b/Engine/Source/Runtime/Experimental/Chaos/Public/Chaos/Segment.h
--- a/Engine/Source/Runtime/Experimental/Chaos/Public/Chaos/Segment.h
+++ b/Engine/Source/Runtime/Experimental/Chaos/Public/Chaos/Segment.h
@@ -66,6 +66,11 @@ namespace Chaos
 			MLength = (T)LengthFloat;
 		}
 
+		FString ToString() const
+		{
+			return FString::Printf(TEXT("Segment: Point: [%f, %f, %f], Axis: [%f, %f, %f], Length: %f"), MPoint.X, MPoint.Y, MPoint.Z, MAxis.X, MAxis.Y, MAxis.Z, MLength);
+		}
+
 		FORCEINLINE TAABB<T, 3> BoundingBox() const
 		{
 			TAABB<T,3> Box(MPoint,MPoint);
diff --git a/Engine/Source/Runtime/Experimental/Chaos/Public/Chaos/Sphere.h b/Engine/Source/Runtime/Experimental/Chaos/Public/Chaos/Sphere.h
--- a/Engine/Source/Runtime/Experimental/Chaos/Public/Chaos/Sphere.h
+++ b/Engine/Source/Runtime/Experimental/Chaos/Public/Chaos/Sphere.h
@@ -268,9 +268,9 @@ namespace Chaos
 			return Center; 
 		}
 
-		virtual FString ToString() const
+		virtual FString ToString() const override
 		{
-			return FString::Printf(TEXT("TSphere Center:%s, Radius:%f"), *Center.ToString(), GetRadius());
+			return FString::Printf(TEXT("Sphere: Center:%s, Radius:%f"), *Center.ToString(), GetRadius());
 		}
 
 		FORCEINLINE void SerializeImp(FArchive& Ar)
diff --git a/Engine/Source/Runtime/Experimental/Chaos/Public/Chaos/Triangle.h b/Engine/Source/Runtime/Experimental/Chaos/Public/Chaos/Triangle.h
--- a/Engine/Source/Runtime/Experimental/Chaos/Public/Chaos/Triangle.h
+++ b/Engine/Source/Runtime/Experimental/Chaos/Public/Chaos/Triangle.h
@@ -327,6 +327,10 @@ namespace Chaos
 			return true;
 		}
 
+		FString ToString() const
+		{
+			return FString::Printf(TEXT("Triangle: A: [%f, %f, %f], B: [%f, %f, %f], C: [%f, %f, %f]"), GetVertex(0).X, GetVertex(0).Y, GetVertex(0).Z, GetVertex(1).X, GetVertex(1).Y, GetVertex(1).Z, GetVertex(2).X, GetVertex(2).Y, GetVertex(2).Z);
+		}
 
 		struct FLineIntersectionToleranceProvider
 		{
