Chaos - adjust inertia when adding a ComNudge

diff --git a/Engine/Source/Runtime/Engine/Private/PhysicsEngine/BodyUtils.cpp b/Engine/Source/Runtime/Engine/Private/PhysicsEngine/BodyUtils.cpp
--- a/Engine/Source/Runtime/Engine/Private/PhysicsEngine/BodyUtils.cpp
+++ b/Engine/Source/Runtime/Engine/Private/PhysicsEngine/BodyUtils.cpp
@@ -11,6 +11,13 @@
 
 namespace BodyUtils
 {
+	namespace CVars
+	{
+		bool bPhysicsComNudgeAdjustInertia = true;
+		FAutoConsoleVariableRef CVarPhysicsComNudgeAffectsInertia(TEXT("p.ComNudgeAffectsInertia"), bPhysicsComNudgeAdjustInertia, TEXT(""));
+	}
+
+
 	inline float KgPerM3ToKgPerCm3(float KgPerM3)
 	{
 		//1m = 100cm => 1m^3 = (100cm)^3 = 1000000cm^3
@@ -69,6 +76,15 @@ namespace BodyUtils
 		MassProps.InertiaTensor *= MassRatio;
 		MassProps.CenterOfMass += MassModifierTransform.TransformVector(OwningBodyInstance->COMNudge);
 
+		// If we move the center of mass, we need to update the inertia using parallel-axis theorem. If we don't do this
+		// and the center of mass is moved significantly it can cause jitter (inertia too small for the contact positions)
+		if (CVars::bPhysicsComNudgeAdjustInertia)
+		{
+			MassProps.InertiaTensor.M[0][0] += MassProps.Mass * OwningBodyInstance->COMNudge.X * OwningBodyInstance->COMNudge.X;
+			MassProps.InertiaTensor.M[1][1] += MassProps.Mass * OwningBodyInstance->COMNudge.Y * OwningBodyInstance->COMNudge.Y;
+			MassProps.InertiaTensor.M[2][2] += MassProps.Mass * OwningBodyInstance->COMNudge.Z * OwningBodyInstance->COMNudge.Z;
+		}
+
 		// Scale the inertia tensor by the owning body instance's InertiaTensorScale
 		// NOTE: PhysX scales the inertia by the mass increase we would get from the scale change, even though we 
 		// don't actually scale the mass at all based on InertiaScale. This is non-intuituve. E.g., you may expect 
diff --git a/Engine/Source/Runtime/Experimental/Chaos/Private/Chaos/ChaosDebugDraw.cpp b/Engine/Source/Runtime/Experimental/Chaos/Private/Chaos/ChaosDebugDraw.cpp
--- a/Engine/Source/Runtime/Experimental/Chaos/Private/Chaos/ChaosDebugDraw.cpp
+++ b/Engine/Source/Runtime/Experimental/Chaos/Private/Chaos/ChaosDebugDraw.cpp
@@ -834,13 +834,21 @@ namespace Chaos
 						// Show the raw inertia in black
 						if (!FVec3::IsNearlyEqual(Rigid->InvIConditioning(), TVec3<FRealSingle>(1), UE_SMALL_NUMBER))
 						{
-							const FVec3 EquivalentBoxSize = Utilities::BoxSizeFromInertia(Rigid->I(), Rigid->M());
-							FDebugDrawQueue::GetInstance().DrawDebugBox(PCOM, 0.5f * Settings.InertiaScale * EquivalentBoxSize, QCOM, FColor::Black, false, 0.0f, 0, Settings.LineThickness);
+							FVec3 EquivalentBoxSize, EquivalentBoxCenter;
+							if (Utilities::BoxFromInertia(Rigid->I(), Rigid->M(), EquivalentBoxCenter, EquivalentBoxSize))
+							{
+								const FVec3 BoxCoM = SpaceTransform.TransformPositionNoScale(PCOM + QCOM * EquivalentBoxCenter);
+								FDebugDrawQueue::GetInstance().DrawDebugBox(BoxCoM, 0.5f * Settings.InertiaScale * EquivalentBoxSize, QCOM, FColor::Black, false, 0.0f, 0, Settings.LineThickness);
+							}
 						}
 
 						// Show the inertia used by the solver in magenta
-						const FVec3 EquivalentConditionedBoxSize = Utilities::BoxSizeFromInertia(Rigid->ConditionedI(), Rigid->M());
-						FDebugDrawQueue::GetInstance().DrawDebugBox(PCOM, 0.5f * Settings.InertiaScale * EquivalentConditionedBoxSize, QCOM, FColor::Magenta, false, 0.0f, 0, Settings.LineThickness);
+						FVec3 EquivalentBoxConditionedSize, EquivalentBoxConditionedCenter;
+						if (Utilities::BoxFromInertia(Rigid->ConditionedI(), Rigid->M(), EquivalentBoxConditionedCenter, EquivalentBoxConditionedSize))
+						{
+							const FVec3 BoxConditionedCoM = SpaceTransform.TransformPositionNoScale(PCOM + QCOM * EquivalentBoxConditionedCenter);
+							FDebugDrawQueue::GetInstance().DrawDebugBox(BoxConditionedCoM, 0.5f * Settings.InertiaScale * EquivalentBoxConditionedSize, QCOM, FColor::Magenta, false, 0.0f, 0, Settings.LineThickness);
+						}
 					}
 				}
 			}
diff --git a/Engine/Source/Runtime/Experimental/Chaos/Public/Chaos/Utilities.h b/Engine/Source/Runtime/Experimental/Chaos/Public/Chaos/Utilities.h
--- a/Engine/Source/Runtime/Experimental/Chaos/Public/Chaos/Utilities.h
+++ b/Engine/Source/Runtime/Experimental/Chaos/Public/Chaos/Utilities.h
@@ -485,20 +485,74 @@ namespace Chaos
 		}
 
 		// Compute the box size that would generate the given (diagonal) inertia
-		inline FVec3 BoxSizeFromInertia(const FVec3& Inertia, const FReal Mass)
+		inline bool BoxFromInertia(const FVec3& InInertia, const FReal Mass, FVec3& OutCenter, FVec3& OutSize)
 		{
+			OutSize = FVec3(0);
+			OutCenter = FVec3(0);
+
 			// System of 3 equations in X^2, Y^2, Z^2
-			// Inertia.X = 1/12 M (Size.Y^2 + Size.Z^2)
-			// Inertia.Y = 1/12 M (Size.Z^2 + Size.X^2)
-			// Inertia.Z = 1/12 M (Size.X^2 + Size.Y^2)
+			//		Inertia.X = 1/12 M (Size.Y^2 + Size.Z^2)
+			//		Inertia.Y = 1/12 M (Size.Z^2 + Size.X^2)
+			//		Inertia.Z = 1/12 M (Size.X^2 + Size.Y^2)
+			// Unless the center of mass has been modified, in which case we have
+			//		Inertia.X = 1/12 M (Size.Y^2 + Size.Z^2) + M D.X^2
+			//		Inertia.Y = 1/12 M (Size.Z^2 + Size.X^2) + M D.Y^2
+			//		Inertia.Z = 1/12 M (Size.X^2 + Size.Y^2) + M D.Z^2
+			// Which will not have a unique solution (3 equations in 6 unknowns).
+			// There's no way to know here that the center of mass was modified so we assume it wasn't unless we cannot
+			// solve the equations and then we must make some guesses to recover an equivalent box.
 			if (Mass > 0)
 			{
-				const FVec3 S = Inertia * 12.0f / Mass;
-				const FMatrix33 R = FMatrix33(-0.5f, 0.5f, 0.5f, 0.5f, -0.5f, 0.5f, 0.5f, 0.5f, -0.5f);
-				const FVec3 XYZSq = R * S;
-				return FVec3(FMath::Sqrt(XYZSq.X), FMath::Sqrt(XYZSq.Y), FMath::Sqrt(XYZSq.Z));
+				// RInv is the inverse of the coefficient matrix (0,1,1)(1,0,1)(0,1,1)
+				const FMatrix33 RInv = FMatrix33(-0.5f, 0.5f, 0.5f, 0.5f, -0.5f, 0.5f, 0.5f, 0.5f, -0.5f);
+				FVec3 Inertia = InInertia / Mass;
+				FVec3 XYZSq = RInv * (Inertia * 12.0f);
+
+				// If we have a shape with a modified center of mass that is outside the equivalent box, we will end up with negative
+				// coefficients here and cannot calculate a box equivalent. To do this properly we need to know what center of mass offset was applied.
+				// But lets try to do something anyway so that debug draw shows something...we'll pretend that the shifted inertia component would be equal to the
+				// smallest component in the absense of the shift. This works "correctly" for a uniform shape (e.g., a box), but will be wrong for everything else!
+				// Also, there's a sign problem since shifting the center of mass in the opposite direction would have altered the inertia the same way.
+				// Net result: I'm not sure how useful this is - see if we can do something better one day (maybe store the ComNudge)
+				if (XYZSq.X < 0)
+				{
+					FReal DXSq = (Inertia.X - FMath::Min(Inertia.Y, Inertia.Z));
+					if (DXSq > 0)
+					{
+						OutCenter.X = FMath::Sqrt(DXSq);
+						Inertia.X -= DXSq;
+						XYZSq = RInv * (Inertia * 12.0f);
+					}
+				}
+				if (XYZSq.Y < 0)
+				{
+					FReal DYSq = (Inertia.Y - FMath::Min(Inertia.X, Inertia.Z));
+					if (DYSq > 0)
+					{
+						OutCenter.Y = FMath::Sqrt(DYSq);
+						Inertia.Y -= DYSq;
+						XYZSq = RInv * (Inertia * 12.0f);
+					}
+				}
+				if (XYZSq.Z < 0)
+				{
+					FReal DZSq = (Inertia.Z - FMath::Min(Inertia.X, Inertia.Y));
+					if (DZSq > 0)
+					{
+						OutCenter.Z = FMath::Sqrt(DZSq);
+						Inertia.Z -= DZSq;
+						XYZSq = RInv * (Inertia * 12.0f);
+					}
+				}
+
+				OutSize = FVec3(
+					FMath::Sqrt(FMath::Max(XYZSq.X, FReal(0))),
+					FMath::Sqrt(FMath::Max(XYZSq.Y, FReal(0))),
+					FMath::Sqrt(FMath::Max(XYZSq.Z, FReal(0))));
+
+				return true;
 			}
-			return FVec3(0);
+			return false;
 		}
 
 		// Replacement for FMath::Wrap that works for integers and returns a value in [Begin, End).
