Chaos - Add ToString() to TAABB

diff --git a/Engine/Source/Runtime/Experimental/Chaos/Public/Chaos/AABB.h b/Engine/Source/Runtime/Experimental/Chaos/Public/Chaos/AABB.h
--- a/Engine/Source/Runtime/Experimental/Chaos/Public/Chaos/AABB.h
+++ b/Engine/Source/Runtime/Experimental/Chaos/Public/Chaos/AABB.h
@@ -549,6 +549,11 @@ namespace Chaos
 			return HashCombine(UE::Math::GetTypeHash(MMin), UE::Math::GetTypeHash(MMax));
 		}
 
+		FString ToString() const
+		{
+			return FString::Printf(TEXT("AABB: Min: [%s], Max: [%s]"), *MMin.ToString(), *MMax.ToString());
+		}
+
 		FORCEINLINE PMatrix<FReal, d, d> GetInertiaTensor(const FReal Mass) const { return GetInertiaTensor(Mass, Extents()); }
 		FORCEINLINE static PMatrix<FReal, 3, 3> GetInertiaTensor(const FReal Mass, const TVector<FReal, 3>& Dim)
 		{
