Geometry Collections replicate initial velocity
- fix GC net dormancy, was assumed to start in DORM_Initial, now checks != Awake as it was sometimes DORM_DormantAll.

diff --git a/Engine/Source/Runtime/Experimental/GeometryCollectionEngine/Private/GeometryCollection/GeometryCollectionComponent.cpp b/Engine/Source/Runtime/Experimental/GeometryCollectionEngine/Private/GeometryCollection/GeometryCollectionComponent.cpp
--- a/Engine/Source/Runtime/Experimental/GeometryCollectionEngine/Private/GeometryCollection/GeometryCollectionComponent.cpp
+++ b/Engine/Source/Runtime/Experimental/GeometryCollectionEngine/Private/GeometryCollection/GeometryCollectionComponent.cpp
@@ -1400,7 +1400,7 @@ void UGeometryCollectionComponent::UpdateRepData()
 					const int32 TransformGroupIdx = PhysicsProxy->GetTransformGroupIndexFromHandle(Root);
 					ensureMsgf(TransformGroupIdx >= 0, TEXT("Non-internal cluster should always have a group index"));
 					ensureMsgf(TransformGroupIdx < TNumericLimits<uint16>::Max(), TEXT("Trying to replicate GC with more than 65k pieces. We assumed uint16 would suffice"));
-					RepData.OneOffActivated.Add(TransformGroupIdx);
+					RepData.OneOffActivated.Add(FGeometryCollectionActivatedCluster(TransformGroupIdx, Root->V(), Root->W()));
 					bClustersChanged = true;
 				}
 			}
@@ -1470,7 +1470,7 @@ void UGeometryCollectionComponent::UpdateRepData()
 			MARK_PROPERTY_DIRTY_FROM_NAME(UGeometryCollectionComponent, RepData, this);
 			++RepData.Version;
 
-			if(Owner->NetDormancy == DORM_Initial)
+			if(Owner->NetDormancy != DORM_Awake)
 			{
 				//If net dormancy is Initial it must be for perf reasons, but since a cluster changed we need to replicate down
 				//TODO: set back to dormant when sim goes to sleep
@@ -1511,7 +1511,17 @@ void UGeometryCollectionComponent::ProcessRepData()
 	//First make sure all one off activations have been applied. This ensures our connectivity graph is the same and we have the same clusters as the server
 	for (; OneOffActivatedProcessed < RepData.OneOffActivated.Num(); ++OneOffActivatedProcessed)
 	{
-		FPBDRigidParticleHandle* OneOff = PhysicsProxy->GetParticles()[RepData.OneOffActivated[OneOffActivatedProcessed]];
+		const FGeometryCollectionActivatedCluster& ActivatedCluster = RepData.OneOffActivated[OneOffActivatedProcessed];
+		FPBDRigidParticleHandle* OneOff = PhysicsProxy->GetParticles()[ActivatedCluster.ActivatedIndex];
+
+		// Set initial velocities if not hard snapping
+		if(!bHardSnap)
+		{
+			// TODO: should we invalidate?
+			OneOff->SetV(ActivatedCluster.InitialLinearVelocity);
+			OneOff->SetW(ActivatedCluster.InitialAngularVelocity);
+		}
+
 		RigidClustering.ReleaseClusterParticles(TArray<FPBDRigidParticleHandle*>{ OneOff });
 	}
 
diff --git a/Engine/Source/Runtime/Experimental/GeometryCollectionEngine/Public/GeometryCollection/GeometryCollectionComponent.h b/Engine/Source/Runtime/Experimental/GeometryCollectionEngine/Public/GeometryCollection/GeometryCollectionComponent.h
--- a/Engine/Source/Runtime/Experimental/GeometryCollectionEngine/Public/GeometryCollection/GeometryCollectionComponent.h
+++ b/Engine/Source/Runtime/Experimental/GeometryCollectionEngine/Public/GeometryCollection/GeometryCollectionComponent.h
@@ -345,6 +345,29 @@ struct FGeometryCollectionClusterRep
 	}
 };
 
+struct FGeometryCollectionActivatedCluster
+{
+	FGeometryCollectionActivatedCluster() = default;
+	FGeometryCollectionActivatedCluster(uint16 Index, const FVector& InitialLinearVel, const FVector& InitialAngularVel )
+		: ActivatedIndex(Index)
+		, InitialLinearVelocity(InitialLinearVel)
+		, InitialAngularVelocity(InitialAngularVel)
+	{}
+
+	uint16 ActivatedIndex;
+	FVector_NetQuantize100 InitialLinearVelocity;
+	FVector_NetQuantize100 InitialAngularVelocity;
+};
+
+FORCEINLINE FArchive& operator<<(FArchive& Ar, FGeometryCollectionActivatedCluster& ActivatedCluster)
+{
+	Ar << ActivatedCluster.ActivatedIndex;
+	Ar << ActivatedCluster.InitialLinearVelocity;
+	Ar << ActivatedCluster.InitialAngularVelocity;
+
+	return Ar;
+}
+
 /**
  * Replicated data for a geometry collection when bEnableReplication is true for
  * that component. See UGeomtryCollectionComponent::UpdateRepData
@@ -361,7 +384,7 @@ struct FGeometryCollectionRepData
 	}
 
 	//Array of one off pieces that became activated
-	TArray<uint16> OneOffActivated;
+	TArray<FGeometryCollectionActivatedCluster> OneOffActivated;
 
 	// Array of cluster data requires to synchronize clients
 	TArray<FGeometryCollectionClusterRep> Clusters;
