Chaos - Support deferred collision in QPBD solver (used by RBAN to perform collision detection after first joint iteration)
- transfer RBAN iterations from legacy solver to QPBD solver if they were changed

diff --git a/Engine/Source/Programs/HeadlessChaos/Private/HeadlessChaosTestCCD.cpp b/Engine/Source/Programs/HeadlessChaos/Private/HeadlessChaosTestCCD.cpp
--- a/Engine/Source/Programs/HeadlessChaos/Private/HeadlessChaosTestCCD.cpp
+++ b/Engine/Source/Programs/HeadlessChaos/Private/HeadlessChaosTestCCD.cpp
@@ -86,7 +86,7 @@ namespace ChaosTest
 		const FReal Dt = 1 / 30.0f;
 		const FReal Fps = 1 / Dt;
 		const FReal BoxHalfSize = 50; // cm
-		const FReal InitialSpeed = BoxHalfSize * 5 * Fps; // More than enough to tunnel
+		const FReal InitialSpeed = (30.0f + BoxHalfSize * 5) * Fps; // More than enough to tunnel
 		using TEvolution = FPBDRigidsEvolutionGBF;
 		FParticleUniqueIndicesMultithreaded UniqueIndices;
 		FPBDRigidsSOAs Particles(UniqueIndices);
diff --git a/Engine/Source/Programs/HeadlessChaos/Public/HeadlessChaosCollisionConstraints.h b/Engine/Source/Programs/HeadlessChaos/Public/HeadlessChaosCollisionConstraints.h
--- a/Engine/Source/Programs/HeadlessChaos/Public/HeadlessChaosCollisionConstraints.h
+++ b/Engine/Source/Programs/HeadlessChaos/Public/HeadlessChaosCollisionConstraints.h
@@ -67,10 +67,13 @@ public:
 	void ComputeConstraints(FReal Dt)
 	{
 		CollisionDetector.GetBroadPhase().SetSpatialAcceleration(&SpatialAcceleration);
-		CollisionDetector.GetNarrowPhase().GetContext().bFilteringEnabled = true;
-		CollisionDetector.GetNarrowPhase().GetContext().bDeferUpdate = false;
-		CollisionDetector.GetNarrowPhase().GetContext().bAllowManifolds = true;
+		CollisionDetector.GetNarrowPhase().GetContext().GetSettings().bFilteringEnabled = true;
+		CollisionDetector.GetNarrowPhase().GetContext().GetSettings().bDeferNarrowPhase = false;
+		CollisionDetector.GetNarrowPhase().GetContext().GetSettings().bAllowManifolds = true;
+		CollisionDetector.GetCollisionContainer().SetDetectorSettings(CollisionDetector.GetNarrowPhase().GetContext().GetSettings());
+		
 		CollisionDetector.DetectCollisions(Dt, nullptr);
+		
 		CollisionDetector.GetCollisionContainer().GetConstraintAllocator().SortConstraintsHandles();
 	}
 
diff --git a/Engine/Source/Programs/HeadlessChaos/Public/HeadlessChaosTestUtility.h b/Engine/Source/Programs/HeadlessChaos/Public/HeadlessChaosTestUtility.h
--- a/Engine/Source/Programs/HeadlessChaos/Public/HeadlessChaosTestUtility.h
+++ b/Engine/Source/Programs/HeadlessChaos/Public/HeadlessChaosTestUtility.h
@@ -149,8 +149,9 @@ namespace ChaosTest {
 		// Settings used for unit tests
 		const float CullDistance = 3.0f;
 		Evolution.GetNarrowPhase().SetBoundsExpansion(CullDistance);
-		Evolution.GetCollisionDetector().GetNarrowPhase().GetContext().bDeferUpdate = false;
-		Evolution.GetCollisionDetector().GetNarrowPhase().GetContext().bAllowManifolds = true;
+		Evolution.GetCollisionDetector().GetNarrowPhase().GetContext().GetSettings().bDeferNarrowPhase = false;
+		Evolution.GetCollisionDetector().GetNarrowPhase().GetContext().GetSettings().bAllowManifolds = true;
+		Evolution.GetCollisionConstraints().SetDetectorSettings(Evolution.GetCollisionDetector().GetNarrowPhase().GetContext().GetSettings());
 		Evolution.GetCollisionConstraints().SetSolverType(EConstraintSolverType::QuasiPbd);
 	}
 
diff --git a/Engine/Source/Runtime/Engine/Private/PhysicsEngine/ImmediatePhysics/ImmediatePhysicsChaos/ImmediatePhysicsSimulation_Chaos.cpp b/Engine/Source/Runtime/Engine/Private/PhysicsEngine/ImmediatePhysics/ImmediatePhysicsChaos/ImmediatePhysicsSimulation_Chaos.cpp
--- a/Engine/Source/Runtime/Engine/Private/PhysicsEngine/ImmediatePhysics/ImmediatePhysicsChaos/ImmediatePhysicsSimulation_Chaos.cpp
+++ b/Engine/Source/Runtime/Engine/Private/PhysicsEngine/ImmediatePhysics/ImmediatePhysicsChaos/ImmediatePhysicsSimulation_Chaos.cpp
@@ -381,10 +381,13 @@ namespace ImmediatePhysics_Chaos
 		Implementation->Collisions.DisableHandles();
 		Implementation->Collisions.SetSolverType(EConstraintSolverType::StandardPbd);
 		Implementation->Joints.SetSolverType(EConstraintSolverType::StandardPbd);
-		Implementation->NarrowPhase.GetContext().bFilteringEnabled = false;
-		Implementation->NarrowPhase.GetContext().bDeferUpdate = true;
-		Implementation->NarrowPhase.GetContext().bAllowManifolds = false;
-		Implementation->NarrowPhase.GetContext().bAllowManifoldReuse = false;
+
+		Implementation->NarrowPhase.GetContext().GetSettings().bFilteringEnabled = false;
+		Implementation->NarrowPhase.GetContext().GetSettings().bAllowManifoldReuse = false;
+		Implementation->NarrowPhase.GetContext().GetSettings().bDeferNarrowPhase = (ChaosImmediate_Collision_DeferNarrowPhase != 0);
+		Implementation->NarrowPhase.GetContext().GetSettings().bAllowManifolds = (ChaosImmediate_Collision_UseManifolds != 0);
+		Implementation->NarrowPhase.GetContext().GetSettings().bAllowCCD = false;
+		Implementation->Collisions.SetDetectorSettings(Implementation->NarrowPhase.GetContext().GetSettings());
 	}
 
 	FSimulation::~FSimulation()
@@ -910,21 +913,16 @@ namespace ImmediatePhysics_Chaos
 			Implementation->Collisions.SetRestitutionEnabled(ChaosImmediate_Collision_RestitutionEnabled != 0);
 			Implementation->Collisions.SetRestitutionThreshold(ChaosImmediate_Collision_RestitutionThresholdMultiplier * InGravity.Size());
 			Implementation->Collisions.SetCollisionsEnabled(ChaosImmediate_Collision_Enabled != 0);
+
 			Implementation->CollisionsRule.SetPriority(ChaosImmediate_Collision_Priority);
 
+			Implementation->NarrowPhase.GetContext().GetSettings().bAllowManifoldReuse = false;
+			Implementation->NarrowPhase.GetContext().GetSettings().bDeferNarrowPhase = (ChaosImmediate_Collision_DeferNarrowPhase != 0);
+			Implementation->NarrowPhase.GetContext().GetSettings().bAllowManifolds = (ChaosImmediate_Collision_UseManifolds != 0);
+			Implementation->Collisions.SetDetectorSettings(Implementation->NarrowPhase.GetContext().GetSettings());
+
 			Implementation->Evolution.SetBoundsExtension(ChaosImmediate_Evolution_BoundsExtension);
 
-			Implementation->NarrowPhase.GetContext().bAllowManifoldReuse = false;
-			if (SolverType == EConstraintSolverType::QuasiPbd)
-			{
-				Implementation->NarrowPhase.GetContext().bDeferUpdate = false;
-				Implementation->NarrowPhase.GetContext().bAllowManifolds = true;
-			}
-			else
-			{
-				Implementation->NarrowPhase.GetContext().bDeferUpdate = (ChaosImmediate_Collision_DeferNarrowPhase != 0);
-				Implementation->NarrowPhase.GetContext().bAllowManifolds = (ChaosImmediate_Collision_UseManifolds != 0);
-			}
 
 			if (ChaosImmediate_Evolution_StepTime > 0)
 			{
diff --git a/Engine/Source/Runtime/Engine/Private/PhysicsEngine/PhysicsAsset.cpp b/Engine/Source/Runtime/Engine/Private/PhysicsEngine/PhysicsAsset.cpp
--- a/Engine/Source/Runtime/Engine/Private/PhysicsEngine/PhysicsAsset.cpp
+++ b/Engine/Source/Runtime/Engine/Private/PhysicsEngine/PhysicsAsset.cpp
@@ -218,6 +218,16 @@ void UPhysicsAsset::Serialize(FArchive& Ar)
 	}
 #endif
 
+	// If the legacy iteration defaults were changed, but the new settings were not, pass the legacy changes on to the new
+	if ((SolverSettings.PositionIterations == 6) && ((SolverIterations.SolverIterations != 3) || (SolverIterations.JointIterations != 2)))
+	{
+		SolverSettings.PositionIterations = SolverIterations.SolverIterations * SolverIterations.JointIterations;
+	}
+	if ((SolverSettings.VelocityIterations == 1) && ((SolverIterations.SolverPushOutIterations != 1) || (SolverIterations.JointIterations != 1)))
+	{
+		SolverSettings.VelocityIterations = SolverIterations.SolverPushOutIterations * SolverIterations.JointIterations;
+	}
+
 	Ar.UsingCustomVersion(FFrameworkObjectVersion::GUID);
 	Ar.UsingCustomVersion(FReleaseObjectVersion::GUID);
 }
diff --git a/Engine/Source/Runtime/Experimental/Chaos/Private/Chaos/Collision/ParticlePairMidPhase.cpp b/Engine/Source/Runtime/Experimental/Chaos/Private/Chaos/Collision/ParticlePairMidPhase.cpp
--- a/Engine/Source/Runtime/Experimental/Chaos/Private/Chaos/Collision/ParticlePairMidPhase.cpp
+++ b/Engine/Source/Runtime/Experimental/Chaos/Private/Chaos/Collision/ParticlePairMidPhase.cpp
@@ -242,7 +242,7 @@ namespace Chaos
 		PHYSICS_CSV_SCOPED_EXPENSIVE(PhysicsVerbose, NarrowPhase_CreateConstraint);
 		check(Constraint == nullptr);
 
-		Constraint = CreateShapePairConstraint(Particle0, Shape0, Particle1, Shape1, CullDistance, ShapePairType, Context.bAllowManifolds);
+		Constraint = CreateShapePairConstraint(Particle0, Shape0, Particle1, Shape1, CullDistance, ShapePairType, Context.GetSettings().bAllowManifolds);
 
 		Constraint->GetContainerCookie().MidPhase = &MidPhase;
 		Constraint->GetContainerCookie().bIsMultiShapePair = false;
@@ -293,7 +293,7 @@ namespace Chaos
 			}
 
 			bool bWasManifoldRestored = false;
-			if (Context.bAllowManifoldReuse && Flags.bEnableManifoldUpdate && bWasUpdatedLastTick)
+			if (Context.GetSettings().bAllowManifoldReuse && Flags.bEnableManifoldUpdate && bWasUpdatedLastTick)
 			{
 				// Update the existing manifold. We can re-use as-is if none of the points have moved much and the bodies have not moved much
 				// NOTE: this can succeed in "restoring" even if we have no manifold points
@@ -311,7 +311,7 @@ namespace Chaos
 				// We will be updating the manifold, if only partially, so update the restore comparison transforms
 				Constraint->SetLastShapeWorldTransforms(ShapeWorldTransform0, ShapeWorldTransform1);
 
-				if (!Context.bDeferUpdate)
+				if (!Context.GetSettings().bDeferNarrowPhase)
 				{
 					// Run the narrow phase
 					Collisions::UpdateConstraint(*Constraint.Get(), ShapeWorldTransform0, ShapeWorldTransform1, Dt);
@@ -320,7 +320,7 @@ namespace Chaos
 
 			// If we have a valid contact, add it to the active list
 			// We also add it to the active list if collision detection is deferred (which is if per-iteration collision detection is enabled like with RBAN)
-			if ((Constraint->GetPhi() <= CullDistance) || Context.bDeferUpdate)
+			if ((Constraint->GetPhi() <= CullDistance) || Context.GetSettings().bDeferNarrowPhase)
 			{
 				if (MidPhase.GetCollisionAllocator().ActivateConstraint(Constraint.Get()))
 				{
@@ -481,7 +481,7 @@ namespace Chaos
 		const FRigidTransform3 ParticleWorldTransform1 = FParticleUtilities::GetActorWorldTransform(P1);
 
 		FCollisionContext LocalContext = Context;
-		LocalContext.CollisionAllocator = this;
+		LocalContext.MultiShapeCollisionDetector = this;
 
 		Collisions::ConstructConstraints(
 			Particle0,
@@ -829,7 +829,7 @@ namespace Chaos
 			FVec3 DirCCD = FVec3(0);
 			const FVec3 DeltaX0 = (ConstParticle0->ObjectState() == EObjectStateType::Kinematic) ? (ConstParticle0->V() * Dt) : (ConstParticle0->P() - ConstParticle0->X());
 			const FVec3 DeltaX1 = (ConstParticle1->ObjectState() == EObjectStateType::Kinematic) ? (ConstParticle1->V() * Dt) : (ConstParticle1->P() - ConstParticle1->X());
-			const bool bUseCCD = Collisions::ShouldUseCCD(Particle0, DeltaX0, Particle1, DeltaX1, DirCCD, LengthCCD, false);
+			const bool bUseCCD = Collisions::ShouldUseCCD(Particle0, DeltaX0, Particle1, DeltaX1, DirCCD, LengthCCD);
 
 			return bUseCCD;
 		}
@@ -873,7 +873,7 @@ namespace Chaos
 		FReal CullDistance = InCullDistance * CullDistanceScale;
 
 		// Enable CCD?
-		const bool bUseCCD = Flags.bIsCCD && ShouldEnableCCD(Dt);
+		const bool bUseCCD = Context.GetSettings().bAllowCCD && Flags.bIsCCD && ShouldEnableCCD(Dt);
 		if (bUseCCD)
 		{
 			const FReal VMax = (FConstGenericParticleHandle(GetParticle0())->V() - FConstGenericParticleHandle(GetParticle1())->V()).GetAbsMax();
@@ -882,7 +882,7 @@ namespace Chaos
 		
 		// Run collision detection on all potentially colliding shape pairs
 		NumActiveConstraints = 0;
-		if (bUseCCD) // Generate CCD constraints as long as AABBs overlap
+		if (bUseCCD)
 		{
 			for (FSingleShapePairCollisionDetector& ShapePair : ShapePairDetectors)
 			{
diff --git a/Engine/Source/Runtime/Experimental/Chaos/Private/Chaos/Collision/SolverCollisionContainer.cpp b/Engine/Source/Runtime/Experimental/Chaos/Private/Chaos/Collision/SolverCollisionContainer.cpp
--- a/Engine/Source/Runtime/Experimental/Chaos/Private/Chaos/Collision/SolverCollisionContainer.cpp
+++ b/Engine/Source/Runtime/Experimental/Chaos/Private/Chaos/Collision/SolverCollisionContainer.cpp
@@ -31,7 +31,6 @@ namespace Chaos
 		FRealSingle Chaos_PBDCollisionSolver_AutoStiffness_MassRatio2 = 0;
 		FAutoConsoleVariableRef CVarChaosPBDCollisionSolverAutoStiffnessMassRatio1(TEXT("p.Chaos.PBDCollisionSolver.AutoStiffness.MassRatio1"), Chaos_PBDCollisionSolver_AutoStiffness_MassRatio1, TEXT(""));
 		FAutoConsoleVariableRef CVarChaosPBDCollisionSolverAutoStiffnessMassRatio2(TEXT("p.Chaos.PBDCollisionSolver.AutoStiffness.MassRatio2"), Chaos_PBDCollisionSolver_AutoStiffness_MassRatio2, TEXT(""));
-
 	}
 	using namespace CVars;
 
@@ -138,19 +137,12 @@ namespace Chaos
 			return Constraint->GetStiffness();
 		}
 
-		/**
-		 * @brief Initialize the data required for the solver and bind to the bodies
-		*/
-		void GatherInput(
-			const FReal Dt,
-			FPBDCollisionConstraint& InConstraint,
+		void BindBodies(
 			const int32 Particle0Level,
 			const int32 Particle1Level,
 			FSolverBodyContainer& SolverBodyContainer,
 			const FPBDCollisionSolverSettings& SolverSettings)
 		{
-			bIsIncremental = Constraint->GetUseIncrementalCollisionDetection();
-
 			// Find the solver bodies for the particles we constrain. This will add them to the container
 			// if they aren't there already, and ensure that they are populated with the latest data.
 			FSolverBody* Body0 = SolverBodyContainer.FindOrAdd(Constraint->GetParticle0(), Dt);
@@ -158,6 +150,11 @@ namespace Chaos
 
 			Body0->SetLevel(Particle0Level);
 			Body1->SetLevel(Particle1Level);
+		
+			Solver.SetSolverBodies(*Body0, *Body1);
+
+			// We should try to remove this - the Constraint should not need to know about solver objects
+			Constraint->SetSolverBodies(Body0, Body1);
 
 			// Friction values. Static and Dynamic friction are applied in the position solve for most shapes.
 			// For quadratic shapes, we run dynamic friction in the velocity solve for better rolling behaviour.
@@ -190,24 +187,18 @@ namespace Chaos
 			const FReal SolverStiffness = CalculateSolverStiffness(Body0, Body1, Chaos_PBDCollisionSolver_AutoStiffness_MassRatio1, Chaos_PBDCollisionSolver_AutoStiffness_MassRatio2);
 			Solver.SetStiffness(FSolverReal(SolverStiffness));
 
-			Solver.SetSolverBodies(*Body0, *Body1);
 			Solver.SolverBody0().SetInvMScale(Constraint->GetInvMassScale0());
 			Solver.SolverBody0().SetInvIScale(Constraint->GetInvInertiaScale0());
 			Solver.SolverBody1().SetInvMScale(Constraint->GetInvMassScale1());
 			Solver.SolverBody1().SetInvIScale(Constraint->GetInvInertiaScale1());
-
-			GatherManifoldPoints(Dt, Body0, Body1);
-
-			// We should try to remove this - the Constraint should not need to know about solver objects
-			Constraint->SetSolverBodies(Body0, Body1);
 		}
 
-		void GatherManifoldPoints(
-			const FReal InDt,
-			const FSolverBody* Body0,
-			const FSolverBody* Body1)
+		void UpdateManifoldPoints(
+			const FReal InDt)
 		{
 			FSolverReal Dt = FSolverReal(InDt);
+			const FConstraintSolverBody& Body0 = Solver.SolverBody0();
+			const FConstraintSolverBody& Body1 = Solver.SolverBody1();
 
 			// We handle incremental manifolds by just collecting any new contacts
 			const int32 BeginPointIndex = Solver.NumManifoldPoints();
@@ -233,8 +224,8 @@ namespace Chaos
 				const FVec3 WorldContactPoint = FReal(0.5) * (WorldContactPoint0 + WorldContactPoint1);
 
 				const FSolverVec3 WorldContactNormal = FSolverVec3(ShapeWorldTransform1.TransformVectorNoScale(ManifoldPoint.ContactPoint.ShapeContactNormal));
-				const FSolverVec3 RelativeContactPosition0 = FSolverVec3(WorldContactPoint - Body0->P());
-				const FSolverVec3 RelativeContactPosition1 = FSolverVec3(WorldContactPoint - Body1->P());
+				const FSolverVec3 RelativeContactPosition0 = FSolverVec3(WorldContactPoint - Body0.P());
+				const FSolverVec3 RelativeContactPosition1 = FSolverVec3(WorldContactPoint - Body1.P());
 				const FSolverReal TargetPhi = FSolverReal(ManifoldPoint.TargetPhi);
 
 				// If we have contact data from a previous tick, use it to calculate the lateral position delta we need
@@ -253,8 +244,8 @@ namespace Chaos
 				}
 				else
 				{
-					const FSolverVec3 ContactVel0 = Body0->V() + FSolverVec3::CrossProduct(Body0->W(), RelativeContactPosition0);
-					const FSolverVec3 ContactVel1 = Body1->V() + FSolverVec3::CrossProduct(Body1->W(), RelativeContactPosition1);
+					const FSolverVec3 ContactVel0 = Body0.V() + FSolverVec3::CrossProduct(Body0.W(), RelativeContactPosition0);
+					const FSolverVec3 ContactVel1 = Body1.V() + FSolverVec3::CrossProduct(Body1.W(), RelativeContactPosition1);
 					const FSolverVec3 ContactVel = ContactVel0 - ContactVel1;
 					WorldFrictionDelta = ContactVel * Dt;
 
@@ -373,6 +364,7 @@ namespace Chaos
 	FPBDCollisionSolverContainer::FPBDCollisionSolverContainer()
 		: FConstraintSolverContainer()
 		, bRequiresIncrementalCollisionDetection(false)
+		, bDeferredCollisionDetection(false)
 	{
 	}
 
@@ -413,8 +405,14 @@ namespace Chaos
 		check(SolverIndex < CollisionSolvers.Num());
 
 		FPBDCollisionSolverAdapter& CollisionSolver = CollisionSolvers[SolverIndex];
+		check(CollisionSolver.GetConstraint() == &Constraint);
 
-		CollisionSolver.GatherInput(Dt, Constraint, Particle0Level, Particle1Level, SolverBodyContainer, SolverSettings);
+		CollisionSolver.BindBodies(Particle0Level, Particle1Level, SolverBodyContainer, SolverSettings);
+
+		if (!bDeferredCollisionDetection)
+		{
+			CollisionSolver.UpdateManifoldPoints(Dt);
+		}
 
 		bRequiresIncrementalCollisionDetection |= CollisionSolver.IsIncrementalManifold();
 	}
@@ -490,12 +488,18 @@ namespace Chaos
 		// Adjust max pushout to attempt to make it iteration count independent
 		const FReal MaxPushOut = (SolverSettings.MaxPushOutVelocity > 0) ? (SolverSettings.MaxPushOutVelocity * Dt) / FReal(NumIts) : 0;
 
-		// Apply the position correction
-		if (bRequiresIncrementalCollisionDetection)
+		// We run collision detection here under two conditions (normally it is run after Integration and before the constraint solver phase):
+		// 1) When deferring collision detection until the solver phase for better joint-collision behaviour (RBAN). In this case, we only do this on the first iteration.
+		// 2) When using incremental manifolds, where we may add/replace manifold points every iteration.
+		const bool bDeferredCollisions = bDeferredCollisionDetection && (It == 0);
+		const bool bIncrementalManifolds = bRequiresIncrementalCollisionDetection;
+		if (bIncrementalManifolds || bDeferredCollisions)
 		{
-			return SolvePositionIncrementalImpl(Dt, BeginIndex, EndIndex, MaxPushOut, bApplyStaticFriction);
+			UpdateCollisions(Dt, BeginIndex, EndIndex);
 		}
-		else if (bApplyStaticFriction)
+
+		// Apply the position correction
+		if (bApplyStaticFriction)
 		{
 			return SolvePositionWithFrictionImpl(Dt, BeginIndex, EndIndex, MaxPushOut, bParallel);
 		}
@@ -505,34 +509,6 @@ namespace Chaos
 		}
 	}
 
-	// Solve position including support for incremental collision detection
-	bool FPBDCollisionSolverContainer::SolvePositionIncrementalImpl(const FReal InDt, const int32 BeginIndex, const int32 EndIndex, const FReal InMaxPushOut, const bool bApplyStaticFriction)
-	{
-		const FSolverReal Dt = FSolverReal(InDt);
-		const FSolverReal MaxPushOut = FSolverReal(InMaxPushOut);
-
-		bool bNeedsAnotherIteration = false;
-		for (int32 SolverIndex = BeginIndex; SolverIndex < EndIndex; ++SolverIndex)
-		{
-			FPBDCollisionSolverAdapter& CollisionSolver = CollisionSolvers[SolverIndex];
-			if (CollisionSolver.IsIncrementalManifold())
-			{
-				Collisions::Update(*CollisionSolver.GetConstraint(), Dt);
-				CollisionSolver.GatherManifoldPoints(Dt, &CollisionSolver.GetSolver().SolverBody0().SolverBody(), &CollisionSolver.GetSolver().SolverBody1().SolverBody());
-			}
-			if (bApplyStaticFriction)
-			{
-				bNeedsAnotherIteration |= CollisionSolver.GetSolver().SolvePositionWithFriction(Dt, MaxPushOut);
-			}
-			else
-			{
-				bNeedsAnotherIteration |= CollisionSolver.GetSolver().SolvePositionNoFriction(Dt, MaxPushOut);
-			}
-
-		}
-		return bNeedsAnotherIteration;
-	}
-
 	// Solve position with friction (last few iterations each tick)
 	bool FPBDCollisionSolverContainer::SolvePositionWithFrictionImpl(const FReal InDt, const int32 BeginIndex, const int32 EndIndex, const FReal InMaxPushOut, const bool bParallel)
 	{
@@ -629,4 +605,20 @@ namespace Chaos
 		}, Chaos::LargeBatchSize);
 	}
 
+	void FPBDCollisionSolverContainer::UpdateCollisions(const FReal InDt, const int32 BeginIndex, const int32 EndIndex)
+	{
+		const FSolverReal Dt = FSolverReal(InDt);
+
+		bool bNeedsAnotherIteration = false;
+		for (int32 SolverIndex = BeginIndex; SolverIndex < EndIndex; ++SolverIndex)
+		{
+			FPBDCollisionSolverAdapter& CollisionSolver = CollisionSolvers[SolverIndex];
+			if (CollisionSolver.IsIncrementalManifold() || bDeferredCollisionDetection)
+			{
+				Collisions::Update(*CollisionSolver.GetConstraint(), Dt);
+				CollisionSolver.UpdateManifoldPoints(Dt);
+			}
+		}
+	}
+
 }
diff --git a/Engine/Source/Runtime/Experimental/Chaos/Private/Chaos/CollisionResolution.cpp b/Engine/Source/Runtime/Experimental/Chaos/Private/Chaos/CollisionResolution.cpp
--- a/Engine/Source/Runtime/Experimental/Chaos/Private/Chaos/CollisionResolution.cpp
+++ b/Engine/Source/Runtime/Experimental/Chaos/Private/Chaos/CollisionResolution.cpp
@@ -81,24 +81,9 @@ FAutoConsoleVariableRef CVarCCDOnlyConsiderDynamicStatic(TEXT("p.Chaos.CCD.OnlyC
 int32 ConstraintsDetailedStats = 0;
 FAutoConsoleVariableRef CVarConstraintsDetailedStats(TEXT("p.Chaos.Constraints.DetailedStats"), ConstraintsDetailedStats, TEXT("When set to 1, will enable more detailed stats."));
 
-int32 AlwaysAddSweptConstraints = 0;
-FAutoConsoleVariableRef CVarAlwaysAddSweptConstraints(TEXT("p.Chaos.Constraints.AlwaysAddSweptConstraints"), AlwaysAddSweptConstraints, TEXT("Since GJKContactPointSwept returns infinity for it's contact data when not hitting anything, some contacts are discarded prematurely. This flag will cause contact points considered for sweeps to never be discarded."));
-
-int32 CCDAllowForceDisable = 1;
-FAutoConsoleVariableRef CVarCCDAllowForceDisable(TEXT("p.Chaos.CCD.AllowForceDisable"), CCDAllowForceDisable, TEXT("Allow force disable CCD."));
-
-int32 CCDManualForceDisable = 0;
-FAutoConsoleVariableRef CVarCCDManualForceDisable(TEXT("p.Chaos.CCD.ManualForceDisable"), CCDManualForceDisable, TEXT("For debugging, manually disable CCD always."));
-
 bool bChaos_Collision_AllowLevelsetManifolds = true;
 FAutoConsoleVariableRef CVarChaosCollisionAllowLevelsetManifolds(TEXT("p.Chaos.Collision.AllowLevelsetManifolds"), bChaos_Collision_AllowLevelsetManifolds, TEXT("Use incremental manifolds for levelset-levelset collision. This does not work well atm - too much rotation in the small pieces"));
 
-bool CCDUseInitialRotationForSweptUpdate = false;
-FAutoConsoleVariableRef CVarCCDUseInitialRotationForSweptUpdate(TEXT("p.Chaos.CCD.UseInitialRotationForSweptUpdate"), CCDUseInitialRotationForSweptUpdate, TEXT("Experimental : Use R instead of Q for SweptUpdate. Will lead to incorrect results in some cases."));
-
-bool CCDAlwaysSweepRemainingDT = false;
-FAutoConsoleVariableRef CVarCCDAlwaysSweepRemainingDT(TEXT("p.Chaos.CCD.AlwaysSweepRemainingDT"), CCDAlwaysSweepRemainingDT, TEXT("Even if we are out of iterations, do the sweep for remaining dt instead of adjusting V. Could cause tunnelling in some cases, but prevents loss of momentum."));
-
 bool CCDNoCullAllShapePairs = true;
 FAutoConsoleVariableRef CVarCCDNoCullAllShapePairs(TEXT("p.Chaos.CCD.NoCullAllShapePairs"), CCDNoCullAllShapePairs, TEXT("Whether to cull contacts early based on phi for sweeps for all shape pairs (not just convex convex)."));
 
@@ -146,18 +131,8 @@ namespace Chaos
 		* Now think about two boxes. If P-X <= 0.5 * BoundingBox().Extents().Min() both boxes, it is impossible for one box to tunnel through the other (it is still possible for a box's corner to move through another box's corner).
 		* So far this is sufficient for us to determine CCD constraints for our primitives. But the caveat is that for arbitrary objects (think about a rod that is not axis aligned), min of a bounding box extent can be too large and tunneling could still happen if P-X falls below this threshold.
 		*/
-		bool ShouldUseCCD(const FGeometryParticleHandle* Particle0, const FVec3& DeltaX0, const FGeometryParticleHandle* Particle1, const FVec3& DeltaX1, FVec3& Dir, FReal& Length, const bool bForceDisableCCD)
+		bool ShouldUseCCD(const FGeometryParticleHandle* Particle0, const FVec3& DeltaX0, const FGeometryParticleHandle* Particle1, const FVec3& DeltaX1, FVec3& Dir, FReal& Length)
 		{
-			if (CCDManualForceDisable)
-			{
-				return false;
-			}
-
-			if (CCDAllowForceDisable > 0 && bForceDisableCCD)
-			{
-				return false;
-			}
-
 			if (CCDOnlyConsiderDynamicStatic > 0 && Particle1->ObjectState() != EObjectStateType::Static)
 			{
 				return false;
@@ -446,7 +421,7 @@ namespace Chaos
 			const TBox<FReal, 3>* Object1 = Implicit1->template GetObject<const TBox<FReal, 3> >();
 			if (ensure(Object0 && Object1))
 			{
-				FPBDCollisionConstraint* Constraint = Context.CollisionAllocator->FindOrCreateConstraint(Particle0, Implicit0, Shape0, nullptr, LocalTransform0, Particle1, Implicit1, Shape1, nullptr, LocalTransform1, CullDistance, EContactShapesType::BoxBox, Context.bAllowManifolds);
+				FPBDCollisionConstraint* Constraint = Context.MultiShapeCollisionDetector->FindOrCreateConstraint(Particle0, Implicit0, Shape0, nullptr, LocalTransform0, Particle1, Implicit1, Shape1, nullptr, LocalTransform1, CullDistance, EContactShapesType::BoxBox, Context.GetSettings().bAllowManifolds);
 				if (T_TRAITS::bImmediateUpdate && (Constraint != nullptr))
 				{
 					const FRigidTransform3 WorldTransform0 = LocalTransform0 * ParticleWorldTransform0;
@@ -484,7 +459,7 @@ namespace Chaos
 			if (ensure(Object0 && Object1))
 			{
 				// @todo(chaos): one-shot manifold
-				FPBDCollisionConstraint* Constraint = Context.CollisionAllocator->FindOrCreateConstraint(Particle0, Implicit0, Shape0, nullptr, LocalTransform0, Particle1, Implicit1, Shape1, nullptr, LocalTransform1, CullDistance, EContactShapesType::BoxHeightField, Context.bAllowManifolds);
+				FPBDCollisionConstraint* Constraint = Context.MultiShapeCollisionDetector->FindOrCreateConstraint(Particle0, Implicit0, Shape0, nullptr, LocalTransform0, Particle1, Implicit1, Shape1, nullptr, LocalTransform1, CullDistance, EContactShapesType::BoxHeightField, Context.GetSettings().bAllowManifolds);
 				if (T_TRAITS::bImmediateUpdate && (Constraint != nullptr))
 				{
 					FRigidTransform3 WorldTransform0 = LocalTransform0 * ParticleWorldTransform0;
@@ -577,7 +552,7 @@ namespace Chaos
 			const TPlane<FReal, 3>* Object1 = Implicit1->template GetObject<const TPlane<FReal, 3> >();
 			if (ensure(Object0 && Object1))
 			{
-				FPBDCollisionConstraint* Constraint = Context.CollisionAllocator->FindOrCreateConstraint(Particle0, Implicit0, Shape0, nullptr, LocalTransform0, Particle1, Implicit1, Shape1, nullptr, LocalTransform1, CullDistance, EContactShapesType::BoxPlane, Context.bAllowManifolds);
+				FPBDCollisionConstraint* Constraint = Context.MultiShapeCollisionDetector->FindOrCreateConstraint(Particle0, Implicit0, Shape0, nullptr, LocalTransform0, Particle1, Implicit1, Shape1, nullptr, LocalTransform1, CullDistance, EContactShapesType::BoxPlane, Context.GetSettings().bAllowManifolds);
 				if (T_TRAITS::bImmediateUpdate && (Constraint != nullptr))
 				{
 					FRigidTransform3 WorldTransform0 = LocalTransform0 * ParticleWorldTransform0;
@@ -614,7 +589,7 @@ namespace Chaos
 				if (const TImplicitObjectScaled<FTriangleMeshImplicitObject>* ScaledTriangleMesh = Implicit1->template GetObject<const TImplicitObjectScaled<FTriangleMeshImplicitObject>>())
 				{
 					// @todo(chaos): one-shot manifold
-					FPBDCollisionConstraint* Constraint = Context.CollisionAllocator->FindOrCreateConstraint(Particle0, Implicit0, Shape0, nullptr, LocalTransform0, Particle1, Implicit1, Shape1, nullptr, LocalTransform1, CullDistance, EContactShapesType::BoxTriMesh, Context.bAllowManifolds);
+					FPBDCollisionConstraint* Constraint = Context.MultiShapeCollisionDetector->FindOrCreateConstraint(Particle0, Implicit0, Shape0, nullptr, LocalTransform0, Particle1, Implicit1, Shape1, nullptr, LocalTransform1, CullDistance, EContactShapesType::BoxTriMesh, Context.GetSettings().bAllowManifolds);
 					if (T_TRAITS::bImmediateUpdate && (Constraint != nullptr))
 					{
 						FRigidTransform3 WorldTransform0 = LocalTransform0 * ParticleWorldTransform0;
@@ -625,7 +600,7 @@ namespace Chaos
 				else if (const FTriangleMeshImplicitObject* TriangleMesh = Implicit1->template GetObject<const FTriangleMeshImplicitObject>())
 				{
 					// @todo(chaos): one-shot manifold
-					FPBDCollisionConstraint* Constraint = Context.CollisionAllocator->FindOrCreateConstraint(Particle0, Implicit0, Shape0, nullptr, LocalTransform0, Particle1, Implicit1, Shape1, nullptr, LocalTransform1, CullDistance, EContactShapesType::BoxTriMesh, Context.bAllowManifolds);
+					FPBDCollisionConstraint* Constraint = Context.MultiShapeCollisionDetector->FindOrCreateConstraint(Particle0, Implicit0, Shape0, nullptr, LocalTransform0, Particle1, Implicit1, Shape1, nullptr, LocalTransform1, CullDistance, EContactShapesType::BoxTriMesh, Context.GetSettings().bAllowManifolds);
 					if (T_TRAITS::bImmediateUpdate && (Constraint != nullptr))
 					{
 						FRigidTransform3 WorldTransform0 = LocalTransform0 * ParticleWorldTransform0;
@@ -663,7 +638,7 @@ namespace Chaos
 			const TSphere<FReal, 3>* Object1 = Implicit1->template GetObject<const TSphere<FReal, 3> >();
 			if (ensure(Object0 && Object1))
 			{
-				FPBDCollisionConstraint* Constraint = Context.CollisionAllocator->FindOrCreateConstraint(Particle0, Implicit0, Shape0, nullptr, LocalTransform0, Particle1, Implicit1, Shape1, nullptr, LocalTransform1, CullDistance, EContactShapesType::SphereSphere, Context.bAllowManifolds);
+				FPBDCollisionConstraint* Constraint = Context.MultiShapeCollisionDetector->FindOrCreateConstraint(Particle0, Implicit0, Shape0, nullptr, LocalTransform0, Particle1, Implicit1, Shape1, nullptr, LocalTransform1, CullDistance, EContactShapesType::SphereSphere, Context.GetSettings().bAllowManifolds);
 				if (T_TRAITS::bImmediateUpdate && (Constraint != nullptr))
 				{
 					FRigidTransform3 WorldTransform0 = LocalTransform0 * ParticleWorldTransform0;
@@ -705,7 +680,7 @@ namespace Chaos
 			const FHeightField* Object1 = Implicit1->template GetObject<const FHeightField >();
 			if (ensure(Object0 && Object1))
 			{
-				FPBDCollisionConstraint* Constraint = Context.CollisionAllocator->FindOrCreateConstraint(Particle0, Implicit0, Shape0, nullptr, LocalTransform0, Particle1, Implicit1, Shape1, nullptr, LocalTransform1, CullDistance, EContactShapesType::SphereHeightField, Context.bAllowManifolds);
+				FPBDCollisionConstraint* Constraint = Context.MultiShapeCollisionDetector->FindOrCreateConstraint(Particle0, Implicit0, Shape0, nullptr, LocalTransform0, Particle1, Implicit1, Shape1, nullptr, LocalTransform1, CullDistance, EContactShapesType::SphereHeightField, Context.GetSettings().bAllowManifolds);
 				if (T_TRAITS::bImmediateUpdate && (Constraint != nullptr))
 				{
 					FRigidTransform3 WorldTransform0 = LocalTransform0 * ParticleWorldTransform0;
@@ -721,7 +696,7 @@ namespace Chaos
 			const FHeightField* Object1 = Implicit1->template GetObject<const FHeightField>();
 			if (ensure(Object0 && Object1))
 			{
-				FPBDCollisionConstraint* Constraint = Context.CollisionAllocator->FindOrCreateSweptConstraint(Particle0, Implicit0, Shape0, nullptr, LocalTransform0, Particle1, Implicit1, Shape1, nullptr, LocalTransform1, CullDistance, EContactShapesType::SphereHeightField);
+				FPBDCollisionConstraint* Constraint = Context.MultiShapeCollisionDetector->FindOrCreateSweptConstraint(Particle0, Implicit0, Shape0, nullptr, LocalTransform0, Particle1, Implicit1, Shape1, nullptr, LocalTransform1, CullDistance, EContactShapesType::SphereHeightField);
 				FRigidTransform3 WorldTransformX0 = LocalTransform0 * FRigidTransform3(Particle0->X(), Particle0->R());
 				FRigidTransform3 WorldTransform1 = LocalTransform1 * ParticleWorldTransform1;
 				UpdateSphereHeightFieldConstraintSwept(Particle0, *Object0, WorldTransformX0, *Object1, WorldTransform1, Dir, Length, Dt, *Constraint);
@@ -752,7 +727,7 @@ namespace Chaos
 			const TPlane<FReal, 3>* Object1 = Implicit1->template GetObject<const TPlane<FReal, 3> >();
 			if (ensure(Object0 && Object1))
 			{
-				FPBDCollisionConstraint* Constraint = Context.CollisionAllocator->FindOrCreateConstraint(Particle0, Implicit0, Shape0, nullptr, LocalTransform0, Particle1, Implicit1, Shape1, nullptr, LocalTransform1, CullDistance, EContactShapesType::SpherePlane, Context.bAllowManifolds);
+				FPBDCollisionConstraint* Constraint = Context.MultiShapeCollisionDetector->FindOrCreateConstraint(Particle0, Implicit0, Shape0, nullptr, LocalTransform0, Particle1, Implicit1, Shape1, nullptr, LocalTransform1, CullDistance, EContactShapesType::SpherePlane, Context.GetSettings().bAllowManifolds);
 				if (T_TRAITS::bImmediateUpdate && (Constraint != nullptr))
 				{
 					FRigidTransform3 WorldTransform0 = LocalTransform0 * ParticleWorldTransform0;
@@ -787,7 +762,7 @@ namespace Chaos
 			const TBox<FReal, 3>* Object1 = Implicit1->template GetObject<const TBox<FReal, 3> >();
 			if (ensure(Object0 && Object1))
 			{
-				FPBDCollisionConstraint* Constraint = Context.CollisionAllocator->FindOrCreateConstraint(Particle0, Implicit0, Shape0, nullptr, LocalTransform0, Particle1, Implicit1, Shape1, nullptr, LocalTransform1, CullDistance, EContactShapesType::SphereBox, Context.bAllowManifolds);
+				FPBDCollisionConstraint* Constraint = Context.MultiShapeCollisionDetector->FindOrCreateConstraint(Particle0, Implicit0, Shape0, nullptr, LocalTransform0, Particle1, Implicit1, Shape1, nullptr, LocalTransform1, CullDistance, EContactShapesType::SphereBox, Context.GetSettings().bAllowManifolds);
 				if (T_TRAITS::bImmediateUpdate && (Constraint != nullptr))
 				{
 					FRigidTransform3 WorldTransform0 = LocalTransform0 * ParticleWorldTransform0;
@@ -822,7 +797,7 @@ namespace Chaos
 			const FCapsule* Object1 = Implicit1->template GetObject<const FCapsule >();
 			if (ensure(Object0 && Object1))
 			{
-				FPBDCollisionConstraint* Constraint = Context.CollisionAllocator->FindOrCreateConstraint(Particle0, Implicit0, Shape0, nullptr, LocalTransform0, Particle1, Implicit1, Shape1, nullptr, LocalTransform1, CullDistance, EContactShapesType::SphereCapsule, Context.bAllowManifolds);
+				FPBDCollisionConstraint* Constraint = Context.MultiShapeCollisionDetector->FindOrCreateConstraint(Particle0, Implicit0, Shape0, nullptr, LocalTransform0, Particle1, Implicit1, Shape1, nullptr, LocalTransform1, CullDistance, EContactShapesType::SphereCapsule, Context.GetSettings().bAllowManifolds);
 				if (T_TRAITS::bImmediateUpdate && (Constraint != nullptr))
 				{
 					FRigidTransform3 WorldTransform0 = LocalTransform0 * ParticleWorldTransform0;
@@ -855,7 +830,7 @@ namespace Chaos
 			const FImplicitObject* Object1 = Implicit1;
 			if (ensure(Object0 && Object1))
 			{
-				FPBDCollisionConstraint* Constraint = Context.CollisionAllocator->FindOrCreateConstraint(Particle0, Implicit0, Shape0, nullptr, LocalTransform0, Particle1, Implicit1, Shape1, nullptr, LocalTransform1, CullDistance, EContactShapesType::SphereConvex, Context.bAllowManifolds);
+				FPBDCollisionConstraint* Constraint = Context.MultiShapeCollisionDetector->FindOrCreateConstraint(Particle0, Implicit0, Shape0, nullptr, LocalTransform0, Particle1, Implicit1, Shape1, nullptr, LocalTransform1, CullDistance, EContactShapesType::SphereConvex, Context.GetSettings().bAllowManifolds);
 				if (T_TRAITS::bImmediateUpdate && (Constraint != nullptr))
 				{
 					FRigidTransform3 WorldTransform0 = LocalTransform0 * ParticleWorldTransform0;
@@ -899,7 +874,7 @@ namespace Chaos
 			{
 				if (const TImplicitObjectScaled<FTriangleMeshImplicitObject>* ScaledTriangleMesh = Implicit1->template GetObject<const TImplicitObjectScaled<FTriangleMeshImplicitObject>>())
 				{
-					FPBDCollisionConstraint* Constraint = Context.CollisionAllocator->FindOrCreateConstraint(Particle0, Implicit0, Shape0, nullptr, LocalTransform0, Particle1, Implicit1, Shape1, nullptr, LocalTransform1, CullDistance, EContactShapesType::SphereTriMesh, Context.bAllowManifolds);
+					FPBDCollisionConstraint* Constraint = Context.MultiShapeCollisionDetector->FindOrCreateConstraint(Particle0, Implicit0, Shape0, nullptr, LocalTransform0, Particle1, Implicit1, Shape1, nullptr, LocalTransform1, CullDistance, EContactShapesType::SphereTriMesh, Context.GetSettings().bAllowManifolds);
 					if (T_TRAITS::bImmediateUpdate && (Constraint != nullptr))
 					{
 						FRigidTransform3 WorldTransform0 = LocalTransform0 * ParticleWorldTransform0;
@@ -909,7 +884,7 @@ namespace Chaos
 				}
 				else if (const FTriangleMeshImplicitObject* TriangleMesh = Implicit1->template GetObject<const FTriangleMeshImplicitObject>())
 				{
-					FPBDCollisionConstraint* Constraint = Context.CollisionAllocator->FindOrCreateConstraint(Particle0, Implicit0, Shape0, nullptr, LocalTransform0, Particle1, Implicit1, Shape1, nullptr, LocalTransform1, CullDistance, EContactShapesType::SphereTriMesh, Context.bAllowManifolds);
+					FPBDCollisionConstraint* Constraint = Context.MultiShapeCollisionDetector->FindOrCreateConstraint(Particle0, Implicit0, Shape0, nullptr, LocalTransform0, Particle1, Implicit1, Shape1, nullptr, LocalTransform1, CullDistance, EContactShapesType::SphereTriMesh, Context.GetSettings().bAllowManifolds);
 					if (T_TRAITS::bImmediateUpdate && (Constraint != nullptr))
 					{
 						FRigidTransform3 WorldTransform0 = LocalTransform0 * ParticleWorldTransform0;
@@ -932,14 +907,14 @@ namespace Chaos
 			{
 				if (const TImplicitObjectScaled<FTriangleMeshImplicitObject>* ScaledTriangleMesh = Implicit1->template GetObject<const TImplicitObjectScaled<FTriangleMeshImplicitObject>>())
 				{
-					FPBDCollisionConstraint* Constraint = Context.CollisionAllocator->FindOrCreateSweptConstraint(Particle0, Implicit0, Shape0, nullptr, LocalTransform0, Particle1, Implicit1, Shape1, nullptr, LocalTransform1, CullDistance, EContactShapesType::SphereTriMesh);
+					FPBDCollisionConstraint* Constraint = Context.MultiShapeCollisionDetector->FindOrCreateSweptConstraint(Particle0, Implicit0, Shape0, nullptr, LocalTransform0, Particle1, Implicit1, Shape1, nullptr, LocalTransform1, CullDistance, EContactShapesType::SphereTriMesh);
 					FRigidTransform3 WorldTransformX0 = LocalTransform0 * FRigidTransform3(Particle0->X(), Particle0->R());
 					FRigidTransform3 WorldTransform1 = LocalTransform1 * ParticleWorldTransform1;
 					UpdateSphereTriangleMeshConstraintSwept(Particle0, *Object0, WorldTransformX0, *ScaledTriangleMesh, WorldTransform1, Dir, Length, Dt, *Constraint);
 				}
 				else if (const FTriangleMeshImplicitObject* TriangleMesh = Implicit1->template GetObject<const FTriangleMeshImplicitObject>())
 				{
-					FPBDCollisionConstraint* Constraint = Context.CollisionAllocator->FindOrCreateSweptConstraint(Particle0, Implicit0, Shape0, nullptr, LocalTransform0, Particle1, Implicit1, Shape1, nullptr, LocalTransform1, CullDistance, EContactShapesType::SphereTriMesh);
+					FPBDCollisionConstraint* Constraint = Context.MultiShapeCollisionDetector->FindOrCreateSweptConstraint(Particle0, Implicit0, Shape0, nullptr, LocalTransform0, Particle1, Implicit1, Shape1, nullptr, LocalTransform1, CullDistance, EContactShapesType::SphereTriMesh);
 					FRigidTransform3 WorldTransformX0 = LocalTransform0 * FRigidTransform3(Particle0->X(), Particle0->R());
 					FRigidTransform3 WorldTransform1 = LocalTransform1 * ParticleWorldTransform1;
 					UpdateSphereTriangleMeshConstraintSwept(Particle0, *Object0, WorldTransformX0, *TriangleMesh, WorldTransform1, Dir, Length, Dt, *Constraint);
@@ -977,7 +952,7 @@ namespace Chaos
 			const FCapsule* Object1 = Implicit1->template GetObject<const FCapsule >();
 			if (ensure(Object0 && Object1))
 			{
-				FPBDCollisionConstraint* Constraint = Context.CollisionAllocator->FindOrCreateConstraint(Particle0, Implicit0, Shape0, nullptr, LocalTransform0, Particle1, Implicit1, Shape1, nullptr, LocalTransform1, CullDistance, EContactShapesType::CapsuleCapsule, Context.bAllowManifolds);
+				FPBDCollisionConstraint* Constraint = Context.MultiShapeCollisionDetector->FindOrCreateConstraint(Particle0, Implicit0, Shape0, nullptr, LocalTransform0, Particle1, Implicit1, Shape1, nullptr, LocalTransform1, CullDistance, EContactShapesType::CapsuleCapsule, Context.GetSettings().bAllowManifolds);
 				if (T_TRAITS::bImmediateUpdate && (Constraint != nullptr))
 				{
 					FRigidTransform3 WorldTransform0 = LocalTransform0 * ParticleWorldTransform0;
@@ -1018,7 +993,7 @@ namespace Chaos
 			const TBox<FReal, 3>* Object1 = Implicit1->template GetObject<const TBox<FReal, 3>>();
 			if (ensure(Object0 && Object1))
 			{
-				FPBDCollisionConstraint* Constraint = Context.CollisionAllocator->FindOrCreateConstraint(Particle0, Implicit0, Shape0, nullptr, LocalTransform0, Particle1, Implicit1, Shape1, nullptr, LocalTransform1, CullDistance, EContactShapesType::CapsuleBox, Context.bAllowManifolds);
+				FPBDCollisionConstraint* Constraint = Context.MultiShapeCollisionDetector->FindOrCreateConstraint(Particle0, Implicit0, Shape0, nullptr, LocalTransform0, Particle1, Implicit1, Shape1, nullptr, LocalTransform1, CullDistance, EContactShapesType::CapsuleBox, Context.GetSettings().bAllowManifolds);
 				if (T_TRAITS::bImmediateUpdate && (Constraint != nullptr))
 				{
 					FRigidTransform3 WorldTransform0 = LocalTransform0 * ParticleWorldTransform0;
@@ -1100,7 +1075,7 @@ namespace Chaos
 			const FImplicitObject* Object1 = Implicit1;
 			if (ensure(Object0 && Object1))
 			{
-				FPBDCollisionConstraint* Constraint = Context.CollisionAllocator->FindOrCreateConstraint(Particle0, Implicit0, Shape0, nullptr, LocalTransform0, Particle1, Implicit1, Shape1, nullptr, LocalTransform1, CullDistance, EContactShapesType::CapsuleConvex, Context.bAllowManifolds);
+				FPBDCollisionConstraint* Constraint = Context.MultiShapeCollisionDetector->FindOrCreateConstraint(Particle0, Implicit0, Shape0, nullptr, LocalTransform0, Particle1, Implicit1, Shape1, nullptr, LocalTransform1, CullDistance, EContactShapesType::CapsuleConvex, Context.GetSettings().bAllowManifolds);
 				if (T_TRAITS::bImmediateUpdate && (Constraint != nullptr))
 				{
 					FRigidTransform3 WorldTransform0 = LocalTransform0 * ParticleWorldTransform0;
@@ -1147,7 +1122,7 @@ namespace Chaos
 			const FHeightField* Object1 = Implicit1->template GetObject<const FHeightField >();
 			if (ensure(Object0 && Object1))
 			{
-				FPBDCollisionConstraint* Constraint = Context.CollisionAllocator->FindOrCreateConstraint(Particle0, Implicit0, Shape0, nullptr, LocalTransform0, Particle1, Implicit1, Shape1, nullptr, LocalTransform1, CullDistance, EContactShapesType::CapsuleHeightField, Context.bAllowManifolds);
+				FPBDCollisionConstraint* Constraint = Context.MultiShapeCollisionDetector->FindOrCreateConstraint(Particle0, Implicit0, Shape0, nullptr, LocalTransform0, Particle1, Implicit1, Shape1, nullptr, LocalTransform1, CullDistance, EContactShapesType::CapsuleHeightField, Context.GetSettings().bAllowManifolds);
 				if (T_TRAITS::bImmediateUpdate && (Constraint != nullptr))
 				{
 					FRigidTransform3 WorldTransform0 = LocalTransform0 * ParticleWorldTransform0;
@@ -1164,7 +1139,7 @@ namespace Chaos
 			const FHeightField* Object1 = Implicit1->template GetObject<const FHeightField >();
 			if (ensure(Object0 && Object1))
 			{
-				FPBDCollisionConstraint* Constraint = Context.CollisionAllocator->FindOrCreateSweptConstraint(Particle0, Implicit0, Shape0, nullptr, LocalTransform0, Particle1, Implicit1, Shape1, nullptr, LocalTransform1, CullDistance, EContactShapesType::CapsuleHeightField);
+				FPBDCollisionConstraint* Constraint = Context.MultiShapeCollisionDetector->FindOrCreateSweptConstraint(Particle0, Implicit0, Shape0, nullptr, LocalTransform0, Particle1, Implicit1, Shape1, nullptr, LocalTransform1, CullDistance, EContactShapesType::CapsuleHeightField);
 				FRigidTransform3 WorldTransformX0 = LocalTransform0 * FRigidTransform3(Particle0->X(), Particle0->R());
 				FRigidTransform3 WorldTransform1 = LocalTransform1 * ParticleWorldTransform1;
 				UpdateCapsuleHeightFieldConstraintSwept(Particle0, *Object0, WorldTransformX0, *Object1, WorldTransform1, Dir, Length, Dt, *Constraint);
@@ -1211,7 +1186,7 @@ namespace Chaos
 			{
 				if (const TImplicitObjectScaled<FTriangleMeshImplicitObject>* ScaledTriangleMesh = Implicit1->template GetObject<const TImplicitObjectScaled<FTriangleMeshImplicitObject>>())
 				{
-					FPBDCollisionConstraint* Constraint = Context.CollisionAllocator->FindOrCreateConstraint(Particle0, Implicit0, Shape0, nullptr, LocalTransform0, Particle1, Implicit1, Shape1,  nullptr, LocalTransform1, CullDistance, EContactShapesType::CapsuleTriMesh, Context.bAllowManifolds);
+					FPBDCollisionConstraint* Constraint = Context.MultiShapeCollisionDetector->FindOrCreateConstraint(Particle0, Implicit0, Shape0, nullptr, LocalTransform0, Particle1, Implicit1, Shape1,  nullptr, LocalTransform1, CullDistance, EContactShapesType::CapsuleTriMesh, Context.GetSettings().bAllowManifolds);
 					if (T_TRAITS::bImmediateUpdate && (Constraint != nullptr))
 					{
 						FRigidTransform3 WorldTransform0 = LocalTransform0 * ParticleWorldTransform0;
@@ -1221,7 +1196,7 @@ namespace Chaos
 				}
 				else if (const FTriangleMeshImplicitObject* TriangleMesh = Implicit1->template GetObject<const FTriangleMeshImplicitObject>())
 				{
-					FPBDCollisionConstraint* Constraint = Context.CollisionAllocator->FindOrCreateConstraint(Particle0, Implicit0, Shape0, nullptr, LocalTransform0, Particle1, Implicit1, Shape1, nullptr, LocalTransform1, CullDistance, EContactShapesType::CapsuleTriMesh, Context.bAllowManifolds);
+					FPBDCollisionConstraint* Constraint = Context.MultiShapeCollisionDetector->FindOrCreateConstraint(Particle0, Implicit0, Shape0, nullptr, LocalTransform0, Particle1, Implicit1, Shape1, nullptr, LocalTransform1, CullDistance, EContactShapesType::CapsuleTriMesh, Context.GetSettings().bAllowManifolds);
 					if (T_TRAITS::bImmediateUpdate && (Constraint != nullptr))
 					{
 						FRigidTransform3 WorldTransform0 = LocalTransform0 * ParticleWorldTransform0;
@@ -1245,14 +1220,14 @@ namespace Chaos
 			{
 				if (const TImplicitObjectScaled<FTriangleMeshImplicitObject>* ScaledTriangleMesh = Implicit1->template GetObject<const TImplicitObjectScaled<FTriangleMeshImplicitObject>>())
 				{
-					FPBDCollisionConstraint* Constraint = Context.CollisionAllocator->FindOrCreateSweptConstraint(Particle0, Implicit0, Shape0, nullptr, LocalTransform0, Particle1, Implicit1, Shape1, nullptr, LocalTransform1, CullDistance, EContactShapesType::CapsuleTriMesh);
+					FPBDCollisionConstraint* Constraint = Context.MultiShapeCollisionDetector->FindOrCreateSweptConstraint(Particle0, Implicit0, Shape0, nullptr, LocalTransform0, Particle1, Implicit1, Shape1, nullptr, LocalTransform1, CullDistance, EContactShapesType::CapsuleTriMesh);
 					FRigidTransform3 WorldTransformX0 = LocalTransform0 * FRigidTransform3(Particle0->X(), Particle0->R());
 					FRigidTransform3 WorldTransform1 = LocalTransform1 * ParticleWorldTransform1;
 					UpdateCapsuleTriangleMeshConstraintSwept(Particle0, *Object0, WorldTransformX0, *ScaledTriangleMesh, WorldTransform1, Dir, Length, Dt, *Constraint);
 				}
 				else if (const FTriangleMeshImplicitObject* TriangleMesh = Implicit1->template GetObject<const FTriangleMeshImplicitObject>())
 				{
-					FPBDCollisionConstraint* Constraint = Context.CollisionAllocator->FindOrCreateSweptConstraint(Particle0, Implicit0, Shape0, nullptr, LocalTransform0, Particle1, Implicit1, Shape1, nullptr, LocalTransform1, CullDistance, EContactShapesType::CapsuleTriMesh);
+					FPBDCollisionConstraint* Constraint = Context.MultiShapeCollisionDetector->FindOrCreateSweptConstraint(Particle0, Implicit0, Shape0, nullptr, LocalTransform0, Particle1, Implicit1, Shape1, nullptr, LocalTransform1, CullDistance, EContactShapesType::CapsuleTriMesh);
 					FRigidTransform3 WorldTransformX0 = LocalTransform0 * FRigidTransform3(Particle0->X(), Particle0->R());
 					FRigidTransform3 WorldTransform1 = LocalTransform1 * ParticleWorldTransform1;
 					UpdateCapsuleTriangleMeshConstraintSwept(Particle0, *Object0, WorldTransformX0, *TriangleMesh, WorldTransform1, Dir, Length, Dt, *Constraint);
@@ -1293,7 +1268,7 @@ namespace Chaos
 			EImplicitObjectType Implicit0Type = Particle0->Geometry()->GetType();
 			EImplicitObjectType Implicit1Type = Particle1->Geometry()->GetType();
 
-			FPBDCollisionConstraint* Constraint = Context.CollisionAllocator->FindOrCreateConstraint(Particle0, Implicit0, Shape0, nullptr, LocalTransform0, Particle1, Implicit1, Shape1, nullptr, LocalTransform1, CullDistance, EContactShapesType::GenericConvexConvex, Context.bAllowManifolds);
+			FPBDCollisionConstraint* Constraint = Context.MultiShapeCollisionDetector->FindOrCreateConstraint(Particle0, Implicit0, Shape0, nullptr, LocalTransform0, Particle1, Implicit1, Shape1, nullptr, LocalTransform1, CullDistance, EContactShapesType::GenericConvexConvex, Context.GetSettings().bAllowManifolds);
 			if (T_TRAITS::bImmediateUpdate && (Constraint != nullptr))
 			{
 				FRigidTransform3 WorldTransform0 = LocalTransform0 * ParticleWorldTransform0;
@@ -1307,7 +1282,7 @@ namespace Chaos
 			CONDITIONAL_SCOPE_CYCLE_COUNTER(STAT_Collisions_ConstructGenericConvexConvexConstraintsSwept, ConstraintsDetailedStats);
 			FGenericParticleHandle RigidParticle0 = FGenericParticleHandle(Particle0);
 			FGenericParticleHandle RigidParticle1 = FGenericParticleHandle(Particle1);
-			FPBDCollisionConstraint* Constraint = Context.CollisionAllocator->FindOrCreateSweptConstraint(Particle0, Implicit0, Shape0, nullptr, LocalTransform0, Particle1, Implicit1, Shape1, nullptr, LocalTransform1, CullDistance, EContactShapesType::GenericConvexConvex);
+			FPBDCollisionConstraint* Constraint = Context.MultiShapeCollisionDetector->FindOrCreateSweptConstraint(Particle0, Implicit0, Shape0, nullptr, LocalTransform0, Particle1, Implicit1, Shape1, nullptr, LocalTransform1, CullDistance, EContactShapesType::GenericConvexConvex);
 			// Rotations of WorldTransformXRx are not actually used in UpdateGenericConvexConvexConstraintSwept.
 			FRigidTransform3 WorldTransformXR0 = LocalTransform0 * FRigidTransform3(StartX0, RigidParticle0->R());
 			FRigidTransform3 WorldTransformXR1 = LocalTransform1 * FRigidTransform3(StartX1, RigidParticle1->R());
@@ -1316,7 +1291,7 @@ namespace Chaos
 			// @todo(zhenglin): If resweep is not enabled, we don't need to process secondary collision constraints so constraints with TOI > 1.0f can be deleted. This could be a potential optimization but I have not made this work yet.
 			// if (!bChaosCollisionCCDEnableResweep && Constraint->TimeOfImpact > 1.0f)
 			// {
-			// 	Context.CollisionAllocator->DestroyConstraint(Constraint);
+			// 	Context.MultiShapeCollisionDetector->DestroyConstraint(Constraint);
 			// }
 		}
 
@@ -1355,7 +1330,7 @@ namespace Chaos
 			const FHeightField* Object1 = Implicit1->template GetObject<const FHeightField >();
 			if (ensure(Implicit0->IsConvex() && Object1))
 			{
-				FPBDCollisionConstraint* Constraint = Context.CollisionAllocator->FindOrCreateConstraint(Particle0, Implicit0, Shape0, nullptr, LocalTransform0, Particle1, Implicit1, Shape1, nullptr, LocalTransform1, CullDistance, EContactShapesType::ConvexHeightField, Context.bAllowManifolds);
+				FPBDCollisionConstraint* Constraint = Context.MultiShapeCollisionDetector->FindOrCreateConstraint(Particle0, Implicit0, Shape0, nullptr, LocalTransform0, Particle1, Implicit1, Shape1, nullptr, LocalTransform1, CullDistance, EContactShapesType::ConvexHeightField, Context.GetSettings().bAllowManifolds);
 				if (T_TRAITS::bImmediateUpdate && (Constraint != nullptr))
 				{
 					FRigidTransform3 WorldTransform0 = LocalTransform0 * ParticleWorldTransform0;
@@ -1371,7 +1346,7 @@ namespace Chaos
 			const FHeightField* Object1 = Implicit1->template GetObject<const FHeightField >();
 			if (ensure(Implicit0->IsConvex() && Object1))
 			{
-				FPBDCollisionConstraint* Constraint = Context.CollisionAllocator->FindOrCreateSweptConstraint(Particle0, Implicit0, Shape0, nullptr, LocalTransform0, Particle1, Implicit1, Shape1, nullptr, LocalTransform1, CullDistance, EContactShapesType::ConvexHeightField);
+				FPBDCollisionConstraint* Constraint = Context.MultiShapeCollisionDetector->FindOrCreateSweptConstraint(Particle0, Implicit0, Shape0, nullptr, LocalTransform0, Particle1, Implicit1, Shape1, nullptr, LocalTransform1, CullDistance, EContactShapesType::ConvexHeightField);
 				FRigidTransform3 WorldTransformX0 = LocalTransform0 * FRigidTransform3(Particle0->X(), Particle0->R());
 				FRigidTransform3 WorldTransform1 = LocalTransform1 * ParticleWorldTransform1;
 				UpdateConvexHeightFieldConstraintSwept(Particle0, *Implicit0, WorldTransformX0, *Object1, WorldTransform1, Dir, Length, Dt, *Constraint);
@@ -1411,7 +1386,7 @@ namespace Chaos
 			CONDITIONAL_SCOPE_CYCLE_COUNTER(STAT_Collisions_ConstructConvexTriangleMeshConstraints, ConstraintsDetailedStats);
 			if (ensure(Implicit0->IsConvex()))
 			{
-				FPBDCollisionConstraint* Constraint = Context.CollisionAllocator->FindOrCreateConstraint(Particle0, Implicit0, Shape0, nullptr, LocalTransform0, Particle1, Implicit1, Shape1, nullptr, LocalTransform1, CullDistance, EContactShapesType::ConvexTriMesh, Context.bAllowManifolds);
+				FPBDCollisionConstraint* Constraint = Context.MultiShapeCollisionDetector->FindOrCreateConstraint(Particle0, Implicit0, Shape0, nullptr, LocalTransform0, Particle1, Implicit1, Shape1, nullptr, LocalTransform1, CullDistance, EContactShapesType::ConvexTriMesh, Context.GetSettings().bAllowManifolds);
 				if (T_TRAITS::bImmediateUpdate && (Constraint != nullptr))
 				{
 					FRigidTransform3 WorldTransform0 = LocalTransform0 * ParticleWorldTransform0;
@@ -1428,14 +1403,14 @@ namespace Chaos
 			{
 				if (const TImplicitObjectScaled<FTriangleMeshImplicitObject>* ScaledTriangleMesh = Implicit1->template GetObject<const TImplicitObjectScaled<FTriangleMeshImplicitObject>>())
 				{
-					FPBDCollisionConstraint* Constraint = Context.CollisionAllocator->FindOrCreateSweptConstraint(Particle0, Implicit0, Shape0, nullptr, LocalTransform0, Particle1, Implicit1, Shape1, nullptr, LocalTransform1, CullDistance, EContactShapesType::ConvexTriMesh);
+					FPBDCollisionConstraint* Constraint = Context.MultiShapeCollisionDetector->FindOrCreateSweptConstraint(Particle0, Implicit0, Shape0, nullptr, LocalTransform0, Particle1, Implicit1, Shape1, nullptr, LocalTransform1, CullDistance, EContactShapesType::ConvexTriMesh);
 					FRigidTransform3 WorldTransformX0 = LocalTransform0 * FRigidTransform3(Particle0->X(), Particle0->R());
 					FRigidTransform3 WorldTransform1 = LocalTransform1 * ParticleWorldTransform1;
 					UpdateConvexTriangleMeshConstraintSwept(Particle0, *Implicit0, WorldTransformX0, *ScaledTriangleMesh, WorldTransform1, Dir, Length, Dt, *Constraint);
 				}
 				else if (const FTriangleMeshImplicitObject* TriangleMesh = Implicit1->template GetObject<const FTriangleMeshImplicitObject>())
 				{
-					FPBDCollisionConstraint* Constraint = Context.CollisionAllocator->FindOrCreateSweptConstraint(Particle0, Implicit0, Shape0, nullptr, LocalTransform0, Particle1, Implicit1, Shape1, nullptr, LocalTransform1, CullDistance, EContactShapesType::ConvexTriMesh);
+					FPBDCollisionConstraint* Constraint = Context.MultiShapeCollisionDetector->FindOrCreateSweptConstraint(Particle0, Implicit0, Shape0, nullptr, LocalTransform0, Particle1, Implicit1, Shape1, nullptr, LocalTransform1, CullDistance, EContactShapesType::ConvexTriMesh);
 					FRigidTransform3 WorldTransformX0 = LocalTransform0 * FRigidTransform3(Particle0->X(), Particle0->R());
 					FRigidTransform3 WorldTransform1 = LocalTransform1 * ParticleWorldTransform1;
 					UpdateConvexTriangleMeshConstraintSwept(Particle0, *Implicit0, WorldTransformX0, *TriangleMesh, WorldTransform1, Dir, Length, Dt, *Constraint);
@@ -1510,7 +1485,7 @@ namespace Chaos
 		template<typename T_TRAITS>
 		void ConstructLevelsetLevelsetConstraints(TGeometryParticleHandle<FReal, 3>* Particle0, TGeometryParticleHandle<FReal, 3>* Particle1, const FImplicitObject* Implicit0, const FPerShapeData* Shape0, const FBVHParticles* Simplicial0, const FImplicitObject* Implicit1, const FPerShapeData* Shape1, const FBVHParticles* Simplicial1, const FRigidTransform3& ParticleWorldTransform0, const FRigidTransform3& LocalTransform0, const FRigidTransform3& ParticleWorldTransform1, const FRigidTransform3& LocalTransform1, const FReal CullDistance, const FReal Dt, const FCollisionContext& Context)
 		{
-			const bool bAllowManifolds = Context.bAllowManifolds && bChaos_Collision_AllowLevelsetManifolds;
+			const bool bAllowManifolds = Context.GetSettings().bAllowManifolds && bChaos_Collision_AllowLevelsetManifolds;
 			FPBDCollisionConstraint* Constraint = nullptr;
 
 			// @todo(chaos): We do not really support a runtime decision for the particle order now that we have persistent collisions. Fix the logic below...?
@@ -1519,11 +1494,11 @@ namespace Chaos
 			int32 P0NumCollisionParticles = Simplicial0 ? Simplicial0->Size() : Particle0->CastToRigidParticle()->CollisionParticlesSize();
 			if (!Particle1->Geometry() || (bIsParticleDynamic0 && !P0NumCollisionParticles && Particle0->Geometry() && !Particle0->Geometry()->IsUnderlyingUnion()))
 			{
-				Constraint = Context.CollisionAllocator->FindOrCreateConstraint(Particle1, Implicit1, Shape1, Simplicial1, LocalTransform1, Particle0, Implicit0, Shape0, Simplicial0, LocalTransform0, CullDistance, EContactShapesType::LevelSetLevelSet, bAllowManifolds);
+				Constraint = Context.MultiShapeCollisionDetector->FindOrCreateConstraint(Particle1, Implicit1, Shape1, Simplicial1, LocalTransform1, Particle0, Implicit0, Shape0, Simplicial0, LocalTransform0, CullDistance, EContactShapesType::LevelSetLevelSet, bAllowManifolds);
 			}
 			else
 			{
-				Constraint = Context.CollisionAllocator->FindOrCreateConstraint(Particle0, Implicit0, Shape0, Simplicial0, LocalTransform0, Particle1, Implicit1, Shape1, Simplicial1, LocalTransform1, CullDistance, EContactShapesType::LevelSetLevelSet, bAllowManifolds);
+				Constraint = Context.MultiShapeCollisionDetector->FindOrCreateConstraint(Particle0, Implicit0, Shape0, Simplicial0, LocalTransform0, Particle1, Implicit1, Shape1, Simplicial1, LocalTransform1, CullDistance, EContactShapesType::LevelSetLevelSet, bAllowManifolds);
 			}
 
 			if (Constraint != nullptr)
@@ -1689,7 +1664,7 @@ namespace Chaos
 			FVec3 DirCCD(0.0f);
 			const FVec3 DeltaX0 = Constraint.GetShapeWorldTransform0().GetTranslation() - WorldTransform0.GetLocation();
 			const FVec3 DeltaX1 = Constraint.GetShapeWorldTransform1().GetTranslation() - WorldTransform1.GetLocation();
-			bool bUseCCD = ShouldUseCCD(Particle0, DeltaX0, Particle1, DeltaX1, DirCCD, LengthCCD, false);
+			bool bUseCCD = ShouldUseCCD(Particle0, DeltaX0, Particle1, DeltaX1, DirCCD, LengthCCD);
 
 			Constraint.SetCCDTimeOfImpact(TNumericLimits<FReal>::Max());
 
@@ -1814,7 +1789,7 @@ namespace Chaos
 			const FVec3 StartX1 = (ConstParticle1->ObjectState() == EObjectStateType::Kinematic) ? (ConstParticle1->P() - ConstParticle1->V() * Dt) : ConstParticle1->X();
 			const FVec3 EndX0 = ConstParticle0->P();
 			const FVec3 EndX1 = ConstParticle1->P();
-			bool bUseCCD = ShouldUseCCD(Particle0, EndX0 - StartX0, Particle1, EndX1 - StartX1, DirCCD, LengthCCD, Context.bForceDisableCCD);
+			bool bUseCCD = Context.GetSettings().bAllowCCD && ShouldUseCCD(Particle0, EndX0 - StartX0, Particle1, EndX1 - StartX1, DirCCD, LengthCCD);
 			const bool bUseGenericSweptConstraints = bUseCCD && (CCDUseGenericSweptConvexConstraints > 0) && bIsConvex0 && bIsConvex1;
 #if CHAOS_COLLISION_CREATE_BOUNDSCHECK
 			if ((Implicit0 != nullptr) && (Implicit1 != nullptr))
@@ -2225,7 +2200,7 @@ namespace Chaos
 						const FPerShapeData* ChildShape0 = ShapesArray[ShapeIdx].Get();
 					
 						// If shape is not sim'd, we may end up iterating over a lot of shapes on particle1's union and wasting time filtering.
-						if (Context.bFilteringEnabled == false || FilterHasSimEnabled(ChildShape0))
+						if (Context.GetSettings().bFilteringEnabled == false || FilterHasSimEnabled(ChildShape0))
 						{
 							LegacyConstructConstraints<T_TRAITS>(Particle0, Particle1, Child0.Get(), FShapeOrShapesArray(ChildShape0), Simplicial0, Implicit1, Shapes1, Simplicial1, ParticleWorldTransform0, LocalTransform0, ParticleWorldTransform1, LocalTransform1, CullDistance, Dt, Context);
 						}
@@ -2282,7 +2257,7 @@ namespace Chaos
 						const FPerShapeData* ChildShape1 = ShapesArray[ShapeIdx].Get();
 
 						// If shape is not sim'd, we may end up iterating over a lot of shapes on particle1's union and wasting time filtering.
-						if (Context.bFilteringEnabled == false || FilterHasSimEnabled(ChildShape1))
+						if (Context.GetSettings().bFilteringEnabled == false || FilterHasSimEnabled(ChildShape1))
 						{
 							LegacyConstructConstraints<T_TRAITS>(Particle0, Particle1, Implicit0, Shapes0, Simplicial0, Child1.Get(), FShapeOrShapesArray(ChildShape1), Simplicial1, ParticleWorldTransform0, LocalTransform0, ParticleWorldTransform1, LocalTransform1, CullDistance, Dt, Context);
 						}
@@ -2343,7 +2318,7 @@ namespace Chaos
 			}
 
 			// Check shape pair filtering if enable
-			if (Context.bFilteringEnabled && !DoCollide(Implicit0Type, Shape0, Implicit1Type, Shape1))
+			if (Context.GetSettings().bFilteringEnabled && !DoCollide(Implicit0Type, Shape0, Implicit1Type, Shape1))
 			{
 				return;
 			}
@@ -2551,7 +2526,7 @@ namespace Chaos
 			INC_DWORD_STAT(STAT_ChaosCollisionCounter_NumParticlePairs);
 
 			// @todo(chaos): remove defer update - it is not a valid option with the new solver
-			bool bDeferUpdate = Context.bDeferUpdate;
+			bool bDeferUpdate = Context.GetSettings().bDeferNarrowPhase;
 			if (bDeferUpdate)
 			{
 				using TTraits = TConstructCollisionTraits<false>;
diff --git a/Engine/Source/Runtime/Experimental/Chaos/Private/Chaos/PBDCollisionConstraints.cpp b/Engine/Source/Runtime/Experimental/Chaos/Private/Chaos/PBDCollisionConstraints.cpp
--- a/Engine/Source/Runtime/Experimental/Chaos/Private/Chaos/PBDCollisionConstraints.cpp
+++ b/Engine/Source/Runtime/Experimental/Chaos/Private/Chaos/PBDCollisionConstraints.cpp
@@ -392,6 +392,7 @@ namespace Chaos
 		{
 			FPBDCollisionSolverContainer& SolverContainer = GetConstraintSolverContainer(SolverData);
 			SolverContainer.SetNum(NumIslandConstraints);
+			SolverContainer.SetIsDeferredCollisionDetection(DetectorSettings.bDeferNarrowPhase);
 		}
 		else
 		{
diff --git a/Engine/Source/Runtime/Experimental/Chaos/Private/Chaos/PBDRigidsEvolutionGBF.cpp b/Engine/Source/Runtime/Experimental/Chaos/Private/Chaos/PBDRigidsEvolutionGBF.cpp
--- a/Engine/Source/Runtime/Experimental/Chaos/Private/Chaos/PBDRigidsEvolutionGBF.cpp
+++ b/Engine/Source/Runtime/Experimental/Chaos/Private/Chaos/PBDRigidsEvolutionGBF.cpp
@@ -58,9 +58,6 @@ namespace Chaos
 		int DisableParticleUpdateVelocityParallelFor = 0;
 		FAutoConsoleVariableRef CVarDisableParticleUpdateVelocityParallelFor(TEXT("p.DisableParticleUpdateVelocityParallelFor"), DisableParticleUpdateVelocityParallelFor, TEXT("Disable Particle Update Velocity ParallelFor and run the update on a single thread"));
 
-		bool bChaosUseCCD = true;
-		FAutoConsoleVariableRef  CVarChaosUseCCD(TEXT("p.Chaos.CCD.UseCCD"), bChaosUseCCD , TEXT("Global flag to turn CCD on or off. Default is true"));
-
 		// NOTE: If we have mrope than 1 CCD iteration (ChaosCollisionCCDConstraintMaxProcessCount), the tight bounding box will cause us to miss secondary CCD collisions if the first one(s) result in a change in direction
 		bool bChaosCollisionCCDUseTightBoundingBox = true;
 		FAutoConsoleVariableRef  CVarChaosCollisionCCDUseTightBoundingBox(TEXT("p.Chaos.Collision.CCD.UseTightBoundingBox"), bChaosCollisionCCDUseTightBoundingBox , TEXT(""));
@@ -435,7 +432,6 @@ void FPBDRigidsEvolutionGBF::AdvanceOneTimeStepImpl(const FReal Dt, const FSubSt
 		UpdateInertiaConditioning();
 	}
 
-	if (bChaosUseCCD)
 	{
 		SCOPE_CYCLE_COUNTER(STAT_Evolution_CCD);
 		CSV_SCOPED_TIMING_STAT(PhysicsVerbose, CCD);
@@ -901,10 +897,6 @@ FPBDRigidsEvolutionGBF::FPBDRigidsEvolutionGBF(FPBDRigidsSOAs& InParticles,THand
 	AddConstraintRule(&JointConstraintRule);
 
 	SetInternalParticleInitilizationFunction([](const FGeometryParticleHandle*, const FGeometryParticleHandle*) {});
-
-	NarrowPhase.GetContext().bFilteringEnabled = true;
-	NarrowPhase.GetContext().bDeferUpdate = false;
-	NarrowPhase.GetContext().bAllowManifolds = false;
 }
 
 FPBDRigidsEvolutionGBF::~FPBDRigidsEvolutionGBF()
diff --git a/Engine/Source/Runtime/Experimental/Chaos/Private/PBDRigidsSolver.cpp b/Engine/Source/Runtime/Experimental/Chaos/Private/PBDRigidsSolver.cpp
--- a/Engine/Source/Runtime/Experimental/Chaos/Private/PBDRigidsSolver.cpp
+++ b/Engine/Source/Runtime/Experimental/Chaos/Private/PBDRigidsSolver.cpp
@@ -242,6 +242,10 @@ namespace Chaos
 		int32 ChaosSolverCollisionAllowManifoldUpdate = 1;
 		FAutoConsoleVariableRef CVarChaosSolverCollisionAllowManifoldUpdate(TEXT("p.Chaos.Solver.Collision.AllowManifoldUpdate"), ChaosSolverCollisionAllowManifoldUpdate, TEXT("Enable/Disable reuse of manifolds between ticks (for small movement)."));
 
+		// Enable/Disable CCD
+		bool bChaosUseCCD = true;
+		FAutoConsoleVariableRef  CVarChaosUseCCD(TEXT("p.Chaos.Solver.UseCCD"), bChaosUseCCD, TEXT("Global flag to turn CCD on or off. Default is true"));
+
 		// Joint cvars
 		float ChaosSolverJointMinSolverStiffness = 1.0f;
 		float ChaosSolverJointMaxSolverStiffness = 1.0f;
@@ -811,10 +815,13 @@ namespace Chaos
 
 	void FPBDRigidsSolver::PrepareAdvanceBy(const FReal DeltaTime)
 	{
-		MEvolution->GetCollisionDetector().GetNarrowPhase().GetContext().bDeferUpdate = (ChaosSolverCollisionDeferNarrowPhase != 0);
-		MEvolution->GetCollisionDetector().GetNarrowPhase().GetContext().bAllowManifolds = (ChaosSolverCollisionUseManifolds != 0);
-		MEvolution->GetCollisionDetector().GetNarrowPhase().GetContext().bAllowManifoldReuse = (ChaosSolverCollisionAllowManifoldUpdate != 0);
-
+		FCollisionDetectorSettings CollisionDetectorSettings = MEvolution->GetCollisionConstraints().GetDetectorSettings();
+		CollisionDetectorSettings.bAllowManifoldReuse = (ChaosSolverCollisionAllowManifoldUpdate != 0);
+		CollisionDetectorSettings.bDeferNarrowPhase = (ChaosSolverCollisionDeferNarrowPhase != 0);
+		CollisionDetectorSettings.bAllowManifolds = (ChaosSolverCollisionUseManifolds != 0);
+		CollisionDetectorSettings.bAllowCCD = bChaosUseCCD;
+		MEvolution->GetCollisionConstraints().SetDetectorSettings(CollisionDetectorSettings);
+		MEvolution->GetCollisionDetector().GetNarrowPhase().GetContext().SetSettings(CollisionDetectorSettings);
 		
 		FPBDJointSolverSettings JointsSettings = MEvolution->GetJointConstraints().GetSettings();
 		JointsSettings.MinSolverStiffness = ChaosSolverJointMinSolverStiffness;
diff --git a/Engine/Source/Runtime/Experimental/Chaos/Public/Chaos/Collision/CollisionContext.h b/Engine/Source/Runtime/Experimental/Chaos/Public/Chaos/Collision/CollisionContext.h
--- a/Engine/Source/Runtime/Experimental/Chaos/Public/Chaos/Collision/CollisionContext.h
+++ b/Engine/Source/Runtime/Experimental/Chaos/Public/Chaos/Collision/CollisionContext.h
@@ -7,42 +7,58 @@ namespace Chaos
 {
 	class FMultiShapePairCollisionDetector;
 
-	/**
-	 * Data passed down into the collision detection functions.
-	 */
-	class FCollisionContext
+	class FCollisionDetectorSettings
 	{
 	public:
-		FCollisionContext()
+		FCollisionDetectorSettings()
 			: bFilteringEnabled(true)
-			, bDeferUpdate(false)
-			, bAllowManifolds(false)
-			, bAllowManifoldReuse(false)
-			, bForceDisableCCD(false)
-			, CollisionAllocator(nullptr)
+			, bDeferNarrowPhase(false)
+			, bAllowManifolds(true)
+			, bAllowManifoldReuse(true)
+			, bAllowCCD(true)
 		{
 		}
 
-		// Whether to check the shape query flags (e.g., Rigid Body nodes have already performed filtering prior to collision detection) [default: true]
+		// Whether to check the shape query flags in the narrow phase (e.g., Rigid Body nodes have already performed filtering prior to collision detection)
 		bool bFilteringEnabled;
 
-		// Whether to defer constraint phi/normal calculation to the Apply step [default: true]. If true, constraints are speculatively created for each shape
-		// pair passed to the narrow phase. This prevents premature culling of constraints, but it can lead to more items in the constraint graph
-		// which could be undesirable in some cases (destruction?).
-		bool bDeferUpdate;
+		// Whether to defer the narrow phase to the constraint-solve phase. This is only enabled by RBAN. It is not useful for the main solver because 
+		// we would not know the contact details when we call the collision modifier callbacks. It is used by RBAN to allow us to run 1 joint iteration 
+		// prior to collision detection which gives better results.
+		bool bDeferNarrowPhase;
 
-		// Whether to use one-shot manifolds where supported [default: false]
+		// Whether to use one-shot manifolds where supported
 		bool bAllowManifolds;
 
-		// Whether we can reuse manifolds between frames if contacts have not moved far [default: false]
+		// Whether we can reuse manifolds between frames if contacts have not moved far
 		bool bAllowManifoldReuse;
 
-		// Force disable CCD
-		bool bForceDisableCCD;
+		// Whether CCD is allowed (disabled for RBAN)
+		bool bAllowCCD;
+
+	};
+
+	/**
+	 * Data passed down into the collision detection functions.
+	 */
+	class FCollisionContext
+	{
+	public:
+		FCollisionContext()
+			: Settings()
+			, MultiShapeCollisionDetector(nullptr)
+		{
+		}
+
+		const FCollisionDetectorSettings& GetSettings() const { return Settings; }
+		FCollisionDetectorSettings& GetSettings() { return Settings; }
+		void SetSettings(const FCollisionDetectorSettings& InSettings) { Settings = InSettings; }
+
+		FCollisionDetectorSettings Settings;
 
 		// This is used in the older collision detection path which is still used for particles that do not flatten their implicit hierrarchies
 		// into the Particle's ShapesArray. Currently this is only Clusters.
 		// @todo(chaos): remove thsi from here and make it a parameter on ConstructCollisions and all inner functions.
-		FMultiShapePairCollisionDetector* CollisionAllocator;
+		FMultiShapePairCollisionDetector* MultiShapeCollisionDetector;
 	};
 }
\ No newline at end of file
diff --git a/Engine/Source/Runtime/Experimental/Chaos/Public/Chaos/Collision/NarrowPhase.h b/Engine/Source/Runtime/Experimental/Chaos/Public/Chaos/Collision/NarrowPhase.h
--- a/Engine/Source/Runtime/Experimental/Chaos/Public/Chaos/Collision/NarrowPhase.h
+++ b/Engine/Source/Runtime/Experimental/Chaos/Public/Chaos/Collision/NarrowPhase.h
@@ -34,12 +34,11 @@ namespace Chaos
 		FReal GetBoundsVelocityInflation() const { return BoundsVelocityInflation; }
 		void SetBoundsVelocityInflation(const FReal InBoundsVelocityInflation) { BoundsVelocityInflation = InBoundsVelocityInflation; }
 
-		void GenerateCollisions(FReal Dt, TGeometryParticleHandle<FReal, 3>* Particle0, TGeometryParticleHandle<FReal, 3>* Particle1, TGeometryParticleHandle<FReal, 3>* SearchParticlePerfHint, const bool bForceDisableCCD)
+		void GenerateCollisions(FReal Dt, TGeometryParticleHandle<FReal, 3>* Particle0, TGeometryParticleHandle<FReal, 3>* Particle1, TGeometryParticleHandle<FReal, 3>* SearchParticlePerfHint)
 		{
 			FParticlePairMidPhase* MidPhase = CollisionAllocator->GetParticlePairMidPhase(Particle0, Particle1, SearchParticlePerfHint);
 			if (MidPhase != nullptr)
 			{
-				Context.bForceDisableCCD = bForceDisableCCD;
 				MidPhase->GenerateCollisions(GetBoundsExpansion(), Dt, GetContext());
 			}
 		}
@@ -50,9 +49,8 @@ namespace Chaos
 		}
 
 		// Use this function if a Mid phase pair is already allocated
-		void GenerateCollisions(FReal Dt, FParticlePairMidPhase* MidPhase, const bool bForceDisableCCD)
+		void GenerateCollisions(FReal Dt, FParticlePairMidPhase* MidPhase)
 		{
-			Context.bForceDisableCCD = bForceDisableCCD;
 			MidPhase->GenerateCollisions(GetBoundsExpansion(), Dt, GetContext());
 		}
 	private:
diff --git a/Engine/Source/Runtime/Experimental/Chaos/Public/Chaos/Collision/ParticlePairBroadPhase.h b/Engine/Source/Runtime/Experimental/Chaos/Public/Chaos/Collision/ParticlePairBroadPhase.h
--- a/Engine/Source/Runtime/Experimental/Chaos/Public/Chaos/Collision/ParticlePairBroadPhase.h
+++ b/Engine/Source/Runtime/Experimental/Chaos/Public/Chaos/Collision/ParticlePairBroadPhase.h
@@ -82,7 +82,7 @@ namespace Chaos
 			const FAABB3& Box1 = ParticleB->WorldSpaceInflatedBounds();
 			if (Box0.Intersects(Box1))
 			{
-				NarrowPhase.GenerateCollisions(Dt, ParticleA, ParticleB, SearchParticle, true);
+				NarrowPhase.GenerateCollisions(Dt, ParticleA, ParticleB, SearchParticle);
 			}
 		}
 
diff --git a/Engine/Source/Runtime/Experimental/Chaos/Public/Chaos/Collision/SolverCollisionContainer.h b/Engine/Source/Runtime/Experimental/Chaos/Public/Chaos/Collision/SolverCollisionContainer.h
--- a/Engine/Source/Runtime/Experimental/Chaos/Public/Chaos/Collision/SolverCollisionContainer.h
+++ b/Engine/Source/Runtime/Experimental/Chaos/Public/Chaos/Collision/SolverCollisionContainer.h
@@ -66,6 +66,9 @@ namespace Chaos
 		 */
 		void SetNum(const int32 MaxCollisions);
 
+		// Set whether we are using deferred collision detection
+		void SetIsDeferredCollisionDetection(const bool bDeferred) { bDeferredCollisionDetection = bDeferred; }
+
 		// Add a solver constraint for the specified constraint and gather the required solver data
 		void PreAddConstraintSolver(const FReal Dt, FPBDCollisionConstraint& Constraint, FSolverBodyContainer& SolverBodyContainer, int32& SolverIndex);
 		void AddConstraintSolver(const FReal Dt, FPBDCollisionConstraint& Constraint, const int32 Particle0Level, const int32 Particle1Level, FSolverBodyContainer& SolverBodyContainer, const FPBDCollisionSolverSettings& SolverSettings);
@@ -82,13 +85,14 @@ namespace Chaos
 		void UpdatePositionShockPropagation(const FReal Dt, const int32 It, const int32 NumIts, const int32 BeginIndex, const int32 EndIndex, const FPBDCollisionSolverSettings& SolverSettings);
 		void UpdateVelocityShockPropagation(const FReal Dt, const int32 It, const int32 NumIts, const int32 BeginIndex, const int32 EndIndex, const FPBDCollisionSolverSettings& SolverSettings);
 		bool SolvePositionImpl(const FReal Dt, const int32 It, const int32 NumIts, const int32 BeginIndex, const int32 EndIndex, const FPBDCollisionSolverSettings& SolverSettings, const bool bParallel);
-		bool SolvePositionIncrementalImpl(const FReal Dt, const int32 BeginIndex, const int32 EndIndex, const FReal MaxPushOut, const bool bApplyStaticFriction);
 		bool SolvePositionWithFrictionImpl(const FReal Dt, const int32 BeginIndex, const int32 EndIndex, const FReal MaxPushOut, const bool bParallel);
 		bool SolvePositionNoFrictionImpl(const FReal Dt, const int32 BeginIndex, const int32 EndIndex, const FReal MaxPushOut, const bool bParallel);
 		bool SolveVelocityImpl(const FReal Dt, const int32 It, const int32 NumIts, const int32 BeginIndex, const int32 EndIndex, const FPBDCollisionSolverSettings& SolverSettings, const bool bParallel);
 		void ScatterOutputImpl(const FReal Dt, const int32 BeginIndex, const int32 EndIndex, const bool bParallel);
+		void UpdateCollisions(const FReal InDt, const int32 BeginIndex, const int32 EndIndex);
 
 		TArray<FPBDCollisionSolverAdapter> CollisionSolvers;
 		bool bRequiresIncrementalCollisionDetection;
+		bool bDeferredCollisionDetection;
 	};
 }
diff --git a/Engine/Source/Runtime/Experimental/Chaos/Public/Chaos/Collision/SpatialAccelerationBroadPhase.h b/Engine/Source/Runtime/Experimental/Chaos/Public/Chaos/Collision/SpatialAccelerationBroadPhase.h
--- a/Engine/Source/Runtime/Experimental/Chaos/Public/Chaos/Collision/SpatialAccelerationBroadPhase.h
+++ b/Engine/Source/Runtime/Experimental/Chaos/Public/Chaos/Collision/SpatialAccelerationBroadPhase.h
@@ -400,7 +400,7 @@ namespace Chaos
 					}
 
 					FParticlePairMidPhase* MidPhasePair = MidPhasePairs[Index];
-					NarrowPhase.GenerateCollisions(Dt, MidPhasePair, false);
+					NarrowPhase.GenerateCollisions(Dt, MidPhasePair);
 				}
 
 				PHYSICS_CSV_CUSTOM_EXPENSIVE(PhysicsCounters, NumFromBroadphase, NumPotentials, ECsvCustomStatOp::Accumulate);
diff --git a/Engine/Source/Runtime/Experimental/Chaos/Public/Chaos/CollisionResolution.h b/Engine/Source/Runtime/Experimental/Chaos/Public/Chaos/CollisionResolution.h
--- a/Engine/Source/Runtime/Experimental/Chaos/Public/Chaos/CollisionResolution.h
+++ b/Engine/Source/Runtime/Experimental/Chaos/Public/Chaos/CollisionResolution.h
@@ -8,9 +8,6 @@
 #include "Chaos/ParticleHandleFwd.h"
 
 
-extern bool CCDUseInitialRotationForSweptUpdate;
-extern bool CCDAlwaysSweepRemainingDT;
-
 namespace Chaos
 {
 	template <typename T, int d>
@@ -58,7 +55,7 @@ namespace Chaos
 		/**
 		 * @brief Whether CCD should be enabled for a contact given the current particle velocities etc
 		*/
-		bool CHAOS_API ShouldUseCCD(const FGeometryParticleHandle* Particle0, const FVec3& DeltaX0, const FGeometryParticleHandle* Particle1, const FVec3& DeltaX1, FVec3& Dir, FReal& Length, const bool bForceDisableCCD);
+		bool CHAOS_API ShouldUseCCD(const FGeometryParticleHandle* Particle0, const FVec3& DeltaX0, const FGeometryParticleHandle* Particle1, const FVec3& DeltaX1, FVec3& Dir, FReal& Length);
 
 		// Update the constraint by re-running collision detection on the shape pair.
 		// @todo(chaos): remove this and use UpdateConstraint instead
diff --git a/Engine/Source/Runtime/Experimental/Chaos/Public/Chaos/PBDCollisionConstraints.h b/Engine/Source/Runtime/Experimental/Chaos/Public/Chaos/PBDCollisionConstraints.h
--- a/Engine/Source/Runtime/Experimental/Chaos/Public/Chaos/PBDCollisionConstraints.h
+++ b/Engine/Source/Runtime/Experimental/Chaos/Public/Chaos/PBDCollisionConstraints.h
@@ -4,6 +4,7 @@
 #include "Chaos/CollisionResolutionTypes.h"
 #include "Chaos/Collision/CollisionApplyType.h"
 #include "Chaos/Collision/CollisionConstraintAllocator.h"
+#include "Chaos/Collision/CollisionContext.h"
 #include "Chaos/Collision/PBDCollisionConstraintHandle.h"
 #include "Chaos/Collision/SolverCollisionContainer.h"
 #include "Chaos/PBDConstraintContainer.h"
@@ -315,6 +316,10 @@ public:
 
 	void UpdateConstraintMaterialProperties(FPBDCollisionConstraint& Contact);
 
+	// @todo(chaos): these settings are duplicated from the NarrowPhase context - fix this
+	const FCollisionDetectorSettings& GetDetectorSettings() const { return DetectorSettings; }
+	void SetDetectorSettings(const FCollisionDetectorSettings& InSettings) { DetectorSettings = InSettings; }
+
 protected:
 	FPBDCollisionConstraint& GetConstraint(int32 Index);
 	FPBDCollisionSolverContainer& GetConstraintSolverContainer(FPBDIslandSolverData& SolverData);
@@ -364,6 +369,9 @@ private:
 	FVec3 GravityDirection;
 	FReal GravitySize;
 
+	// Settings for the collision detector
+	FCollisionDetectorSettings DetectorSettings;
+
 	// Settings for the low-level collision solvers
 	FPBDCollisionSolverSettings SolverSettings;
 
