Fixed issue with Chaos warning: 'Cannot make a copy of unique array of type xxxxxxxx within the managed array collection. Regenerate unique pointer attributes if needed.'

diff --git a/Engine/Source/Runtime/Experimental/Chaos/Public/GeometryCollection/ManagedArray.h b/Engine/Source/Runtime/Experimental/Chaos/Public/GeometryCollection/ManagedArray.h
--- a/Engine/Source/Runtime/Experimental/Chaos/Public/GeometryCollection/ManagedArray.h
+++ b/Engine/Source/Runtime/Experimental/Chaos/Public/GeometryCollection/ManagedArray.h
@@ -566,13 +566,13 @@ template<> inline void InitHelper(TArray<TYPE>& Array, const TManagedArrayBase<T
 }
 
 typedef TUniquePtr<Chaos::TGeometryParticle<Chaos::FReal, 3>> LOCAL_MA_UniqueTGeometryParticle;
-UNSUPPORTED_UNIQUE_ARRAY_COPIES(LOCAL_MA_UniqueTGeometryParticle, "Chaos::TGeometryParticle");
+UNSUPPORTED_UNIQUE_ARRAY_COPIES(LOCAL_MA_UniqueTGeometryParticle, TEXT("Chaos::TGeometryParticle"));
 
 typedef TUniquePtr<Chaos::FBVHParticles, TDefaultDelete<Chaos::FBVHParticles>> LOCAL_MA_UniqueTBVHParticles;
-UNSUPPORTED_UNIQUE_ARRAY_COPIES(LOCAL_MA_UniqueTBVHParticles, "Chaos::FBVHParticles");
+UNSUPPORTED_UNIQUE_ARRAY_COPIES(LOCAL_MA_UniqueTBVHParticles, TEXT("Chaos::FBVHParticles"));
 
 typedef TUniquePtr<TArray<UE::Math::TVector<float>>> LOCAL_MA_UniqueTArrayTVector;
-UNSUPPORTED_UNIQUE_ARRAY_COPIES(LOCAL_MA_UniqueTArrayTVector, "TArray<UE::Math::TVector<float>>");
+UNSUPPORTED_UNIQUE_ARRAY_COPIES(LOCAL_MA_UniqueTArrayTVector, TEXT("TArray<UE::Math::TVector<float>>"));
 
 template<class InElementType>
 class TManagedArray : public TManagedArrayBase<InElementType>
