make fracture mode's leaf/cluster selectors only select the leaves/clusters at the current level

diff --git a/Engine/Plugins/Experimental/ChaosEditor/Source/FractureEditor/Private/FractureToolSelectors.cpp b/Engine/Plugins/Experimental/ChaosEditor/Source/FractureEditor/Private/FractureToolSelectors.cpp
--- a/Engine/Plugins/Experimental/ChaosEditor/Source/FractureEditor/Private/FractureToolSelectors.cpp
+++ b/Engine/Plugins/Experimental/ChaosEditor/Source/FractureEditor/Private/FractureToolSelectors.cpp
@@ -313,7 +313,7 @@ FText UFractureToolSelectLeaf::GetDisplayText() const
 
 FText UFractureToolSelectLeaf::GetTooltipText() const
 {
-	return FText(NSLOCTEXT("Fracture", "FractureToolSelectLeafTooltip", "Select leaf bones that represent rigid bodies"));
+	return FText(NSLOCTEXT("Fracture", "FractureToolSelectLeafTooltip", "Select leaf bones that represent rigid bodies at current level"));
 }
 
 FSlateIcon UFractureToolSelectLeaf::GetToolIcon() const
@@ -323,7 +323,7 @@ FSlateIcon UFractureToolSelectLeaf::GetToolIcon() const
 
 void UFractureToolSelectLeaf::RegisterUICommand(FFractureEditorCommands* BindingContext)
 {
-	UI_COMMAND_EXT(BindingContext, UICommandInfo, "SelectLeaf", "Leaf", "Select (rigid) leaf bones.", EUserInterfaceActionType::Button, FInputChord());
+	UI_COMMAND_EXT(BindingContext, UICommandInfo, "SelectLeaf", "Leaf", "Select (rigid) leaf bones at current level.", EUserInterfaceActionType::Button, FInputChord());
 	BindingContext->SelectLeaves = UICommandInfo;
 }
 
@@ -345,7 +345,7 @@ FText UFractureToolSelectCluster::GetDisplayText() const
 
 FText UFractureToolSelectCluster::GetTooltipText() const
 {
-	return FText(NSLOCTEXT("Fracture", "FractureToolSelectClusterTooltip", "Select cluster bones"));
+	return FText(NSLOCTEXT("Fracture", "FractureToolSelectClusterTooltip", "Select cluster bones at current level"));
 }
 
 FSlateIcon UFractureToolSelectCluster::GetToolIcon() const
@@ -355,7 +355,7 @@ FSlateIcon UFractureToolSelectCluster::GetToolIcon() const
 
 void UFractureToolSelectCluster::RegisterUICommand(FFractureEditorCommands* BindingContext)
 {
-	UI_COMMAND_EXT(BindingContext, UICommandInfo, "SelectCluster", "Cluster", "Select cluster bones.", EUserInterfaceActionType::Button, FInputChord());
+	UI_COMMAND_EXT(BindingContext, UICommandInfo, "SelectCluster", "Cluster", "Select cluster bones at current level.", EUserInterfaceActionType::Button, FInputChord());
 	BindingContext->SelectClusters = UICommandInfo;
 }
 
diff --git a/Engine/Source/Runtime/Experimental/GeometryCollectionEngine/Private/GeometryCollection/GeometryCollectionComponent.cpp b/Engine/Source/Runtime/Experimental/GeometryCollectionEngine/Private/GeometryCollection/GeometryCollectionComponent.cpp
--- a/Engine/Source/Runtime/Experimental/GeometryCollectionEngine/Private/GeometryCollection/GeometryCollectionComponent.cpp
+++ b/Engine/Source/Runtime/Experimental/GeometryCollectionEngine/Private/GeometryCollection/GeometryCollectionComponent.cpp
@@ -2800,12 +2800,15 @@ void FScopedColorEdit::SelectBones(GeometryCollection::ESelectionMode SelectionM
 		case GeometryCollection::ESelectionMode::Leaves:
 		{
 			ResetBoneSelection();
+			int32 ViewLevel = GetViewLevel();
 			TArray<int32> BonesToSelect;
 			FGeometryCollectionClusteringUtility::GetBonesToLevel(GeometryCollectionPtr.Get(), GetViewLevel(), BonesToSelect, true, true);
 			const TManagedArray<int32>& SimType = GeometryCollectionPtr->SimulationType;
+			const TManagedArray<int32>& Levels = GeometryCollectionPtr->GetAttribute<int32>("Level", FGeometryCollection::TransformGroup);
 			BonesToSelect.SetNum(Algo::RemoveIf(BonesToSelect, [&](int32 BoneIdx)
 				{
-					return SimType[BoneIdx] != FGeometryCollection::ESimulationTypes::FST_Rigid;
+					return SimType[BoneIdx] != FGeometryCollection::ESimulationTypes::FST_Rigid
+						|| (ViewLevel != -1 && Levels[BoneIdx] != ViewLevel);
 				}));
 			AppendSelectedBones(BonesToSelect);
 		}
@@ -2814,12 +2817,15 @@ void FScopedColorEdit::SelectBones(GeometryCollection::ESelectionMode SelectionM
 		case GeometryCollection::ESelectionMode::Clusters:
 		{
 			ResetBoneSelection();
+			int32 ViewLevel = GetViewLevel();
 			TArray<int32> BonesToSelect;
-			FGeometryCollectionClusteringUtility::GetBonesToLevel(GeometryCollectionPtr.Get(), GetViewLevel(), BonesToSelect, true, true);
+			FGeometryCollectionClusteringUtility::GetBonesToLevel(GeometryCollectionPtr.Get(), ViewLevel, BonesToSelect, true, true);
 			const TManagedArray<int32>& SimType = GeometryCollectionPtr->SimulationType;
+			const TManagedArray<int32>& Levels = GeometryCollectionPtr->GetAttribute<int32>("Level", FGeometryCollection::TransformGroup);
 			BonesToSelect.SetNum(Algo::RemoveIf(BonesToSelect, [&](int32 BoneIdx)
 			{
-				return SimType[BoneIdx] != FGeometryCollection::ESimulationTypes::FST_Clustered;
+				return SimType[BoneIdx] != FGeometryCollection::ESimulationTypes::FST_Clustered
+					|| (ViewLevel != -1 && Levels[BoneIdx] != ViewLevel);
 			}));
 			AppendSelectedBones(BonesToSelect);
 		}
