Chaos : ManagedArrayCollection
- Expose as a USTRUCT
... requires removing restrictions on copy constructor and equality operator.

diff --git a/Engine/Source/Runtime/Experimental/Chaos/Public/GeometryCollection/ManagedArray.h b/Engine/Source/Runtime/Experimental/Chaos/Public/GeometryCollection/ManagedArray.h
--- a/Engine/Source/Runtime/Experimental/Chaos/Public/GeometryCollection/ManagedArray.h
+++ b/Engine/Source/Runtime/Experimental/Chaos/Public/GeometryCollection/ManagedArray.h
@@ -12,7 +12,7 @@
 #include "Chaos/BVHParticles.h"
 #include "Math/Vector.h"
 
-class FManagedArrayCollection;
+struct FManagedArrayCollection;
 DEFINE_LOG_CATEGORY_STATIC(UManagedArrayLogging, NoLogging, All);
 
 template <typename T>
@@ -152,13 +152,6 @@ public:
 		check(false);
 	}
 
-	/** Return unmanaged copy of array with input indices. */
-	virtual FManagedArrayBase * NewCopy(const TArray<int32> & DeletionList)
-	{
-		check(false);
-		return nullptr;
-	}
-
 	/** The length of the array.*/
 	virtual int32 Num() const 
 	{
diff --git a/Engine/Source/Runtime/Experimental/Chaos/Public/GeometryCollection/ManagedArrayCollection.h b/Engine/Source/Runtime/Experimental/Chaos/Public/GeometryCollection/ManagedArrayCollection.h
--- a/Engine/Source/Runtime/Experimental/Chaos/Public/GeometryCollection/ManagedArrayCollection.h
+++ b/Engine/Source/Runtime/Experimental/Chaos/Public/GeometryCollection/ManagedArrayCollection.h
@@ -2,12 +2,15 @@
 
 #pragma once
 
+#include "CoreMinimal.h"
 #include "GeometryCollection/ManagedArray.h"
 #include "GeometryCollection/ManagedArrayTypes.h"
 #include "UObject/Object.h"
 #include "UObject/ObjectMacros.h"
 #include "GeometryCollection/GeometryCollectionSection.h"
 
+#include "ManagedArrayCollection.generated.h"
+
 class FSimulationProperties;
 
 namespace Chaos
@@ -37,19 +40,20 @@ namespace Chaos
 *
 *
 */
-class CHAOS_API FManagedArrayCollection
+USTRUCT()
+struct CHAOS_API FManagedArrayCollection
 {
+	GENERATED_USTRUCT_BODY()
 	friend FSimulationProperties;
 
 public:
-
 	FManagedArrayCollection();
 	virtual ~FManagedArrayCollection() {}
 
-	FManagedArrayCollection(const FManagedArrayCollection&) = delete;
-	FManagedArrayCollection& operator=(const FManagedArrayCollection&) = delete;
+	FManagedArrayCollection(const FManagedArrayCollection& In) { In.CopyTo(this); }
+	FManagedArrayCollection& operator=(const FManagedArrayCollection& In) { Reset(); In.CopyTo(this); return *this; }
 	FManagedArrayCollection(FManagedArrayCollection&&) = default;
-	FManagedArrayCollection& operator=(FManagedArrayCollection&&) = default;
+	FManagedArrayCollection& operator=(FManagedArrayCollection&&)= default;
 
 	static int8 Invalid;
 	typedef EManagedArrayType EArrayType;
@@ -117,6 +121,12 @@ public:
 	T* NewCopy() const
 	{
 		T* Collection = new T();
+		CopyTo(Collection);
+		return Collection;
+	}
+
+	void CopyTo(FManagedArrayCollection* Collection) const
+	{
 		if (!Map.IsEmpty())
 		{
 			for (const TTuple<FKeyType, FValueType>& Entry : Map)
@@ -135,7 +145,6 @@ public:
 				Collection->CopyAttribute(*this, Entry.Key.Get<0>(), Entry.Key.Get<1>());
 			}
 		}
-		return Collection;
 	}
 
 
@@ -307,6 +316,12 @@ public:
 		return *(static_cast<TManagedArray<T>*>(Map[Key].Value));
 	};
 
+	/**
+	* Clear the internal data. 
+	*/
+	virtual void Reset() { Map.Reset(); GroupInfo.Reset(); MakeDirty(); }
+
+
 	/**
 	* Remove the element at index and reindex the dependent arrays 
 	*/
@@ -491,7 +506,21 @@ private:
 			, bExternalValue(false)
 			, Value(&In) {};
 
-		
+		FValueType(const FValueType& Other)
+			: ArrayType(Other.ArrayType)
+			, GroupIndexDependency(Other.GroupIndexDependency)
+			, Saved(Other.Saved)
+			, bExternalValue(false)
+			, Value(nullptr)
+		{
+			if (Other.Value)
+			{
+				this->Value = NewManagedTypedArray(this->ArrayType);
+				this->Value->Resize(Other.Value->Num());
+				this->Value->Init(*Other.Value);
+			}
+		};
+
 		FValueType(FValueType&& Other)
 			: ArrayType(Other.ArrayType)
 			, GroupIndexDependency(Other.GroupIndexDependency)
@@ -513,8 +542,24 @@ private:
 			}
 		}
 
-		FValueType(const FValueType&) = delete;	//no copies because we are treating this as a simple unique ptr (not using TUniquePtr because it's only true when bExternalValue is false)
-		FValueType& operator=(const FValueType& Other) = delete;
+
+		FValueType& operator=(const FValueType& Other) 
+		{
+			this->ArrayType = Other.ArrayType;
+			this->GroupIndexDependency = Other.GroupIndexDependency;
+			this->Saved = Other.Saved;
+			this->bExternalValue = false;
+			this->Value = nullptr;
+			if (Other.Value)
+			{
+				this->Value = NewManagedTypedArray(this->ArrayType);
+				this->Value->Resize(Other.Value->Num());
+				this->Value->Init(*Other.Value);
+			}
+			return *this;
+		}
+
+
 	};
 
 	virtual void SetDefaults(FName Group, uint32 StartSize, uint32 NumElements) {};
diff --git a/Engine/Source/Runtime/Experimental/GeometryCollectionEngine/Public/GeometryCollection/GeometryCollectionObject.h b/Engine/Source/Runtime/Experimental/GeometryCollectionEngine/Public/GeometryCollection/GeometryCollectionObject.h
--- a/Engine/Source/Runtime/Experimental/GeometryCollectionEngine/Public/GeometryCollection/GeometryCollectionObject.h
+++ b/Engine/Source/Runtime/Experimental/GeometryCollectionEngine/Public/GeometryCollection/GeometryCollectionObject.h
@@ -15,7 +15,7 @@
 class UMaterialInterface;
 class UGeometryCollectionCache;
 class FGeometryCollection;
-class FManagedArrayCollection;
+struct FManagedArrayCollection;
 struct FGeometryCollectionSection;
 struct FSharedSimulationParameters;
 class UDataflow;
