Don't clear bDisabled flag from constraint mid-tick, only on Activate

diff --git a/Engine/Source/Runtime/Experimental/Chaos/Private/Chaos/Collision/PBDCollisionConstraint.cpp b/Engine/Source/Runtime/Experimental/Chaos/Private/Chaos/Collision/PBDCollisionConstraint.cpp
--- a/Engine/Source/Runtime/Experimental/Chaos/Private/Chaos/Collision/PBDCollisionConstraint.cpp
+++ b/Engine/Source/Runtime/Experimental/Chaos/Private/Chaos/Collision/PBDCollisionConstraint.cpp
@@ -533,7 +533,6 @@
 			ShapeWorldTransform1 = ShapeCoMRelativeTransform1 * ParticleCoMTransform1;
 		}
 
-		Flags.bDisabled = false;
 		ClosestManifoldPointIndex = INDEX_NONE;
 		Material.Reset();
 
@@ -605,7 +604,6 @@
 		Material.Reset();
 		ManifoldPoints.Reset();
 		ExpectedNumManifoldPoints = 0;
-		Flags.bDisabled = false;
 		Flags.bWasManifoldRestored = false;
 	}
 
@@ -619,7 +617,6 @@
 
 		// Reset current closest point
 		ClosestManifoldPointIndex = INDEX_NONE;
-		Flags.bDisabled = false;
 		Material.Reset();
 
 		// How many manifold points we expect. E.g., for Box-box this will be 4 or 1 depending on whether
diff --git a/Engine/Source/Runtime/Experimental/Chaos/Public/Chaos/Collision/PBDCollisionConstraint.h b/Engine/Source/Runtime/Experimental/Chaos/Public/Chaos/Collision/PBDCollisionConstraint.h
--- a/Engine/Source/Runtime/Experimental/Chaos/Public/Chaos/Collision/PBDCollisionConstraint.h
+++ b/Engine/Source/Runtime/Experimental/Chaos/Public/Chaos/Collision/PBDCollisionConstraint.h
@@ -367,6 +367,7 @@ namespace Chaos
 		void Activate()
 		{
 			AccumulatedImpulse = FVec3(0);
+			SetDisabled(false);
 		}
 
 		// @todo(chaos): half of this API is wrong for the new multi-point manifold constraints. Remove it
