Chaos - more logging to trap invalid convex structure data

diff --git a/Engine/Source/Runtime/Experimental/Chaos/Private/Chaos/Convex.cpp b/Engine/Source/Runtime/Experimental/Chaos/Private/Chaos/Convex.cpp
--- a/Engine/Source/Runtime/Experimental/Chaos/Private/Chaos/Convex.cpp
+++ b/Engine/Source/Runtime/Experimental/Chaos/Private/Chaos/Convex.cpp
@@ -350,7 +350,11 @@ namespace Chaos
 	// Store the structure data with the convex. This is used by manifold generation, for example
 	void FConvex::CreateStructureData(TArray<TArray<int32>>&& PlaneVertexIndices)
 	{
-		StructureData.SetPlaneVertices(MoveTemp(PlaneVertexIndices), Vertices.Num());
+		const bool bSuccess = StructureData.SetPlaneVertices(MoveTemp(PlaneVertexIndices), Vertices.Num());
+		if (!bSuccess || !StructureData.IsValid())
+		{
+			UE_LOG(LogChaos, Error, TEXT("Unable to create structure data for %s"), *ToStringFull());
+		}
 	}
 
 	void FConvex::MovePlanesAndRebuild(const FRealType InDelta)
diff --git a/Engine/Source/Runtime/Experimental/Chaos/Public/Chaos/Convex.h b/Engine/Source/Runtime/Experimental/Chaos/Public/Chaos/Convex.h
--- a/Engine/Source/Runtime/Experimental/Chaos/Public/Chaos/Convex.h
+++ b/Engine/Source/Runtime/Experimental/Chaos/Public/Chaos/Convex.h
@@ -722,7 +722,51 @@ namespace Chaos
 
 		virtual FString ToString() const override
 		{
-			return FString::Printf(TEXT("Convex: Verts: %d, Planes: %d, Margin %f"), NumVertices(), NumPlanes(), GetMargin());
+			return ToStringSummary();
+		}
+
+		// A one-line summary of the convex
+		FString ToStringSummary() const
+		{
+			return FString::Printf(TEXT("Convex: Verts: %d, Edges %d, Planes: %d, Margin %f"), NumVertices(), NumEdges(), NumPlanes(), GetMargin());
+		}
+
+		// Print all the datas
+		FString ToStringFull() const
+		{
+			FString S = ToStringSummary();
+			S.Append(TEXT("\n"));
+			for (int32 VertexIndex = 0; VertexIndex < NumVertices(); ++VertexIndex)
+			{
+				const FVec3 Vert = GetVertex(VertexIndex);
+				S.Append(FString::Printf(TEXT("  Vertex %d: [%f, %f, %f]\n"), VertexIndex, Vert.X, Vert.Y, Vert.Z));
+			}
+			for (int32 PlaneIndex = 0; PlaneIndex < NumPlanes(); ++PlaneIndex)
+			{
+				const TPlaneConcrete<FReal, 3> Plane = GetPlane(PlaneIndex);
+				S.Append(FString::Printf(TEXT("  Plane %d: Normal: [%f, %f, %f], Distance: %f, Verts: ["), PlaneIndex, Plane.Normal().X, Plane.Normal().Y, Plane.Normal().Z, FVec3::DotProduct(Plane.X(), Plane.Normal())));
+				const int32 PlaneVertexCount = NumPlaneVertices(PlaneIndex);
+				for (int32 PlaneVertexIndex = 0; PlaneVertexIndex < PlaneVertexCount; ++PlaneVertexIndex)
+				{
+					S.Append(FString::Printf(TEXT("%d"), GetPlaneVertex(PlaneIndex, PlaneVertexIndex)));
+					if (PlaneVertexIndex < PlaneVertexCount - 1)
+					{
+						S.Append(FString::Printf(TEXT(", ")));
+					}
+				}
+				S.Append(FString::Printf(TEXT("]\n")));
+			}
+			S.Append(FString::Printf(TEXT("  Edges: ")));
+			for (int32 EdgeIndex = 0; EdgeIndex < NumEdges(); ++EdgeIndex)
+			{
+				S.Append(FString::Printf(TEXT("[%d, %d]"), GetEdgeVertex(EdgeIndex, 0), GetEdgeVertex(EdgeIndex, 1)));
+				if (EdgeIndex < NumEdges() - 1)
+				{
+					S.Append(FString::Printf(TEXT(", ")));
+				}
+			}
+			S.Append(FString::Printf(TEXT("\n")));
+			return S;
 		}
 
 		const TArray<FVec3Type>& GetVertices() const
diff --git a/Engine/Source/Runtime/Experimental/Chaos/Public/Chaos/ConvexHalfEdgeStructureData.h b/Engine/Source/Runtime/Experimental/Chaos/Public/Chaos/ConvexHalfEdgeStructureData.h
--- a/Engine/Source/Runtime/Experimental/Chaos/Public/Chaos/ConvexHalfEdgeStructureData.h
+++ b/Engine/Source/Runtime/Experimental/Chaos/Public/Chaos/ConvexHalfEdgeStructureData.h
@@ -338,6 +338,8 @@ namespace Chaos
 
 			if ((InPlaneVertices.Num() > MaxIndex) || (HalfEdgeCount > MaxIndex) || (InNumVertices > MaxIndex))
 			{
+				// We should never get here. See GetRequiredIndexType::GetRequiredIndexType which calls CanMake() on eachn index type until it fits
+				UE_LOG(LogChaos, Error, TEXT("Unable to create structure data for convex. MaxIndex too small (%d bytes) for Planes: %d HalfEdges: %d Verts: %d"), sizeof(FIndex), InPlaneVertices.Num(), HalfEdgeCount, InNumVertices);
 				return false;
 			}
 
@@ -434,6 +436,8 @@ namespace Chaos
 
 			BuildVertexPlanes();
 
+			BuildUniqueEdgeList();
+
 			return true;
 		}
 
@@ -446,13 +450,17 @@ namespace Chaos
 			Ar << Vertices;
 
 			const bool bHasUniqueEdgeList = Ar.CustomVer(FPhysicsObjectVersion::GUID) >= FPhysicsObjectVersion::ChaosConvexHasUniqueEdgeSet;
-			if (Ar.IsLoading() && !bHasUniqueEdgeList)
+			if (bHasUniqueEdgeList)
 			{
-				BuildUniqueEdgeList();
+				Ar << Edges;
 			}
-			else
+
+			// Handle older data without the edge list and also an issue where some assets were saved without having their EdgeList generated (now fixed)
+			// Avoid adding a new custom version for that fix because we want to integrate this into other streams
+			if (Ar.IsLoading() && (Edges.Num() == 0) && (HalfEdges.Num() > 0))
 			{
-				Ar << Edges;
+				BuildUniqueEdgeList();
+				ensureMsgf(Edges.Num() > 0, TEXT("Invalid edge data on convex in Load"));
 			}
 
 			if (Ar.IsLoading())
diff --git a/Engine/Source/Runtime/Experimental/Chaos/Public/Chaos/ConvexStructureData.h b/Engine/Source/Runtime/Experimental/Chaos/Public/Chaos/ConvexStructureData.h
--- a/Engine/Source/Runtime/Experimental/Chaos/Public/Chaos/ConvexStructureData.h
+++ b/Engine/Source/Runtime/Experimental/Chaos/Public/Chaos/ConvexStructureData.h
@@ -202,6 +202,14 @@ namespace Chaos
 				[]() { return (int32)INDEX_NONE; });
 		}
 
+
+		inline int32 NumHalfEdges() const
+		{
+			return ConstDataOp(
+				[](const auto& ConcreteData) { return ConcreteData.NumHalfEdges(); },
+				[]() { return 0; });
+		}
+
 		inline int32 NumEdges() const
 		{
 			return ConstDataOp(
@@ -230,14 +238,14 @@ namespace Chaos
 		}
 
 		// Initialize the structure data from the set of vertices for each face of the convex
-		void SetPlaneVertices(const TArray<TArray<int32>>& InPlaneVertices, int32 NumVerts)
+		bool SetPlaneVertices(const TArray<TArray<int32>>& InPlaneVertices, int32 NumVerts)
 		{
 			const EIndexType NewIndexType = GetRequiredIndexType(InPlaneVertices, NumVerts);
 			CreateDataContainer(NewIndexType);
 
-			NonConstDataOp(
-				[InPlaneVertices, NumVerts](auto& ConcreteData) { ConcreteData.SetPlaneVertices(InPlaneVertices, NumVerts); },
-				[]() {});
+			return NonConstDataOp(
+				[InPlaneVertices, NumVerts](auto& ConcreteData) { return ConcreteData.SetPlaneVertices(InPlaneVertices, NumVerts); },
+				[]() { return true; });
 		}
 
 		void Serialize(FArchive& Ar)
