chaos geometry collection : use slerp for rotation interpolation and use external strain for crumbling methods

diff --git a/Engine/Source/Runtime/Experimental/Chaos/Private/PhysicsProxy/GeometryCollectionPhysicsProxy.cpp b/Engine/Source/Runtime/Experimental/Chaos/Private/PhysicsProxy/GeometryCollectionPhysicsProxy.cpp
--- a/Engine/Source/Runtime/Experimental/Chaos/Private/PhysicsProxy/GeometryCollectionPhysicsProxy.cpp
+++ b/Engine/Source/Runtime/Experimental/Chaos/Private/PhysicsProxy/GeometryCollectionPhysicsProxy.cpp
@@ -81,6 +81,12 @@ FAutoConsoleVariableRef CVarGeometryCollectionAlwaysGenerateConnectionGraph(
 	bGeometryCollectionAlwaysGenerateConnectionGraph,
 	TEXT("When enabled, always  generate the cluster's connection graph instead of using the rest collection stored one - Note: this should only be used for troubleshooting.[def: false]"));
 
+bool bGeometryCollectionCrumbleUsingStrain = true;
+FAutoConsoleVariableRef CVarGeometryCollectionCrumbleUsingStrain(
+	TEXT("p.GeometryCollection.CrumbleUsingStrain"),
+	bGeometryCollectionCrumbleUsingStrain,
+	TEXT("When enabled, cluster crumbling is using strain instead of directly calling the clustering methods.[def: true]"));
+
 DEFINE_LOG_CATEGORY_STATIC(UGCC_LOG, Error, All);
 
 static const FSharedSimulationSizeSpecificData& GetSizeSpecificData(const TArray<FSharedSimulationSizeSpecificData>& SizeSpecificData, const FGeometryCollection& RestCollection, const int32 TransformIndex, const FBox& BoundingBox);
@@ -1526,6 +1532,20 @@ void FGeometryCollectionPhysicsProxy::DisableParticles(TArray<int32>&& Transform
 			});
 	}
 }
+static void ApplyStrainToAllClusterChildren(Chaos::FRigidClustering& Clustering, Chaos::FPBDRigidClusteredParticleHandle* ClusteredParent, Chaos::FReal StrainValue)
+{
+	if (TArray<Chaos::FPBDRigidParticleHandle*>* Children = Clustering.GetChildrenMap().Find(ClusteredParent))
+	{
+		for (Chaos::FPBDRigidParticleHandle* ChildHandle: *Children)
+		{
+			if (Chaos::FPBDRigidClusteredParticleHandle* ClusteredChildHandle = ChildHandle->CastToClustered())
+			{
+				ClusteredChildHandle->SetExternalStrain(FMath::Max(ClusteredChildHandle->GetExternalStrain(), StrainValue));
+			}
+		}
+	}
+}
+
 
 void FGeometryCollectionPhysicsProxy::BreakInternalClusterParents(TArray<int32>&& TransformGroupIndices)
 {
@@ -1554,7 +1574,14 @@ void FGeometryCollectionPhysicsProxy::BreakInternalClusterParents(TArray<int32>&
 					Chaos::FRigidClustering& Clustering = RBDSolver->GetEvolution()->GetRigidClustering();
 					if (!ClusteredParent->Disabled() && Clustering.GetChildrenMap().Contains(ClusteredParent))
 					{
-						Clustering.ReleaseClusterParticles(ClusteredParent, true /* bForceRelease */);
+						if (bGeometryCollectionCrumbleUsingStrain)
+						{
+							ApplyStrainToAllClusterChildren(Clustering, ClusteredParent, TNumericLimits<Chaos::FReal>::Max());
+						}
+						else
+						{
+							Clustering.ReleaseClusterParticles(ClusteredParent, true /* bForceRelease */);
+						}
 					}
 				}
 			}
@@ -1577,7 +1604,14 @@ void FGeometryCollectionPhysicsProxy::BreakClusters(TArray<int32>&& TransformGro
 				{
 					if (!ClusterHandle->Disabled() && Clustering.GetChildrenMap().Contains(ClusterHandle))
 					{
-						Clustering.ReleaseClusterParticles(ClusterHandle, true /* bForceRelease */);
+						if (bGeometryCollectionCrumbleUsingStrain)
+						{
+							ApplyStrainToAllClusterChildren(Clustering, ClusterHandle, TNumericLimits<Chaos::FReal>::Max());
+						}
+						else
+						{
+							Clustering.ReleaseClusterParticles(ClusterHandle, true /* bForceRelease */);
+						}
 					}
 				}
 			}
@@ -2277,7 +2311,7 @@ bool FGeometryCollectionPhysicsProxy::PullFromPhysicsState(const Chaos::FDirtyGe
 					}
 					
 					const Chaos::FRotation3 OldR = GTParticle.R();
-					const Chaos::FRotation3 NewR = FMath::Lerp(Prev.ParticleRs[TransformGroupIndex], Next.ParticleRs[TransformGroupIndex], *Alpha);
+					const Chaos::FRotation3 NewR = FQuat::Slerp(Prev.ParticleRs[TransformGroupIndex], Next.ParticleRs[TransformGroupIndex], *Alpha);
 					const bool bNeedUpdateR = (NewR != OldR); 
 					if (bNeedUpdateR)
 					{
