Chaos geometry collection : add support for CCD , linear and angular damping

diff --git a/Engine/Source/Runtime/Experimental/Chaos/Private/PhysicsProxy/GeometryCollectionPhysicsProxy.cpp b/Engine/Source/Runtime/Experimental/Chaos/Private/PhysicsProxy/GeometryCollectionPhysicsProxy.cpp
--- a/Engine/Source/Runtime/Experimental/Chaos/Private/PhysicsProxy/GeometryCollectionPhysicsProxy.cpp
+++ b/Engine/Source/Runtime/Experimental/Chaos/Private/PhysicsProxy/GeometryCollectionPhysicsProxy.cpp
@@ -752,30 +752,6 @@ void FGeometryCollectionPhysicsProxy::CreateNonClusteredParticles(Chaos::FPBDRig
 	}
 }
 
-void FGeometryCollectionPhysicsProxy::SetGravityEnabled(const Chaos::FPBDRigidsSolver& RigidsSolver, bool bEnabled)
-{
-	Chaos::FPerParticleGravity& GravityForces = RigidsSolver.GetEvolution()->GetGravityForces();
-	for (int32 HandleIdx = 0; HandleIdx < SolverParticleHandles.Num(); ++HandleIdx)
-	{
-		if (Chaos::TPBDRigidParticleHandle<Chaos::FReal, 3>* Handle = SolverParticleHandles[HandleIdx])
-		{
-			Handle->SetGravityEnabled(bEnabled);
-		}
-	}
-}
-
-void FGeometryCollectionPhysicsProxy::EnableInertiaConditioning(const Chaos::FPBDRigidsSolver& RigidsSolver, bool bEnabled)
-{
-	Chaos::FPerParticleGravity& GravityForces = RigidsSolver.GetEvolution()->GetGravityForces();
-	for (int32 HandleIdx = 0; HandleIdx < SolverParticleHandles.Num(); ++HandleIdx)
-	{
-		if (Chaos::TPBDRigidParticleHandle<Chaos::FReal, 3>* Handle = SolverParticleHandles[HandleIdx])
-		{
-			Handle->SetInertiaConditioningEnabled(bEnabled);
-		}
-	}
-}
-
 void FGeometryCollectionPhysicsProxy::SetSleepingState(const Chaos::FPBDRigidsSolver& RigidsSolver)
 {
 	for (FClusterHandle* Handle: SolverParticleHandles)
@@ -1121,9 +1097,6 @@ void FGeometryCollectionPhysicsProxy::InitializeBodiesPT(Chaos::FPBDRigidsSolver
 		}
 #endif // TODO_REIMPLEMENT_RIGID_CACHING
 
-		const bool bEnableGravity = Parameters.EnableGravity && !DisableGeometryCollectionGravity;
-		SetGravityEnabled(*RigidsSolver, bEnableGravity);
-
 		const bool bStartSleeping =
 				(Parameters.ObjectType == EObjectStateTypeEnum::Chaos_Object_Sleeping
 			 || (Parameters.ObjectType == EObjectStateTypeEnum::Chaos_Object_Dynamic && !Parameters.StartAwake));
@@ -1132,8 +1105,20 @@ void FGeometryCollectionPhysicsProxy::InitializeBodiesPT(Chaos::FPBDRigidsSolver
 			SetSleepingState(*RigidsSolver);
 		}
 
-		EnableInertiaConditioning(*RigidsSolver, Parameters.UseInertiaConditioning);
-
+		// apply various features on the handles 
+		const bool bEnableGravity = Parameters.EnableGravity && !DisableGeometryCollectionGravity;
+		for (Chaos::FPBDRigidParticleHandle* Handle: SolverParticleHandles)
+		{
+			if (Handle)
+			{
+				Handle->SetGravityEnabled(bEnableGravity);
+				Handle->SetCCDEnabled(Parameters.UseCCD);
+				Handle->SetInertiaConditioningEnabled(Parameters.UseInertiaConditioning);
+				Handle->SetLinearEtherDrag(Parameters.LinearDamping);
+				Handle->SetAngularEtherDrag(Parameters.AngularDamping);
+			}
+		}
+		
 		// call DirtyParticle to make sure the acceleration structure is up to date with all the changes happening here
 		DirtyAllParticles(*RigidsSolver);
 
diff --git a/Engine/Source/Runtime/Experimental/Chaos/Public/GeometryCollection/GeometryCollectionSimulationCoreTypes.h b/Engine/Source/Runtime/Experimental/Chaos/Public/GeometryCollection/GeometryCollectionSimulationCoreTypes.h
--- a/Engine/Source/Runtime/Experimental/Chaos/Public/GeometryCollection/GeometryCollectionSimulationCoreTypes.h
+++ b/Engine/Source/Runtime/Experimental/Chaos/Public/GeometryCollection/GeometryCollectionSimulationCoreTypes.h
@@ -201,6 +201,9 @@ struct FSimulationParameters
 		, RemoveOnFractureEnabled(false)
 		, EnableGravity(true)
 		, UseInertiaConditioning(true)
+		, UseCCD(false)
+		, LinearDamping(0.01f)
+		, AngularDamping(0)
 		, SimulationFilterData()
 		, QueryFilterData()
 		, UserData(nullptr)
@@ -242,6 +245,9 @@ struct FSimulationParameters
 		, RemoveOnFractureEnabled(Other.RemoveOnFractureEnabled)
 		, EnableGravity(Other.EnableGravity)
 		, UseInertiaConditioning(Other.UseInertiaConditioning)
+		, UseCCD(Other.UseCCD)
+		, LinearDamping(Other.LinearDamping)
+		, AngularDamping(Other.AngularDamping)
 		, SimulationFilterData(Other.SimulationFilterData)
 		, QueryFilterData(Other.QueryFilterData)
 		, UserData(Other.UserData)
@@ -306,6 +312,9 @@ struct FSimulationParameters
 	
 	bool EnableGravity;
 	bool UseInertiaConditioning;
+	bool UseCCD;
+	float LinearDamping;
+	float AngularDamping;
 
 	FCollisionFilterData SimulationFilterData;
 	FCollisionFilterData QueryFilterData;
diff --git a/Engine/Source/Runtime/Experimental/Chaos/Public/PhysicsProxy/GeometryCollectionPhysicsProxy.h b/Engine/Source/Runtime/Experimental/Chaos/Public/PhysicsProxy/GeometryCollectionPhysicsProxy.h
--- a/Engine/Source/Runtime/Experimental/Chaos/Public/PhysicsProxy/GeometryCollectionPhysicsProxy.h
+++ b/Engine/Source/Runtime/Experimental/Chaos/Public/PhysicsProxy/GeometryCollectionPhysicsProxy.h
@@ -278,8 +278,6 @@ protected:
 		const Chaos::FClusterCreationParameters & Parameters,
 		const Chaos::FUniqueIdx* ExistingIndex);
 
-	void SetGravityEnabled(const Chaos::FPBDRigidsSolver& RigidsSolver, bool bEnabled);
-	void EnableInertiaConditioning(const Chaos::FPBDRigidsSolver& RigidsSolver, bool bEnabled);
 	void SetSleepingState(const Chaos::FPBDRigidsSolver& RigidsSolver);
 	void DirtyAllParticles(const Chaos::FPBDRigidsSolver& RigidsSolver);
 	
diff --git a/Engine/Source/Runtime/Experimental/GeometryCollectionEngine/Private/GeometryCollection/GeometryCollectionComponent.cpp b/Engine/Source/Runtime/Experimental/GeometryCollectionEngine/Private/GeometryCollection/GeometryCollectionComponent.cpp
--- a/Engine/Source/Runtime/Experimental/GeometryCollectionEngine/Private/GeometryCollection/GeometryCollectionComponent.cpp
+++ b/Engine/Source/Runtime/Experimental/GeometryCollectionEngine/Private/GeometryCollection/GeometryCollectionComponent.cpp
@@ -2299,6 +2299,9 @@ void UGeometryCollectionComponent::RegisterAndInitializePhysicsProxy()
 		SimulationParameters.RemoveOnFractureEnabled = SimulationParameters.Shared.RemoveOnFractureIndices.Num() > 0;
 		SimulationParameters.EnableGravity = BodyInstance.bEnableGravity;
 		SimulationParameters.UseInertiaConditioning = BodyInstance.IsInertiaConditioningEnabled();
+		SimulationParameters.UseCCD = BodyInstance.bUseCCD;
+		SimulationParameters.LinearDamping = BodyInstance.LinearDamping;
+		SimulationParameters.AngularDamping = BodyInstance.AngularDamping;
 		SimulationParameters.WorldTransform = GetComponentToWorld();
 		SimulationParameters.UserData = static_cast<void*>(&PhysicsUserData);
 
diff --git a/Engine/Source/Runtime/Experimental/GeometryCollectionEngine/Public/GeometryCollection/GeometryCollectionObject.h b/Engine/Source/Runtime/Experimental/GeometryCollectionEngine/Public/GeometryCollection/GeometryCollectionObject.h
--- a/Engine/Source/Runtime/Experimental/GeometryCollectionEngine/Public/GeometryCollection/GeometryCollectionObject.h
+++ b/Engine/Source/Runtime/Experimental/GeometryCollectionEngine/Public/GeometryCollection/GeometryCollectionObject.h
@@ -566,12 +566,12 @@ public:
 	UPROPERTY(EditAnywhere, BlueprintReadOnly, Category = Removal, meta = (DisplayName = "Removal Duration", EditCondition="bRemoveOnMaxSleep"))
 	FVector2D RemovalDuration;
 
-	/** when on non-sleeping, slow moving pieces will be considered as sleeping, this help removal of jittery but not really moving objects. */
-	UPROPERTY(EditAnywhere, BlueprintReadOnly, Category = Removal, meta = (DisplayName = "Slow*-Moving as sleeping", EditCondition="bRemoveOnMaxSleep"))
+	/** when on non-sleeping, slow moving pieces will be considered as sleeping, this helps removal of jittery but not really moving objects. */
+	UPROPERTY(EditAnywhere, BlueprintReadOnly, Category = Removal, meta = (DisplayName = "Slow-Moving as sleeping", EditCondition="bRemoveOnMaxSleep"))
 	bool bSlowMovingAsSleeping;
 
 	/** When slow moving detection is on, this defines the linear velocity thresholds in cm/s to consider the object as sleeping . */
-	UPROPERTY(EditAnywhere, BlueprintReadOnly, Category = Removal, meta = (DisplayName = "Slow*-Moving Velocity Threshold", EditCondition="bRemoveOnMaxSleep && bSlowMovingAsSleeping"))
+	UPROPERTY(EditAnywhere, BlueprintReadOnly, Category = Removal, meta = (DisplayName = "Slow-Moving Velocity Threshold", EditCondition="bRemoveOnMaxSleep && bSlowMovingAsSleeping"))
 	float SlowMovingVelocityThreshold;
 	
 	/*
