Chaos detsruction : imporved propgation breaking model
- Add p.GraphPropagationBasedCollisionFactor cvar to control the ratio of energy propagated
- moved propagation logic after breaking logic to be order independent

diff --git a/Engine/Source/Runtime/Experimental/Chaos/Private/Chaos/PBDRigidClustering.cpp b/Engine/Source/Runtime/Experimental/Chaos/Private/Chaos/PBDRigidClustering.cpp
--- a/Engine/Source/Runtime/Experimental/Chaos/Private/Chaos/PBDRigidClustering.cpp
+++ b/Engine/Source/Runtime/Experimental/Chaos/Private/Chaos/PBDRigidClustering.cpp
@@ -50,6 +50,10 @@ namespace Chaos
 
 	int32 GraphPropagationBasedCollisionImpulseProcessing = 0;
 	FAutoConsoleVariableRef CVarGraphPropagationBasedCollisionImpulseProcessing(TEXT("p.GraphPropagationBasedCollisionImpulseProcessing"), GraphPropagationBasedCollisionImpulseProcessing, TEXT("when processing collision impulse toc ompute strain, pick the closest child from the impact point and propagate using the connection graph [def: 0]"));
+
+	float GraphPropagationBasedCollisionFactor = 1;
+	FAutoConsoleVariableRef CVarGraphPropagationBasedCollisionFactor(TEXT("p.GraphPropagationBasedCollisionFactor"), GraphPropagationBasedCollisionFactor, TEXT("when p.GraphPropagationBasedCollisionImpulseProcessing is on, the percentage [0-1] of remaining damage that is distributed to the connected pieces"));
+
 	
 	//==========================================================================
 	// TPBDRigidClustering
@@ -592,32 +596,11 @@ namespace Chaos
 				continue;
 			}
 
+			// @todo(chaos) eventually shoudl get rid of collision impulse array and only use external strain
 			const FReal MaxAppliedStrain = FMath::Max(Child->CollisionImpulses(), Child->GetExternalStrain());
 			if ((MaxAppliedStrain >= Child->Strain()) || bForceRelease)
 			{
 				//UE_LOG(LogTemp, Warning, TEXT("Releasing child %d from parent %p due to strain %.5f Exceeding internal strain %.5f (Source: %s)"), ChildIdx, ClusteredParticle, ChildStrain, Child->Strain(), bForceRelease ? TEXT("Forced by caller") : ExternalStrainMap ? TEXT("External") : TEXT("Collision"));
-
-				// if necessary propagate through the graph
-				if (GraphPropagationBasedCollisionImpulseProcessing && !bForceRelease)
-				{
-					const FReal RemainingStrain = (MaxAppliedStrain - Child->Strain());
-					if (RemainingStrain > 0)
-					{
-						// todo(chaos) : could do better and have something weighted on distance with a falloff maybe ?  
-						const FReal RemainingStrainPerSibling = RemainingStrain / Child->ConnectivityEdges().Num();
-						for (const TConnectivityEdge<FReal>& Edge: Child->ConnectivityEdges())
-						{
-							if (Edge.Sibling)
-							{
-								if (FPBDRigidClusteredParticleHandle* ClusteredSibling = Edge.Sibling->CastToClustered())
-								{
-									// todo(chaos) this may currently be non optimal as we are in the apply loop and this may be cleareed right after 
-									ClusteredSibling->CollisionImpulses() +=  RemainingStrainPerSibling;
-								}
-							}
-						}
-					}
-				}
 				
 				// The piece that hits just breaks off - we may want more control 
 				// by looking at the edges of this piece which would give us cleaner 
@@ -638,13 +621,55 @@ namespace Chaos
 				{
 					SendBreakingEvent(Child);
 				}
-			}
 
-			// no need anymore of the external strain let's clear it
-			// @todo(chaos) impulse collision are right now cleared somewhere else, but we should consolidate that eventually
-			Child->ClearExternalStrain();
+				if (GraphPropagationBasedCollisionImpulseProcessing && !bForceRelease)
+				{
+					// keep max applied strain in external strain
+					// @todo(chaos) eventually shoudl get rid of collision impulse array and only use external strain
+					Child->SetExternalStrain(MaxAppliedStrain);
+				}
+				else
+				{
+					Child->ClearExternalStrain();	
+				}
+			}
+			else
+			{
+				// no need anymore of the external strain let's clear it
+				Child->ClearExternalStrain();				
+			}
 		}
 
+		// if necessary propagate through the graph
+		if (GraphPropagationBasedCollisionImpulseProcessing && !bForceRelease)
+		{
+			for (FPBDRigidParticleHandle* ActivatedChild: ActivatedChildren)
+			{
+				FPBDRigidClusteredParticleHandle* ClusteredChild = ActivatedChild->CastToClustered();
+
+				const FReal RemainingStrain = (ClusteredChild->GetExternalStrain() - ClusteredChild->Strain());
+				if (RemainingStrain > 0)
+				{
+					const FReal AdjustedRemainingStrain = GraphPropagationBasedCollisionFactor * RemainingStrain;
+					// todo(chaos) : could do better and have something weighted on distance with a falloff maybe ?  
+					const FReal RemainingStrainPerSibling = AdjustedRemainingStrain / ClusteredChild->ConnectivityEdges().Num();
+					for (const TConnectivityEdge<FReal>& Edge: ClusteredChild->ConnectivityEdges())
+					{
+						if (Edge.Sibling)
+						{
+							if (FPBDRigidClusteredParticleHandle* ClusteredSibling = Edge.Sibling->CastToClustered())
+							{
+								// todo(chaos) this may currently be non optimal as we are in the apply loop and this may be cleareed right after
+								ClusteredSibling->SetExternalStrain(FMath::Max(ClusteredSibling->GetExternalStrain(), RemainingStrainPerSibling));
+							}
+						}
+					}
+				}
+				// finally reset the external strain
+				ClusteredChild->ClearExternalStrain();
+			}
+		}
+		
 		if (ActivatedChildren.Num() > 0)
 		{
 			if (Children.Num() == 0)
