Dataflow - Connections as UStruct

diff --git a/Engine/Plugins/Experimental/Dataflow/Source/DataflowEditor/Private/Dataflow/DataflowEditorActions.cpp b/Engine/Plugins/Experimental/Dataflow/Source/DataflowEditor/Private/Dataflow/DataflowEditorActions.cpp
--- a/Engine/Plugins/Experimental/Dataflow/Source/DataflowEditor/Private/Dataflow/DataflowEditorActions.cpp
+++ b/Engine/Plugins/Experimental/Dataflow/Source/DataflowEditor/Private/Dataflow/DataflowEditorActions.cpp
@@ -65,9 +65,9 @@ void FDataflowEditorCommands::EvaluateNodes(const FGraphPanelSelectionSet& Selec
 					{
 						if (DataflowNode->GetOutputs().Num())
 						{
-							for (Dataflow::FConnection* NodeOutput : DataflowNode->GetOutputs())
+							for (FDataflowConnection* NodeOutput : DataflowNode->GetOutputs())
 							{
-								Evaluate(DataflowNode.Get(), NodeOutput);
+								Evaluate(DataflowNode.Get(), (FDataflowOutput*)NodeOutput);
 							}
 						}
 						else
diff --git a/Engine/Plugins/Experimental/Dataflow/Source/DataflowEditor/Private/Dataflow/DataflowGraphEditor.cpp b/Engine/Plugins/Experimental/Dataflow/Source/DataflowEditor/Private/Dataflow/DataflowGraphEditor.cpp
--- a/Engine/Plugins/Experimental/Dataflow/Source/DataflowEditor/Private/Dataflow/DataflowGraphEditor.cpp
+++ b/Engine/Plugins/Experimental/Dataflow/Source/DataflowEditor/Private/Dataflow/DataflowGraphEditor.cpp
@@ -64,10 +64,10 @@ void SDataflowGraphEditor::EvaluateNode()
 	}
 	else
 	{
-		FDataflowEditorCommands::FGraphEvaluationCallback LocalEvaluateCallback = [](FDataflowNode* Node, Dataflow::FConnection* Out)
+		FDataflowEditorCommands::FGraphEvaluationCallback LocalEvaluateCallback = [](FDataflowNode* Node, FDataflowOutput* Out)
 		{
-			float EvalTime = FGameTime::GetTimeSinceAppStart().GetRealTimeSeconds();
-			return Node->Evaluate(Dataflow::FContext(EvalTime), Out);
+			Dataflow::FContext Context(FGameTime::GetTimeSinceAppStart().GetRealTimeSeconds());
+			Node->Evaluate(Context, Out);
 		};
 
 		FDataflowEditorCommands::EvaluateNodes(GetSelectedNodes(), LocalEvaluateCallback);
diff --git a/Engine/Plugins/Experimental/Dataflow/Source/DataflowEditor/Public/Dataflow/DataflowEditorActions.h b/Engine/Plugins/Experimental/Dataflow/Source/DataflowEditor/Public/Dataflow/DataflowEditorActions.h
--- a/Engine/Plugins/Experimental/Dataflow/Source/DataflowEditor/Public/Dataflow/DataflowEditorActions.h
+++ b/Engine/Plugins/Experimental/Dataflow/Source/DataflowEditor/Public/Dataflow/DataflowEditorActions.h
@@ -6,11 +6,7 @@
 #include "Framework/Commands/Commands.h"
 #include "Styling/AppStyle.h"
 
-namespace Dataflow 
-{
-	class FConnection;
-}
-
+struct FDataflowOutput;
 class IStructureDetailsView;
 class UDataflow;
 class UDataflowEdNode;
@@ -44,7 +40,7 @@ public:
 class DATAFLOWEDITOR_API FDataflowEditorCommands
 {
 public:
-	typedef TFunction<void(FDataflowNode*, Dataflow::FConnection*)> FGraphEvaluationCallback;
+	typedef TFunction<void(FDataflowNode*, FDataflowOutput*)> FGraphEvaluationCallback;
 	typedef TFunction<void(const FGeometry& MyGeometry, const FDragDropEvent& DragDropEvent)> FOnDragDropEventCallback;
 
 	static void Register();
diff --git a/Engine/Plugins/Experimental/Dataflow/Source/DataflowEditor/Public/Dataflow/DataflowGraphEditor.h b/Engine/Plugins/Experimental/Dataflow/Source/DataflowEditor/Public/Dataflow/DataflowGraphEditor.h
--- a/Engine/Plugins/Experimental/Dataflow/Source/DataflowEditor/Public/Dataflow/DataflowGraphEditor.h
+++ b/Engine/Plugins/Experimental/Dataflow/Source/DataflowEditor/Public/Dataflow/DataflowGraphEditor.h
@@ -21,9 +21,9 @@
 class FDataflowEditorToolkit;
 class FDataflowSNodeFactory;
 class UDataflow;
+struct FDataflowConnection;
 namespace Dataflow {
 	class FContext;
-	class FConnection;
 }
 /**
  * The SDataflowGraphEditpr class is a specialization of SGraphEditor
diff --git a/Engine/Plugins/Experimental/Dataflow/Source/DataflowNodes/Private/Dataflow/DataflowSkeletalMeshNodes.cpp b/Engine/Plugins/Experimental/Dataflow/Source/DataflowNodes/Private/Dataflow/DataflowSkeletalMeshNodes.cpp
--- a/Engine/Plugins/Experimental/Dataflow/Source/DataflowNodes/Private/Dataflow/DataflowSkeletalMeshNodes.cpp
+++ b/Engine/Plugins/Experimental/Dataflow/Source/DataflowNodes/Private/Dataflow/DataflowSkeletalMeshNodes.cpp
@@ -17,31 +17,42 @@ namespace Dataflow
 		DATAFLOW_NODE_REGISTER_CREATION_FACTORY(FSkeletalMeshBoneDataflowNode);
 	}
 }
-void FGetSkeletalMeshDataflowNode::Evaluate(const Dataflow::FContext& Context, Dataflow::FConnection* Out) const
+
+void FGetSkeletalMeshDataflowNode::Evaluate(Dataflow::FContext& Context, const FDataflowOutput* Out) const
 {
-	SkeletalMeshOut.SetValue(nullptr, Context);
-	if (SkeletalMesh)
-	{
-		SkeletalMeshOut.SetValue(SkeletalMesh, Context);
-	}
-	else if (const Dataflow::FEngineContext* EngineContext = Context.AsType<Dataflow::FEngineContext>())
+	typedef TObjectPtr<const USkeletalMesh> DataType;
+	if (Out->IsA<DataType>(&SkeletalMesh))
 	{
-		if (const USkeletalMesh* SkeletalMeshFromOwner = Dataflow::Reflection::FindObjectPtrProperty<USkeletalMesh>(
-			EngineContext->Owner, PropertyName))
+		GetOutput(&SkeletalMesh)->SetValue<DataType>(SkeletalMesh, Context); // prime to avoid ensure
+
+		if (SkeletalMesh)
+		{
+			GetOutput(&SkeletalMesh)->SetValue<DataType>(SkeletalMesh, Context);
+		}
+		else if (const Dataflow::FEngineContext* EngineContext = Context.AsType<Dataflow::FEngineContext>())
 		{
-			SkeletalMeshOut.SetValue(SkeletalMeshFromOwner, Context);
+			if (const USkeletalMesh* SkeletalMeshFromOwner = Dataflow::Reflection::FindObjectPtrProperty<USkeletalMesh>(
+				EngineContext->Owner, PropertyName))
+			{
+				GetOutput(&SkeletalMesh)->SetValue<DataType>(DataType(SkeletalMeshFromOwner), Context);
+			}
 		}
 	}
 }
 
-void FSkeletalMeshBoneDataflowNode::Evaluate(const Dataflow::FContext& Context, Dataflow::FConnection* Out) const
+void FSkeletalMeshBoneDataflowNode::Evaluate(Dataflow::FContext& Context, const FDataflowOutput* Out) const
 {
-	BoneIndexOut.SetValue(INDEX_NONE, Context);
-
-	if (const USkeletalMesh* SkeletalMesh = SkeletalMeshIn.GetValue(Context))
+	typedef TObjectPtr<const USkeletalMesh> InDataType;
+	if (Out->IsA<int>(&BoneIndexOut))
 	{
-		int32 Index = SkeletalMesh->GetRefSkeleton().FindBoneIndex(BoneName);
-		BoneIndexOut.SetValue(Index, Context);
+		GetOutput(&BoneIndexOut)->SetValue<int>(INDEX_NONE, Context); // prime to avoid ensure
+
+		if( InDataType InSkeletalMesh = GetInput(&SkeletalMesh)->GetValue<InDataType>(Context, SkeletalMesh) )
+		{
+			int32 Index = InSkeletalMesh->GetRefSkeleton().FindBoneIndex(BoneName);
+			Out->SetValue<int>(Index, Context);
+		}
+
 	}
 }
 
diff --git a/Engine/Plugins/Experimental/Dataflow/Source/DataflowNodes/Private/Dataflow/DataflowStaticMeshNodes.cpp b/Engine/Plugins/Experimental/Dataflow/Source/DataflowNodes/Private/Dataflow/DataflowStaticMeshNodes.cpp
--- a/Engine/Plugins/Experimental/Dataflow/Source/DataflowNodes/Private/Dataflow/DataflowStaticMeshNodes.cpp
+++ b/Engine/Plugins/Experimental/Dataflow/Source/DataflowNodes/Private/Dataflow/DataflowStaticMeshNodes.cpp
@@ -11,20 +11,26 @@ namespace Dataflow
 	}
 }
 
-void FGetStaticMeshDataflowNode::Evaluate(const Dataflow::FContext& Context, Dataflow::FConnection* Out) const
+void FGetStaticMeshDataflowNode::Evaluate(Dataflow::FContext& Context, const FDataflowOutput* Out) const
 {
-	StaticMeshOut.SetValue(nullptr, Context);
-	if (StaticMesh)
+	typedef TObjectPtr<const UStaticMesh> DataType;
+	if (Out->IsA<DataType>(&StaticMesh))
 	{
-		StaticMeshOut.SetValue(StaticMesh, Context);
-	}
-	else if (const Dataflow::FEngineContext* EngineContext = Context.AsType<Dataflow::FEngineContext>())
-	{
-		if (const UStaticMesh* StaticMeshFromOwner = Dataflow::Reflection::FindObjectPtrProperty<UStaticMesh>(
-			EngineContext->Owner, PropertyName))
+		GetOutput(&StaticMesh)->SetValue<DataType>(StaticMesh, Context); // prime to avoid ensure
+
+		if (StaticMesh)
 		{
-			StaticMeshOut.SetValue(StaticMeshFromOwner, Context);
+			GetOutput(&StaticMesh)->SetValue<DataType>(StaticMesh, Context);
+		}
+		else if (const Dataflow::FEngineContext* EngineContext = Context.AsType<Dataflow::FEngineContext>())
+		{
+			if (const UStaticMesh* StaticMeshFromOwner = Dataflow::Reflection::FindObjectPtrProperty<UStaticMesh>(
+				EngineContext->Owner, PropertyName))
+			{
+				GetOutput(&StaticMesh)->SetValue<DataType>(DataType(StaticMeshFromOwner), Context);
+			}
 		}
 	}
 }
 
+
diff --git a/Engine/Plugins/Experimental/Dataflow/Source/DataflowNodes/Public/Dataflow/DataflowNodesConnectionTypes.h b/Engine/Plugins/Experimental/Dataflow/Source/DataflowNodes/Public/Dataflow/DataflowNodesConnectionTypes.h
--- a/Engine/Plugins/Experimental/Dataflow/Source/DataflowNodes/Public/Dataflow/DataflowNodesConnectionTypes.h
+++ b/Engine/Plugins/Experimental/Dataflow/Source/DataflowNodes/Public/Dataflow/DataflowNodesConnectionTypes.h
@@ -12,9 +12,8 @@ namespace Dataflow
 {
 	typedef const UStaticMesh* UStaticMeshPtr;
 	typedef const USkeletalMesh* USkeletalMeshPtr;
-	//typedef TSharedPtr<FManagedArrayCollection> FManagedArrayCollectionSharedPtr;
-	DATAFLOW_CONNECTION_TYPE_PRIMITIVE(DATAFLOWNODES_API, UStaticMeshPtr, StaticMeshPtr);
-	DATAFLOW_CONNECTION_TYPE_PRIMITIVE(DATAFLOWNODES_API, USkeletalMeshPtr, SkeletalMeshPtr);
-	DATAFLOW_CONNECTION_TYPE_PRIMITIVE(DATAFLOWNODES_API, FString, String);
-	DATAFLOW_CONNECTION_TYPE_PRIMITIVE(DATAFLOWNODES_API, FVector3f, Vector3f);
+	//DATAFLOW_CONNECTION_TYPE_PRIMITIVE(DATAFLOWNODES_API, UStaticMeshPtr, StaticMeshPtr);
+	//DATAFLOW_CONNECTION_TYPE_PRIMITIVE(DATAFLOWNODES_API, USkeletalMeshPtr, SkeletalMeshPtr);
+	//DATAFLOW_CONNECTION_TYPE_PRIMITIVE(DATAFLOWNODES_API, FString, String);
+	//DATAFLOW_CONNECTION_TYPE_PRIMITIVE(DATAFLOWNODES_API, FVector3f, Vector3f);
 }
diff --git a/Engine/Plugins/Experimental/Dataflow/Source/DataflowNodes/Public/Dataflow/DataflowSkeletalMeshNodes.h b/Engine/Plugins/Experimental/Dataflow/Source/DataflowNodes/Public/Dataflow/DataflowSkeletalMeshNodes.h
--- a/Engine/Plugins/Experimental/Dataflow/Source/DataflowNodes/Public/Dataflow/DataflowSkeletalMeshNodes.h
+++ b/Engine/Plugins/Experimental/Dataflow/Source/DataflowNodes/Public/Dataflow/DataflowSkeletalMeshNodes.h
@@ -9,8 +9,8 @@
 #include "Dataflow/DataflowEngine.h"
 
 
-#include "DataflowSkeletalMeshNodes.generated.h"
 
+#include "DataflowSkeletalMeshNodes.generated.h"
 
 DEFINE_LOG_CATEGORY_STATIC(LogDataflowSkeletalMeshNodes, Log, All);
 
@@ -24,24 +24,23 @@ struct DATAFLOWNODES_API FGetSkeletalMeshDataflowNode: public FDataflowNode
 
 public:
 	
-	UPROPERTY(EditAnywhere, Category = "Dataflow")
-	TObjectPtr<USkeletalMesh> SkeletalMesh = nullptr;
-
+	UPROPERTY(EditAnywhere, Category = "Dataflow", meta = (DataflowOutput, DisplayName = "SkeletalMesh"))
+	TObjectPtr<const USkeletalMesh> SkeletalMesh = nullptr;
 
 	UPROPERTY(EditAnywhere, Category = "Dataflow" )
 	FName PropertyName = "SkeletalMesh";
 
-	Dataflow::TOutput<Dataflow::USkeletalMeshPtr> SkeletalMeshOut;
-
 	FGetSkeletalMeshDataflowNode(const Dataflow::FNodeParameters& InParam, FGuid InGuid = FGuid::NewGuid())
 		: FDataflowNode(InParam, InGuid)
-		, SkeletalMeshOut({ FName("SkeletalMeshOut"), this })
-	{}
+	{
+		RegisterOutputConnection(&SkeletalMesh);
+	}
 
 
-	virtual void Evaluate(const Dataflow::FContext& Context, Dataflow::FConnection* Out) const override;
+	virtual void Evaluate(Dataflow::FContext& Context, const FDataflowOutput* Out) const override;
 };
 
+
 USTRUCT()
 struct DATAFLOWNODES_API FSkeletalMeshBoneDataflowNode : public FDataflowNode
 {
@@ -52,17 +51,21 @@ public:
 	UPROPERTY(EditAnywhere, Category = "Dataflow")
 	FName BoneName;
 
-	Dataflow::TInput<Dataflow::USkeletalMeshPtr> SkeletalMeshIn;
-	Dataflow::TOutput<int> BoneIndexOut;
+	UPROPERTY(meta = (DataflowInput, DisplayName = "SkeletalMesh"))
+	TObjectPtr<const USkeletalMesh> SkeletalMesh = nullptr;
+
+	UPROPERTY(meta = (DataflowOutput, DisplayName = "Index"))
+	int BoneIndexOut = INDEX_NONE;
 
 	FSkeletalMeshBoneDataflowNode(const Dataflow::FNodeParameters& InParam, FGuid InGuid = FGuid::NewGuid())
 		: FDataflowNode(InParam, InGuid)
-		, SkeletalMeshIn({ FName("SkeletalMeshIn"), this })
-		, BoneIndexOut({ FName("BoneIndexOut"), this, INDEX_NONE })
-	{}
+	{
+		RegisterInputConnection(&SkeletalMesh);
+		RegisterOutputConnection(&BoneIndexOut);
+	}
 
 
-	virtual void Evaluate(const Dataflow::FContext& Context, Dataflow::FConnection* Out) const override;
+	virtual void Evaluate(Dataflow::FContext& Context, const FDataflowOutput* Out) const override;
 };
 
 namespace Dataflow
diff --git a/Engine/Plugins/Experimental/Dataflow/Source/DataflowNodes/Public/Dataflow/DataflowStaticMeshNodes.h b/Engine/Plugins/Experimental/Dataflow/Source/DataflowNodes/Public/Dataflow/DataflowStaticMeshNodes.h
--- a/Engine/Plugins/Experimental/Dataflow/Source/DataflowNodes/Public/Dataflow/DataflowStaticMeshNodes.h
+++ b/Engine/Plugins/Experimental/Dataflow/Source/DataflowNodes/Public/Dataflow/DataflowStaticMeshNodes.h
@@ -12,7 +12,6 @@
 
 #include "DataflowStaticMeshNodes.generated.h"
 
-
 DEFINE_LOG_CATEGORY_STATIC(LogDataflowStaticMeshNodes, Log, All);
 
 class UStaticMesh;
@@ -21,26 +20,24 @@ USTRUCT()
 struct DATAFLOWNODES_API FGetStaticMeshDataflowNode : public FDataflowNode
 {
 	GENERATED_USTRUCT_BODY()
-		DATAFLOW_NODE_DEFINE_INTERNAL(FGetStaticMeshDataflowNode, "StaticMesh", "Dataflow", "Static Mesh")
+	DATAFLOW_NODE_DEFINE_INTERNAL(FGetStaticMeshDataflowNode, "StaticMesh", "Dataflow", "Static Mesh")
 
 public:
 
-	UPROPERTY(EditAnywhere, Category = "Dataflow")
-		TObjectPtr<UStaticMesh> StaticMesh = nullptr;
-
+	UPROPERTY(EditAnywhere, Category = "Dataflow", meta = (DataflowOutput, DisplayName = "StaticMesh"))
+	TObjectPtr<const UStaticMesh> StaticMesh = nullptr;
 
 	UPROPERTY(EditAnywhere, Category = "Dataflow")
-		FName PropertyName = "StaticMesh";
-
-	Dataflow::TOutput<Dataflow::UStaticMeshPtr> StaticMeshOut;
+	FName PropertyName = "StaticMesh";
 
 	FGetStaticMeshDataflowNode(const Dataflow::FNodeParameters& InParam, FGuid InGuid = FGuid::NewGuid())
 		: FDataflowNode(InParam, InGuid)
-		, StaticMeshOut({ FName("StaticMeshOut"), this })
-	{}
+	{
+		RegisterOutputConnection(&StaticMesh);
+	}
 
 
-	virtual void Evaluate(const Dataflow::FContext& Context, Dataflow::FConnection* Out) const override;
+	virtual void Evaluate(Dataflow::FContext& Context, const FDataflowOutput* Out) const override;
 };
 
 
diff --git a/Engine/Plugins/Experimental/GeometryCollectionPlugin/Source/GeometryCollectionEditor/Private/GeometryCollection/GeometryCollectionEditorToolkit.cpp b/Engine/Plugins/Experimental/GeometryCollectionPlugin/Source/GeometryCollectionEditor/Private/GeometryCollection/GeometryCollectionEditorToolkit.cpp
--- a/Engine/Plugins/Experimental/GeometryCollectionPlugin/Source/GeometryCollectionEditor/Private/GeometryCollection/GeometryCollectionEditorToolkit.cpp
+++ b/Engine/Plugins/Experimental/GeometryCollectionPlugin/Source/GeometryCollectionEditor/Private/GeometryCollection/GeometryCollectionEditorToolkit.cpp
@@ -94,10 +94,11 @@ TSharedRef<SGraphEditor> FGeometryCollectionEditorToolkit::CreateGraphEditorWidg
 	using namespace Dataflow;
 	IDataflowEditorPlugin& DataflowEditorModule = FModuleManager::LoadModuleChecked<IDataflowEditorPlugin>(TEXT("DataflowEditor"));
 
-	FDataflowEditorCommands::FGraphEvaluationCallback Evaluate = [&](FDataflowNode* Node, Dataflow::FConnection* Out)
+	FDataflowEditorCommands::FGraphEvaluationCallback Evaluate = [&](FDataflowNode* Node, FDataflowOutput* Out)
 	{
 		float EvalTime = FGameTime::GetTimeSinceAppStart().GetRealTimeSeconds();
-		return Node->Evaluate(FEngineContext(GeometryCollection, Dataflow, EvalTime, FString("UGeometryCollection")), Out);
+		FEngineContext Context(GeometryCollection, Dataflow, EvalTime, FString("UGeometryCollection"));
+		Node->Evaluate(Context, Out);
 	};
 
 	return SNew(SDataflowGraphEditor, GeometryCollection)
diff --git a/Engine/Plugins/Experimental/GeometryCollectionPlugin/Source/GeometryCollectionNodes/Private/Dataflow/GeometryCollectionNodes.cpp b/Engine/Plugins/Experimental/GeometryCollectionPlugin/Source/GeometryCollectionNodes/Private/Dataflow/GeometryCollectionNodes.cpp
--- a/Engine/Plugins/Experimental/GeometryCollectionPlugin/Source/GeometryCollectionNodes/Private/Dataflow/GeometryCollectionNodes.cpp
+++ b/Engine/Plugins/Experimental/GeometryCollectionPlugin/Source/GeometryCollectionNodes/Private/Dataflow/GeometryCollectionNodes.cpp
@@ -23,9 +23,9 @@ namespace Dataflow
 	}
 }
 
-void FGetCollectionAssetDataflowNode::Evaluate(const Dataflow::FContext& Context, Dataflow::FConnection* Out) const
+void FGetCollectionAssetDataflowNode::Evaluate(Dataflow::FContext& Context, const FDataflowOutput* Out) const
 {
-	if (&Output == Out)
+	if (Out->IsA<DataType>(&Output))
 	{
 		if (const Dataflow::FEngineContext* EngineContext = Context.AsType<Dataflow::FEngineContext>())
 		{
@@ -33,107 +33,107 @@ void FGetCollectionAssetDataflowNode::Evaluate(const Dataflow::FContext& Context
 			{
 				if (const TSharedPtr<FGeometryCollection, ESPMode::ThreadSafe> AssetCollection = CollectionAsset->GetGeometryCollection())
 				{
-					FManagedArrayCollection* NewCollection = AssetCollection->NewCopy<FManagedArrayCollection>();
-					Output.SetValue(DataType(NewCollection), Context);
+					GetOutput(&Output)->SetValue<DataType>(DataType(*AssetCollection), Context);
 				}
 			}
 		}
 	}
 }
 
-void FExampleCollectionEditDataflowNode::Evaluate(const Dataflow::FContext& Context, Dataflow::FConnection* Out) const
+void FExampleCollectionEditDataflowNode::Evaluate(Dataflow::FContext& Context, const FDataflowOutput* Out) const
 {
-	if (&Output == Out)
+	if (Out->IsA<DataType>(&Collection))
 	{
-		DataType Collection = Input.GetValue(Context);
+		DataType InCollection = GetInput(&Collection)->GetValue<DataType>(Context, Collection);
+
 		if (Active)
 		{
-			TManagedArray<FVector3f>* Vertex = Collection->FindAttribute<FVector3f>("Vertex", "Vertices");
+			TManagedArray<FVector3f>* Vertex = InCollection.FindAttribute<FVector3f>("Vertex", "Vertices");
 			for (int i = 0; i < Vertex->Num(); i++)
 			{
 				(*Vertex)[i][1] *= Scale;
 			}
 		}
-		Output.SetValue(Collection, Context);
+		Out->SetValue<DataType>(InCollection, Context);
 	}
 }
 
-void FSetCollectionAssetDataflowNode::Evaluate(const Dataflow::FContext& Context, Dataflow::FConnection* Out) const
+void FSetCollectionAssetDataflowNode::Evaluate(Dataflow::FContext& Context, const FDataflowOutput* Out) const
 {
-	if (Out == nullptr)
+	DataType InCollection = GetInput(&Collection)->GetValue<DataType>(Context, Collection);
+
+	if (const Dataflow::FEngineContext* EngineContext = Context.AsType<Dataflow::FEngineContext>())
 	{
-		DataType Collection = Input.GetValue(Context);
-		if (const Dataflow::FEngineContext* EngineContext = Context.AsType<Dataflow::FEngineContext>())
+		if (UGeometryCollection* CollectionAsset = Cast<UGeometryCollection>(EngineContext->Owner))
 		{
-			if (UGeometryCollection* CollectionAsset = Cast<UGeometryCollection>(EngineContext->Owner))
-			{
-				TSharedPtr<FGeometryCollection, ESPMode::ThreadSafe> NewCollection(Collection->NewCopy<FGeometryCollection>());
-				CollectionAsset->SetGeometryCollection(NewCollection);
-				CollectionAsset->InvalidateCollection();
-			}
+			TSharedPtr<FGeometryCollection, ESPMode::ThreadSafe> NewCollection(InCollection.NewCopy<FGeometryCollection>());
+			CollectionAsset->SetGeometryCollection(NewCollection);
+			CollectionAsset->InvalidateCollection();
 		}
 	}
 }
 
-void FResetGeometryCollectionDataflowNode::Evaluate(const Dataflow::FContext& Context, Dataflow::FConnection* Out) const
+void FResetGeometryCollectionDataflowNode::Evaluate(Dataflow::FContext& Context, const FDataflowOutput* Out) const
 {
-	ManagedArrayOut.SetValue(TSharedPtr<FManagedArrayCollection>(nullptr), Context);
-
-	if (const Dataflow::FEngineContext* EngineContext = Context.AsType<Dataflow::FEngineContext>())
+	if (Out->IsA<DataType>(&Collection))
 	{
-		if (UGeometryCollection* GeometryCollectionObject = Cast<UGeometryCollection>(EngineContext->Owner))
-		{
-			GeometryCollectionObject->Reset();
+		Out->SetValue<DataType>(Collection, Context); // prime to avoid ensure
 
-			const UObject* Owner = EngineContext->Owner;
-			FName AName("GeometrySource");
-			if (Owner && Owner->GetClass())
+		if (const Dataflow::FEngineContext* EngineContext = Context.AsType<Dataflow::FEngineContext>())
+		{
+			if (UGeometryCollection* GeometryCollectionObject = Cast<UGeometryCollection>(EngineContext->Owner))
 			{
-				if (const ::FProperty* UEProperty = Owner->GetClass()->FindPropertyByName(AName))
+				GeometryCollectionObject->Reset();
+
+				const UObject* Owner = EngineContext->Owner;
+				FName AName("GeometrySource");
+				if (Owner && Owner->GetClass())
 				{
-					if (const FArrayProperty* ArrayProperty = CastField<FArrayProperty>(UEProperty))
+					if (const ::FProperty* UEProperty = Owner->GetClass()->FindPropertyByName(AName))
 					{
-						FScriptArrayHelper_InContainer ArrayHelper(ArrayProperty, Owner);
-						const int32 ArraySize = ArrayHelper.Num();
-						for (int32 Index = 0; Index < ArraySize; ++Index)
+						if (const FArrayProperty* ArrayProperty = CastField<FArrayProperty>(UEProperty))
 						{
-							if (FGeometryCollectionSource* SourceObject = (FGeometryCollectionSource*)(ArrayHelper.GetRawPtr(Index)))
+							FScriptArrayHelper_InContainer ArrayHelper(ArrayProperty, Owner);
+							const int32 ArraySize = ArrayHelper.Num();
+							for (int32 Index = 0; Index < ArraySize; ++Index)
 							{
-								if (UObject* ResolvedObject = SourceObject->SourceGeometryObject.ResolveObject())
+								if (FGeometryCollectionSource* SourceObject = (FGeometryCollectionSource*)(ArrayHelper.GetRawPtr(Index)))
 								{
-									if (UStaticMesh* StaticMesh = Cast<UStaticMesh>(ResolvedObject))
+									if (UObject* ResolvedObject = SourceObject->SourceGeometryObject.ResolveObject())
 									{
-										TArray<UMaterialInterface*> Materials;
-										Materials.Reserve(StaticMesh->GetStaticMaterials().Num());
-
-										for (int32 Index2 = 0; Index2 < StaticMesh->GetStaticMaterials().Num(); ++Index2)
+										if (UStaticMesh* StaticMesh = Cast<UStaticMesh>(ResolvedObject))
 										{
-											UMaterialInterface* CurrMaterial = StaticMesh->GetMaterial(Index2);
-											Materials.Add(CurrMaterial);
-										}
+											TArray<UMaterialInterface*> Materials;
+											Materials.Reserve(StaticMesh->GetStaticMaterials().Num());
+
+											for (int32 Index2 = 0; Index2 < StaticMesh->GetStaticMaterials().Num(); ++Index2)
+											{
+												UMaterialInterface* CurrMaterial = StaticMesh->GetMaterial(Index2);
+												Materials.Add(CurrMaterial);
+											}
 
-										// Geometry collections usually carry the selection material, which we'll delete before appending
-										UMaterialInterface* BoneSelectedMaterial = LoadObject<UMaterialInterface>(nullptr, UGeometryCollection::GetSelectedMaterialPath(), nullptr, LOAD_None, nullptr);
-										GeometryCollectionObject->Materials.Remove(BoneSelectedMaterial);
-										Materials.Remove(BoneSelectedMaterial);
+											// Geometry collections usually carry the selection material, which we'll delete before appending
+											UMaterialInterface* BoneSelectedMaterial = LoadObject<UMaterialInterface>(nullptr, UGeometryCollection::GetSelectedMaterialPath(), nullptr, LOAD_None, nullptr);
+											GeometryCollectionObject->Materials.Remove(BoneSelectedMaterial);
+											Materials.Remove(BoneSelectedMaterial);
 
-										FGeometryCollectionEngineConversion::AppendStaticMesh(StaticMesh, Materials, FTransform(), GeometryCollectionObject);
+											FGeometryCollectionEngineConversion::AppendStaticMesh(StaticMesh, Materials, FTransform(), GeometryCollectionObject);
 
+										}
 									}
 								}
 							}
 						}
 					}
 				}
-			}
-			GeometryCollectionObject->UpdateConvexGeometry();
-			GeometryCollectionObject->InitializeMaterials();
-			GeometryCollectionObject->InvalidateCollection();
+				GeometryCollectionObject->UpdateConvexGeometry();
+				GeometryCollectionObject->InitializeMaterials();
+				GeometryCollectionObject->InvalidateCollection();
 
-			if (const TSharedPtr<FGeometryCollection, ESPMode::ThreadSafe> AssetCollection = GeometryCollectionObject->GetGeometryCollection())
-			{
-				FManagedArrayCollection* NewCollection = AssetCollection->NewCopy<FManagedArrayCollection>();
-				ManagedArrayOut.SetValue(DataType(NewCollection), Context);
+				if (const TSharedPtr<FGeometryCollection, ESPMode::ThreadSafe> AssetCollection = GeometryCollectionObject->GetGeometryCollection())
+				{
+					Out->SetValue<DataType>(*AssetCollection, Context);
+				}
 			}
 		}
 	}
@@ -141,3 +141,4 @@ void FResetGeometryCollectionDataflowNode::Evaluate(const Dataflow::FContext& Co
 
 
 
+
diff --git a/Engine/Plugins/Experimental/GeometryCollectionPlugin/Source/GeometryCollectionNodes/Public/Dataflow/GeometryCollectionNodeConnectionTypes.h b/Engine/Plugins/Experimental/GeometryCollectionPlugin/Source/GeometryCollectionNodes/Public/Dataflow/GeometryCollectionNodeConnectionTypes.h
--- a/Engine/Plugins/Experimental/GeometryCollectionPlugin/Source/GeometryCollectionNodes/Public/Dataflow/GeometryCollectionNodeConnectionTypes.h
+++ b/Engine/Plugins/Experimental/GeometryCollectionPlugin/Source/GeometryCollectionNodes/Public/Dataflow/GeometryCollectionNodeConnectionTypes.h
@@ -10,8 +10,8 @@
 
 namespace Dataflow
 {
-	typedef TSharedPtr<FManagedArrayCollection> FManagedArrayCollectionSharedPtr;
+	//typedef TSharedPtr<FManagedArrayCollection> FManagedArrayCollectionSharedPtr;
 	 
-	DATAFLOW_CONNECTION_TYPE_SHARED_POINTER(GEOMETRYCOLLECTIONNODES_API, FManagedArrayCollectionSharedPtr, ManagedArrayCollectionSharedPtr)
+	//DATAFLOW_CONNECTION_TYPE_SHARED_POINTER(GEOMETRYCOLLECTIONNODES_API, FManagedArrayCollectionSharedPtr, ManagedArrayCollectionSharedPtr)
 	
 }
diff --git a/Engine/Plugins/Experimental/GeometryCollectionPlugin/Source/GeometryCollectionNodes/Public/Dataflow/GeometryCollectionNodes.h b/Engine/Plugins/Experimental/GeometryCollectionPlugin/Source/GeometryCollectionNodes/Public/Dataflow/GeometryCollectionNodes.h
--- a/Engine/Plugins/Experimental/GeometryCollectionPlugin/Source/GeometryCollectionNodes/Public/Dataflow/GeometryCollectionNodes.h
+++ b/Engine/Plugins/Experimental/GeometryCollectionPlugin/Source/GeometryCollectionNodes/Public/Dataflow/GeometryCollectionNodes.h
@@ -6,9 +6,12 @@
 
 #include "CoreMinimal.h"
 #include "Dataflow/DataflowEngine.h"
+#include "GeometryCollection/ManagedArrayCollection.h"
+
 
 #include "GeometryCollectionNodes.generated.h"
 
+
 USTRUCT()
 struct FGetCollectionAssetDataflowNode : public FDataflowNode
 {
@@ -16,16 +19,18 @@ struct FGetCollectionAssetDataflowNode : public FDataflowNode
 	DATAFLOW_NODE_DEFINE_INTERNAL(FGetCollectionAssetDataflowNode, "GetCollectionAsset", "GeometryCollection", "")
 
 public:
-	typedef Dataflow::FManagedArrayCollectionSharedPtr DataType;
+	typedef FManagedArrayCollection DataType;
 
-	Dataflow::TOutput<DataType> Output;
+	UPROPERTY(meta = (DataflowOutput, DisplayName = "Collection"))
+	mutable FManagedArrayCollection Output;
 
 	FGetCollectionAssetDataflowNode(const Dataflow::FNodeParameters& InParam, FGuid InGuid = FGuid::NewGuid())
 		: FDataflowNode(InParam, InGuid)
-		, Output({ FName("CollectionOut"), this })
-	{}
+	{
+		RegisterOutputConnection(&Output);
+	}
 
-	virtual void Evaluate(const Dataflow::FContext& Context, Dataflow::FConnection* Out) const override;
+	virtual void Evaluate(Dataflow::FContext& Context, const FDataflowOutput* Out) const override;
 };
 
 USTRUCT()
@@ -35,23 +40,25 @@ struct FExampleCollectionEditDataflowNode : public FDataflowNode
 	DATAFLOW_NODE_DEFINE_INTERNAL(FExampleCollectionEditDataflowNode, "ExampleCollectionEdit", "GeometryCollection", "")
 
 public:
+	typedef FManagedArrayCollection DataType;
+
 	UPROPERTY(EditAnywhere, Category = "Dataflow");
 	bool Active = true;
 
 	UPROPERTY(EditAnywhere, Category = "Dataflow");
 	float Scale = 1.0;
 
-	typedef Dataflow::FManagedArrayCollectionSharedPtr DataType;
-	Dataflow::TInput<DataType> Input;
-	Dataflow::TOutput<DataType> Output;
+	UPROPERTY(meta = (DataflowInput, DataflowOutput, DisplayName = "Collection"))
+	FManagedArrayCollection Collection;
 
 	FExampleCollectionEditDataflowNode(const Dataflow::FNodeParameters& InParam, FGuid InGuid = FGuid::NewGuid())
 		: FDataflowNode(InParam, InGuid)
-		, Input({ FName("CollectionIn"), this })
-		, Output({ FName("CollectionOut"), this })
-	{}
+	{
+		RegisterInputConnection(&Collection);
+		RegisterOutputConnection(&Collection);
+	}
 
-	virtual void Evaluate(const Dataflow::FContext& Context, Dataflow::FConnection* Out) const override;
+	virtual void Evaluate(Dataflow::FContext& Context, const FDataflowOutput* Out) const override;
 
 };
 
@@ -62,15 +69,18 @@ struct FSetCollectionAssetDataflowNode : public FDataflowNode
 	DATAFLOW_NODE_DEFINE_INTERNAL(FSetCollectionAssetDataflowNode, "SetCollectionAsset", "GeometryCollection", "")
 
 public:
-	typedef Dataflow::FManagedArrayCollectionSharedPtr DataType;
-	Dataflow::TInput<DataType> Input;
+	typedef FManagedArrayCollection DataType;
+
+	UPROPERTY(meta = (DataflowInput, DisplayName = "Collection"))
+	FManagedArrayCollection Collection;
 
 	FSetCollectionAssetDataflowNode(const Dataflow::FNodeParameters& InParam, FGuid InGuid = FGuid::NewGuid())
 		: FDataflowNode(InParam, InGuid)
-		, Input({ FName("CollectionIn"), this })
-	{}
+	{
+		RegisterInputConnection(&Collection);
+	}
 
-	virtual void Evaluate(const Dataflow::FContext& Context, Dataflow::FConnection* Out) const override;
+	virtual void Evaluate(Dataflow::FContext& Context, const FDataflowOutput* Out) const override;
 
 };
 
@@ -81,17 +91,19 @@ struct FResetGeometryCollectionDataflowNode : public FDataflowNode
 	DATAFLOW_NODE_DEFINE_INTERNAL(FResetGeometryCollectionDataflowNode, "ResetGeometryCollection", "GeometryCollection", "")
 
 public:
-	typedef Dataflow::FManagedArrayCollectionSharedPtr DataType;
+	typedef FManagedArrayCollection DataType;
 
-	Dataflow::TOutput<DataType> ManagedArrayOut;
+	UPROPERTY(meta = (DataflowOutput, DisplayName = "Collection"))
+	FManagedArrayCollection Collection;
 
 	FResetGeometryCollectionDataflowNode(const Dataflow::FNodeParameters& InParam, FGuid InGuid = FGuid::NewGuid())
 		: FDataflowNode(InParam, InGuid)
-		, ManagedArrayOut({ FName("ManagedArrayOut"), this })
-	{}
+	{
+		RegisterOutputConnection(&Collection);
+	}
 
 
-	virtual void Evaluate(const Dataflow::FContext& Context, Dataflow::FConnection* Out) const override;
+	virtual void Evaluate(Dataflow::FContext& Context, const FDataflowOutput* Out) const override;
 
 };
 
diff --git a/Engine/Source/Runtime/Experimental/Dataflow/Core/Private/Dataflow/DataflowConnection.cpp b/Engine/Source/Runtime/Experimental/Dataflow/Core/Private/Dataflow/DataflowConnection.cpp
--- a/Engine/Source/Runtime/Experimental/Dataflow/Core/Private/Dataflow/DataflowConnection.cpp
+++ b/Engine/Source/Runtime/Experimental/Dataflow/Core/Private/Dataflow/DataflowConnection.cpp
@@ -5,24 +5,23 @@
 #include "Dataflow/DataflowNodeParameters.h"
 #include "Dataflow/DataflowNode.h"
 
-namespace Dataflow
-{
-	FConnection::FConnection(FPin::EDirection InDirection, FName InType, FName InName, FDataflowNode* InOwningNode, FGuid InGuid)
-		: Direction(InDirection)
-		, Type(InType)
-		, Name(InName)
-		, Guid(InGuid)
-		, OwningNode(InOwningNode)
-	{}
+FDataflowConnection::FDataflowConnection(Dataflow::FPin::EDirection InDirection, FName InType, FName InName, FDataflowNode* InOwningNode, FProperty* InProperty, FGuid InGuid)
+	: Direction(InDirection)
+	, Type(InType)
+	, Name(InName)
+	, OwningNode(InOwningNode)
+	, Property(InProperty)
+	, Guid(InGuid)
+{}
 
-	void FConnection::BindInput(FDataflowNode* InNode, FConnection* That) 
-	{
-		if (InNode && That) { InNode->AddInput(That); }
-	}
-	void FConnection::BindOutput(FDataflowNode* InNode, FConnection* That)
+
+uint32 FDataflowConnection::GetOffset() const
+{
+	if (ensure(Property != nullptr))
 	{
-		if (InNode && That) { InNode->AddOutput(That); }
+		return Property->GetOffset_ForInternal();
 	}
-
+	return (uint32)INDEX_NONE;
 }
 
+
diff --git a/Engine/Source/Runtime/Experimental/Dataflow/Core/Private/Dataflow/DataflowGraph.cpp b/Engine/Source/Runtime/Experimental/Dataflow/Core/Private/Dataflow/DataflowGraph.cpp
--- a/Engine/Source/Runtime/Experimental/Dataflow/Core/Private/Dataflow/DataflowGraph.cpp
+++ b/Engine/Source/Runtime/Experimental/Dataflow/Core/Private/Dataflow/DataflowGraph.cpp
@@ -21,11 +21,11 @@ namespace Dataflow
 
 	void FGraph::RemoveNode(TSharedPtr<FDataflowNode> Node)
 	{
-		for (FConnection* Output : Node->GetOutputs())
+		for (FDataflowConnection* Output : Node->GetOutputs())
 		{
 			if (Output)
 			{
-				for (FConnection* Input : Output->GetConnectedInputs())
+				for (FDataflowConnection* Input : Output->GetConnectedInputs())
 				{
 					if (Input)
 					{
@@ -34,12 +34,12 @@ namespace Dataflow
 				}
 			}
 		}
-		for (FConnection* Input : Node->GetInputs())
+		for (FDataflowConnection* Input : Node->GetInputs())
 		{
 			if (Input)
 			{
-				TArray<FConnection*> Outputs = Input->GetConnectedOutputs();
-				for (FConnection* Out : Outputs)
+				TArray<FDataflowConnection*> Outputs = Input->GetConnectedOutputs();
+				for (FDataflowConnection* Out : Outputs)
 				{
 					if (Out)
 					{
@@ -51,22 +51,22 @@ namespace Dataflow
 		Nodes.Remove(Node);
 	}
 
-	void FGraph::ClearConnections(FConnection* Connection)
+	void FGraph::ClearConnections(FDataflowConnection* Connection)
 	{
 		// Todo(dataflow) : do this without triggering a invalidation. 
 		//            or implement a better sync for the EdGraph and DataflowGraph
 		if (Connection->GetDirection() == FPin::EDirection::INPUT)
 		{
-			TArray<FConnection*> BaseOutputs = Connection->GetConnectedOutputs();
-			for (FConnection* Output : BaseOutputs)
+			TArray<FDataflowConnection*> BaseOutputs = Connection->GetConnectedOutputs();
+			for (FDataflowConnection* Output : BaseOutputs)
 			{
 				Disconnect(Connection, Output);
 			}
 		}
 		else if (Connection->GetDirection() == FPin::EDirection::OUTPUT)
 		{
-			TArray<FConnection*> BaseInputs = Connection->GetConnectedInputs();
-			for (FConnection* Input : BaseInputs)
+			TArray<FDataflowConnection*> BaseInputs = Connection->GetConnectedInputs();
+			for (FDataflowConnection* Input : BaseInputs)
 			{
 				Disconnect(Input, Connection);
 			}
@@ -74,7 +74,7 @@ namespace Dataflow
 	}
 
 
-	void FGraph::Connect(FConnection* Input, FConnection* Output)
+	void FGraph::Connect(FDataflowConnection* Input, FDataflowConnection* Output)
 	{
 		if (ensure(Input && Output))
 		{
@@ -86,7 +86,7 @@ namespace Dataflow
 		}
 	}
 
-	void FGraph::Disconnect(FConnection* Input, FConnection* Output)
+	void FGraph::Disconnect(FDataflowConnection* Input, FDataflowConnection* Output)
 	{
 		Input->RemoveConnection(Output);
 		Output->RemoveConnection(Input);
@@ -115,11 +115,11 @@ namespace Dataflow
 
 				DATAFLOW_OPTIONAL_BLOCK_WRITE_BEGIN()
 				{
-					TArray< FConnection* > IO = Node->GetOutputs();
+					TArray< FDataflowConnection* > IO = Node->GetOutputs();
 					IO.Append(Node->GetInputs());
 					ArNum = IO.Num();
 					Ar << ArNum;
-					for (FConnection* Conn : IO)
+					for (FDataflowConnection* Conn : IO)
 					{
 						ArGuid = Conn->GetGuid();
 						ArType = Conn->GetType();
@@ -142,7 +142,7 @@ namespace Dataflow
 			int32 ArNum = 0;
 
 			TMap<FGuid, TSharedPtr<FDataflowNode> > NodeGuidMap;
-			TMap<FGuid, FConnection* > ConnectionGuidMap;
+			TMap<FGuid, FDataflowConnection* > ConnectionGuidMap;
 
 			Ar << ArNum;
 			for (int32 Ndx = ArNum; Ndx > 0; Ndx--)
@@ -157,7 +157,7 @@ namespace Dataflow
 
 					int ArNumInner;
 					Ar << ArNumInner;
-					TArray< FConnection* > IO = Node->GetOutputs();  IO.Append(Node->GetInputs());
+					TArray< FDataflowConnection* > IO = Node->GetOutputs();  IO.Append(Node->GetInputs());
 					for (int Cdx = 0; Cdx < ArNumInner; Cdx++)
 					{
 						Ar << ArGuid << ArType << ArName;
diff --git a/Engine/Source/Runtime/Experimental/Dataflow/Core/Private/Dataflow/DataflowInputOutput.cpp b/Engine/Source/Runtime/Experimental/Dataflow/Core/Private/Dataflow/DataflowInputOutput.cpp
--- a/Engine/Source/Runtime/Experimental/Dataflow/Core/Private/Dataflow/DataflowInputOutput.cpp
+++ b/Engine/Source/Runtime/Experimental/Dataflow/Core/Private/Dataflow/DataflowInputOutput.cpp
@@ -2,12 +2,123 @@
 
 #include "Dataflow/DataflowInputOutput.h"
 
+#include "ChaosLog.h"
 #include "Dataflow/DataflowNodeParameters.h"
 #include "Dataflow/DataflowNode.h"
 
-namespace Dataflow
+
+FDataflowInput FDataflowInput::NoOpInput = FDataflowInput();
+FDataflowOutput FDataflowOutput::NoOpOutput = FDataflowOutput();
+
+
+FDataflowInput::FDataflowInput(const Dataflow::FInputParameters& Param, FGuid InGuid)
+	: FDataflowConnection(Dataflow::FPin::EDirection::INPUT, Param.Type, Param.Name, Param.Owner, Param.Property, InGuid)
+	, Connection(nullptr)
+{
+}
+
+
+
+bool FDataflowInput::AddConnection(FDataflowConnection* InOutput)
+{
+	ensure(Connection == nullptr);
+	if (ensure(InOutput->GetType() == this->GetType()))
+	{
+		Connection = (FDataflowOutput*)InOutput;
+		return true;
+	}
+	return false;
+}
+
+bool FDataflowInput::RemoveConnection(FDataflowConnection* InOutput)
+{
+	if (ensure(Connection == (FDataflowOutput*)InOutput))
+	{
+		Connection = nullptr;
+		return true;
+	}
+	return false;
+}
+
+TArray< FDataflowConnection* > FDataflowInput::GetConnectedOutputs()
+{
+	TArray<FDataflowConnection* > RetList;
+	if (FDataflowConnection* Conn = GetConnection())
+	{
+		RetList.Add(Conn);
+	}
+	return RetList;
+}
+
+void FDataflowInput::Invalidate()
 {
-	// Base Types
+	OwningNode->InvalidateOutputs();
+}
+
+//
+//
+//  Output
+//
+//
+//
+
 
+FDataflowOutput::FDataflowOutput(const Dataflow::FOutputParameters& Param, FGuid InGuid)
+	: FDataflowConnection(Dataflow::FPin::EDirection::OUTPUT, Param.Type, Param.Name, Param.Owner, Param.Property, InGuid)
+{
+}
+
+const TArray<FDataflowInput*>& FDataflowOutput::GetConnections() const { return Connections; }
+TArray<FDataflowInput*>& FDataflowOutput::GetConnections() { return Connections; }
+
+TArray<FDataflowConnection*> FDataflowOutput::GetConnectedInputs()
+{
+	TArray<FDataflowConnection*> RetList;
+	RetList.Reserve(Connections.Num());
+	for (FDataflowInput* Ptr : Connections) { RetList.Add(Ptr); }
+	return RetList;
+}
+
+
+bool FDataflowOutput::AddConnection(FDataflowConnection* InOutput)
+{
+	if (ensure(InOutput->GetType() == this->GetType()))
+	{
+		Connections.Add((FDataflowInput*)InOutput);
+		return true;
+	}
+	return false;
+}
+
+bool FDataflowOutput::RemoveConnection(FDataflowConnection* InInput)
+{
+	Connections.RemoveSwap((FDataflowInput*)InInput); return true;
+}
+
+bool FDataflowOutput::Evaluate(Dataflow::FContext& Context) const
+{
+	check(OwningNode);
+
+	if (OwningNode->bActive)
+	{
+		OwningNode->Evaluate(Context, this);
+		if (!Context.HasData(CacheKey()))
+		{
+			ensureMsgf(false, TEXT("Failed to evaluate output (%s:%s)"), *OwningNode->GetName().ToString(), *GetName().ToString());
+		}
+	}
+	return OwningNode->bActive;
+}
+
+void FDataflowOutput::Invalidate()
+{
+	if (CacheKeyValue != UINT_MAX)
+	{
+		CacheKeyValue = UINT_MAX;
+		for (FDataflowConnection* Con : GetConnections())
+		{
+			Con->Invalidate();
+		}
+	}
 }
 
diff --git a/Engine/Source/Runtime/Experimental/Dataflow/Core/Private/Dataflow/DataflowNode.cpp b/Engine/Source/Runtime/Experimental/Dataflow/Core/Private/Dataflow/DataflowNode.cpp
--- a/Engine/Source/Runtime/Experimental/Dataflow/Core/Private/Dataflow/DataflowNode.cpp
+++ b/Engine/Source/Runtime/Experimental/Dataflow/Core/Private/Dataflow/DataflowNode.cpp
@@ -2,69 +2,284 @@
 
 #include "Dataflow/DataflowNode.h"
 
-#include "Dataflow/DataflowNode.h"
+#include "Dataflow/DataflowInputOutput.h"
 #include "Dataflow/DataflowArchive.h"
 
-void FDataflowNode::AddInput(Dataflow::FConnection* InPtr)
+
+
+//
+// Inputs
+//
+
+void FDataflowNode::AddInput(FDataflowConnection* InPtr)
+{
+	if (InPtr)
+	{
+		for (TPair<uint32, FDataflowConnection*> Elem : Inputs)
+		{
+			FDataflowConnection* In = Elem.Value;
+			ensureMsgf(!In->GetName().IsEqual(InPtr->GetName()), TEXT("Add Input Failed: Existing Node input already defined with name (%s)"), *InPtr->GetName().ToString());
+		}
+
+		if (!Inputs.Contains(InPtr->Property->GetOffset_ForInternal()))
+		{
+			Inputs.Add(InPtr->Property->GetOffset_ForInternal(), InPtr);
+		}
+		else
+		{
+			Inputs[InPtr->Property->GetOffset_ForInternal()] = InPtr;
+		}
+	}
+}
+
+FDataflowInput* FDataflowNode::FindInput(FName InName)
+{
+	for (TPair<uint32, FDataflowConnection*> Elem : Inputs)
+	{
+		FDataflowConnection* Con = Elem.Value;
+		if (Con->GetName().IsEqual(InName))
+		{
+			return (FDataflowInput*)Con;
+		}
+	}
+	return nullptr;
+}
+
+
+const FDataflowInput* FDataflowNode::FindInput(const void* Reference) const
+{
+	for (TPair<uint32, FDataflowConnection*> Elem : Inputs)
+	{
+		FDataflowConnection* Con = Elem.Value;
+		if (Con->RealAddress() == (size_t)Reference)
+		{
+			return (FDataflowInput*)Con;
+		}
+	}
+	return nullptr;
+}
+
+const FDataflowInput* FDataflowNode::GetInput(const void* Reference) const
+{
+	for (TPair<uint32, FDataflowConnection*> Elem : Inputs)
+	{
+		FDataflowConnection* Con = Elem.Value;
+		size_t AddressA = Con->RealAddress();
+		size_t AddressB = (size_t)Reference;
+		if (Con->RealAddress() == (size_t)Reference)
+		{
+			return (FDataflowInput*)Con;
+		}
+	}
+	return &FDataflowInput::NoOpInput;
+}
+
+TArray< FDataflowConnection* > FDataflowNode::GetInputs() const
+{
+	TArray< FDataflowConnection* > Result;
+	Inputs.GenerateValueArray(Result);
+	return Result;
+}
+
+void FDataflowNode::ClearInputs()
+{
+	for (TPair<uint32, FDataflowConnection*> Elem : Inputs)
+	{
+		FDataflowConnection* Con = Elem.Value;
+		delete Con;
+	}
+	Inputs.Reset();
+}
+
+
+//
+// Outputs
+//
+
+
+void FDataflowNode::AddOutput(FDataflowConnection* InPtr)
 {
-	for (Dataflow::FConnection* In : Inputs)
+	if (InPtr)
 	{
-		ensureMsgf(!In->GetName().IsEqual(InPtr->GetName()), TEXT("Add Input Failed: Existing Node input already defined with name (%s)"), *InPtr->GetName().ToString());
+		for (TPair<uint32, FDataflowConnection*> Elem : Outputs)
+		{
+			FDataflowConnection* Out = Elem.Value;
+			ensureMsgf(!Out->GetName().IsEqual(InPtr->GetName()), TEXT("Add Output Failed: Existing Node output already defined with name (%s)"), *InPtr->GetName().ToString());
+		}
+
+		if (!Outputs.Contains(InPtr->Property->GetOffset_ForInternal()))
+		{
+			Outputs.Add(InPtr->Property->GetOffset_ForInternal(), InPtr);
+		}
+		else
+		{
+			Outputs[InPtr->Property->GetOffset_ForInternal()] = InPtr;
+		}
+	}
+}
+
+
+
+FDataflowOutput* FDataflowNode::FindOutput(FName InName)
+{
+	for (TPair<uint32, FDataflowConnection*> Elem : Outputs)
+	{
+		FDataflowConnection* Con = Elem.Value;
+		if (Con->GetName().IsEqual(InName))
+		{
+			return (FDataflowOutput*)Con;
+		}
+	}
+	return nullptr;
+}
+
+const FDataflowOutput* FDataflowNode::FindOutput(const void* Reference) const
+{
+	for (TPair<uint32, FDataflowConnection*> Elem : Outputs)
+	{
+		FDataflowConnection* Con = Elem.Value;
+		if (Con->RealAddress() == (size_t)Reference)
+		{
+			return (FDataflowOutput*)Con;
+		}
 	}
-	Inputs.Add(InPtr);
+	return nullptr;
 }
 
-void FDataflowNode::AddOutput(Dataflow::FConnection* InPtr)
+const FDataflowOutput* FDataflowNode::GetOutput(const void* Reference) const
 {
-	for (Dataflow::FConnection* Out : Outputs)
+	for (TPair<uint32, FDataflowConnection*> Elem : Outputs)
 	{
-		ensureMsgf(!Out->GetName().IsEqual(InPtr->GetName()), TEXT("Add Output Failed: Existing Node output already defined with name (%s)"), *InPtr->GetName().ToString());
+		FDataflowConnection* Con = Elem.Value;
+		if (Con->RealAddress() == (size_t)Reference)
+		{
+			return (FDataflowOutput*)Con;
+		}
 	}
-	Outputs.Add(InPtr);
+	return &FDataflowOutput::NoOpOutput;
 }
 
+TArray< FDataflowConnection* > FDataflowNode::GetOutputs() const
+{
+	TArray< FDataflowConnection* > Result;
+	Outputs.GenerateValueArray(Result);
+	return Result;
+}
+
+
+void FDataflowNode::ClearOutputs()
+{
+	for (TPair<uint32, FDataflowConnection*> Elem : Outputs)
+	{
+		FDataflowConnection* Con = Elem.Value;
+		delete Con;
+	}
+	Outputs.Reset();
+}
 
 
 TArray<Dataflow::FPin> FDataflowNode::GetPins() const
 {
 	TArray<Dataflow::FPin> RetVal;
-	for (Dataflow::FConnection* Con : Inputs)
+	for (TPair<uint32, FDataflowConnection*> Elem : Inputs)
+	{
+		FDataflowConnection* Con = Elem.Value;
 		RetVal.Add({ Dataflow::FPin::EDirection::INPUT,Con->GetType(), Con->GetName() });
-	for (Dataflow::FConnection* Con : Outputs)
+	}
+	for (TPair<uint32, FDataflowConnection*> Elem : Outputs)
+	{
+		FDataflowConnection* Con = Elem.Value;
 		RetVal.Add({ Dataflow::FPin::EDirection::OUTPUT,Con->GetType(), Con->GetName() });
+	}
 	return RetVal;
 }
 
 void FDataflowNode::InvalidateOutputs()
 {
-	for (Dataflow::FConnection* Output : Outputs)
+	for (TPair<uint32, FDataflowConnection*> Elem : Outputs)
 	{
-		Output->Invalidate();
+		FDataflowConnection* Con = Elem.Value;
+		Con->Invalidate();
 	}
 }
 
-Dataflow::FConnection* FDataflowNode::FindInput(FName InName) const
+void FDataflowNode::RegisterInputConnection(const void* Data)
 {
-	for (Dataflow::FConnection* Input : Inputs)
+	if (TUniquePtr<FStructOnScope> ScriptOnStruct = TUniquePtr<FStructOnScope>(NewScructOnScope()))
 	{
-		if (Input->GetName().IsEqual(InName))
+		if (const UStruct* Struct = ScriptOnStruct->GetStruct())
 		{
-			return Input;
+			for (TFieldIterator<FProperty> PropertyIt(Struct); PropertyIt; ++PropertyIt)
+			{
+				FProperty* Property = *PropertyIt;
+				size_t RealAddress = (size_t)this + Property->GetOffset_ForInternal();
+				if (RealAddress == (size_t)Data)
+				{
+					FName PropName(Property->GetName());
+					FName PropType(Property->GetCPPType());
+					AddInput(new FDataflowInput({ PropType, PropName, this, Property }));
+				}
+			}
 		}
 	}
-	return nullptr;
 }
 
-Dataflow::FConnection* FDataflowNode::FindOutput(FName InName) const
+void FDataflowNode::RegisterOutputConnection(const void* Data)
 {
-	for (Dataflow::FConnection* Output : Outputs)
+	if (TUniquePtr<FStructOnScope> ScriptOnStruct = TUniquePtr<FStructOnScope>(NewScructOnScope()))
 	{
-		if (Output->GetName().IsEqual(InName))
+		if (const UStruct* Struct = ScriptOnStruct->GetStruct())
 		{
-			return Output;
+			for (TFieldIterator<FProperty> PropertyIt(Struct); PropertyIt; ++PropertyIt)
+			{
+				FProperty* Property = *PropertyIt;
+				size_t RealAddress = (size_t)this + Property->GetOffset_ForInternal();
+				if (RealAddress == (size_t)Data)
+				{
+					FName PropName(Property->GetName());
+					FName PropType(Property->GetCPPType());
+					AddOutput(new FDataflowOutput({ PropType, PropName, this, Property }));
+				}
+			}
 		}
 	}
-	return nullptr;
+}
+
+
+bool FDataflowNode::ValidateConnections()
+{
+	bool bResult = true;
+	if (TUniquePtr<FStructOnScope> ScriptOnStruct = TUniquePtr<FStructOnScope>(NewScructOnScope()))
+	{
+		if (const UStruct* Struct = ScriptOnStruct->GetStruct())
+		{
+			for (TFieldIterator<FProperty> PropertyIt(Struct); PropertyIt; ++PropertyIt)
+			{
+				FProperty* Property = *PropertyIt;
+				FName PropName(Property->GetName());
+				FName PropType(Property->GetCPPType());
+#if WITH_EDITORONLY_DATA
+				if (Property->HasMetaData(TEXT("DataflowInput")))
+				{
+					FDataflowInput* Input = FindInput(PropName);
+					if (ensureMsgf(Input, TEXT("Missing dataflow RegisterInputConnection in constructor for (%s:%s)"), *GetName().ToString(), *PropName.ToString()))
+					{
+						bResult = false;
+					}
+				}
+				if (Property->HasMetaData(TEXT("DataflowOutput")))
+				{
+					FDataflowOutput* Output = FindOutput(PropName);
+					if(ensureMsgf(Output,TEXT("Missing dataflow RegisterOutputConnection in constructor for (%s:%s)"), *GetName().ToString(),*PropName.ToString()))
+					{
+						bResult = false;
+					}
+				}
+#endif
+			}
+		}
+	}
+	return bResult;
 }
 
 
diff --git a/Engine/Source/Runtime/Experimental/Dataflow/Core/Public/Dataflow/DataflowConnection.h b/Engine/Source/Runtime/Experimental/Dataflow/Core/Public/Dataflow/DataflowConnection.h
--- a/Engine/Source/Runtime/Experimental/Dataflow/Core/Public/Dataflow/DataflowConnection.h
+++ b/Engine/Source/Runtime/Experimental/Dataflow/Core/Public/Dataflow/DataflowConnection.h
@@ -5,6 +5,7 @@
 #include "CoreMinimal.h"
 #include "Dataflow/DataflowNodeParameters.h"
 
+#include "DataflowConnection.generated.h"
 
 struct FDataflowNode;
 
@@ -26,50 +27,62 @@ namespace Dataflow
 		FName Type;
 		FName Name;
 	};
+}
 
-	//
-	// Input Output Base
-	//
+//
+// Input Output Base
+//
+USTRUCT()
+struct DATAFLOWCORE_API FDataflowConnection
+{
+	GENERATED_USTRUCT_BODY()
 
-	class DATAFLOWCORE_API FConnection
-	{
+protected:
+	Dataflow::FPin::EDirection Direction;
+	FName Type;
+	FName Name;
+	FDataflowNode* OwningNode = nullptr;
+	FProperty* Property = nullptr;
+	FGuid  Guid;
 
-	protected:
-		FPin::EDirection Direction;
-		FName Type;
-		FName Name;
-		FGuid  Guid;
-		FDataflowNode* OwningNode = nullptr;
+	friend struct FDataflowNode;
 
-		friend struct FDataflowNode;
-		static void BindInput(FDataflowNode* InNode, FConnection*);
-		static void BindOutput(FDataflowNode* InNode, FConnection*);
+public:
+	FDataflowConnection() {};
+	FDataflowConnection(Dataflow::FPin::EDirection Direction, FName InType, FName InName, FDataflowNode* OwningNode = nullptr, FProperty* InProperty = nullptr, FGuid InGuid = FGuid::NewGuid());
+	virtual ~FDataflowConnection() {};
 
+	FDataflowNode* GetOwningNode() { return OwningNode; }
+	const FDataflowNode* GetOwningNode() const { return OwningNode; }
 
-	public:
-		FConnection(FPin::EDirection Direction, FName InType, FName InName, FDataflowNode* OwningNode = nullptr, FGuid InGuid = FGuid::NewGuid());
-		virtual ~FConnection() {};
+	Dataflow::FPin::EDirection GetDirection() const { return Direction; }
+	uint32 GetOffset( ) const;
 
-		FDataflowNode* GetOwningNode() { return OwningNode; }
-		const FDataflowNode* GetOwningNode() const { return OwningNode; }
+	FName GetType() const { return Type; }
 
-		FPin::EDirection GetDirection() const { return Direction; }
+	FGuid GetGuid() const { return Guid; }
+	void SetGuid(FGuid InGuid) { Guid = InGuid; }
 
-		FName GetType() const { return Type; }
+	FName GetName() const { return Name; }
+	void SetName(FName InName) { Name = InName; }
 
-		FGuid GetGuid() const { return Guid; }
-		void SetGuid(FGuid InGuid) { Guid = InGuid; }
+	size_t RealAddress() const { ensure(OwningNode);  return (size_t)OwningNode + (size_t)GetOffset(); };
+	size_t CacheKey() const { return RealAddress(); };
 
-		FName GetName() const { return Name; }
-		void SetName(FName InName) { Name = InName; }
+	virtual bool AddConnection(FDataflowConnection* In) { return false; };
+	virtual bool RemoveConnection(FDataflowConnection* In) { return false; }
 
+	virtual TArray< FDataflowConnection* > GetConnectedInputs() { return TArray<FDataflowConnection* >(); }
+	virtual TArray< FDataflowConnection* > GetConnectedOutputs() { return TArray<FDataflowConnection* >(); }
 
-		virtual bool AddConnection(FConnection* In) { return false; };
-		virtual bool RemoveConnection(FConnection* In) { return false; }
+	template<class T>
+	bool IsA(const T* InVar) const
+	{
+		return (size_t)OwningNode + (size_t)GetOffset() == (size_t)InVar;
 
-		virtual TArray< FConnection* > GetConnectedInputs() { return TArray<FConnection* >(); }
-		virtual TArray< FConnection* > GetConnectedOutputs() { return TArray<FConnection* >(); }
-		virtual void Invalidate() {};
+	}
 
-	};
-}
+	virtual void Invalidate() {};
+	virtual bool Evaluate(Dataflow::FContext& Context) const { return false; };
+
+};
diff --git a/Engine/Source/Runtime/Experimental/Dataflow/Core/Public/Dataflow/DataflowGraph.h b/Engine/Source/Runtime/Experimental/Dataflow/Core/Public/Dataflow/DataflowGraph.h
--- a/Engine/Source/Runtime/Experimental/Dataflow/Core/Public/Dataflow/DataflowGraph.h
+++ b/Engine/Source/Runtime/Experimental/Dataflow/Core/Public/Dataflow/DataflowGraph.h
@@ -7,9 +7,10 @@
 #include "Dataflow/DataflowNode.h"
 #include "Serialization/Archive.h"
 
+struct FDataflowConnection;
+
 namespace Dataflow
 {
-	class FConnection;
 
 	struct FLink {
 		FGuid InputNode;
@@ -80,9 +81,9 @@ namespace Dataflow
 
 		void RemoveNode(TSharedPtr<FDataflowNode> Node);
 
-		void ClearConnections(FConnection*);
-		void Connect(FConnection* Input, FConnection* Output);
-		void Disconnect(FConnection* Input, FConnection* Output);
+		void ClearConnections(FDataflowConnection*);
+		void Connect(FDataflowConnection* Input, FDataflowConnection* Output);
+		void Disconnect(FDataflowConnection* Input, FDataflowConnection* Output);
 
 		virtual void Serialize(FArchive& Ar);
 		const TSet<FName>& GetDisabledNodes() const { return DisabledNodes; }
diff --git a/Engine/Source/Runtime/Experimental/Dataflow/Core/Public/Dataflow/DataflowInputOutput.h b/Engine/Source/Runtime/Experimental/Dataflow/Core/Public/Dataflow/DataflowInputOutput.h
--- a/Engine/Source/Runtime/Experimental/Dataflow/Core/Public/Dataflow/DataflowInputOutput.h
+++ b/Engine/Source/Runtime/Experimental/Dataflow/Core/Public/Dataflow/DataflowInputOutput.h
@@ -4,260 +4,146 @@
 
 #include "CoreMinimal.h"
 #include "Dataflow/DataflowNodeParameters.h"
-#include "Dataflow/DataflowNode.h"
 #include "Dataflow/DataflowConnection.h"
+#include "Dataflow/DataflowNode.h"
 #include "Templates/Function.h"
 
-namespace Dataflow
-{
-	template<typename T> class TOutput;
+#include "DataflowInputOutput.generated.h"
 
-	//
-	//  Input
-	//
 
-	template<class T>
-	struct DATAFLOWCORE_API TInputParameters {
-		TInputParameters(FName InName = FName(""), FDataflowNode * InOwner = nullptr, T InDefault = T())
-			: Type(GraphConnectionTypeName<T>())
+struct FDataflowOutput;
+
+//
+//  Input
+//
+namespace Dataflow
+{
+	struct DATAFLOWCORE_API FInputParameters {
+		FInputParameters(FName InType = FName(""), FName InName = FName(""), FDataflowNode * InOwner = nullptr, FProperty * InProperty = nullptr)
+			: Type(InType)
 			, Name(InName)
 			, Owner(InOwner)
-			, Default(InDefault) {}
+			, Property(InProperty){}
 		FName Type;
 		FName Name;
 		FDataflowNode* Owner = nullptr;
-		T Default = T();
+		FProperty* Property = nullptr;
 	};
+}
 
-	template<class T>
-	class TInput : public FConnection
-	{
-		typedef FConnection Super;
-		friend class FConnection;
-
-		T Default = T();
-		TOutput<T>* Connection;
-	public:
-		TInput(const TInputParameters<T>& Param = {}, FGuid InGuid = FGuid::NewGuid())
-			: FConnection(FPin::EDirection::INPUT, Param.Type, Param.Name, Param.Owner, InGuid)
-			, Default(Param.Default)
-			, Connection(nullptr)
-		{
-			Super::BindInput(Param.Owner, this);
-		}
-	
-		const T& GetDefault() const { return Default; }
-		
-		const TOutput<T>* GetConnection() const { return Connection; }
-		TOutput<T>* GetConnection() { return Connection; }
+USTRUCT()
+struct FDataflowInput : public FDataflowConnection
+{
+	GENERATED_USTRUCT_BODY()
 
-		virtual bool AddConnection(FConnection* InOutput) override
-		{ 
-			ensure(Connection == nullptr);
-			if (ensure(InOutput->GetType()==this->GetType()))
-			{
-				Connection = (TOutput<T>*)InOutput;
-				return true;
-			}
-			return false;
-		}
+	static FDataflowInput NoOpInput;
 
-		virtual bool RemoveConnection(FConnection* InOutput) override
-		{ 
-			if (ensure(Connection == (TOutput<T>*)InOutput))
-			{
-				Connection = nullptr;
-				return true;
-			}
-			return false;
-		}
+	friend struct FDataflowConnection;
 
-		T GetValue(const FContext& Context) const
-		{
-			if (Connection)
-			{
-				return Connection->GetValue(Context);
-			}
-			return Default;
-		}
+	FDataflowOutput* Connection;
+public:
+	FDataflowInput(const Dataflow::FInputParameters& Param = {}, FGuid InGuid = FGuid::NewGuid());
+
+	virtual bool AddConnection(FDataflowConnection* InOutput) override;
+	virtual bool RemoveConnection(FDataflowConnection* InOutput) override;
 
-		void SetValue(const T& Value, const FContext& Context)
+	FDataflowOutput* GetConnection() { return Connection; }
+	const FDataflowOutput* GetConnection() const { return Connection; }
+	virtual TArray< FDataflowConnection* > GetConnectedOutputs() override;
+
+	// @todo(dataflow) : needs testing. 
+	template<class T> const T& GetValue(Dataflow::FContext& Context, const T& Default) const
+	{
+		if (FDataflowConnection* ConnectionBase = ((FDataflowConnection*)Connection) )
 		{
-			Default = Value;
-			if (!Connection)
+			if (!ConnectionBase->Evaluate(Context))
 			{
-				OwningNode->InvalidateOutputs();
+				Context.SetData(ConnectionBase->CacheKey(), new Dataflow::ContextCache<T>(Property, new T(Default)));
 			}
-		}
-
-		virtual TArray< FConnection* > GetConnectedOutputs() override
-		{ 
-			TArray<FConnection* > RetList;
-			if (GetConnection())
+			if (Context.HasData(ConnectionBase->CacheKey()))
 			{
-				RetList.Add(GetConnection());
+				return Context.GetDataReference<T>(ConnectionBase->CacheKey(), Default);
 			}
-			return RetList;
 		}
+		return Default;
+	}
 
-		virtual void Invalidate() override
-		{
-			OwningNode->InvalidateOutputs();
-		}
+	virtual void Invalidate() override;
+};
 
-	};
 
-
-	//
-	// Output
-	//
-
-	template<class T>
-	struct DATAFLOWCORE_API TOutputParameters 
+//
+// Output
+//
+namespace Dataflow
+{
+	struct DATAFLOWCORE_API FOutputParameters
 	{
-		typedef TFunction<T(const  Dataflow::FContext& Context, Dataflow::FConnection*)> PassthroughCallback;
-
-		TOutputParameters(FName InName = FName(""), FDataflowNode * InOwner = nullptr, T InDefault = T())
-			: Type(GraphConnectionTypeName<T>())
+		FOutputParameters(FName InType = FName(""), FName InName = FName(""), FDataflowNode* InOwner = nullptr, FProperty* InProperty = nullptr)
+			: Type(InType)
 			, Name(InName)
 			, Owner(InOwner)
-			, Default(InDefault){}
+			, Property(InProperty) {}
 
-		TOutputParameters(FName InName, FDataflowNode* InOwner, T InDefault, PassthroughCallback InPassthrough)
-			: Type(GraphConnectionTypeName<T>())
-			, Name(InName)
-			, Owner(InOwner)
-			, Default(InDefault)
-			, Passthrough(InPassthrough) {}
 		FName Type;
 		FName Name;
 		FDataflowNode* Owner = nullptr;
-		T Default = T();
-		PassthroughCallback Passthrough;
+		FProperty* Property = nullptr;
 	};
+}
+USTRUCT()
+struct DATAFLOWCORE_API FDataflowOutput : public FDataflowConnection
+{
+	GENERATED_USTRUCT_BODY()
 
-	template<class T>
-	class DATAFLOWCORE_API TOutput : public FConnection
-	{
-		typedef TFunction<T(const  Dataflow::FContext& Context, Dataflow::FConnection*)> PassthroughCallback;
-		typedef FConnection Super;
-		friend class FConnection;
-
-		mutable uint32 CacheKey = UINT_MAX;
-		mutable TCacheValue<T> Cache;
-
-		TArray< TInput<T>* > Connections;
-		PassthroughCallback Passthrough;
-		T Default = T();
-
-	public:
-		TOutput(const TOutputParameters<T>& Param = {}, FGuid InGuid = FGuid::NewGuid())
-			: FConnection(FPin::EDirection::OUTPUT, Param.Type, Param.Name, Param.Owner, InGuid)
-			, Passthrough(Param.Passthrough)
-			, Default(Param.Default)
-		{
-			Super::BindOutput(Param.Owner, this);
-		}
+	friend struct FDataflowConnection;
 
-		const TArray<TInput<T>*>& GetConnections() const { return Connections; }
-		TArray<TInput<T>* >& GetConnections() { return Connections; }
+	mutable uint32 CacheKeyValue = UINT_MAX;
+	mutable Dataflow::TCacheValue Cache;
 
-		virtual TArray<FConnection*> GetConnectedInputs() override
-		{ 
-			TArray<FConnection*> RetList;
-			RetList.Reserve(Connections.Num());
-			for (TInput<T>* Ptr : Connections) { RetList.Add(Ptr); }
-			return RetList;
-		}
+	TArray< FDataflowInput* > Connections;
 
+public:
+	static FDataflowOutput NoOpOutput;
 
-		virtual bool AddConnection(FConnection* InOutput) override
-		{
-			if (ensure(InOutput->GetType() == this->GetType()))
-			{
-				Connections.Add((TInput<T>*)InOutput);
-				return true;
-			}
-			return false;
-		}
-		
-		virtual bool RemoveConnection(FConnection* InInput) override { Connections.RemoveSwap((TInput<T>*)InInput); return true; }
+	FDataflowOutput(const Dataflow::FOutputParameters& Param = {}, FGuid InGuid = FGuid::NewGuid());
 
-		void SetValue(T InVal, const FContext& Context) const
-		{
-			CacheKey = Context.GetTypeHash();
-			Cache.Data = InVal;
-		}
+	TArray<FDataflowInput*>& GetConnections();
+	const TArray<FDataflowInput*>& GetConnections() const;
+
+	virtual TArray<FDataflowConnection*> GetConnectedInputs() override;
+
+	virtual bool AddConnection(FDataflowConnection* InOutput) override;
 
-		T GetValue(const FContext& Context)
+	virtual bool RemoveConnection(FDataflowConnection* InInput) override;
+
+	template<class T>
+	void SetValue(const T& InVal, Dataflow::FContext& Context) const
+	{
+		if (Property)
 		{
-			if (CacheKey != Context.GetTypeHash())
-			{
-				if (OwningNode->bActive)
-				{
-					OwningNode->Evaluate(Context, this);
-				}
-				else
-				{
-					SetValue(Default, Context);
-					if (Passthrough)
-					{
-						SetValue(Passthrough(Context, this), Context);
-					}
-				}
-			}
-			ensure(CacheKey == Context.GetTypeHash());
-			return DeepCopy<T>(Cache.Data);
+			Context.SetData(CacheKey(), new Dataflow::ContextCache<T>(Property, new T(InVal)));
 		}
+	}
 
-		virtual void Invalidate() override
+	template<class T> const T& GetValue(Dataflow::FContext& Context, const T& Default) const
+	{
+		if (!this->Evaluate(Context))
 		{
-			if (CacheKey != UINT_MAX)
-			{
-				CacheKey = UINT_MAX;
-				Cache = T();
-
-				for (FConnection* Con : GetConnections())
-				{
-					Con->Invalidate();
-				}
-			}
+			Context.SetData(CacheKey(), new Dataflow::ContextCache<T>(Property, new T(Default)));
 		}
 
-	};
 
+		if (Context.HasData(CacheKey()))
+		{
+			return Context.GetDataReference<T>(CacheKey(), Default);
+		}
 
-//
-// Used this macros to defined dataflow connection types 
-//
+		return Default;
+	}
 
-#define DATAFLOW_CONNECTION_TYPE_PRIMITIVE(D, a,A)					\
-	template<> inline FName D GraphConnectionTypeName<a>()			\
-		{ return FName(TEXT(#A)); }									\
-	template<> inline a D DeepCopy<a>(const a & Val)				\
-		{ return Val; }												\
-	template struct D TOutputParameters<a>;							\
-	template class D TOutput<a>;									\
-	template struct D TInputParameters<a>;							\
-	template class D TInput<a>;
-
-#define DATAFLOW_CONNECTION_TYPE_SHARED_POINTER(D, a,A)				\
-	template<> inline FName D GraphConnectionTypeName<a>()			\
-		{ return FName(TEXT(#A)); }									\
-	template<> inline a D DeepCopy<a>(const a& Val)					\
-		{ return Val ? TSharedPtr<a::ElementType>(Val->NewCopy())	\
-		: TSharedPtr<a::ElementType>(nullptr); }					\
-	template struct D TOutputParameters<a>;							\
-	template class D TOutput<a>;									\
-	template struct D TInputParameters<a>;							\
-	template class D TInput<a>;
-
-
-DATAFLOW_CONNECTION_TYPE_PRIMITIVE(DATAFLOWCORE_API, bool, Bool)
-DATAFLOW_CONNECTION_TYPE_PRIMITIVE(DATAFLOWCORE_API, char, Char)
-DATAFLOW_CONNECTION_TYPE_PRIMITIVE(DATAFLOWCORE_API, int, Integer)
-DATAFLOW_CONNECTION_TYPE_PRIMITIVE(DATAFLOWCORE_API, uint8, UInt8)
-DATAFLOW_CONNECTION_TYPE_PRIMITIVE(DATAFLOWCORE_API, float, Float)
-DATAFLOW_CONNECTION_TYPE_PRIMITIVE(DATAFLOWCORE_API, double, Double)
-}
+		
+	bool Evaluate(Dataflow::FContext& Context) const;
+	virtual void Invalidate() override;
 
+};
\ No newline at end of file
diff --git a/Engine/Source/Runtime/Experimental/Dataflow/Core/Public/Dataflow/DataflowNode.h b/Engine/Source/Runtime/Experimental/Dataflow/Core/Public/Dataflow/DataflowNode.h
--- a/Engine/Source/Runtime/Experimental/Dataflow/Core/Public/Dataflow/DataflowNode.h
+++ b/Engine/Source/Runtime/Experimental/Dataflow/Core/Public/Dataflow/DataflowNode.h
@@ -2,6 +2,7 @@
 
 #pragma once
 
+#include "ChaosLog.h"
 #include "CoreMinimal.h"
 #include "Dataflow/DataflowConnection.h"
 #include "Dataflow/DataflowNodeParameters.h"
@@ -9,6 +10,9 @@
 
 #include "DataflowNode.generated.h"
 
+struct FDataflowInput;
+struct FDataflowOutput;
+
 namespace Dataflow {
 	struct DATAFLOWCORE_API FNodeParameters {
 		FName Name;
@@ -17,22 +21,27 @@ namespace Dataflow {
 }
 
 
-//
-// FNode
-//
+/**
+* FNode
+*		Base class for node based evaluation within the Dataflow graph. 
+* 
+* Note : Do NOT create mutable variables in the classes derived from FDataflowNode. The state
+*        is stored on the FContext. The Evaluate is const to allow support for multithreaded
+*        evaluation. 
+*/
 USTRUCT()
 struct DATAFLOWCORE_API FDataflowNode
 {
 	GENERATED_USTRUCT_BODY()
 
 	friend class Dataflow::FGraph;
-	friend class Dataflow::FConnection;
+	friend struct FDataflowConnection;
 
 	FGuid Guid;
 	FName Name;
 
-	TArray< Dataflow::FConnection* > Inputs;
-	TArray< Dataflow::FConnection* > Outputs;
+	TMap< int, FDataflowConnection* > Inputs;
+	TMap< int, FDataflowConnection* > Outputs;
 
 	UPROPERTY(EditAnywhere, Category = "Dataflow")
 	bool bActive = true;
@@ -40,14 +49,16 @@ struct DATAFLOWCORE_API FDataflowNode
 	FDataflowNode()
 		: Guid(FGuid())
 		, Name("Invalid")
-	{}
+	{
+	}
 
 	FDataflowNode(const Dataflow::FNodeParameters& Param, FGuid InGuid = FGuid::NewGuid())
 		: Guid(InGuid)
 		, Name(Param.Name)
-	{}
+	{
+	}
 
-	virtual ~FDataflowNode() {}
+	virtual ~FDataflowNode() { ClearInputs(); ClearOutputs(); }
 
 	FGuid GetGuid() const { return Guid; }
 	FName GetName() const { return Name; }
@@ -59,53 +70,90 @@ struct DATAFLOWCORE_API FDataflowNode
 	virtual FString GetTags() const { return ""; }
 	virtual FString GetToolTip() const { return ""; }
 
+	//
+	// Connections
+	//
 
 	TArray<Dataflow::FPin> GetPins() const;
-	Dataflow::FConnection* FindInput(FName Name) const;
-	Dataflow::FConnection* FindOutput(FName Name) const;
 
-	virtual void Evaluate(const  Dataflow::FContext& Context, Dataflow::FConnection*) const { ensure(false); }
-	void InvalidateOutputs();
+	void AddInput(FDataflowConnection* InPtr);
+	FDataflowInput* FindInput(FName Name);
+	const FDataflowInput* FindInput(const void* Reference) const;
+	const FDataflowInput* GetInput(const void* Reference) const;
+	TArray< FDataflowConnection* > GetInputs() const;
+	void ClearInputs();
 
-	void AddInput(Dataflow::FConnection* InPtr);
-	const TArray< Dataflow::FConnection* >& GetInputs() const { return Inputs; }
-	TArray< Dataflow::FConnection* >& GetInputs() { return Inputs; }
+	void AddOutput(FDataflowConnection* InPtr);
+	FDataflowOutput* FindOutput(FName Name);
+	const FDataflowOutput* FindOutput(const void* Reference) const;
+	const FDataflowOutput* GetOutput(const void* Reference) const;
+	TArray< FDataflowConnection* > GetOutputs() const;
+	void ClearOutputs();
 
-	void AddOutput(Dataflow::FConnection* InPtr);
-	const TArray< Dataflow::FConnection* >& GetOutputs() const { return Outputs; }
-	TArray< Dataflow::FConnection* >& GetOutputs() { return Outputs; }
+	//
+	//  Struct Support
+	//
 
 	virtual void SerializeInternal(FArchive& Ar) { check(false); }
 	virtual FStructOnScope* NewScructOnScope() { return nullptr; }
+
+	/** Register the Input and Outputs after the creation in the factory */
+	void RegisterInputConnection(const void*);
+	void RegisterOutputConnection(const void*);
+
+	//
+	// Evaluation
+	//
+	virtual void Evaluate(Dataflow::FContext& Context, const FDataflowOutput*) const { ensure(false); }
+
+	void InvalidateOutputs();
+
+
+	template<class T>
+	void SetValue(Dataflow::FContext& Context, const T& NewValue, const T* Reference) const
+	{
+		if (const FDataflowConnection* Con = FindOutput(Reference))
+		{
+			Context.SetData(Con->RealAddress(), new Dataflow::ContextCache<T>( Con->Property, new T(NewValue)));
+		}
+		else
+		{
+			ensure(false);
+		}
+	}
+
+	bool ValidateConnections();
+
 };
 
 namespace Dataflow
 {
 	//
-	// Used these macros to register dataflow nodes. 
+	// Use these macros to register dataflow nodes. 
 	//
 
-#define DATAFLOW_NODE_REGISTER_CREATION_FACTORY(A)								\
-	FNodeFactory::GetInstance()->RegisterNode(									\
-		{A::StaticType(),A::StaticDisplay(),A::StaticCategory(),				\
-			A::StaticTags(),A::StaticToolTip()},								\
-		[](const FNewNodeParameters& InParam)									\
-			{return new A({InParam.Name}, InParam.Guid); });
-
-#define DATAFLOW_NODE_DEFINE_INTERNAL(TYPE, DISPLAY_NAME, CATEGORY, TAGS)		\
-public:																			\
-	static FName StaticType() {return #TYPE;}									\
-	static FName StaticDisplay() {return DISPLAY_NAME;}							\
-	static FName StaticCategory() {return CATEGORY;}							\
-	static FString StaticTags() {return TAGS;}									\
-	static FString StaticToolTip() {return FString("Create a dataflow node.");}	\
-	virtual FName GetType() const { return #TYPE; }								\
-	virtual FStructOnScope* NewScructOnScope() override {						\
-	   return new FStructOnScope(TYPE::StaticStruct(), (uint8*)this);}			\
-	virtual void SerializeInternal(FArchive& Ar) override {						\
-		UScriptStruct* const Struct = TYPE::StaticStruct();						\
-		Struct->SerializeTaggedProperties(Ar, (uint8*)this,						\
-		Struct, nullptr);}														\
+#define DATAFLOW_NODE_REGISTER_CREATION_FACTORY(A)									\
+	FNodeFactory::GetInstance()->RegisterNode(										\
+		{A::StaticType(),A::StaticDisplay(),A::StaticCategory(),					\
+			A::StaticTags(),A::StaticToolTip()},									\
+		[](const FNewNodeParameters& InParam){										\
+				A* Val=new A({InParam.Name}, InParam.Guid);							\
+				Val->ValidateConnections(); return Val;});
+
+#define DATAFLOW_NODE_DEFINE_INTERNAL(TYPE, DISPLAY_NAME, CATEGORY, TAGS)			\
+public:																				\
+	static FName StaticType() {return #TYPE;}										\
+	static FName StaticDisplay() {return DISPLAY_NAME;}								\
+	static FName StaticCategory() {return CATEGORY;}								\
+	static FString StaticTags() {return TAGS;}										\
+	static FString StaticToolTip() {return FString("Create a dataflow node.");}		\
+	virtual FName GetType() const { return #TYPE; }									\
+	virtual FStructOnScope* NewScructOnScope() override {							\
+	   return new FStructOnScope(TYPE::StaticStruct(), (uint8*)this);}				\
+	virtual void SerializeInternal(FArchive& Ar) override {							\
+		UScriptStruct* const Struct = TYPE::StaticStruct();							\
+		Struct->SerializeTaggedProperties(Ar, (uint8*)this,							\
+		Struct, nullptr);}															\
 	virtual FName GetDisplayName() const override { return TYPE::StaticDisplay(); }	\
 	virtual FName GetCategory() const override { return TYPE::StaticCategory(); }	\
 	virtual FString GetTags() const override { return TYPE::StaticTags(); }			\
diff --git a/Engine/Source/Runtime/Experimental/Dataflow/Core/Public/Dataflow/DataflowNodeFactory.h b/Engine/Source/Runtime/Experimental/Dataflow/Core/Public/Dataflow/DataflowNodeFactory.h
--- a/Engine/Source/Runtime/Experimental/Dataflow/Core/Public/Dataflow/DataflowNodeFactory.h
+++ b/Engine/Source/Runtime/Experimental/Dataflow/Core/Public/Dataflow/DataflowNodeFactory.h
@@ -7,10 +7,11 @@
 #include "ChaosLog.h"
 
 struct FDataflowNode;
+struct FDataflowConnection;
+
 
 namespace Dataflow
 {
-	class FConnection;
 
 	struct DATAFLOWCORE_API FNewNodeParameters {
 		FGuid Guid;
diff --git a/Engine/Source/Runtime/Experimental/Dataflow/Core/Public/Dataflow/DataflowNodeParameters.h b/Engine/Source/Runtime/Experimental/Dataflow/Core/Public/Dataflow/DataflowNodeParameters.h
--- a/Engine/Source/Runtime/Experimental/Dataflow/Core/Public/Dataflow/DataflowNodeParameters.h
+++ b/Engine/Source/Runtime/Experimental/Dataflow/Core/Public/Dataflow/DataflowNodeParameters.h
@@ -8,14 +8,51 @@ class  UDataflow;
 
 namespace Dataflow
 {
+	struct ContextCacheBase 
+	{
+		ContextCacheBase(FProperty* InProperty = nullptr, void* InData = nullptr)
+			: Property(InProperty)
+			, Data(InData)
+		{}
+		virtual ~ContextCacheBase() {}
+
+		FProperty* Property = nullptr;
+		void* Data = nullptr;
+	};
+
+	template<class T>
+	struct ContextCache : public ContextCacheBase 
+	{
+		ContextCache(FProperty* InProperty = nullptr, void* InData = nullptr)
+			: ContextCacheBase(InProperty,InData)
+		{}
+		virtual ~ContextCache() 
+		{ 
+			delete (T*)Data; 
+		}
+	};
+
 	class DATAFLOWCORE_API FContext
 	{
+		TMap<int64, ContextCacheBase*> DataStore;
+
 	public:
+
+
+
 		FContext(float InTime, FString InType = FString(""))
-		: Timestamp(InTime)
-		, Type(StaticType().Append(InType))
+			: Timestamp(InTime)
+			, Type(StaticType().Append(InType))
 		{}
 
+		~FContext()
+		{
+			for (TTuple<int64, ContextCacheBase*> Elem : DataStore)
+			{
+				delete Elem.Value;
+			}
+		}
+
 		float Timestamp = 0.f;
 		FString Type;
 		static FString StaticType() { return "FContext"; }
@@ -34,14 +71,58 @@ namespace Dataflow
 			}
 			return nullptr;
 		}
+
+		void SetData(size_t Key, ContextCacheBase* InData)
+		{
+			int64 IntKey = (int64)Key;
+			if (!DataStore.Contains(IntKey))
+			{
+				DataStore.Add(IntKey, InData);
+			}
+			else
+			{
+				DataStore[IntKey] = InData;
+			}
+		}
+
+		template<class T>
+		bool GetData(size_t Key, T& OutData)
+		{
+			int64 IntKey = (int64)Key;
+
+			if (DataStore.Contains(IntKey))
+			{
+				// @todo(dataflow) : type check
+				OutData = *(T*)DataStore[IntKey]->Data;
+				return true;
+			}
+			return false;
+		}
+
+		template<class T>
+		const T& GetDataReference(size_t Key, const T& Default)
+		{
+			int64 IntKey = (int64)Key;
+
+			if (DataStore.Contains(Key))
+			{
+				// @todo(dataflow) : type check
+				return *(T*)DataStore[Key]->Data;
+			}
+			return Default;
+		}
+
+		bool HasData(size_t Key)
+		{
+			int64 IntKey = (int64)Key;
+
+			return DataStore.Contains(IntKey);
+		}
 	};
 
 
-	template<class T>
-	struct TCacheValue  {
-		TCacheValue(T InData = T())
-			: Data(InData) {}
-		T Data;
+	struct TCacheValue {
+		TCacheValue() {}
 	};
 }
 
diff --git a/Engine/Source/Runtime/Experimental/Dataflow/Engine/Private/Dataflow/DataflowEdNode.cpp b/Engine/Source/Runtime/Experimental/Dataflow/Engine/Private/Dataflow/DataflowEdNode.cpp
--- a/Engine/Source/Runtime/Experimental/Dataflow/Engine/Private/Dataflow/DataflowEdNode.cpp
+++ b/Engine/Source/Runtime/Experimental/Dataflow/Engine/Private/Dataflow/DataflowEdNode.cpp
@@ -56,42 +56,48 @@ void UDataflowEdNode::PinConnectionListChanged(UEdGraphPin* Pin)
 	{
 		if (TSharedPtr<FDataflowNode> DataflowNode = DataflowGraph->FindBaseNode(DataflowNodeGuid))
 		{
-			if (Dataflow::FConnection* ConnectionInput = DataflowNode->FindInput(FName(Pin->GetName())) )
+			if (Pin->Direction == EEdGraphPinDirection::EGPD_Input)
 			{
-				DataflowGraph->ClearConnections(ConnectionInput);
-				for (UEdGraphPin* LinkedCon : Pin->LinkedTo)
+				if (FDataflowConnection* ConnectionInput = DataflowNode->FindInput(FName(Pin->GetName())))
 				{
-					if (UDataflowEdNode* LinkedNode = Cast<UDataflowEdNode>(LinkedCon->GetOwningNode()))
+					DataflowGraph->ClearConnections(ConnectionInput);
+					for (UEdGraphPin* LinkedCon : Pin->LinkedTo)
 					{
-						if (ensure(LinkedNode->IsBound()))
+						if (UDataflowEdNode* LinkedNode = Cast<UDataflowEdNode>(LinkedCon->GetOwningNode()))
 						{
-							if (TSharedPtr<FDataflowNode> LinkedDataflowNode = DataflowGraph->FindBaseNode(LinkedNode->GetDataflowNodeGuid()))
+							if (ensure(LinkedNode->IsBound()))
 							{
-								if (Dataflow::FConnection* LinkedConBase = LinkedDataflowNode->FindOutput(FName(LinkedCon->GetName())))
+								if (TSharedPtr<FDataflowNode> LinkedDataflowNode = DataflowGraph->FindBaseNode(LinkedNode->GetDataflowNodeGuid()))
 								{
-									DataflowGraph->Connect(ConnectionInput, LinkedConBase);
+									if (FDataflowConnection* LinkedConBase = LinkedDataflowNode->FindOutput(FName(LinkedCon->GetName())))
+									{
+										DataflowGraph->Connect(ConnectionInput, LinkedConBase);
+									}
 								}
 							}
 						}
 					}
 				}
 			}
-			else if (Dataflow::FConnection* Con = DataflowNode->FindOutput(FName(Pin->GetName())))
+			else if (Pin->Direction == EEdGraphPinDirection::EGPD_Output)
 			{
-				if (Dataflow::FConnection* ConnectionOutput = DataflowNode->FindOutput(FName(Pin->GetName())))
+				if (FDataflowConnection* Con = DataflowNode->FindOutput(FName(Pin->GetName())))
 				{
-					DataflowGraph->ClearConnections(ConnectionOutput);
-					for (UEdGraphPin* LinkedCon : Pin->LinkedTo)
+					if (FDataflowConnection* ConnectionOutput = DataflowNode->FindOutput(FName(Pin->GetName())))
 					{
-						if (UDataflowEdNode* LinkedNode = Cast<UDataflowEdNode>(LinkedCon->GetOwningNode()))
+						DataflowGraph->ClearConnections(ConnectionOutput);
+						for (UEdGraphPin* LinkedCon : Pin->LinkedTo)
 						{
-							if (ensure(LinkedNode->IsBound()))
+							if (UDataflowEdNode* LinkedNode = Cast<UDataflowEdNode>(LinkedCon->GetOwningNode()))
 							{
-								if (TSharedPtr<FDataflowNode> LinkedDataflowNode = DataflowGraph->FindBaseNode(LinkedNode->GetDataflowNodeGuid()))
+								if (ensure(LinkedNode->IsBound()))
 								{
-									if (Dataflow::FConnection* LinkedConBase = LinkedDataflowNode->FindInput(FName(LinkedCon->GetName())))
+									if (TSharedPtr<FDataflowNode> LinkedDataflowNode = DataflowGraph->FindBaseNode(LinkedNode->GetDataflowNodeGuid()))
 									{
-										DataflowGraph->Connect(LinkedConBase, ConnectionOutput);
+										if (FDataflowConnection* LinkedConBase = LinkedDataflowNode->FindInput(FName(LinkedCon->GetName())))
+										{
+											DataflowGraph->Connect(LinkedConBase, ConnectionOutput);
+										}
 									}
 								}
 							}
