Chaos : Flesh
Add corotated cache and associated tests. Add elastic FEM and associated tests. Add a primitive version of Newton's method with hyperelasticity. Compile fix for some loop in CG functions.

diff --git a/Engine/Source/Runtime/Experimental/Chaos/Private/Chaos/NewtonCorotatedCache.cpp b/Engine/Source/Runtime/Experimental/Chaos/Private/Chaos/NewtonCorotatedCache.cpp
new file mode 100644
--- /dev/null
+++ b/Engine/Source/Runtime/Experimental/Chaos/Private/Chaos/NewtonCorotatedCache.cpp
@@ -0,0 +1,105 @@
+// Copyright Epic Games, Inc. All Rights Reserved.
+#include "Chaos/NewtonCorotatedCache.h"
+#include "Chaos/ImplicitQRSVD.h"
+
+
+
+namespace Chaos::Softs
+{
+
+template<typename T>
+void CorotatedCache<T>::UpdateCache(const Chaos::PMatrix<T, 3, 3>& F) 
+{
+	Chaos::PMatrix<T, 3, 3> R((T)0.), S((T)0.), JFinvT((T)0.);
+	Chaos::PolarDecomposition(F, R, S);
+
+	JFInvTCache[0] = F.GetAt(1, 1) * F.GetAt(2, 2) - F.GetAt(2, 1) * F.GetAt(1, 2); JFInvTCache[1] = F.GetAt(2, 0) * F.GetAt(1, 2) - F.GetAt(1, 0) * F.GetAt(2, 2);  
+	JFInvTCache[2] = F.GetAt(1, 0) * F.GetAt(2, 1) - F.GetAt(2, 0) * F.GetAt(1, 1); JFInvTCache[3] = F.GetAt(2, 1) * F.GetAt(0, 2) - F.GetAt(0, 1) * F.GetAt(2, 2);
+	JFInvTCache[4] = F.GetAt(0, 0) * F.GetAt(2, 2) - F.GetAt(2, 0) * F.GetAt(0, 2); JFInvTCache[5] = F.GetAt(2, 0) * F.GetAt(0, 1) - F.GetAt(0, 0) * F.GetAt(2, 1);
+	JFInvTCache[6] = F.GetAt(0, 1) * F.GetAt(1, 2) - F.GetAt(1, 1) * F.GetAt(0, 2); JFInvTCache[7] = F.GetAt(1, 0) * F.GetAt(0, 2) - F.GetAt(0, 0) * F.GetAt(1, 2);
+	JFInvTCache[8] = F.GetAt(0, 0)* F.GetAt(1, 1) - F.GetAt(1, 0) * F.GetAt(0, 1);
+
+	Chaos::PMatrix<T, 3, 3> D = (S.GetAt(0, 0) + S.GetAt(1, 1) + S.GetAt(2, 2)) * Chaos::PMatrix<T, 3, 3>::Identity - S;
+	Chaos::PMatrix<T, 3, 3> DInv = D.Inverse();
+
+	for (int32 i = 0; i < 3; i++) 
+	{
+		for (int32 j = 0; j < 3; j++)
+		{
+			DInvCache[i * 3 + j] = DInv.GetAt(i, j);
+			RCache[i * 3 + j] = R.GetAt(i, j);
+		}
+	}
+
+	JCache = F.Determinant();
+
+}
+
+template<typename T>
+void CorotatedCache<T>::P(const Chaos::PMatrix<T, 3, 3>& F, const T mu, const T lambda, Chaos::PMatrix<T, 3, 3>& P) 
+{
+	Chaos::PMatrix<T, 3, 3> R((T)0.), JFinvT((T)0.);
+	for (int32 i = 0; i < 3; i++)
+	{
+		for (int32 j = 0; j < 3; j++)
+		{
+			R.SetAt(i, j, RCache[i * 3 + j]);
+			JFinvT.SetAt(i, j, JFInvTCache[i * 3 + j]);
+		}
+	}
+	P = (T)2. * mu * (F - R) + lambda * (JCache - (T)1.) * JFinvT;
+}
+
+template<typename T>
+void CorotatedCache<T>::deltaP(const Chaos::PMatrix<T, 3, 3>& F, const Chaos::PMatrix<T, 3, 3>& dF, const T mu, const T lambda, Chaos::PMatrix<T, 3, 3>& dP) 
+{
+	Chaos::PMatrix<T, 3, 3> R((T)0.), S((T)0.), JFinvT((T)0.), dR((T)0.), dJFinvT((T)0.), Dinv((T)0.);
+	for (int32 i = 0; i < 3; i++) 
+	{
+		for (int32 j = 0; j < 3; j++)
+		{
+			R.SetAt(i, j, RCache[i * 3 + j]);
+			JFinvT.SetAt(i, j, JFInvTCache[i * 3 + j]);
+			Dinv.SetAt(i, j, DInvCache[i * 3 + j]);
+		}
+	}
+
+	T dJ = (T)0.;
+	for (int32 alpha = 0; alpha < 3; alpha++) {
+		for (int32 beta = 0; beta < 3; beta++) {
+			dJ += JFinvT.GetAt(alpha, beta) * dF.GetAt(alpha, beta);
+		}
+	}
+
+	//dR
+	Chaos::PMatrix<T, 3, 3> A = dF * R.GetTransposed(), B((T)0.);
+	TVector<T, 3> a((T)0.), b((T)0.);
+	b[0] = A.GetAt(1, 2) - A.GetAt(2, 1); b[1] = A.GetAt(2, 0) - A.GetAt(0, 2); b[2] = A.GetAt(0, 1) - A.GetAt(1, 0);
+	a = Dinv.GetTransposed() * b;
+	B.SetAt(0, 0, (T)0.); B.SetAt(0, 1, a[2]); B.SetAt(0, 2, -a[1]);
+	B.SetAt(1, 0, -a[2]); B.SetAt(1, 1, (T)0.); B.SetAt(1, 2, a[0]);
+	B.SetAt(2, 0, a[1]); B.SetAt(2, 1, -a[0]); B.SetAt(2, 2, (T)0.);
+	dR = B * R;
+
+	dJFinvT.SetAt(0, 0, F.GetAt(2, 2) * dF.GetAt(1, 1) - F.GetAt(2, 1) * dF.GetAt(1, 2) - F.GetAt(1, 2) * dF.GetAt(2, 1) + F.GetAt(1, 1) * dF.GetAt(2, 2));
+	dJFinvT.SetAt(0, 1, -F.GetAt(2, 2) * dF.GetAt(1, 0) + F.GetAt(2, 0) * dF.GetAt(1, 2) + F.GetAt(1, 2) * dF.GetAt(2, 0) - F.GetAt(1, 0) * dF.GetAt(2, 2));
+	dJFinvT.SetAt(0, 2, F.GetAt(2, 1) * dF.GetAt(1, 0) - F.GetAt(2, 0) * dF.GetAt(1, 1) - F.GetAt(1, 1) * dF.GetAt(2, 0) + F.GetAt(1, 0) * dF.GetAt(2, 1));
+	dJFinvT.SetAt(1, 0, -F.GetAt(2, 2) * dF.GetAt(0, 1) + F.GetAt(2, 1) * dF.GetAt(0, 2) + F.GetAt(0, 2) * dF.GetAt(2, 1) - F.GetAt(0, 1) * dF.GetAt(2, 2));
+	dJFinvT.SetAt(1, 1, F.GetAt(2, 2) * dF.GetAt(0, 0) - F.GetAt(2, 0) * dF.GetAt(0, 2) - F.GetAt(0, 2) * dF.GetAt(2, 0) + F.GetAt(0, 0) * dF.GetAt(2, 2));
+	dJFinvT.SetAt(1, 2, -F.GetAt(2, 1) * dF.GetAt(0, 0) + F.GetAt(2, 0) * dF.GetAt(0, 1) + F.GetAt(0, 1) * dF.GetAt(2, 0) - F.GetAt(0, 0) * dF.GetAt(2, 1));
+	dJFinvT.SetAt(2, 0, F.GetAt(1, 2) * dF.GetAt(0, 1) - F.GetAt(1, 1) * dF.GetAt(0, 2) - F.GetAt(0, 2) * dF.GetAt(1, 1) + F.GetAt(0, 1) * dF.GetAt(1, 2));
+	dJFinvT.SetAt(2, 1, -F.GetAt(1, 2) * dF.GetAt(0, 0) + F.GetAt(1, 0) * dF.GetAt(0, 2) + F.GetAt(0, 2) * dF.GetAt(1, 0) - F.GetAt(0, 0) * dF.GetAt(1, 2));
+	dJFinvT.SetAt(2, 2, F.GetAt(1, 1) * dF.GetAt(0, 0) - F.GetAt(1, 0) * dF.GetAt(0, 1) - F.GetAt(0, 1) * dF.GetAt(1, 0) + F.GetAt(0, 0) * dF.GetAt(1, 1));
+
+	dP = (T)2. * mu * (dF - dR) + lambda * dJ * JFinvT + lambda * (JCache - (T)1.) * dJFinvT;
+
+}
+
+} 
+
+template CHAOS_API void Chaos::Softs::CorotatedCache<Chaos::FRealSingle>::UpdateCache(const Chaos::PMatrix<Chaos::FRealSingle, 3, 3>&);
+template CHAOS_API void Chaos::Softs::CorotatedCache<Chaos::FRealSingle>::P(const Chaos::PMatrix<Chaos::FRealSingle, 3, 3>&, const Chaos::FRealSingle, const Chaos::FRealSingle, Chaos::PMatrix<Chaos::FRealSingle, 3, 3>&);
+template CHAOS_API void Chaos::Softs::CorotatedCache<Chaos::FRealSingle>::deltaP(const Chaos::PMatrix<Chaos::FRealSingle, 3, 3>&, const Chaos::PMatrix<Chaos::FRealSingle, 3, 3>&, const Chaos::FRealSingle, const Chaos::FRealSingle, Chaos::PMatrix<Chaos::FRealSingle, 3, 3>&);
+template CHAOS_API void Chaos::Softs::CorotatedCache<Chaos::FReal>::UpdateCache(const Chaos::PMatrix<Chaos::FReal, 3, 3>&);
+template CHAOS_API void Chaos::Softs::CorotatedCache<Chaos::FReal>::P(const Chaos::PMatrix<Chaos::FReal, 3, 3>&, const Chaos::FReal, const Chaos::FReal, Chaos::PMatrix<Chaos::FReal, 3, 3>&);
+template CHAOS_API void Chaos::Softs::CorotatedCache<Chaos::FReal>::deltaP(const Chaos::PMatrix<Chaos::FReal, 3, 3>&, const Chaos::PMatrix<Chaos::FReal, 3, 3>&, const Chaos::FReal, const Chaos::FReal, Chaos::PMatrix<Chaos::FReal, 3, 3>&);
diff --git a/Engine/Source/Runtime/Experimental/Chaos/Private/Chaos/NewtonEvolution.cpp b/Engine/Source/Runtime/Experimental/Chaos/Private/Chaos/NewtonEvolution.cpp
new file mode 100644
--- /dev/null
+++ b/Engine/Source/Runtime/Experimental/Chaos/Private/Chaos/NewtonEvolution.cpp
@@ -0,0 +1,850 @@
+// Copyright Epic Games, Inc. All Rights Reserved.
+#include "Chaos/NewtonEvolution.h"
+
+#include "Chaos/Framework/Parallel.h"
+#include "Chaos/PBDTriangleMeshIntersections.h"
+#include "Chaos/PerParticleDampVelocity.h"
+#include "Chaos/PerParticleGravity.h"
+#include "Chaos/PerParticlePBDCCDCollisionConstraint.h"
+#include "Chaos/PerParticlePBDCollisionConstraint.h"
+#include "Chaos/VelocityField.h"
+#include "ChaosStats.h"
+#include "HAL/IConsoleManager.h"
+#include "Chaos/NewtonCorotatedCache.h"
+#include "Misc/FileHelper.h"
+#include "Misc/Paths.h"
+#include "HAL/PlatformFileManager.h"
+#include "Chaos/Math/Krylov.h"
+
+//DECLARE_CYCLE_STAT(TEXT("Chaos Newton Advance Time"), STAT_ChaosNewtonVAdvanceTime, STATGROUP_Chaos);
+////DECLARE_CYCLE_STAT(TEXT("Chaos Newton Velocity Damping State Update"), STAT_ChaosNewtonVelocityDampUpdateState, STATGROUP_Chaos);
+//DECLARE_CYCLE_STAT(TEXT("Chaos Newton Velocity Field Update Forces"), STAT_ChaosNewtonVelocityFieldUpdateForces, STATGROUP_Chaos);
+//DECLARE_CYCLE_STAT(TEXT("Chaos Newton Velocity Damping"), STAT_ChaosNewtonVelocityDampUpdate, STATGROUP_Chaos);
+//DECLARE_CYCLE_STAT(TEXT("Chaos Newton Pre Iteration Updates"), STAT_ChaosNewtonPreIterationUpdates, STATGROUP_Chaos);
+//DECLARE_CYCLE_STAT(TEXT("Chaos Newton Integrate"), STAT_ChaosClothSolverIntegrate, STATGROUP_Chaos);
+//DECLARE_CYCLE_STAT(TEXT("Chaos Newton Iteration Loop"), STAT_ChaosNewtonIterationLoop, STATGROUP_Chaos);
+//DECLARE_CYCLE_STAT(TEXT("Chaos Newton Post Iteration Updates"), STAT_ChaosNewtonPostIterationUpdates, STATGROUP_Chaos);
+//DECLARE_CYCLE_STAT(TEXT("Chaos Newton Constraint Rules"), STAT_ChaosNewtonConstraintRule, STATGROUP_Chaos);
+//DECLARE_CYCLE_STAT(TEXT("Chaos Newton Post Collision Constraint Rules"), STAT_ChaosNewtonPostCollisionConstraintRule, STATGROUP_Chaos);
+//DECLARE_CYCLE_STAT(TEXT("Chaos Newton Self Collision"), STAT_ChaosNewtonSelfCollisionRule, STATGROUP_Chaos);
+//DECLARE_CYCLE_STAT(TEXT("Chaos Newton Collision Rule"), STAT_ChaosNewtonCollisionRule, STATGROUP_Chaos);
+//DECLARE_CYCLE_STAT(TEXT("Chaos Newton Collider Friction"), STAT_ChaosNewtonCollisionRuleFriction, STATGROUP_Chaos);
+//DECLARE_CYCLE_STAT(TEXT("Chaos Newton Collider Kinematic Update"), STAT_ChaosNewtonCollisionKinematicUpdate, STATGROUP_Chaos);
+//DECLARE_CYCLE_STAT(TEXT("Chaos Newton Clear Collided Array"), STAT_ChaosNewtonClearCollidedArray, STATGROUP_Chaos);
+//DECLARE_CYCLE_STAT(TEXT("Chaos Newton Constraints Init"), STAT_ChaosXNewtonConstraintsInit, STATGROUP_Chaos);
+
+TAutoConsoleVariable<bool> CVarChaosNewtonEvolutionUseNestedParallelFor(TEXT("p.Chaos.NewtonEvolution.UseNestedParallelFor"), true, TEXT(""), ECVF_Cheat);
+TAutoConsoleVariable<bool> CVarChaosNewtonEvolutionFastPositionBasedFriction(TEXT("p.Chaos.NewtonEvolution.FastPositionBasedFriction"), true, TEXT(""), ECVF_Cheat);
+TAutoConsoleVariable<bool> CVarChaosNewtonEvolutionUseSmoothTimeStep(TEXT("p.Chaos.NewtonEvolution.UseSmoothTimeStep"), true, TEXT(""), ECVF_Cheat);
+TAutoConsoleVariable<int32> CVarChaosNewtonEvolutionMinParallelBatchSize(TEXT("p.Chaos.NewtonEvolution.MinParallelBatchSize"), 300, TEXT(""), ECVF_Cheat);
+TAutoConsoleVariable<bool> CVarChaosNewtonEvolutionWriteCCDContacts(TEXT("p.Chaos.NewtonEvolution.WriteCCDContacts"), false, TEXT("Write CCD collision contacts and normals potentially causing the CCD collision threads to lock, allowing for debugging of these contacts."), ECVF_Cheat);
+TAutoConsoleVariable<bool> CVarChaosNewtonEvolutionParallelIntegrate(TEXT("p.Chaos.NewtonEvolution.ParalleIntegrate"), false, TEXT("Run the integration step in parallel for."), ECVF_Cheat);
+
+//#if INTEL_ISPC && !UE_BUILD_SHIPPING
+//bool bChaos_PostIterationUpdates_ISPC_Enabled = true;
+//FAutoConsoleVariableRef CVarChaosPostIterationUpdatesISPCEnabled(TEXT("p.Chaos.PostIterationUpdates.ISPC"), bChaos_PostIterationUpdates_ISPC_Enabled, TEXT("Whether to use ISPC optimizations in Newton Post iteration updates"));
+//
+//static_assert(sizeof(ispc::FVector3f) == sizeof(Chaos::Softs::FSolverVec3), "sizeof(ispc::FVector3f) != sizeof(Chaos::Softs::FSolverVec3");
+//#endif
+
+namespace Chaos::Softs {
+
+void FNewtonEvolution::AddGroups(int32 NumGroups)
+{
+	// Add elements
+	const uint32 Offset = TArrayCollection::Size();
+	TArrayCollection::AddElementsHelper(NumGroups);
+
+	// Set defaults
+	for (uint32 GroupId = Offset; GroupId < TArrayCollection::Size(); ++GroupId)
+	{
+		MGroupGravityAccelerations[GroupId] = MGravity;
+		MGroupCollisionThicknesses[GroupId] = MCollisionThickness;
+		MGroupCoefficientOfFrictions[GroupId] = MCoefficientOfFriction;
+		MGroupDampings[GroupId] = MDamping;
+		MGroupLocalDampings[GroupId] = MLocalDamping;
+		MGroupUseCCDs[GroupId] = false;
+	}
+}
+
+void FNewtonEvolution::ResetGroups()
+{
+	TArrayCollection::ResizeHelper(0);
+	AddGroups(1);  // Add default group
+}
+
+FNewtonEvolution::FNewtonEvolution(
+	FSolverParticles&& InParticles,
+	FSolverRigidParticles&& InGeometryParticles,
+	TArray<TVec3<int32>>&& CollisionTriangles,
+	const TArray<TVector<int32, 4>>& InMesh,
+	TArray<TArray<TVector<int32, 2>>>&& InIncidentElements,
+	int32 NumNewtonIterations,
+	int32 NumCGIterations,
+	const TArray<int32>& ConstrainedVertices,
+	const TArray<FSolverVec3>& BCPositions,
+	FSolverReal CollisionThickness,
+	FSolverReal SelfCollisionsThickness,
+	FSolverReal CoefficientOfFriction,
+	FSolverReal Damping,
+	FSolverReal LocalDamping,
+	FSolverReal EMesh,
+	FSolverReal NuMesh,
+	FSolverReal NewtonTol,
+	FSolverReal CGTol,
+	bool bWriteDebugInfoIn)
+	: MParticles(MoveTemp(InParticles))
+	, MParticlesActiveView(MParticles)
+	, MCollisionParticles(MoveTemp(InGeometryParticles))
+	, MCollisionParticlesActiveView(MCollisionParticles)
+	, MConstraintInitsActiveView(MConstraintInits)
+	, MConstraintRulesActiveView(MConstraintRules)
+	, MPostCollisionConstraintRulesActiveView(MPostCollisionConstraintRules)
+	, MNumNewtonIterations(NumNewtonIterations)
+	, MNumCGIterations(NumCGIterations)
+	, MMesh(InMesh)
+	, MIncidentElements(MoveTemp(InIncidentElements))
+	, MGravity(FSolverVec3((FSolverReal)0., (FSolverReal)0., (FSolverReal)980.665))
+	, MCollisionThickness(CollisionThickness)
+	, MCoefficientOfFriction(CoefficientOfFriction)
+	, MDamping(Damping)
+	, MLocalDamping(LocalDamping)
+	, MTime(0)
+	, MSmoothDt(1.f / 30.f)  // Initialize filtered timestep at 30fps
+	, MNewtonTol(NewtonTol)
+	, MCGTol(CGTol)
+	, MConstrainedVertices(ConstrainedVertices)
+	, MBCPositions(BCPositions)
+	, bWriteDebugInfo(bWriteDebugInfoIn)
+{
+	// Add group arrays
+	TArrayCollection::AddArray(&MGroupGravityAccelerations);
+	TArrayCollection::AddArray(&MGroupVelocityAndPressureFields);
+	TArrayCollection::AddArray(&MGroupForceRules);
+	TArrayCollection::AddArray(&MGroupCollisionThicknesses);
+	TArrayCollection::AddArray(&MGroupCoefficientOfFrictions);
+	TArrayCollection::AddArray(&MGroupDampings);
+	TArrayCollection::AddArray(&MGroupLocalDampings);
+	TArrayCollection::AddArray(&MGroupUseCCDs);
+	AddGroups(1);  // Add default group
+
+	// Add particle arrays
+	MParticles.AddArray(&MParticleGroupIds);
+	MCollisionParticles.AddArray(&MCollisionTransforms);
+	MCollisionParticles.AddArray(&MCollided);
+	MCollisionParticles.AddArray(&MCollisionParticleGroupIds);
+	ConsCaches.CorotatedCache.Init(CorotatedCache<FSolverReal>(), InMesh.Num());
+
+	Lambda = EMesh * NuMesh / (((FSolverReal)1. + NuMesh) * ((FSolverReal)1. - (FSolverReal)2. * NuMesh));
+	Mu = EMesh / ((FSolverReal)2. * ((FSolverReal)1. + NuMesh));
+	
+	MVn.Init(FSolverVec3(0.f), MParticles.Size());
+
+	UpdatePositionBasedState = [this](const bool update_damping = false) {
+		PhysicsParallelFor(MMesh.Num(), [&](const int32 e) {
+			Chaos::PMatrix<FSolverReal, 3, 3> Fe = MEFem->F(e, MParticles);
+			ConsCaches.CorotatedCache[e].UpdateCache(Fe);
+		});
+	};
+
+	PStress = [this](const PMatrix<FSolverReal, 3, 3>& F, PMatrix<FSolverReal, 3, 3>& P, const int32 e) {
+		ConsCaches.CorotatedCache[e].P(F, Mu, Lambda, P);
+	};
+
+	DeltaPStress = [this](const PMatrix<FSolverReal, 3, 3>& F, const PMatrix<FSolverReal, 3, 3>& deltaF, PMatrix<FSolverReal, 3, 3>& deltaP, const int32 e) {
+		ConsCaches.CorotatedCache[e].deltaP(F, deltaF, Mu, Lambda, deltaP);
+	};
+	
+
+	//set up lambdas from constraint info
+	ProjectBCs = [this](TArray<FSolverVec3>& y) {
+		for (int32 j = 0; j < this->MConstrainedVertices.Num(); j++) {
+			y[this->MConstrainedVertices[j]] = FSolverVec3(0.f);
+		}
+	};
+
+	SetBCs = [this](const FSolverReal time, TArray<FSolverVec3>& y) {
+		for (int32 j = 0; j < this->MConstrainedVertices.Num(); j++) {
+			y[this->MConstrainedVertices[j]] = this->MBCPositions[j];
+		}
+	};
+
+	MAddExternalForce = [this](const FSolverReal time, const FSolverReal dt, const TArray<FSolverReal>& nodal_mass, bool divide_mass, TArray<FSolverVec3>& v) {
+		//TODO: make following parallel
+		for (int32 i = 0; i < v.Num(); ++i) {
+			for (int32 j = 0; j < 3; j++) {
+				if (divide_mass) {
+					v[i][j] -= dt * MGravity[j];
+				}
+				else {
+					v[i][j] -= dt * MGravity[j] * nodal_mass[i];
+				}
+  		}
+		}
+	};
+
+	FString file = FPaths::ProjectDir();
+	file.Append(TEXT("/HoudiniOutput/NewtonLog.txt"));
+	FFileHelper::SaveStringToFile(FString(TEXT("Running Newton's Method\r\n")), *file);
+	FFileHelper::SaveStringToFile(FString(TEXT("Current Configuration:\r\n")), *file, FFileHelper::EEncodingOptions::AutoDetect, &IFileManager::Get(), EFileWrite::FILEWRITE_Append);
+	FFileHelper::SaveStringToFile((FString(TEXT("Newton Tol = ")) + FString::SanitizeFloat(MNewtonTol) + FString(TEXT("\r\n"))), *file, FFileHelper::EEncodingOptions::AutoDetect, &IFileManager::Get(), EFileWrite::FILEWRITE_Append);
+	FFileHelper::SaveStringToFile((FString(TEXT("CG Tol = ")) + FString::SanitizeFloat(MCGTol) + FString(TEXT("\r\n"))), *file, FFileHelper::EEncodingOptions::AutoDetect, &IFileManager::Get(), EFileWrite::FILEWRITE_Append);
+	FFileHelper::SaveStringToFile((FString(TEXT("Max Newton Iters = ")) + FString::FromInt(MNumNewtonIterations) + FString(TEXT("\r\n"))), *file, FFileHelper::EEncodingOptions::AutoDetect, &IFileManager::Get(), EFileWrite::FILEWRITE_Append);
+	FFileHelper::SaveStringToFile((FString(TEXT("Max CG Iters = ")) + FString::FromInt(MNumCGIterations) + FString(TEXT("\r\n"))), *file, FFileHelper::EEncodingOptions::AutoDetect, &IFileManager::Get(), EFileWrite::FILEWRITE_Append);
+	FFileHelper::SaveStringToFile(FString(TEXT("ConsModel: Corotated\r\n")), *file, FFileHelper::EEncodingOptions::AutoDetect, &IFileManager::Get(), EFileWrite::FILEWRITE_Append);
+
+}
+
+void FNewtonEvolution::InitFEM() {
+	MEFem.Reset(new ElasticFEM<FSolverReal, FSolverParticles>(MParticles, MMesh, MIncidentElements));
+	Measure = MEFem->Measure;
+}
+
+
+void FNewtonEvolution::WriteOutputLog(const int32 Frame) 
+{
+	FString file = FPaths::ProjectDir();
+	file.Append(TEXT("/HoudiniOutput/NewtonLog.txt"));
+	FFileHelper::SaveStringToFile(FString(TEXT("---------------------------------------------------\r\n")), *file, FFileHelper::EEncodingOptions::AutoDetect, &IFileManager::Get(), EFileWrite::FILEWRITE_Append);
+	FFileHelper::SaveStringToFile((FString(TEXT("Frame = ")) + FString::FromInt(Frame) + FString(TEXT("\r\n"))), *file, FFileHelper::EEncodingOptions::AutoDetect, &IFileManager::Get(), EFileWrite::FILEWRITE_Append);
+	FFileHelper::SaveStringToFile((FString(TEXT("Time = ")) + FString::SanitizeFloat(MTime) + FString(TEXT("\r\n"))), *file, FFileHelper::EEncodingOptions::AutoDetect, &IFileManager::Get(), EFileWrite::FILEWRITE_Append);
+	FFileHelper::SaveStringToFile(FString(TEXT("---------------------------------------------------\r\n")), *file, FFileHelper::EEncodingOptions::AutoDetect, &IFileManager::Get(), EFileWrite::FILEWRITE_Append);
+}
+
+	void FNewtonEvolution::ResetParticles()
+	{
+		// Reset particles
+		MParticles.Resize(0);
+		MParticlesActiveView.Reset();
+
+		// Reset particle groups
+		ResetGroups();
+	}
+
+	int32 FNewtonEvolution::AddParticleRange(int32 NumParticles, uint32 GroupId, bool bActivate)
+	{
+		if (NumParticles)
+		{
+			const int32 Offset = (int32)MParticles.Size();
+
+			MParticles.AddParticles(NumParticles);
+
+			MVn.Init(FSolverVec3(0.f), MParticles.Size());
+
+			// Initialize the new particles' group ids
+			for (int32 i = Offset; i < (int32)MParticles.Size(); ++i)
+			{
+				MParticleGroupIds[i] = GroupId;
+			}
+
+			// Resize the group parameter arrays
+			const uint32 GroupSize = TArrayCollection::Size();
+			if (GroupId >= GroupSize)
+			{
+				AddGroups(GroupId + 1 - GroupSize);
+			}
+
+			// Add range
+			MParticlesActiveView.AddRange(NumParticles, bActivate);
+
+			return Offset;
+		}
+		return INDEX_NONE;
+	}
+
+	void FNewtonEvolution::ResetCollisionParticles(int32 NumParticles)
+	{
+		MCollisionParticles.Resize(NumParticles);
+		MCollisionParticlesActiveView.Reset(NumParticles);
+	}
+
+	int32 FNewtonEvolution::AddCollisionParticleRange(int32 NumParticles, uint32 GroupId, bool bActivate)
+	{
+		if (NumParticles)
+		{
+			const int32 RangeOffset = (int32)MCollisionParticles.Size();
+
+			MCollisionParticles.AddParticles(NumParticles);
+
+			// Initialize the new particles' group ids
+			for (int32 i = RangeOffset; i < (int32)MCollisionParticles.Size(); ++i)
+			{
+				MCollisionParticleGroupIds[i] = GroupId;
+			}
+
+			// Add range
+			MCollisionParticlesActiveView.AddRange(NumParticles, bActivate);
+
+			return RangeOffset;
+		}
+		return INDEX_NONE;
+	}
+
+	int32 FNewtonEvolution::AddConstraintInitRange(int32 NumConstraints, bool bActivate)
+	{
+		// Add new constraint init functions
+		MConstraintInits.AddDefaulted(NumConstraints);
+
+		// Add range
+		return MConstraintInitsActiveView.AddRange(NumConstraints, bActivate);
+	}
+
+	int32 FNewtonEvolution::AddConstraintRuleRange(int32 NumConstraints, bool bActivate)
+	{
+		// Add new constraint rule functions
+		MConstraintRules.AddDefaulted(NumConstraints);
+
+		// Add range
+		return MConstraintRulesActiveView.AddRange(NumConstraints, bActivate);
+	}
+
+	int32 FNewtonEvolution::AddPostCollisionConstraintRuleRange(int32 NumConstraints, bool bActivate)
+	{
+		// Add new constraint rule functions
+		MPostCollisionConstraintRules.AddDefaulted(NumConstraints);
+
+		// Add range
+		return MPostCollisionConstraintRulesActiveView.AddRange(NumConstraints, bActivate);
+	}
+
+	template <typename Func1, typename Func2>
+	void FNewtonEvolution::ComputeNegativeBackwardEulerResidual(const FSolverParticles& InParticles, const TArray<TArray<TVector<int32, 2>>>& IncidentElements, const TArray<FSolverReal>& NodalMass, const TArray<FSolverVec3>& Vn, const FSolverParticles& Xn, Func1 P, Func2 AddExternalForce, const FSolverReal Time, const FSolverReal Dt, TArray<FSolverVec3>& Residual)
+	{
+		//std::fill(std::execution::par_unseq, residual.begin(), residual.end(), Vector());
+		MEFem->AddInternalElasticForce(InParticles, P, Residual, &IncidentElements);
+		//wc.AddInternalForce(x, residual);
+		AddExternalForce(Time + Dt, (FSolverReal)1., NodalMass, false, Residual);
+
+		const FSolverReal DtSquared = Dt * Dt;
+//#pragma omp parallel for
+		for (int32 i = 0; i < int32(InParticles.Size()); ++i) {
+			for (int32 alpha = 0; alpha < 3; alpha++) {
+				Residual[i][alpha] *= DtSquared;
+				Residual[i][alpha] -= NodalMass[i] * (InParticles.P(i)[alpha] - InParticles.X(i)[alpha] - Dt * Vn[i][alpha]);
+			}
+		}
+	}
+
+	FSolverReal ComputeVectorL2Norm(const TArray<FSolverVec3>& InVector)
+	{
+		FSolverReal Sum = (FSolverReal)0.;
+		for (FSolverVec3 Vector : InVector)
+		{
+			for (int32 i = 0; i < 3; i++)
+				Sum += Vector[i] * Vector[i];
+		}
+		return FMath::Sqrt(Sum);
+	}
+
+
+
+	template<typename Func1, typename Func2, typename Func3, typename Func4, typename Func5, typename Func6>
+	void FNewtonEvolution::DoNewtonStep(const int32 max_it_newton, const FSolverReal newton_tol, const int32 max_it_cg, const FSolverReal cg_tol, Func1 P, Func2 dP, const FSolverReal Time, const FSolverReal Dt, const TArray<TArray<TVector<int32, 2>>>& IncidentElements, const TArray<FSolverReal>& nodal_mass, Func3 set_bcs, Func4 project_bcs, Func5 add_external_force, Func6 update_position_based_state, FSolverParticles& InParticles, TArray<FSolverReal>& ResidualNorm, bool use_cg, FSolverReal cg_prctg_reduce, bool no_verbose) {
+		FString file = FPaths::ProjectDir();
+		file.Append(TEXT("/HoudiniOutput/NewtonLog.txt"));
+		
+		for (int32 it = 0; it < MNumNewtonIterations; it++)
+		{
+			/////////////////////////
+			//compute Newton Residual
+			/////////////////////////
+			TArray<FSolverVec3> Residual;
+			Residual.Init(FSolverVec3(0.f), InParticles.Size());
+			update_position_based_state(false);
+
+			ComputeNegativeBackwardEulerResidual(InParticles, IncidentElements, nodal_mass,  MVn, InParticles, P, add_external_force, Time, Dt, Residual);
+			project_bcs(Residual);
+
+			//when possible, discard unused nodes in loops
+			FSolverReal Norm = ComputeVectorL2Norm(Residual);
+
+			FFileHelper::SaveStringToFile((FString(TEXT("Newton Residual two Norm = ")) + FString::SanitizeFloat(Norm) + FString(TEXT("\r\n"))), *file, FFileHelper::EEncodingOptions::AutoDetect, &IFileManager::Get(), EFileWrite::FILEWRITE_Append);
+			if (Norm < newton_tol) {
+				FFileHelper::SaveStringToFile((FString(TEXT("Newton converged to tolerance ")) + FString::SanitizeFloat(newton_tol) + FString(TEXT(" in = ")) + FString::FromInt(it) + FString(TEXT("\r\n"))), *file, FFileHelper::EEncodingOptions::AutoDetect, &IFileManager::Get(), EFileWrite::FILEWRITE_Append);
+				//std::cout << "Newton converged to tolerance " << newton_tol << " in = " << it << "iterations" << std::endl;
+				return;
+			}
+
+
+			auto multiply = [this, &project_bcs, &dP, &nodal_mass, &Dt](TArray<FSolverVec3>& y, const TArray<FSolverVec3>& _x) {
+				//TODO: this is not strictly necessary since P*_x should be _x in the context of CG
+				TArray<FSolverVec3> x_proj = _x;
+				project_bcs(x_proj);
+
+				y.Init(FSolverVec3((FSolverReal)0.), x_proj.Num());
+
+				//defining the damping delta stress lambda to be aware of Dt:
+				this->MEFem->AddNegativeInternalElasticForceDifferential(this->MParticles, dP, x_proj, y, &(this->MIncidentElements));
+
+				//wc.AddNegativeInternalForceDifferential(x_proj, y);
+
+
+				//add effect of mass matrix and scale by Dt^2
+				if (!use_list) {
+					PhysicsParallelFor(y.Num(), [&](const int32 i) {
+						for (int32 c = 0; c < 3; c++) {
+							y[i][c] = nodal_mass[i] * x_proj[i][c] + Dt * Dt * y[i][c];
+						}
+						});
+				}
+				else {
+					//use_list will specify a sparser set of nodes for the loop
+					PhysicsParallelFor(use_list->Num(), [&](const int32 i) {
+						int32 j = (*use_list)[i];
+						for (int32 c = 0; c < 3; c++) {
+							y[j][c] = nodal_mass[j] * x_proj[j][c] + Dt * Dt * y[j][c];
+						}
+						});
+				}
+
+				project_bcs(y);
+			};
+
+			/////////////////////
+			//define PCG lambdas
+			/////////////////////
+			auto multiply_F = [this, &nodal_mass, &Dt](TArray<FSolverVec3>& Y, const TArray<FSolverVec3>& X) {
+				// TGSLAssert(Y.size() == X.size() && X.size() == nodal_mass.size(), "BackwardEuler: attempting multiply precoditioners between vectors of differing size\n");
+
+				PhysicsParallelFor(X.Num(), [&](const int32 i) {
+					if (nodal_mass[i] > (FSolverReal)0.) {
+						for (int32 j = 0; j < 3; j++) {
+							Y[i][j] = Dt * X[i][j] / (FMath::Sqrt(nodal_mass[i]));
+						}
+					}
+					else {
+						for (int32 j = 0; j < 3; j++) {
+							Y[i][j] = Dt * X[i][j];
+						}
+					}
+					});
+			};
+
+			auto AddScaled = [&](TArray<FSolverVec3>& a, const TArray<FSolverVec3>& b, const FSolverReal s) {
+				//TGSLAssert(a.size() == b.size(), "attempting addScaled with vectors of differing size\n");
+				if (!use_list) {
+					PhysicsParallelFor(a.Num(), [&](const int32 i) {
+						for (int32 j = 0; j < 3; j++)
+							a[i][j] += b[i][j] * s;
+						});
+				}
+				else {
+
+					PhysicsParallelFor(use_list->Num(), [&](const int32 i) {
+						int32 j = (*use_list)[i];
+						for (size_t alpha = 0; alpha < 3; alpha++)
+							a[j][alpha] += b[j][alpha] * s;
+						});
+				}
+			};
+
+			auto DotProduct = [&](const TArray<FSolverVec3>& X, const TArray<FSolverVec3>& Y) {
+				//TGSLAssert(Y.size() == X.size(), "attempting dot product with vectors of differing size\n");
+				FSolverReal sum = 0.f;
+				//TV sum_omp(omp_get_max_threads(), 0);
+				//TODO(Yizhou C): Make following parallel if possible
+				if (!use_list) {
+					for (int32 i = 0; i < X.Num(); i++) {
+						for (int32 j = 0; j < 3; j++) {
+							sum += X[i][j] * Y[i][j];
+						}
+					}
+				}
+				else
+				{
+					for (int32 i = 0; i < use_list->Num(); i++) {
+						int32 j = (*use_list)[i];
+						for (int32 alpha = 0; alpha < 3; alpha++) {
+							sum += X[j][alpha] * Y[j][alpha];
+						}
+					}
+				}
+
+				return sum;
+			};
+
+			auto ComputeLinearResidual = [&](FSolverReal& r, const TArray<FSolverVec3>& x, const TArray<FSolverVec3>& b) {
+				TArray<FSolverVec3> res, Fres;
+				res.Init(FSolverVec3((FSolverReal)0.), x.Num());
+				Fres.Init(FSolverVec3((FSolverReal)0.), x.Num());
+				multiply(res, x);
+				AddScaled(res, b, FSolverReal(-1));
+				multiply_F(Fres, res);
+				r = FSolverReal(0.); r = FMath::Sqrt(DotProduct(Fres, Fres));
+			};
+
+			auto Set = [&](TArray<FSolverVec3>& X, const TArray<FSolverVec3>& Y) {
+				//TGSLAssert(Y.size() == X.size(), "attempting set with vectors of differing size\n");
+				if (!use_list) {
+					PhysicsParallelFor(X.Num(), [&](const int32 i) {
+						for (int32 j = 0; j < 3; j++)
+							X[i][j] = Y[i][j];
+						});
+				}
+				else {
+					PhysicsParallelFor(use_list->Num(), [&](const int32 i) {
+						int32 j = (*use_list)[i];
+						for (int32 alpha = 0; alpha < 3; alpha++)
+							X[j][alpha] = Y[j][alpha];
+						});
+				}
+			};
+
+			auto SetScaled = [&](TArray<FSolverVec3>& a, const TArray<FSolverVec3>& b, FSolverReal s) {
+				//TGSLAssert(a.size() == b.size(), "attempting setScaled with vectors of differing size\n");
+				if (!use_list) {
+					PhysicsParallelFor(a.Num(), [&](const int32 i) {
+						for (int32 j = 0; j < 3; j++)
+							a[i][j] = b[i][j] * s;
+						});
+				}
+				else {
+
+					PhysicsParallelFor(use_list->Num(), [&](const int32 i) {
+						int32 j = (*use_list)[i];
+						for (int32 alpha = 0; alpha < 3; alpha++)
+							a[j][alpha] = b[j][alpha] * s;
+						});
+				}
+			};
+
+
+			/////////////////////////
+			//solve linearized system
+			/////////////////////////
+			TArray<FSolverVec3> delta_x;
+			delta_x.Init(FSolverVec3(0.f), InParticles.Size());
+
+			//Chaos::LanczosCG<FSolverReal>(multiply, delta_x, Residual, max_it_cg, cg_tol, use_list); 
+			Chaos::LanczosCG(multiply, delta_x, Residual, max_it_cg, cg_tol, nullptr);
+			////////////////////////
+			//add Newton increment
+			////////////////////////
+			if (!use_list) {
+				PhysicsParallelFor(InParticles.Size(), [&](const int32 i) {
+					for (size_t alpha = 0; alpha < 3; alpha++) {
+						InParticles.P(i)[alpha] += delta_x[i][alpha];
+					}
+					});
+			}
+			else {
+				PhysicsParallelFor(use_list->Num(), [&](const int32 i) {
+					int32 j = (*use_list)[i];
+					for (size_t alpha = 0; alpha < 3; alpha++) {
+						InParticles.P(j)[alpha] += delta_x[j][alpha];
+					}
+					});
+			}
+
+			////////////////////////////////
+			//check linear Residual after cg
+			////////////////////////////////
+			TArray<FSolverVec3> LinearResidual;
+			LinearResidual.Init(FSolverVec3((FSolverReal)0.), InParticles.Size());
+			//LinearResidual.Init(InParticles.Size(), FSolverVec3(0.f));
+			multiply(LinearResidual, delta_x);
+
+			if (!use_list) {
+				PhysicsParallelFor(InParticles.Size(), [&](const int32 i) {
+					for (int32 alpha = 0; alpha < 3; alpha++) {
+						LinearResidual[i][alpha] -= Residual[i][alpha];
+					}
+					});
+			}
+			else {
+				PhysicsParallelFor(use_list->Num(), [&](const int32 i) {
+					int32 j = (*use_list)[i];
+					for (int32 alpha = 0; alpha < 3; alpha++) {
+						LinearResidual[j][alpha] -= Residual[j][alpha];
+					}
+					});
+			}
+
+			Norm = ComputeVectorL2Norm(LinearResidual);
+			FFileHelper::SaveStringToFile((FString(TEXT("Linear Residual two Norm = ")) + FString::SanitizeFloat(Norm) + FString(TEXT("\r\n"))), *file, FFileHelper::EEncodingOptions::AutoDetect, &IFileManager::Get(), EFileWrite::FILEWRITE_Append);
+			//std::cout << "Linear Residual two Norm = " << std::scientific << Norm << std::endl;
+
+			//if (!no_verbose)
+			//	ResidualNorm[it] = Norm;
+
+			//if (it == max_it_newton - 1 && hard_newton_tol)
+			//	TGSLAssert(false, "ImplicitElasticFEM::BackwardEuler: Newton did not converge.");
+		}
+
+	}
+
+	template<bool bForceRule, bool bVelocityField, bool bDampVelocityRule>
+	void FNewtonEvolution::PreIterationUpdate(
+		const FSolverReal Dt,
+		const int32 Offset,
+		const int32 Range,
+		const int32 MinParallelBatchSize)
+	{
+		const uint32 ParticleGroupId = MParticleGroupIds[Offset];
+		const TFunction<void(FSolverParticles&, const FSolverReal, const int32)>& ForceRule = MGroupForceRules[ParticleGroupId];
+		const FSolverVec3& Gravity = MGroupGravityAccelerations[ParticleGroupId];
+		FVelocityAndPressureField& VelocityAndPressureField = MGroupVelocityAndPressureFields[ParticleGroupId];
+
+		if (bVelocityField)
+		{
+			//TRACE_CPUPROFILER_EVENT_SCOPE(ChaosNewtonVelocityFieldUpdateForces);
+			//SCOPE_CYCLE_COUNTER(STAT_ChaosNewtonVelocityFieldUpdateForces);
+			VelocityAndPressureField.UpdateForces(MParticles, Dt);  // Update force per surface element
+		}
+
+		FPerParticleDampVelocity DampVelocityRule(MGroupLocalDampings[ParticleGroupId]);
+		if (bDampVelocityRule)
+		{
+			//TRACE_CPUPROFILER_EVENT_SCOPE(ChaosNewtonVelocityDampUpdateState);
+			//SCOPE_CYCLE_COUNTER(STAT_ChaosNewtonVelocityDampUpdateState);
+			DampVelocityRule.UpdatePositionBasedState(MParticles, Offset, Range);
+		}
+
+
+		{
+			//TRACE_CPUPROFILER_EVENT_SCOPE(ChaosClothSolverIntegrate);
+			//SCOPE_CYCLE_COUNTER(STAT_ChaosClothSolverIntegrate);
+
+			constexpr FSolverReal DampingFrequency = (FSolverReal)60.;  // The damping value is the percentage of velocity removed per frame when running at 60Hz
+			const FSolverReal Damping = FMath::Clamp(MGroupDampings[ParticleGroupId], (FSolverReal)0., (FSolverReal)1.);
+			FSolverReal DampingPowDt;
+			FSolverReal DampingIntegrated;
+			if (Damping > (FSolverReal)1. - (FSolverReal)UE_KINDA_SMALL_NUMBER)
+			{
+				DampingIntegrated = DampingPowDt = (FSolverReal)0.;
+			}
+			else if (Damping > (FSolverReal)UE_SMALL_NUMBER)
+			{
+				const FSolverReal LogValueByFrequency = FMath::Loge((FSolverReal)1. - Damping) * DampingFrequency;
+
+				DampingPowDt = FMath::Exp(LogValueByFrequency * Dt);  // DampingPowDt = FMath::Pow(OneMinusDamping, Dt * DampingFrequency);
+				DampingIntegrated = (DampingPowDt - (FSolverReal)1.) / LogValueByFrequency;
+			}
+			else
+			{
+				DampingPowDt = (FSolverReal)1.;
+				DampingIntegrated = Dt;
+			}
+
+			const int32 RangeSize = Range - Offset;
+			PhysicsParallelFor(RangeSize,
+				[this, &Offset, &ForceRule, &Gravity, &VelocityAndPressureField, &DampVelocityRule, DampingPowDt, DampingIntegrated, Dt](int32 i)
+				{
+					const int32 Index = Offset + i;
+					if (MParticles.InvM(Index) != (FSolverReal)0.)  // Process dynamic particles
+					{
+						// Init forces with GravityForces
+						MParticles.Acceleration(Index) = Gravity;
+
+						// Force Rule
+						if (bForceRule)
+						{
+							ForceRule(MParticles, Dt, Index); // F += M * A
+						}
+
+						// Velocity Field
+						if (bVelocityField)
+						{
+							VelocityAndPressureField.Apply(MParticles, Dt, Index);
+						}
+
+						// Euler Step Velocity
+						//original code:
+						/*MParticles.V(Index) += MParticles.Acceleration(Index) * MSmoothDt;*/
+						
+						//new code:
+						MParticles.V(Index) += MParticles.Acceleration(Index) * Dt;
+
+						// Damp Velocity Rule
+						if (bDampVelocityRule)
+						{
+							DampVelocityRule.ApplyFast(MParticles, Dt, Index);
+						}
+						
+						//original code:
+						// Euler Step with point damping integration
+						//MParticles.P(Index) = MParticles.X(Index) + MParticles.V(Index) * DampingIntegrated;
+
+						//MParticles.V(Index) *= DampingPowDt;
+
+						//new code:
+						MParticles.P(Index) = MParticles.X(Index) + MParticles.V(Index) * Dt;
+
+						MParticles.V(Index) *= 1.f;
+
+					}
+					else  // Process kinematic particles
+					{
+						MKinematicUpdate(MParticles, Dt, MTime, Index);
+					}
+				}, RangeSize < MinParallelBatchSize);
+		}
+	}
+
+	void FNewtonEvolution::AdvanceOneTimeStep(const FSolverReal Dt, const bool bSmoothDt)
+	{
+		//TRACE_CPUPROFILER_EVENT_SCOPE(FNewtonEvolution_AdvanceOneTimeStep);
+		//SCOPE_CYCLE_COUNTER(STAT_ChaosNewtonVAdvanceTime);
+
+		// Advance time
+		MTime += Dt;
+
+		// Filter delta time to smoothen time variations and prevent unwanted vibrations, works best on Forces
+		if (bSmoothDt && CVarChaosNewtonEvolutionUseSmoothTimeStep.GetValueOnAnyThread())
+		{
+			constexpr FSolverReal DeltaTimeDecay = (FSolverReal)0.1;
+			MSmoothDt += (Dt - MSmoothDt) * DeltaTimeDecay;
+		}
+		else
+		{
+			MSmoothDt = Dt;
+		}
+
+		// Don't bother with threaded execution if we don't have enough work to make it worth while.
+		const bool bUseSingleThreadedRange = !CVarChaosNewtonEvolutionUseNestedParallelFor.GetValueOnAnyThread();
+		const int32 MinParallelBatchSize = !CVarChaosNewtonEvolutionParallelIntegrate.GetValueOnAnyThread() ?
+			TNumericLimits<int32>::Max() :  // Disable
+			CVarChaosNewtonEvolutionMinParallelBatchSize.GetValueOnAnyThread(); // TODO: 1000 is a guess, tune this!
+		const bool bWriteCCDContacts = CVarChaosNewtonEvolutionWriteCCDContacts.GetValueOnAnyThread();
+
+		{
+			TRACE_CPUPROFILER_EVENT_SCOPE(ChaosNewtonPreIterationUpdates);
+
+			MParticlesActiveView.RangeFor(
+				[this, Dt, MinParallelBatchSize](FSolverParticles& Particles, int32 Offset, int32 Range)
+				{
+					const uint32 ParticleGroupId = MParticleGroupIds[Offset];
+
+					if (MGroupVelocityAndPressureFields[ParticleGroupId].IsActive())
+					{
+						if (MGroupLocalDampings[ParticleGroupId] > (FSolverReal)0.)
+						{
+							if (MGroupForceRules[ParticleGroupId])  // VeloctiyFields, Damping, Forces  // Damping?????
+							{
+								PreIterationUpdate</*bForceRule =*/ true, /*bVelocityField =*/ true, /*bDampVelocityRule =*/ true>(Dt, Offset, Range, MinParallelBatchSize);
+							}
+							else  // VeloctiyFields, Damping
+							{
+								PreIterationUpdate</*bForceRule =*/ false, /*bVelocityField =*/ true, /*bDampVelocityRule =*/ true>(Dt, Offset, Range, MinParallelBatchSize);
+							}
+						}
+						else  // No Damping
+						{
+							if (MGroupForceRules[ParticleGroupId])  // VeloctiyFields, Forces
+							{
+								PreIterationUpdate</*bForceRule =*/ true, /*bVelocityField =*/ true, /*bDampVelocityRule =*/ false>(Dt, Offset, Range, MinParallelBatchSize);
+							}
+							else  // VeloctiyFields
+							{
+								PreIterationUpdate</*bForceRule =*/ false, /*bVelocityField =*/ true, /*bDampVelocityRule =*/ false>(Dt, Offset, Range, MinParallelBatchSize);
+							}
+						}
+					}
+					else   // No Velocity Fields
+					{
+						if (MGroupLocalDampings[ParticleGroupId] > (FSolverReal)0.)
+						{
+							if (MGroupForceRules[ParticleGroupId])  // VeloctiyFields, Damping, Forces
+							{
+								PreIterationUpdate</*bForceRule =*/ true, /*bVelocityField =*/ false, /*bDampVelocityRule =*/ true>(Dt, Offset, Range, MinParallelBatchSize);
+							}
+							else  // VeloctiyFields, Damping
+							{
+								PreIterationUpdate</*bForceRule =*/ false, /*bVelocityField =*/ false, /*bDampVelocityRule =*/ true>(Dt, Offset, Range, MinParallelBatchSize);
+							}
+						}
+						else  // No Damping
+						{
+							if (MGroupForceRules[ParticleGroupId])  // VeloctiyFields, Forces
+							{
+								PreIterationUpdate</*bForceRule =*/ true, /*bVelocityField =*/ false, /*bDampVelocityRule =*/ false>(Dt, Offset, Range, MinParallelBatchSize);
+							}
+							else  // VeloctiyFields
+							{
+								PreIterationUpdate</*bForceRule =*/ false, /*bVelocityField =*/ false, /*bDampVelocityRule =*/ false>(Dt, Offset, Range, MinParallelBatchSize);
+							}
+						}
+					}
+				}, bUseSingleThreadedRange);
+		}
+
+		// Collision update
+		{
+			if (MCollisionKinematicUpdate)
+			{
+				TRACE_CPUPROFILER_EVENT_SCOPE(ChaosNewtonCollisionKinematicUpdate);
+				//SCOPE_CYCLE_COUNTER(STAT_ChaosNewtonCollisionKinematicUpdate);
+
+				MCollisionParticlesActiveView.SequentialFor(
+					[this, Dt](FSolverRigidParticles& CollisionParticles, int32 Index)
+					{
+						// Store active collision particle frames prior to the kinematic update for CCD collisions
+						MCollisionTransforms[Index] = FSolverRigidTransform3(CollisionParticles.X(Index), CollisionParticles.R(Index));
+
+						// Update collision transform and velocity
+						MCollisionKinematicUpdate(CollisionParticles, Dt, MTime, Index);
+					});
+			}
+
+			{
+				TRACE_CPUPROFILER_EVENT_SCOPE(ChaosNewtonClearCollidedArray);
+				//SCOPE_CYCLE_COUNTER(STAT_ChaosNewtonClearCollidedArray);
+				memset(MCollided.GetData(), 0, MCollided.Num() * sizeof(bool));
+			}
+		}
+
+		// Constraint init (clear XNewton's Lambdas, init self collisions)
+		{
+			TRACE_CPUPROFILER_EVENT_SCOPE(ChaosXNewtonConstraintsInit);
+			//SCOPE_CYCLE_COUNTER(STAT_ChaosXNewtonConstraintsInit);
+			MConstraintInitsActiveView.SequentialFor(
+				[this, Dt](TArray<TFunction<void(FSolverParticles&, const FSolverReal)>>& ConstraintInits, int32 Index)
+				{
+					ConstraintInits[Index](MParticles, Dt);
+				});
+		}
+
+		// Collision rule initializations
+		MCollisionContacts.Reset();
+		MCollisionNormals.Reset();
+
+		// Iteration loop
+		{
+			TRACE_CPUPROFILER_EVENT_SCOPE(ChaosNewtonIterationLoop);
+			//SCOPE_CYCLE_COUNTER(STAT_ChaosNewtonIterationLoop);
+
+
+			DoNewtonStep(MNumNewtonIterations, MNewtonTol, MNumCGIterations, MCGTol, PStress, DeltaPStress, MTime, Dt, MIncidentElements, MParticles.GetM(), SetBCs, ProjectBCs, MAddExternalForce, UpdatePositionBasedState, MParticles, MResidualNorm, true, NoVerbose);
+
+
+			{
+				//SCOPE_CYCLE_COUNTER(STAT_ChaosNewtonPostIterationUpdates);
+
+				// Particle update, V = (P - X) / Dt; X = P;
+
+				{
+					MParticlesActiveView.ParallelFor(
+						[Dt, this](FSolverParticles& Particles, int32 Index)
+						{
+							Particles.V(Index) = (Particles.P(Index) - Particles.X(Index)) / Dt;
+							Particles.X(Index) = Particles.P(Index);
+							this->MVn[Index] = Particles.V(Index);
+						}, MinParallelBatchSize);
+				}
+			}
+		}
+
+		// The following is not currently been used by the cloth solver implementation at the moment
+		//if (!CVarChaosNewtonEvolutionFastPositionBasedFriction.GetValueOnAnyThread() && MCoefficientOfFriction > 0)
+		//{
+		//	SCOPE_CYCLE_COUNTER(STAT_ChaosNewtonCollisionRuleFriction);
+		//	MParticlesActiveView.ParallelFor(
+		//		[&CollisionRule, Dt](FSolverParticles& Particles, int32 Index)
+		//		{
+		//			TRACE_CPUPROFILER_EVENT_SCOPE(ChaosNewtonCollisionRuleFriction);
+		//			CollisionRule.ApplyFriction(Particles, Dt, Index);
+		//		}, bUseSingleThreadedRange, MinParallelBatchSize);
+		//}
+	}
+
+}  // End namespace Chaos::Softs
diff --git a/Engine/Source/Runtime/Experimental/Chaos/Public/Chaos/Math/Krylov.h b/Engine/Source/Runtime/Experimental/Chaos/Public/Chaos/Math/Krylov.h
--- a/Engine/Source/Runtime/Experimental/Chaos/Public/Chaos/Math/Krylov.h
+++ b/Engine/Source/Runtime/Experimental/Chaos/Public/Chaos/Math/Krylov.h
@@ -22,11 +22,12 @@ void LanczosCG(
 	{
 		T result = T(0);
 		checkfSlow(x.Num() == y.Num(), TEXT("LanczosCG: trying to take the dot product with vectors of different size."));
-		PhysicsParallelFor(x.Num(), [&](const int32 i)
+		//PhysicsParallelFor(x.Num(), [&](const int32 i)
+		for (int32 i = 0; i < x.Num(); i++)
 		{
 			result += x[i] * y[i];
-		},
-		x.Num() < 1000);
+		}/*,
+		x.Num() < 1000);*/
 		return result;
 	};
 
@@ -77,8 +78,8 @@ void LanczosCG(
 		T result = T(0);
 		if (use_list != nullptr) 
 		{
-			PhysicsParallelFor(use_list->Num(), [&](const int32 i)
-			{
+			//PhysicsParallelFor(use_list->Num(), [&](const int32 i)
+			//{
 				for (int32 i = 0; i < use_list->Num(); ++i) 
 				{
 					T dot = T(0);
@@ -89,12 +90,14 @@ void LanczosCG(
 					}
 					result += dot;
 				}
-			},
-			use_list->Num() < 1000);
+			//},
+			//use_list->Num() < 1000);
 		}
 		else 
 		{
-			PhysicsParallelFor(x.Num(), [&](const int32 i)
+			//PhysicsParallelFor(x.Num(), [&](const int32 i)
+			//{
+			for (int32 i = 0; i < x.Num(); i++)
 			{
 				T dot = T(0);
 				for (size_t alpha = 0; alpha < d; alpha++) 
@@ -102,8 +105,8 @@ void LanczosCG(
 					dot += x[i][alpha] * y[i][alpha];
 				}
 				result += dot;
-			},
-			x.Num() < 1000);
+			}/*,
+			x.Num() < 1000);*/
 		}
 		return result;
 	};
diff --git a/Engine/Source/Runtime/Experimental/Chaos/Public/Chaos/NewtonCorotatedCache.h b/Engine/Source/Runtime/Experimental/Chaos/Public/Chaos/NewtonCorotatedCache.h
new file mode 100644
--- /dev/null
+++ b/Engine/Source/Runtime/Experimental/Chaos/Public/Chaos/NewtonCorotatedCache.h
@@ -0,0 +1,30 @@
+// Copyright Epic Games, Inc. All Rights Reserved.
+#pragma once 
+#include "Chaos/Matrix.h"
+
+namespace Chaos::Softs
+{
+
+template <typename T>
+struct CorotatedCache 
+{
+	TArray<T> JFInvTCache, RCache;
+	TArray<T> DInvCache;
+
+	T JCache;
+
+
+	CorotatedCache()
+	{
+		JFInvTCache.Init((T)0., 3*3);
+		RCache.Init((T)0., 3*3);
+		DInvCache.Init((T)0., 3*3);
+	};
+
+	void UpdateCache(const Chaos::PMatrix<T, 3, 3>& F);
+	void P(const Chaos::PMatrix<T, 3, 3>& F, const T mu, const T lambda, Chaos::PMatrix<T, 3, 3>& P);
+	void deltaP(const Chaos::PMatrix<T, 3, 3>& F, const Chaos::PMatrix<T, 3, 3>& dF, const T mu, const T lambda, Chaos::PMatrix<T, 3, 3>& dP);
+
+};
+
+}
diff --git a/Engine/Source/Runtime/Experimental/Chaos/Public/Chaos/NewtonElasticFEM.h b/Engine/Source/Runtime/Experimental/Chaos/Public/Chaos/NewtonElasticFEM.h
new file mode 100644
--- /dev/null
+++ b/Engine/Source/Runtime/Experimental/Chaos/Public/Chaos/NewtonElasticFEM.h
@@ -0,0 +1,274 @@
+// Copyright Epic Games, Inc. All Rights Reserved.
+#pragma once
+
+#include "Chaos/Core.h"
+#include "Chaos/ArrayCollection.h"
+#include "Chaos/PBDActiveView.h"
+#include "Chaos/PBDSoftsEvolutionFwd.h"
+#include "Chaos/PBDSoftsSolverParticles.h"
+#include "Chaos/KinematicGeometryParticles.h"
+#include "Chaos/VelocityField.h"
+
+namespace Chaos::Softs
+{
+
+
+template <typename T, typename ParticleType>
+class ElasticFEM 
+{
+public:
+	const TArray<TVector<int32, 4>>& Mesh;
+	TArray<T> DmInverse;
+	TArray<T> Measure;
+	TArray<TVector<T, 3>> ElementForces;  // extra storage in case of parallel force and force differential computation
+	int32 Np;// the total number of particles in the arrays that mesh is indexing
+	TArray<T> MNodalMass;
+	TArray<TArray<TVector<int32, 2>>> MIncidentElements;
+
+	ElasticFEM(
+		const ParticleType& InParticles,
+		const TArray<TVector<int32, 4>>& InMesh,
+		const TArray<TArray<TVector<int32, 2>>>& IncidentElementsIn
+	)
+		: Mesh(InMesh), MIncidentElements(IncidentElementsIn)
+	{
+		
+		Np = InParticles.Size();
+		DmInverse.Init((T)0., 9 * Mesh.Num());
+		Measure.Init((T)0., Mesh.Num());
+		for (int e = 0; e < Mesh.Num(); e++)
+		{
+			PMatrix<T, 3, 3> Dm = DsInit(e, InParticles);
+			PMatrix<T, 3, 3> DmInv = Dm.Inverse();
+			for (int r = 0; r < 3; r++) {
+				for (int c = 0; c < 3; c++) {
+					DmInverse[(3 * 3) * e + 3 * r + c] = DmInv.GetAt(r, c);
+				}
+			}
+
+			Measure[e] = Dm.Determinant() / (T)6.;
+
+			//part of preprocessing: if inverted element is found, 
+			//invert it so that the measure is positive
+			if (Measure[e] < (T)0.)
+			{
+
+				Measure[e] = -Measure[e];
+			}
+		}
+	}
+
+	void ComputeNodalMass(const T Density)
+	{
+		MNodalMass.Init((T)0., Np);
+		for (int32 e = 0; e < Mesh.Num(); e++) 
+		{
+			T NodeMassContribution = Density * Measure[e] / (T)4.;
+			for (int32 i = 0; i < 4; i++) 
+			{
+				MNodalMass[Mesh[e][i]] += NodeMassContribution;
+			}
+		}
+	}
+
+	PMatrix<T, 3, 3> DsInit(const int e, const ParticleType& InParticles) const 
+	{
+		PMatrix<T, 3, 3> Result((T)0.);
+		for (int i = 0; i < 3; i++) {
+			for (int c = 0; c < 3; c++) {
+				Result.SetAt(c, i, InParticles.X(Mesh[e][i + 1])[c] - InParticles.X(Mesh[e][0])[c]);
+			}
+		}
+		return Result;
+	}
+
+
+	PMatrix<T, 3, 3> Ds(const int e, const ParticleType& InParticles) const 
+	{
+		PMatrix<T, 3, 3> Result((T)0.);
+		for (int i = 0; i < 3; i++) {
+			for (int c = 0; c < 3; c++) {
+				Result.SetAt(c, i, InParticles.P(Mesh[e][i + 1])[c] - InParticles.P(Mesh[e][0])[c]);
+			}
+		}
+		return Result;
+	}
+
+	PMatrix<T, 3, 3> Ds(const int e, const TArray<TVector<T, 3>>& InParticles) const 
+	{
+		PMatrix<T, 3, 3> Result((T)0.);
+		for (int i = 0; i < 3; i++) {
+			for (int c = 0; c < 3; c++) {
+				Result.SetAt(c, i, InParticles[Mesh[e][i + 1]][c] - InParticles[Mesh[e][0]][c]);
+			}
+		}
+		return Result;
+	}
+
+
+	PMatrix<T, 3, 3> F(const int e, const ParticleType& InParticles) const {
+		return ElementDmInv(e) * Ds(e, InParticles);
+	}
+
+	PMatrix<T, 3, 3> F(const int e, const TArray<TVector<T, 3>>& InParticles) const {
+		return ElementDmInv(e) * Ds(e, InParticles);
+	}
+
+	PMatrix<T, 3, 3> ElementDmInv(const int e) const 
+	{
+		PMatrix<T, 3, 3> DmInv((T)0.);
+		for (int r = 0; r < 3; r++) {
+			for (int c = 0; c < 3; c++) {
+				DmInv.SetAt(r, c, DmInverse[(3 * 3) * e + 3 * r + c]);
+			}
+		}
+		return DmInv;
+	}
+
+	template <typename Func>
+	void AddInternalElasticForce(const ParticleType& InParticles, Func P, TArray<TVector<T, 3>>& Force, const TArray<TArray<TVector<int32, 2>>>* IncidentElements = nullptr, T Scale = (T)1., const TArray<T>* NodalMass = 0)
+	{
+		//compute the force per node from the mesh, DmInv and P
+		//incident elements optionally used for parallelism
+		//scale and nodal mass can be used to multiply/divide the force entries so that this can be used with time stepping
+
+		Force.Init(TVector<T, 3>((T)0.), Force.Num());
+
+		if (NodalMass)
+			ensureMsgf(NodalMass->Num() == InParticles.Size(), TEXT("NewtonEvolution::AddInternalElasticForce: mass and position sized inconistently."));
+
+		if (!IncidentElements) {
+			for (int32 e = 0; e < Mesh.Num() ; e++) {
+				PMatrix<T, 3, 3> Fe = F(e, InParticles);
+				PMatrix<T, 3, 3> Pe(0.f);
+				P(Fe, Pe, e);
+				PMatrix<T, 3, 3> g = -Measure[e] * (ElementDmInv(e).GetTransposed()) * Pe;
+				for (int32 ie = 0; ie < 3; ie++) {
+					T InverseMass = 1.f;
+					if (NodalMass)
+						InverseMass = 1.f / NodalMass -> operator[](Mesh[e][ie+1]);
+					for (int32 c = 0; c < 3; c++) {
+						Force[Mesh[e][ie+1]][c] += Scale * InverseMass * g.GetAt(c, ie);
+					}
+				}
+				T InverseMass = 1.f;
+				if (NodalMass)
+					InverseMass = 1.f / NodalMass->operator[](Mesh[e][0]);
+				for (int32 c = 0; c < 3; c++) {
+					for (int32 h = 0; h < 3; h++) {
+						Force[Mesh[e][0]][c] -= Scale * InverseMass * g.GetAt(c, h);
+					}
+				}
+			}
+		}
+		else {
+			ElementForces.Init(TVector<T, 3>(0.f), 4 * Mesh.Num());
+			PhysicsParallelFor(Mesh.Num(), [&](const int32 e)
+				{
+					PMatrix<T, 3, 3> Fe = F(e, InParticles);
+					PMatrix<T, 3, 3> Pe(0.f);
+					P(Fe, Pe, e);
+					PMatrix<T, 3, 3> g = -Measure[e] * (ElementDmInv(e).GetTransposed()) * Pe;
+					for (int32 ie = 0; ie < 3; ie++) {
+						for (int32 c = 0; c < 3; c++) {
+							ElementForces[4 * e + ie + 1][c] += g.GetAt(c, ie);
+						}
+					}
+					for (int32 c = 0; c < 3; c++) {
+						for (int32 h = 0; h < 3; h++) {
+							ElementForces[4 * e][c] -= g.GetAt(c, h);
+						}
+					}
+				});
+
+
+			if (!NodalMass) {
+				//#pragma omp parallel for
+				PhysicsParallelFor(IncidentElements->Num(), [&](const int32 i)
+					{
+						for (int32 e = 0; e < (*IncidentElements)[i].Num(); e++) {
+							for (int32 c = 0; c < 3; c++) {
+								Force[Mesh[(*IncidentElements)[i][e][0]][(*IncidentElements)[i][e][1]]][c] += Scale * ElementForces[4* (*IncidentElements)[i][e][0]+ (*IncidentElements)[i][e][1]][c];
+							}
+						}
+					});
+			}
+			else {
+				//#pragma omp parallel for
+				PhysicsParallelFor(IncidentElements->Num(), [&](const int32 i)
+					{
+						for (int32 e = 0; e < (*IncidentElements)[i].Num(); e++) {
+							for (int32 c = 0; c < 3; c++) {
+								Force[Mesh[(*IncidentElements)[i][e][0]][(*IncidentElements)[i][e][1]]][c] += Scale * ElementForces[4 * (*IncidentElements)[i][e][0] + (*IncidentElements)[i][e][1]][c] / (*NodalMass)[Mesh[(*IncidentElements)[i][e][0]][(*IncidentElements)[i][e][1]]];
+							}
+						}
+					});
+			}
+		}
+	}
+
+	template <typename Func>
+	void AddNegativeInternalElasticForceDifferential(const ParticleType& InParticles, Func dP, const TArray<TVector<T, 3>>& DeltaParticles, TArray<TVector<T, 3>>& ndf, const TArray<TArray<TVector<int32, 2>>>* IncidentElements = nullptr)
+	{
+		//matrix free elastic potential energy Hessian based differentials
+		//compute the force differential per node from the mesh, DmInv and dP
+		//incident elements optionally used for parallelism
+
+		ndf.Init(TVector<T, 3>(0.f), ndf.Num());
+
+		if (!IncidentElements) {
+			for (int32 e = 0; e < Mesh.Num() ; e++) {
+				PMatrix<T, 3, 3> Fe = F(e, InParticles), dFe = F(e, DeltaParticles);
+				PMatrix<T, 3, 3> dPe;
+				dP(Fe, dFe, dPe, e);
+				PMatrix<T, 3, 3> dg = Measure[e] * (ElementDmInv(e).GetTransposed()) * dPe;
+				for (int32 ie = 0; ie < 3; ie++) {
+					for (int32 c = 0; c < 3; c++) {
+						ndf[Mesh[e][ie+1]][c] += dg.GetAt(c, ie);
+					}
+				}
+				for (int32 c = 0; c < 3; c++) {
+					for (int32 h = 0; h < 3; h++) {
+						ndf[Mesh[e][0]][c] -= dg.GetAt(c, h);
+					}
+				}
+			}
+		}
+		else {
+			ElementForces.Init(TVector<T, 3>(0.f), Mesh.Num() * 4);
+			PhysicsParallelFor(Mesh.Num(), [&](const int32 e) {
+				//for (int32 e = 0; e < int32(Mesh.Num() ); ++e) {
+				PMatrix<T, 3, 3> Fe = F(e, InParticles), dFe = F(e, DeltaParticles);
+				PMatrix<T, 3, 3> dPe;
+				dP(Fe, dFe, dPe, e);
+				PMatrix<T, 3, 3> dg = Measure[e] * (ElementDmInv(e).GetTransposed()) * dPe;
+				for (int32 ie = 0; ie < 3; ie++)
+				{
+					for (int32 c = 0; c < 3; c++)
+					{
+						ElementForces[4 * e + ie + 1][c] += dg.GetAt(c, ie);
+					}
+				}
+				for (int32 c = 0; c < 3; c++) {
+					for (int32 h = 0; h < 3; h++) {
+						ElementForces[4 * e][c] -= dg.GetAt(c, h);
+					}
+				}
+				});
+
+			//#pragma omp parallel for
+			for (int32 i = 0; i < int32(IncidentElements->Num()); ++i) {
+				for (int32 e = 0; e < (*IncidentElements)[i].Num(); e++) {
+					for (int32 c = 0; c < 3; c++) {
+						ndf[Mesh[(*IncidentElements)[i][e][0]][(*IncidentElements)[i][e][1]]][c] += ElementForces[4 * (*IncidentElements)[i][e][0] + (*IncidentElements)[i][e][1]][c];
+					}
+				}
+			}
+		}
+	}
+
+
+
+};
+
+}
diff --git a/Engine/Source/Runtime/Experimental/Chaos/Public/Chaos/NewtonEvolution.h b/Engine/Source/Runtime/Experimental/Chaos/Public/Chaos/NewtonEvolution.h
new file mode 100644
--- /dev/null
+++ b/Engine/Source/Runtime/Experimental/Chaos/Public/Chaos/NewtonEvolution.h
@@ -0,0 +1,272 @@
+// Copyright Epic Games, Inc. All Rights Reserved.
+#pragma once
+
+#include "Chaos/Core.h"
+#include "Chaos/ArrayCollection.h"
+#include "Chaos/PBDActiveView.h"
+#include "Chaos/PBDSoftsEvolutionFwd.h"
+#include "Chaos/PBDSoftsSolverParticles.h"
+#include "Chaos/NewtonElasticFEM.h"
+#include "Chaos/NewtonCorotatedCache.h"
+#include "Chaos/KinematicGeometryParticles.h"
+#include "Chaos/VelocityField.h"
+
+namespace Chaos::Softs
+{
+
+struct ConsModelCaches {
+	TArray<Chaos::Softs::CorotatedCache<FSolverReal>> CorotatedCache;
+	//TArray<CorotatedFiberCache> CorotatedFiberCache;
+};
+
+class CHAOS_API FNewtonEvolution : public TArrayCollection
+{
+public:
+	// TODO: Tidy up this constructor (and update Headless Chaos)
+	FNewtonEvolution(
+		FSolverParticles&& InParticles,
+		FSolverRigidParticles&& InGeometryParticles,
+		TArray<TVec3<int32>>&& CollisionTriangles,
+		const TArray<TVector<int32, 4>>& InMesh,
+		TArray<TArray<TVector<int32, 2>>>&& InIncidentElements,
+		int32 NumNewtonIterations = 5,
+		int32 NumCGIterations = 20,
+		const TArray<int32>& ConstrainedVertices = TArray<int32>(),
+		const TArray<FSolverVec3>& BCPositions = TArray<FSolverVec3>(),
+		FSolverReal CollisionThickness = (FSolverReal)0.,
+		FSolverReal SelfCollisionsThickness = (FSolverReal)0.,
+		FSolverReal CoefficientOfFriction = (FSolverReal)0.,
+		FSolverReal Damping = (FSolverReal)0.04,
+		FSolverReal LocalDamping = (FSolverReal)0.,
+		FSolverReal EMesh = (FSolverReal)1000.,
+		FSolverReal NuMesh = (FSolverReal).3,
+		FSolverReal NewtonTol = (FSolverReal) 1e-6,
+		FSolverReal CGTol = (FSolverReal) 1e-8,
+		bool bWriteDebugInfoIn = true);
+	~FNewtonEvolution() {}
+
+	// Advance one time step. Filter the input time step if specified.
+	void AdvanceOneTimeStep(const FSolverReal Dt, const bool bSmoothDt = true);
+
+	// Remove all particles, will also reset all rules
+	void ResetParticles();
+
+	// Add particles and initialize group ids. Return the index of the first added particle.
+	int32 AddParticleRange(int32 NumParticles, uint32 GroupId, bool bActivate);
+
+	// Return the number of particles of the block starting at Offset
+	int32 GetParticleRangeSize(int32 Offset) const { return MParticlesActiveView.GetRangeSize(Offset); }
+
+	// Set a block of particles active or inactive, using the index of the first added particle to identify the block.
+	void ActivateParticleRange(int32 Offset, bool bActivate) { MParticlesActiveView.ActivateRange(Offset, bActivate); }
+
+	// Particles accessors
+	const FSolverParticles& Particles() const { return MParticles; }
+	FSolverParticles& Particles() { return MParticles; }
+	const TPBDActiveView<FSolverParticles>& ParticlesActiveView() { return MParticlesActiveView; }
+
+	TArray<TArray<TVector<int32, 2>>>& IncidentElements() { return MIncidentElements; }
+
+	const TArray<uint32>& ParticleGroupIds() const { return MParticleGroupIds; }
+
+	// Remove all collision particles
+	void ResetCollisionParticles(int32 NumParticles = 0);
+
+	// Add collision particles and initialize group ids. Return the index of the first added particle.
+	// Use INDEX_NONE as GroupId for collision particles that affect all particle groups.
+	int32 AddCollisionParticleRange(int32 NumParticles, uint32 GroupId, bool bActivate);
+
+	// Set a block of collision particles active or inactive, using the index of the first added particle to identify the block.
+	void ActivateCollisionParticleRange(int32 Offset, bool bActivate) { MCollisionParticlesActiveView.ActivateRange(Offset, bActivate); }
+
+	// Return the number of particles of the block starting at Offset
+	int32 GetCollisionParticleRangeSize(int32 Offset) const { return MCollisionParticlesActiveView.GetRangeSize(Offset); }
+
+	//// collision particles accessors
+	const FSolverRigidParticles& CollisionParticles() const { return MCollisionParticles; }
+	FSolverRigidParticles& CollisionParticles() { return MCollisionParticles; }
+	const TArray<uint32>& CollisionParticleGroupIds() const { return MCollisionParticleGroupIds; }
+	const TPBDActiveView<FSolverRigidParticles>& CollisionParticlesActiveView() { return MCollisionParticlesActiveView; }
+
+
+	// Reset all constraint init and rule functions.
+	void ResetConstraintRules()
+	{
+		MConstraintInits.Reset();
+		MConstraintRules.Reset();
+		MPostCollisionConstraintRules.Reset();
+		MConstraintInitsActiveView.Reset();
+		MConstraintRulesActiveView.Reset();
+		MPostCollisionConstraintRulesActiveView.Reset();
+	}
+
+	// Add constraints. Return the index of the first added constraint.
+	int32 AddConstraintInitRange(int32 NumConstraints, bool bActivate);
+	int32 AddConstraintRuleRange(int32 NumConstraints, bool bActivate);
+	int32 AddPostCollisionConstraintRuleRange(int32 NumConstraints, bool bActivate);
+
+	// Return the number of particles of the block starting at Offset
+	int32 GetConstraintInitRangeSize(int32 Offset) const { return MConstraintInitsActiveView.GetRangeSize(Offset); }
+	int32 GetConstraintRuleRangeSize(int32 Offset) const { return MConstraintRulesActiveView.GetRangeSize(Offset); }
+	int32 GetPostCollisionConstraintRuleRangeSize(int32 Offset) const { return MPostCollisionConstraintRulesActiveView.GetRangeSize(Offset); }
+
+	// Set a block of constraints active or inactive, using the index of the first added particle to identify the block.
+	void ActivateConstraintInitRange(int32 Offset, bool bActivate) { MConstraintInitsActiveView.ActivateRange(Offset, bActivate); }
+	void ActivateConstraintRuleRange(int32 Offset, bool bActivate) { MConstraintRulesActiveView.ActivateRange(Offset, bActivate); }
+	void ActivatePostCollisionConstraintRuleRange(int32 Offset, bool bActivate) { MPostCollisionConstraintRulesActiveView.ActivateRange(Offset, bActivate); }
+
+	// Constraint accessors
+	const TArray<TFunction<void(FSolverParticles&, const FSolverReal)>>& ConstraintInits() const { return MConstraintInits; }
+	TArray<TFunction<void(FSolverParticles&, const FSolverReal)>>& ConstraintInits() { return MConstraintInits; }
+	const TArray<TFunction<void(FSolverParticles&, const FSolverReal)>>& ConstraintRules() const { return MConstraintRules; }
+	TArray<TFunction<void(FSolverParticles&, const FSolverReal)>>& ConstraintRules() { return MConstraintRules; }
+	const TArray<TFunction<void(FSolverParticles&, const FSolverReal)>>& PostCollisionConstraintRules() const { return MPostCollisionConstraintRules; }
+	TArray<TFunction<void(FSolverParticles&, const FSolverReal)>>& PostCollisionConstraintRules() { return MPostCollisionConstraintRules; }
+
+	void SetKinematicUpdateFunction(TFunction<void(FSolverParticles&, const FSolverReal, const FSolverReal, const int32)> KinematicUpdate) { MKinematicUpdate = KinematicUpdate; }
+	void SetCollisionKinematicUpdateFunction(TFunction<void(FSolverRigidParticles&, const FSolverReal, const FSolverReal, const int32)> KinematicUpdate) { MCollisionKinematicUpdate = KinematicUpdate; }
+
+	TFunction<void(FSolverParticles&, const FSolverReal, const int32)>& GetForceFunction(const uint32 GroupId = 0) { return MGroupForceRules[GroupId]; }
+	const TFunction<void(FSolverParticles&, const FSolverReal, const int32)>& GetForceFunction(const uint32 GroupId = 0) const { return MGroupForceRules[GroupId]; }
+
+	const FSolverVec3& GetGravity(const uint32 GroupId = 0) const { check(GroupId < TArrayCollection::Size()); return MGroupGravityAccelerations[GroupId]; }
+	void SetGravity(const FSolverVec3& Acceleration, const uint32 GroupId = 0) { check(GroupId < TArrayCollection::Size()); MGroupGravityAccelerations[GroupId] = Acceleration; }
+
+	FVelocityAndPressureField& GetVelocityAndPressureField(const uint32 GroupId = 0) { check(GroupId < TArrayCollection::Size()); return MGroupVelocityAndPressureFields[GroupId]; }
+	const FVelocityAndPressureField& GetVelocityAndPressureField(const uint32 GroupId = 0) const { check(GroupId < TArrayCollection::Size()); return MGroupVelocityAndPressureFields[GroupId]; }
+
+	UE_DEPRECATED(5.1, "Chaos::Softs::FVelocityField has been renamed FVelocityAndPressureField to match its new behavior.")
+		PRAGMA_DISABLE_DEPRECATION_WARNINGS
+		FVelocityField& GetVelocityField(const uint32 GroupId = 0) { return GetVelocityAndPressureField(GroupId); }
+	PRAGMA_ENABLE_DEPRECATION_WARNINGS
+		UE_DEPRECATED(5.1, "Chaos::Softs::FVelocityField has been renamed FVelocityAndPressureField to match its new behavior.")
+		PRAGMA_DISABLE_DEPRECATION_WARNINGS
+		const FVelocityField& GetVelocityField(const uint32 GroupId = 0) const { return GetVelocityAndPressureField(GroupId); }
+	PRAGMA_ENABLE_DEPRECATION_WARNINGS
+
+
+	int32 GetNewtonIterations() const { return MNumNewtonIterations; }
+	void SetNewtonIterations(const int32 Iterations) { MNumNewtonIterations = Iterations; }
+
+	FSolverReal GetCollisionThickness(const uint32 GroupId = 0) const { check(GroupId < TArrayCollection::Size()); return MGroupCollisionThicknesses[GroupId]; }
+	void SetCollisionThickness(const FSolverReal CollisionThickness, const uint32 GroupId = 0) { check(GroupId < TArrayCollection::Size()); MGroupCollisionThicknesses[GroupId] = CollisionThickness; }
+
+	FSolverReal GetCoefficientOfFriction(const uint32 GroupId = 0) const { check(GroupId < TArrayCollection::Size()); return MGroupCoefficientOfFrictions[GroupId]; }
+	void SetCoefficientOfFriction(const FSolverReal CoefficientOfFriction, const uint32 GroupId = 0) { check(GroupId < TArrayCollection::Size()); MGroupCoefficientOfFrictions[GroupId] = CoefficientOfFriction; }
+
+	FSolverReal GetDamping(const uint32 GroupId = 0) const { check(GroupId < TArrayCollection::Size()); return MGroupDampings[GroupId]; }
+	void SetDamping(const FSolverReal Damping, const uint32 GroupId = 0) { check(GroupId < TArrayCollection::Size()); MGroupDampings[GroupId] = Damping; }
+
+	FSolverReal GetLocalDamping(const uint32 GroupId = 0) const { check(GroupId < TArrayCollection::Size()); return MGroupLocalDampings[GroupId]; }
+	void SetLocalDamping(const FSolverReal LocalDamping, const uint32 GroupId = 0) { check(GroupId < TArrayCollection::Size()); MGroupLocalDampings[GroupId] = LocalDamping; }
+
+	bool GetUseCCD(const uint32 GroupId = 0) const { check(GroupId < TArrayCollection::Size()); return MGroupUseCCDs[GroupId]; }
+	void SetUseCCD(const bool bUseCCD, const uint32 GroupId = 0) { check(GroupId < TArrayCollection::Size()); MGroupUseCCDs[GroupId] = bUseCCD; }
+
+	UE_DEPRECATED(4.27, "Use GetCollisionStatus() instead")
+		const bool Collided(int32 index) { return MCollided[index]; }
+
+	const TArray<bool>& GetCollisionStatus() { return MCollided; }
+	const TArray<FSolverVec3>& GetCollisionContacts() const { return MCollisionContacts; }
+	const TArray<FSolverVec3>& GetCollisionNormals() const { return MCollisionNormals; }
+
+	FSolverReal GetTime() const { return MTime; }
+		
+	template <typename Func1, typename Func2>
+	void ComputeNegativeBackwardEulerResidual(const FSolverParticles& InParticles, const TArray<TArray<TVector<int32, 2>>>& IncidentElements, const TArray<FSolverReal>& NodalMass, const TArray<FSolverVec3>& Vn, const FSolverParticles& Xn, Func1 P, Func2 AddExternalForce, const FSolverReal Time, const FSolverReal Dt, TArray<FSolverVec3>& Residual);
+		
+	template<typename Func1, typename Func2, typename Func3, typename Func4, typename Func5, typename Func6>
+	void DoNewtonStep(const int32 max_it_newton, const FSolverReal newton_tol, const int32 max_it_cg, const FSolverReal cg_tol, Func1 P, Func2 dP, const FSolverReal time, const FSolverReal dt, const TArray<TArray<TVector<int32, 2>>>& incident_elements, const TArray<FSolverReal>& nodal_mass, Func3 set_bcs, Func4 project_bcs, Func5 add_external_force, Func6 update_position_based_state, FSolverParticles& InParticles, TArray<FSolverReal>& residual_norm, bool use_cg = false, FSolverReal cg_prctg_reduce = 0, bool no_verbose = false);
+
+	void InitFEM();
+
+	void WriteOutputLog(const int32 Frame);
+
+private:
+	// Add simulation groups and set default values
+	void AddGroups(int32 NumGroups);
+	// Reset simulation groups
+	void ResetGroups();
+	// Selected versions of the pre-iteration updates (euler step, force, velocity field. damping updates)..
+	template<bool bForceRule, bool bVelocityField, bool bDampVelocityRule>
+	void PreIterationUpdate(const FSolverReal Dt, const int32 Offset, const int32 Range, const int32 MinParallelBatchSize);
+	
+
+private:
+	FSolverParticles MParticles;
+	TPBDActiveView<FSolverParticles> MParticlesActiveView;
+	FSolverRigidParticles MCollisionParticles;
+	TPBDActiveView<FSolverRigidParticles> MCollisionParticlesActiveView;
+
+	TArrayCollectionArray<FSolverRigidTransform3> MCollisionTransforms;  // Used for CCD to store the initial state before the kinematic update
+	TArrayCollectionArray<bool> MCollided;
+	TArrayCollectionArray<uint32> MCollisionParticleGroupIds;  // Used for per group parameters for collision particles
+	TArrayCollectionArray<uint32> MParticleGroupIds;  // Used for per group parameters for particles
+	TArray<FSolverVec3> MCollisionContacts;
+	TArray<FSolverVec3> MCollisionNormals;
+
+	TArrayCollectionArray<FSolverVec3> MGroupGravityAccelerations;
+	TArrayCollectionArray<FVelocityAndPressureField> MGroupVelocityAndPressureFields;
+	TArrayCollectionArray<TFunction<void(FSolverParticles&, const FSolverReal, const int32)>> MGroupForceRules;
+	TArrayCollectionArray<FSolverReal> MGroupCollisionThicknesses;
+	TArrayCollectionArray<FSolverReal> MGroupCoefficientOfFrictions;
+	TArrayCollectionArray<FSolverReal> MGroupDampings;
+	TArrayCollectionArray<FSolverReal> MGroupLocalDampings;
+	TArrayCollectionArray<bool> MGroupUseCCDs;
+
+	TArray<TFunction<void(FSolverParticles&, const FSolverReal)>> MConstraintInits;
+	TPBDActiveView<TArray<TFunction<void(FSolverParticles&, const FSolverReal)>>> MConstraintInitsActiveView;
+	TArray<TFunction<void(FSolverParticles&, const FSolverReal)>> MConstraintRules;
+	TPBDActiveView<TArray<TFunction<void(FSolverParticles&, const FSolverReal)>>> MConstraintRulesActiveView;
+	TArray<TFunction<void(FSolverParticles&, const FSolverReal)>> MPostCollisionConstraintRules;
+	TPBDActiveView<TArray<TFunction<void(FSolverParticles&, const FSolverReal)>>> MPostCollisionConstraintRulesActiveView;
+
+	TFunction<void(FSolverParticles&, const FSolverReal, const FSolverReal, const int32)> MKinematicUpdate;
+	TFunction<void(FSolverRigidParticles&, const FSolverReal, const FSolverReal, const int32)> MCollisionKinematicUpdate;
+	//Newton specific lambdas:
+	TFunction<void(const bool)> UpdatePositionBasedState;
+	TFunction<void(const PMatrix<FSolverReal, 3, 3>&, PMatrix<FSolverReal, 3, 3>&, const int32)> PStress;
+	TFunction<void(const PMatrix<FSolverReal, 3, 3>&, const PMatrix<FSolverReal, 3, 3>&, PMatrix<FSolverReal, 3, 3>&, const int32)> DeltaPStress;
+	TFunction<void(TArray<FSolverVec3>&)> ProjectBCs;
+	TFunction<void(const FSolverReal, TArray<FSolverVec3>&)> SetBCs;
+	TFunction<void(const FSolverReal, const FSolverReal, const TArray<FSolverReal>&, bool, TArray<FSolverVec3>&)> MAddExternalForce;
+
+	int32 MNumNewtonIterations;
+	int32 MNumCGIterations;
+	TArray<TVector<int32, 4>> MMesh;
+	TArray<TArray<TVector<int32, 2>>> MIncidentElements;
+	TArray<FSolverReal> Measure;
+	//Material Constants:
+	FSolverReal Mu;
+	FSolverReal Lambda;
+
+	FSolverVec3 MGravity;
+	FSolverReal MCollisionThickness;
+	FSolverReal MCoefficientOfFriction;
+	FSolverReal MDamping;
+	FSolverReal MLocalDamping;
+	FSolverReal MTime;
+	FSolverReal MSmoothDt;
+	FSolverReal MNewtonTol;
+	FSolverReal MCGTol;
+	TArray<FSolverVec3> ElementForces;
+	TArray<FSolverVec3> MVn;
+	TArray<FSolverReal> MNodalMass;
+	TArray<FSolverReal> MResidualNorm;
+	TArray<int32> MConstrainedVertices;
+	TArray<FSolverVec3> MBCPositions;
+	bool NoVerbose = false;
+	ConsModelCaches ConsCaches;
+	TUniquePtr<ElasticFEM<FSolverReal, FSolverParticles>> MEFem;
+	TUniquePtr<TArray<int32>> use_list;
+	bool bWriteDebugInfo;
+};
+
+
+}  // End namespace Chaos::Softs
+
+#if !defined(CHAOS_POST_ITERATION_UPDATES_ISPC_ENABLED)
+#define CHAOS_POST_ITERATION_UPDATES_ISPC_ENABLED 1
+#endif
+
+
