Fix crash when removing selection material from a geometry collection

diff --git a/Engine/Source/Runtime/Experimental/GeometryCollectionEngine/Private/GeometryCollection/GeometryCollectionSceneProxy.cpp b/Engine/Source/Runtime/Experimental/GeometryCollectionEngine/Private/GeometryCollection/GeometryCollectionSceneProxy.cpp
--- a/Engine/Source/Runtime/Experimental/GeometryCollectionEngine/Private/GeometryCollection/GeometryCollectionSceneProxy.cpp
+++ b/Engine/Source/Runtime/Experimental/GeometryCollectionEngine/Private/GeometryCollection/GeometryCollectionSceneProxy.cpp
@@ -962,7 +962,7 @@ void FGeometryCollectionSceneProxy::GetDynamicMeshElements(const TArray<const FS
 			// bone selection is already contained in the rendered colors
 			// #note: This renders the geometry again but with the bone selection material.  Ideally we'd have one render pass and one
 			// material.
-			if ((bShowBoneColors || bEnableBoneSelection) && !bSuppressSelectionMaterial)
+			if ((bShowBoneColors || bEnableBoneSelection) && !bSuppressSelectionMaterial && Materials.IsValidIndex(BoneSelectionMaterialID))
 			{
 				FMaterialRenderProxy* MaterialRenderProxy = Materials[BoneSelectionMaterialID]->GetRenderProxy();
 
