Dataflow Node : add non const GetInput/GetOutput

diff --git a/Engine/Source/Runtime/Experimental/Dataflow/Core/Private/Dataflow/DataflowNode.cpp b/Engine/Source/Runtime/Experimental/Dataflow/Core/Private/Dataflow/DataflowNode.cpp
--- a/Engine/Source/Runtime/Experimental/Dataflow/Core/Private/Dataflow/DataflowNode.cpp
+++ b/Engine/Source/Runtime/Experimental/Dataflow/Core/Private/Dataflow/DataflowNode.cpp
@@ -59,6 +59,19 @@ const FDataflowInput* FDataflowNode::FindInput(const void* Reference) const
 	return nullptr;
 }
 
+FDataflowInput* FDataflowNode::FindInput(void* Reference)
+{
+	for (TPair<uint32, FDataflowConnection*> Elem : Inputs)
+	{
+		FDataflowConnection* Con = Elem.Value;
+		if (Con->RealAddress() == (size_t)Reference)
+		{
+			return (FDataflowInput*)Con;
+		}
+	}
+	return nullptr;
+}
+
 const FDataflowInput* FDataflowNode::GetInput(const void* Reference) const
 {
 	for (TPair<uint32, FDataflowConnection*> Elem : Inputs)
@@ -146,6 +159,19 @@ const FDataflowOutput* FDataflowNode::FindOutput(const void* Reference) const
 	return nullptr;
 }
 
+FDataflowOutput* FDataflowNode::FindOutput(void* Reference)
+{
+	for (TPair<uint32, FDataflowConnection*> Elem : Outputs)
+	{
+		FDataflowConnection* Con = Elem.Value;
+		if (Con->RealAddress() == (size_t)Reference)
+		{
+			return (FDataflowOutput*)Con;
+		}
+	}
+	return nullptr;
+}
+
 const FDataflowOutput* FDataflowNode::GetOutput(const void* Reference) const
 {
 	for (TPair<uint32, FDataflowConnection*> Elem : Outputs)
diff --git a/Engine/Source/Runtime/Experimental/Dataflow/Core/Public/Dataflow/DataflowNode.h b/Engine/Source/Runtime/Experimental/Dataflow/Core/Public/Dataflow/DataflowNode.h
--- a/Engine/Source/Runtime/Experimental/Dataflow/Core/Public/Dataflow/DataflowNode.h
+++ b/Engine/Source/Runtime/Experimental/Dataflow/Core/Public/Dataflow/DataflowNode.h
@@ -78,6 +78,7 @@ struct DATAFLOWCORE_API FDataflowNode
 
 	void AddInput(FDataflowConnection* InPtr);
 	FDataflowInput* FindInput(FName Name);
+	FDataflowInput* FindInput(void* Reference);
 	const FDataflowInput* FindInput(const void* Reference) const;
 	const FDataflowInput* GetInput(const void* Reference) const;
 	TArray< FDataflowConnection* > GetInputs() const;
@@ -85,6 +86,7 @@ struct DATAFLOWCORE_API FDataflowNode
 
 	void AddOutput(FDataflowConnection* InPtr);
 	FDataflowOutput* FindOutput(FName Name);
+	FDataflowOutput* FindOutput(void* Reference);
 	const FDataflowOutput* FindOutput(const void* Reference) const;
 	const FDataflowOutput* GetOutput(const void* Reference) const;
 	TArray< FDataflowConnection* > GetOutputs() const;
