Remove old condition in ComputeSweptContactTimeToTargetPhi which set CCD TOI to infinity if end phi is shallower than start phi by a certain tolerance. This is no longer valid to do, since a CCD constraint with TOI =inf indicates no collision at the end position.

For the same reason, when EndPhi > IgnorePhi leads to TOI = 1, we still return true from ComputeSweptContactTOIAndPhiAtTargetPenetration so that GJKContactPointSwept will populate its contact and return a TOI. In the case of contacts that are shallow enough that we just set TOI to 1, immediately downgrade a swept constraint to a regular one inside mid-phase.

diff --git a/Engine/Source/Runtime/Experimental/Chaos/Private/Chaos/Collision/ContactPointsMiscShapes.cpp b/Engine/Source/Runtime/Experimental/Chaos/Private/Chaos/Collision/ContactPointsMiscShapes.cpp
--- a/Engine/Source/Runtime/Experimental/Chaos/Private/Chaos/Collision/ContactPointsMiscShapes.cpp
+++ b/Engine/Source/Runtime/Experimental/Chaos/Private/Chaos/Collision/ContactPointsMiscShapes.cpp
@@ -123,12 +123,6 @@ namespace Chaos
 			return InfiniteTOI;
 		}
 
-		// If contact is moving in the right direction or not moving ignore the contact
-		if (EndPhi > StartPhi - MovementTolerance)
-		{
-			return InfiniteTOI;
-		}
-
 		// If we penetrate by less than the IgnorePhi, treat it as TOI=1. This mean no CCD impulse and the non-CCD 
 		// solve is expected to handle it. This improves the behaviour when we are sliding along a surface at
 		// above CCD speeds - we don't want to handle TOI events with the floor
@@ -146,7 +140,7 @@ namespace Chaos
 	}
 
 	// Modify the time of impact so that the contact depth is TargetPenetration. If penetration at T=1 is less than TargetPenetration, TOI will be "infinity".
-	// Returns true if we have a TOI less than 1 (i.e., a contact at TargetPenetration or more that needs CCD processing)
+	// Returns true if we have a TOI less than or equal to 1 (i.e., a contact at TargetPenetration or more that needs CCD processing)
 	bool ComputeSweptContactTOIAndPhiAtTargetPenetration(const FReal DirDotNormal, const FReal SweepLength, const FReal HitDistance, const FReal IgnorePenetration, const FReal TargetPenetration, FReal& OutTOI, FReal& OutPhi)
 	{
 		// Parse the GJK result to get TOI and Phi (Interpretation depends on whether TOI=0 or not - see GJKRaycast2)
@@ -165,7 +159,7 @@ namespace Chaos
 		OutTOI = TargetTOI;
 		OutPhi = (TargetTOI == 0) ? StartPhi : TargetPhi;
 
-		return (OutTOI < FReal(1));
+		return (OutTOI <= FReal(1));
 	}
 
 	bool ComputeSweptContactTOIAndPhiAtTargetPenetration(const FVec3& ContactNormal, const FVec3& Dir, const FReal SweepLength, const FReal HitDistance, const FReal IgnorePenetration, const FReal TargetPenetration, FReal& OutTOI, FReal& OutPhi)
diff --git a/Engine/Source/Runtime/Experimental/Chaos/Private/Chaos/Collision/ParticlePairMidPhase.cpp b/Engine/Source/Runtime/Experimental/Chaos/Private/Chaos/Collision/ParticlePairMidPhase.cpp
--- a/Engine/Source/Runtime/Experimental/Chaos/Private/Chaos/Collision/ParticlePairMidPhase.cpp
+++ b/Engine/Source/Runtime/Experimental/Chaos/Private/Chaos/Collision/ParticlePairMidPhase.cpp
@@ -377,6 +377,14 @@ namespace Chaos
 			const FRigidTransform3 CCDShapeWorldTransform0 = Constraint->ImplicitTransform[0] * CCDParticleWorldTransform0;
 			const FRigidTransform3 CCDShapeWorldTransform1 = Constraint->ImplicitTransform[1] * CCDParticleWorldTransform1;
 			Collisions::UpdateConstraintSwept(*Constraint.Get(), CCDShapeWorldTransform0, CCDShapeWorldTransform1, Dt);
+
+			// If we did get a hit but it's at TOI = 1, treat this constraint as a regular non-swept constraint.
+			if (Constraint->GetCCDTimeOfImpact() == FReal(1))
+			{
+				Collisions::UpdateConstraint(*Constraint.Get(), Constraint->GetShapeWorldTransform0(), Constraint->GetShapeWorldTransform1(), Dt);
+				Constraint->SetCCDEnabled(false);
+			}
+
 			MidPhase.GetCollisionAllocator().ActivateConstraint(Constraint.Get());
 			LastUsedEpoch = MidPhase.GetCollisionAllocator().GetCurrentEpoch();
 
