Chaos - Fix GJK when we end up with a degenerate simplex
- Add extra logging options to help track down future issues

diff --git a/Engine/Source/Runtime/Experimental/Chaos/Private/Chaos/CollisionOneShotManifolds.cpp b/Engine/Source/Runtime/Experimental/Chaos/Private/Chaos/CollisionOneShotManifolds.cpp
--- a/Engine/Source/Runtime/Experimental/Chaos/Private/Chaos/CollisionOneShotManifolds.cpp
+++ b/Engine/Source/Runtime/Experimental/Chaos/Private/Chaos/CollisionOneShotManifolds.cpp
@@ -602,6 +602,14 @@ namespace Chaos
 				}
 			}
 		}
+		
+		// SelectContactPlane was failing. This issue was tracked down to a degenerate simplex in GJK, but we may need this again one day...
+		template<typename ConvexImplicitType>
+		void CheckPlaneIndex(const int32 PlaneIndex, const ConvexImplicitType& Convex, const FVec3 X, const FVec3 N, const FReal MaxDistance, const int32 VertexIndex)
+		{
+			CHAOS_COLLISIONERROR_CLOG(PlaneIndex == INDEX_NONE, TEXT("SelectContactPlane: Invalid PlaneIndex %d, X: [%f, %f, %f], N: [%f, %f, %f], MaxDistance: %f, VertexIndex: %d\n%s"), PlaneIndex, X.X, X.Y, X.Z, N.X, N.Y, N.Z, MaxDistance, VertexIndex, *Convex.ToString());
+			CHAOS_COLLISIONERROR_ENSURE(PlaneIndex != INDEX_NONE);
+		}
 
 		// Select one of the planes on the convex to use as the contact plane, given an estimated contact position and opposing 
 		// normal from GJK with margins (which gives the shapes rounded corners/edges).
@@ -652,7 +660,7 @@ namespace Chaos
 				BestPlaneIndex = Convex.GetMostOpposingPlane(N);
 			}
 
-			check(BestPlaneIndex != INDEX_NONE);
+			CheckPlaneIndex(BestPlaneIndex, Convex, X, N, InMaxDistance, VertexIndex);
 			return BestPlaneIndex;
 		}
 
@@ -717,56 +725,16 @@ namespace Chaos
 				BestPlaneIndex = ScaledConvex.GetMostOpposingPlane(N);
 			}
 
-			check(BestPlaneIndex != INDEX_NONE);
+			CheckPlaneIndex(BestPlaneIndex, ScaledConvex, X, N, InMaxDistance, VertexIndex);
 			return BestPlaneIndex;
 		}
 
+		// GJK was failing to generate vertex indices. This issue was tracked down to a degenerate simplex in GJK, but we may need this again one day...
 		template<typename ConvexImplicitType>
 		bool CheckVertexIndex(const ConvexImplicitType& Convex, const int32 VertexIndex)
 		{
-			ensureMsgf(VertexIndex != INDEX_NONE, TEXT("GJKContactPointMargin invalid vertex index (type:%d) %d"), Convex.GetType(), VertexIndex);
-			return VertexIndex != INDEX_NONE;
-		}
-
-		template<>
-		bool CheckVertexIndex(const FTriangle& Triangle, const int32 VertexIndex)
-		{
-			ensureMsgf(VertexIndex != INDEX_NONE, TEXT("GJKContactPointMargin invalid Triangle VertexIndex %d"), VertexIndex);
-			return VertexIndex != INDEX_NONE;
-		}
-
-		template<>
-		bool CheckVertexIndex(const FImplicitCapsule3& Capsule, const int32 VertexIndex)
-		{
-			ensureMsgf(VertexIndex != INDEX_NONE, TEXT("GJKContactPointMargin invalid Capsule VertexIndex %d"), VertexIndex);
-			return VertexIndex != INDEX_NONE;
-		}
-
-		template<>
-		bool CheckVertexIndex(const FImplicitBox3& Box, const int32 VertexIndex)
-		{
-			ensureMsgf(VertexIndex != INDEX_NONE, TEXT("GJKContactPointMargin invalid Box VertexIndex %d"), VertexIndex);
-			return VertexIndex != INDEX_NONE;
-		}
-
-		template<>
-		bool CheckVertexIndex(const FImplicitConvex3& Convex, const int32 VertexIndex)
-		{
-			ensureMsgf(VertexIndex != INDEX_NONE, TEXT("GJKContactPointMargin invalid Convex VertexIndex %d [%d, %d, %d]"), VertexIndex, Convex.NumVertices(), Convex.NumEdges(), Convex.NumPlanes());
-			return VertexIndex != INDEX_NONE;
-		}
-
-		template<>
-		bool CheckVertexIndex(const TImplicitObjectInstanced<FImplicitConvex3>& Convex, const int32 VertexIndex)
-		{
-			ensureMsgf(VertexIndex != INDEX_NONE, TEXT("GJKContactPointMargin invalid InstancedConvex VertexIndex %d [%d, %d, %d]"), VertexIndex, Convex.NumVertices(), Convex.NumEdges(), Convex.NumPlanes());
-			return VertexIndex != INDEX_NONE;
-		}
-
-		template<>
-		bool CheckVertexIndex(const TImplicitObjectScaled<FImplicitConvex3>& Convex, const int32 VertexIndex)
-		{
-			ensureMsgf(VertexIndex != INDEX_NONE, TEXT("GJKContactPointMargin invalid ScaledConvex VertexIndex %d [%d, %d, %d]"), VertexIndex, Convex.NumVertices(), Convex.NumEdges(), Convex.NumPlanes());
+			CHAOS_COLLISIONERROR_CLOG(VertexIndex == INDEX_NONE, TEXT("GJKContactPointMargin invalid vertex index %d %s"), VertexIndex, *Convex.ToString());
+			CHAOS_COLLISIONERROR_ENSURE(VertexIndex != INDEX_NONE);
 			return VertexIndex != INDEX_NONE;
 		}
 
@@ -833,6 +801,10 @@ namespace Chaos
 			// @todo(chaos): get the vertex index from GJK and use to to get the plane
 			const FVec3 SeparationDirectionLocalConvex1 = Convex2ToConvex1Transform.TransformVectorNoScale(GJKContactPoint.ShapeContactNormal);
 			const int32 MostOpposingPlaneIndexConvex1 = SelectContactPlane(Convex1, GJKContactPoint.ShapeContactPoints[0], SeparationDirectionLocalConvex1, Margin1, VertexIndexA);
+			if (MostOpposingPlaneIndexConvex1 == INDEX_NONE)
+			{
+				return;
+			}
 			FVec3 BestPlanePosition1, BestPlaneNormal1;
 			Convex1.GetPlaneNX(MostOpposingPlaneIndexConvex1, BestPlaneNormal1, BestPlanePosition1);
 			const FReal BestPlaneDotNormalConvex1 = !bConvex1IsCapsule ? FMath::Abs(FVec3::DotProduct(-SeparationDirectionLocalConvex1, BestPlaneNormal1)) : -FLT_MAX;
@@ -840,6 +812,10 @@ namespace Chaos
 			// Now for Convex2
 			const FVec3 SeparationDirectionLocalConvex2 = GJKContactPoint.ShapeContactNormal;
 			const int32 MostOpposingPlaneIndexConvex2 = SelectContactPlane(Convex2, GJKContactPoint.ShapeContactPoints[1], -SeparationDirectionLocalConvex2, Margin2, VertexIndexB);
+			if (MostOpposingPlaneIndexConvex2 == INDEX_NONE)
+			{
+				return;
+			}
 			FVec3 BestPlanePosition2, BestPlaneNormal2;
 			Convex2.GetPlaneNX(MostOpposingPlaneIndexConvex2, BestPlaneNormal2, BestPlanePosition2);
 			const FReal BestPlaneDotNormalConvex2 = FMath::Abs(FVec3::DotProduct(SeparationDirectionLocalConvex2, BestPlaneNormal2));
@@ -1075,6 +1051,10 @@ namespace Chaos
 			// Find the best plane on the convex (most opposing direction)
 			const FReal PlaneSearchDistance = FReal(0);	// use the miniumum (see SelectContactPlane)
 			const int32 ConvexPlaneIndex = SelectContactPlane(Convex, GJKContactPoint.ShapeContactPoints[0], GJKContactPoint.ShapeContactNormal, PlaneSearchDistance, VertexIndexA);
+			if (ConvexPlaneIndex == INDEX_NONE)
+			{
+				return;
+			}
 			FVec3 ConvexPlanePosition, ConvexPlaneNormal;
 			Convex.GetPlaneNX(ConvexPlaneIndex, ConvexPlaneNormal, ConvexPlanePosition);
 			const FReal ConvexPlaneNormalDotContactNormal = FVec3::DotProduct(ConvexPlaneNormal, GJKContactPoint.ShapeContactNormal);
@@ -1267,6 +1247,10 @@ namespace Chaos
 			// @todo(chaos): handle zero margins...
 			const FReal PlaneSearchDistance = ConvexMargin;
 			const int32 ConvexPlaneIndex = SelectContactPlane(Convex, GJKContactPoint.ShapeContactPoints[0], GJKContactPoint.ShapeContactNormal, PlaneSearchDistance, VertexIndexA);
+			if (ConvexPlaneIndex == INDEX_NONE)
+			{
+				return;
+			}
 			FVec3 ConvexPlanePosition, ConvexPlaneNormal;
 			Convex.GetPlaneNX(ConvexPlaneIndex, ConvexPlaneNormal, ConvexPlanePosition);
 			const FReal ConvexPlaneNormalDotContactNormal = FVec3::DotProduct(ConvexPlaneNormal, GJKContactPoint.ShapeContactNormal);
diff --git a/Engine/Source/Runtime/Experimental/Chaos/Public/Chaos/GJK.h b/Engine/Source/Runtime/Experimental/Chaos/Public/Chaos/GJK.h
--- a/Engine/Source/Runtime/Experimental/Chaos/Public/Chaos/GJK.h
+++ b/Engine/Source/Runtime/Experimental/Chaos/Public/Chaos/GJK.h
@@ -17,6 +17,21 @@
 #include "ChaosCheck.h"
 #include "ChaosLog.h"
 
+// Enable some logging or ensures to trap issues in collision detection, usually related to degeneracies in GJK/EPA
+#define CHAOS_COLLISIONERROR_LOG_ENABLED		((!UE_BUILD_TEST && !UE_BUILD_SHIPPING) && 0)
+#define CHAOS_COLLISIONERROR_ENSURE_ENABLED		((!UE_BUILD_TEST && !UE_BUILD_SHIPPING) && 1)
+
+#if CHAOS_COLLISIONERROR_LOG_ENABLED
+#define CHAOS_COLLISIONERROR_CLOG(Condition, Fmt, ...) UE_CLOG((Condition), LogChaos, Error, Fmt, __VA_ARGS__)
+#else
+#define CHAOS_COLLISIONERROR_CLOG(Condition, Fmt, ...)
+#endif
+
+#if CHAOS_COLLISIONERROR_ENSURE_ENABLED
+#define CHAOS_COLLISIONERROR_ENSURE(X) ensure(X)
+#else
+#define CHAOS_COLLISIONERROR_ENSURE(X)
+#endif
 
 namespace Chaos
 {
@@ -29,17 +44,8 @@ namespace Chaos
 		const int32 MaxIterations = 32;
 		const bool bLimitExceeded = (NumIterations >= MaxIterations);
 
-#if !UE_BUILD_TEST && !UE_BUILD_SHIPPING
-		if (bLimitExceeded)
-		{
-			static bool bLogged = false;
-			if (!bLogged)
-			{
-				UE_LOG(LogChaos, Warning, TEXT("GJK hit iteration limit with shapes:\n    A: %s\n    B: %s"), *A.ToString(), *B.ToString());
-				bLogged = true;
-			}
-		}
-#endif
+		CHAOS_COLLISIONERROR_CLOG(bLimitExceeded, TEXT("GJK hit iteration limit with shapes:\n    A: %s\n    B: %s"), *A.ToString(), *B.ToString());
+		CHAOS_COLLISIONERROR_ENSURE(!bLimitExceeded);
 
 		return bLimitExceeded;
 	}
@@ -974,6 +980,11 @@ namespace Chaos
 			V = SimplexFindClosestToOrigin2(Simplex, NumVerts, SimplexData.Barycentric, SimplexData.As, SimplexData.Bs);
 			T NewDistance = V.Size();
 
+			// If we hit this error, we probably have a degenerate simplex that is not being detected in SimplexFindClosestToOrigin2
+			CHAOS_COLLISIONERROR_CLOG(V.ContainsNaN(), TEXT("SimplexFindClosestToOrigin2 NaN, NumVerts: %d, Simplex: [[%f, %f, %f], [%f %f %f], [%f %f %f], [%f %f %f]]"), 
+				NumVerts, Simplex[0].X, Simplex[0].Y, Simplex[0].Z, Simplex[1].X, Simplex[1].Y, Simplex[1].Z, Simplex[2].X, Simplex[2].Y, Simplex[2].Z, Simplex[3].X, Simplex[3].Y, Simplex[3].Z);
+			CHAOS_COLLISIONERROR_ENSURE(!V.ContainsNaN());
+
 			// Are we overlapping or too close for GJK to get a good result?
 			bIsContact = (NewDistance < Epsilon);
 
diff --git a/Engine/Source/Runtime/Experimental/Chaos/Public/Chaos/Simplex.h b/Engine/Source/Runtime/Experimental/Chaos/Public/Chaos/Simplex.h
--- a/Engine/Source/Runtime/Experimental/Chaos/Public/Chaos/Simplex.h
+++ b/Engine/Source/Runtime/Experimental/Chaos/Public/Chaos/Simplex.h
@@ -421,9 +421,16 @@ namespace Chaos
 		// Previously was return VectorMultiplyAdd(AC, w, VectorMultiplyAdd(AB, v, A));
 		const TVec3<T> TriNormal = TVec3<T>::CrossProduct(AB, AC);
 		const T TriNormal2 = TVec3<T>::DotProduct(TriNormal, TriNormal);
-		const TVec3<T> TriNormalOverSize2 = TriNormal / TriNormal2;
-		const T SignedDistance = TVec3<T>::DotProduct(A, TriNormalOverSize2);
-		return TriNormal * SignedDistance;
+		if (TriNormal2 > std::numeric_limits<T>::min())
+		{
+			const TVec3<T> TriNormalOverSize2 = TriNormal / TriNormal2;
+			const T SignedDistance = TVec3<T>::DotProduct(A, TriNormalOverSize2);
+			return TriNormal * SignedDistance;
+		}
+
+		// If we get here we hit a degenerate Simplex (all 3 verts in a line/point)
+		// Let's just exit GJK with whatever normal and distance it had last iteration...
+		return FVec3(0);
 	}
 
 	template <typename T>
