Chaos : add item index to queries for geometry collection
- added a crumbling and apply strain BP functions to geometry collection component

diff --git a/Engine/Source/Runtime/Engine/Private/Collision/CollisionConversions.cpp b/Engine/Source/Runtime/Engine/Private/Collision/CollisionConversions.cpp
--- a/Engine/Source/Runtime/Engine/Private/Collision/CollisionConversions.cpp
+++ b/Engine/Source/Runtime/Engine/Private/Collision/CollisionConversions.cpp
@@ -19,6 +19,7 @@
 #include "PhysicsEngine/CollisionQueryFilterCallback.h"
 #include "GameFramework/LightWeightInstanceManager.h"
 #include "PhysicsProxy/SingleParticlePhysicsProxy.h"
+#include "PhysicsProxy/GeometryCollectionPhysicsProxy.h"
 
 // Used to place overlaps into a TMap when deduplicating them
 struct FOverlapKey
@@ -151,6 +152,16 @@ static void SetHitResultFromShapeAndFaceIndex(const FPhysicsShape& Shape,  const
 			OwningComponent = PossibleOwner;
 			OutResult.Item = INDEX_NONE;
 			OutResult.BoneName = NAME_None;
+
+			// if we have a geometry component let's extract the index of the active piece we have a hit for
+			if (const IPhysicsProxyBase* ActorProxy = Actor.GetProxy())
+			{
+				if (ActorProxy->GetType() == EPhysicsProxyType::GeometryCollectionType)
+				{
+					const FGeometryCollectionPhysicsProxy* ConcreteProxy = static_cast<const FGeometryCollectionPhysicsProxy*>(ActorProxy);
+					OutResult.Item = ConcreteProxy->GetTransformGroupIndexFromGTParticle(&Actor);
+				}
+			}
 		}
 		else
 		{
diff --git a/Engine/Source/Runtime/Experimental/Chaos/Private/PhysicsProxy/GeometryCollectionPhysicsProxy.cpp b/Engine/Source/Runtime/Experimental/Chaos/Private/PhysicsProxy/GeometryCollectionPhysicsProxy.cpp
--- a/Engine/Source/Runtime/Experimental/Chaos/Private/PhysicsProxy/GeometryCollectionPhysicsProxy.cpp
+++ b/Engine/Source/Runtime/Experimental/Chaos/Private/PhysicsProxy/GeometryCollectionPhysicsProxy.cpp
@@ -464,6 +464,7 @@ void FGeometryCollectionPhysicsProxy::Initialize(Chaos::FPBDRigidsEvolutionBase
 		{
 			GTParticles[Index] = Chaos::FGeometryParticle::CreateParticle();
 			Chaos::FGeometryParticle* P = GTParticles[Index].Get();
+			GTParticlesToTransformGroupIndex.Add(P, Index);
 
 			GTParticles[Index]->SetUniqueIdx(Evolution->GenerateUniqueIdx());
 
@@ -1604,6 +1605,56 @@ void FGeometryCollectionPhysicsProxy::BreakClusters(TArray<int32>&& TransformGro
 	}
 }
 
+void FGeometryCollectionPhysicsProxy::ApplyStrain(int32 TransformGroupIndex, const FVector& Location, float StrainValue)
+{
+	check(IsInGameThread());
+
+	if (Chaos::FPhysicsSolver* RBDSolver = GetSolver<Chaos::FPhysicsSolver>())
+	{
+		RBDSolver->EnqueueCommandImmediate([this, StrainValue, RBDSolver, Location, TransformGroupIndex]()
+		{
+			// @todo(chaos) move most of the logic into a method in clustering 
+			Chaos::FRigidClustering& Clustering = RBDSolver->GetEvolution()->GetRigidClustering();
+			if (Chaos::FPBDRigidClusteredParticleHandle* ClusterHandle = SolverParticleHandles[TransformGroupIndex])
+			{
+				const bool bIsCluster = (ClusterHandle->ClusterIds().NumChildren > 0);
+				const bool bHasInternalClusterParent = ClusterHandle->Parent()? ClusterHandle->Parent()->InternalCluster(): false;
+				if (bHasInternalClusterParent || (!bIsCluster && !ClusterHandle->Disabled()))
+				{
+					ClusterHandle->SetExternalStrain(FMath::Max(ClusterHandle->GetExternalStrain(), StrainValue));
+				}
+				else
+				{
+					if (bIsCluster)
+					{
+						if (const TArray<Chaos::FPBDRigidParticleHandle*>* ChildrenHandles = Clustering.GetChildrenMap().Find(ClusterHandle))
+						{
+							FVector::FReal ClosestSquaredDist = TNumericLimits<FVector::FReal>::Max();
+							Chaos::FPBDRigidParticleHandle* ClosestChildHandle = nullptr; 
+							for (Chaos::FPBDRigidParticleHandle* ChildHandle: *ChildrenHandles)
+							{
+								const FVector::FReal SquaredDist = (ChildHandle->X() - Location).SizeSquared();
+								if (SquaredDist < ClosestSquaredDist)
+								{
+									ClosestSquaredDist = SquaredDist;
+									ClosestChildHandle = ChildHandle;
+								}
+							}
+							if (ClosestChildHandle)
+							{
+								if (Chaos::FPBDRigidClusteredParticleHandle* ClusteredChildHandle = ClosestChildHandle->CastToClustered())
+								{
+									ClusteredChildHandle->SetExternalStrain(FMath::Max(ClusteredChildHandle->GetExternalStrain(), StrainValue));
+								}
+							}
+						}
+					}
+				}
+			}
+		});
+	}
+}
+
 int32 FGeometryCollectionPhysicsProxy::CalculateHierarchyLevel(const FGeometryDynamicCollection& DynamicCollection, int32 TransformIndex)
 {
 	int32 Level = 0;
diff --git a/Engine/Source/Runtime/Experimental/Chaos/Public/PhysicsProxy/GeometryCollectionPhysicsProxy.h b/Engine/Source/Runtime/Experimental/Chaos/Public/PhysicsProxy/GeometryCollectionPhysicsProxy.h
--- a/Engine/Source/Runtime/Experimental/Chaos/Public/PhysicsProxy/GeometryCollectionPhysicsProxy.h
+++ b/Engine/Source/Runtime/Experimental/Chaos/Public/PhysicsProxy/GeometryCollectionPhysicsProxy.h
@@ -229,21 +229,29 @@ public:
 
 	int32 GetTransformGroupIndexFromHandle(const FParticleHandle* Handle) const
 	{
-		if (HandleToTransformGroupIndex.Contains(Handle))
+		if (const int32* TransformGroupIndex = HandleToTransformGroupIndex.Find(Handle))
 		{
-			return HandleToTransformGroupIndex[Handle];
+			return *TransformGroupIndex;
 		}
-		else
+		return INDEX_NONE;
+	}
+
+	int32 GetTransformGroupIndexFromGTParticle(const Chaos::FGeometryParticle* GTPParticle) const
+	{
+		if (const int32* TransformGroupIndex = GTParticlesToTransformGroupIndex.Find(GTPParticle))
 		{
-			return INDEX_NONE;
+			return *TransformGroupIndex;
 		}
+		return INDEX_NONE;
 	}
 
 	bool GetIsObjectDynamic() const { return IsObjectDynamic; }
 
 	void DisableParticles(TArray<int32>&& TransformGroupIndices);
+
 	void BreakInternalClusterParents(TArray<int32>&& TransformGroupIndices);
 	void BreakClusters(TArray<int32>&& TransformGroupIndices);
+	void ApplyStrain(int32 TransformGroupIndex, const FVector& Location, float StrainValue);
 
 	FProxyInterpolationData& GetInterpolationData() { return InterpolationData; }
 	const FProxyInterpolationData& GetInterpolationData() const { return InterpolationData; }
@@ -324,6 +332,7 @@ private:
 	bool bInitializedForReplication = false;	//Indicates that initialization for replication has finished
 
 	TManagedArray<TUniquePtr<Chaos::FGeometryParticle>> GTParticles;
+	TMap<Chaos::FGeometryParticle*, int32> GTParticlesToTransformGroupIndex;
 
 	// These are read on both threads and should not be changed
 	const FCollisionFilterData SimFilter;
diff --git a/Engine/Source/Runtime/Experimental/GeometryCollectionEngine/Private/GeometryCollection/GeometryCollectionComponent.cpp b/Engine/Source/Runtime/Experimental/GeometryCollectionEngine/Private/GeometryCollection/GeometryCollectionComponent.cpp
--- a/Engine/Source/Runtime/Experimental/GeometryCollectionEngine/Private/GeometryCollection/GeometryCollectionComponent.cpp
+++ b/Engine/Source/Runtime/Experimental/GeometryCollectionEngine/Private/GeometryCollection/GeometryCollectionComponent.cpp
@@ -3729,6 +3729,43 @@ void UGeometryCollectionComponent::IncrementBreakTimer(float DeltaTime)
 	}
 }
 
+void UGeometryCollectionComponent::ApplyExternalStrain(int32 Index, const FVector& Location, float Strain)
+{
+	if (DynamicCollection && PhysicsProxy)
+	{
+		const int32 NumTransform = DynamicCollection->NumElements(FGeometryCollection::TransformGroup);
+		if (Index >= 0 && Index < NumTransform)
+		{
+			PhysicsProxy->ApplyStrain(Index, Location, Strain);
+		}
+	}
+}
+
+bool UGeometryCollectionComponent::CrumbleCluster(int32 Index)
+{
+	bool Result = false;
+	if (DynamicCollection && PhysicsProxy)
+	{
+		const int32 NumTransform = DynamicCollection->NumElements(FGeometryCollection::TransformGroup);
+		if (Index >= 0 && Index < NumTransform)
+		{
+			const TManagedArray<uint8>& InternalClusterParentType = DynamicCollection->GetAttribute<uint8>("InternalClusterParentTypeArray", FGeometryCollection::TransformGroup);
+			const bool bHasInternalClusterParent = (InternalClusterParentType[Index] != (uint8)Chaos::EInternalClusterType::None);
+			const bool bIsCluster = (DynamicCollection->Children[Index].Num() > 0);
+			if (bIsCluster)
+			{
+				PhysicsProxy->BreakClusters({Index});
+				Result = true;
+			}
+			else if (bHasInternalClusterParent)
+			{
+				PhysicsProxy->BreakInternalClusterParents({Index});
+				Result = true;
+			}
+		}
+	}
+	return Result;
+}
 
 bool UGeometryCollectionComponent::CalculateInnerSphere(int32 TransformIndex, FSphere& SphereOut) const
 {
diff --git a/Engine/Source/Runtime/Experimental/GeometryCollectionEngine/Public/GeometryCollection/GeometryCollectionComponent.h b/Engine/Source/Runtime/Experimental/GeometryCollectionEngine/Public/GeometryCollection/GeometryCollectionComponent.h
--- a/Engine/Source/Runtime/Experimental/GeometryCollectionEngine/Public/GeometryCollection/GeometryCollectionComponent.h
+++ b/Engine/Source/Runtime/Experimental/GeometryCollectionEngine/Public/GeometryCollection/GeometryCollectionComponent.h
@@ -474,6 +474,25 @@ public:
 	UPROPERTY(EditAnywhere, Category = "ChaosPhysics", meta = (DisplayName = "Chaos Solver"))
 	TObjectPtr<AChaosSolverActor> ChaosSolverActor;
 
+	/**
+	 * Apply an external strain a specific piece of the geometry collection
+	 * if the piece is a cluster then the strain will be applied to the closest child   
+	 * @param Index index of the piece to apply strain on
+	 * @param Location world location of where to apply the strain
+	 * @param Strain strain / damage to apply 
+	 */
+	UFUNCTION(BlueprintCallable, Category = "ChaosPhysics")
+	void ApplyExternalStrain(int32 Index, const FVector& Location, float Strain);
+
+	/**
+	 * Crumbe a cluster into all its pieces
+	 * if the piece is not a cluster but is part of an internal/partial cluster, then the owning cluster will crumble   
+	 * @param Index index of the piece to crumble
+	 * @return  true of a crumble operation was performed
+	 */
+	UFUNCTION(BlueprintCallable, Category = "ChaosPhysics")
+	bool CrumbleCluster(int32 Index);
+	
 	/** RestCollection */
 	UFUNCTION(BlueprintCallable, Category = "ChaosPhysics")
 	void SetRestCollection(const UGeometryCollection * RestCollectionIn);
