Add Laplacian() and ComputeFiberField(), along with supporting 3x3 matrix operations.  Add functionality to Chaos::TUniformGrid for dealing with nodes, and translating flat indices to/from multi-indices.  Add a utility that generates a tetrahedral mesh on a grid.

diff --git a/Engine/Source/Programs/HeadlessChaos/Private/HeadlessChaosTestMath.cpp b/Engine/Source/Programs/HeadlessChaos/Private/HeadlessChaosTestMath.cpp
--- a/Engine/Source/Programs/HeadlessChaos/Private/HeadlessChaosTestMath.cpp
+++ b/Engine/Source/Programs/HeadlessChaos/Private/HeadlessChaosTestMath.cpp
@@ -3,8 +3,10 @@
 #include "HeadlessChaos.h"
 #include "HeadlessChaosTestUtility.h"
 #include "Chaos/Math/Krylov.h"
+#include "Chaos/Math/Poisson.h"
 #include "Chaos/Matrix.h"
 #include "Chaos/Utilities.h"
+#include "Chaos/UniformGrid.h"
 #include "Chaos/Vector.h"
 #include "GenericPlatform/GenericPlatformMath.h"
 
@@ -167,4 +169,108 @@ namespace ChaosTest
 		EXPECT_NEAR(r, 0., 1.e-10);
 	}
 #endif // 0
+
+	TEST(MathTests, TestLaplacian)
+	{
+		int32 N = 3;
+		double dx = 1. / N;
+		TVector<double, 3> Origin(0, 0, 0);
+		TVector<double, 3> MinCorner = Origin;
+		TVector<double, 3> MaxCorner = Origin + TVector<double, 3>(dx * N, dx * N, dx * N);
+		TUniformGrid<double, 3> Grid(MinCorner, MaxCorner, TVector<int32, 3>(N, N, N), 0);
+		
+		TArray<TVector<int, 4>> Mesh;
+		TArray<TVector<double, 3>> X;
+		Chaos::Utilities::TetMeshFromGrid<double>(Grid, Mesh, X);
+
+		TArray<TArray<int>> IncidentElementsLocalIndex;
+		TArray<TArray<int>> IncidentElements = Chaos::Utilities::ComputeIncidentElements(Mesh, &IncidentElementsLocalIndex);
+
+		TArray<double> u; u.SetNum(X.Num());
+		TVector<double, 3> A(1, 2, 3);
+		for (int32 i = 0; i < X.Num(); i++)
+		{
+			u[i] = 0.;
+			for (int32 alpha = 0; alpha < 3; alpha++)
+			{
+				u[i] += A[alpha] * X[i][alpha];
+			}
+		}
+
+		TArray<double> De_inverse, measure;
+		Chaos::ComputeDeInverseAndElementMeasures<double>(Mesh, X, De_inverse, measure);
+
+		TArray<double> Lu; Lu.SetNum(X.Num());
+		Chaos::Laplacian(Mesh, IncidentElements, IncidentElementsLocalIndex, De_inverse, measure, u, Lu);
+
+		for (int32 i = 0; i < Grid.GetNumNodes(); i++)
+		{
+			Chaos::TVector<int32, 3> MIndex;
+			Grid.FlatToMultiIndex(i, MIndex, true);
+			if (Grid.InteriorNode(MIndex)) 
+			{
+				EXPECT_NEAR(Lu[i], 0., 1.e-14);
+			}
+		}
+
+		double Energy = Chaos::LaplacianEnergy<double>(Mesh, De_inverse, measure, u);
+		EXPECT_NEAR(Energy, .5 * Chaos::Utilities::DotProduct(u, Lu), 1e-12);
+
+		srand(0);
+		for (int32 i = 0; i < u.Num(); i++) 
+		{
+			u[i] = 2. * rand() / RAND_MAX - 1.;
+		}
+
+		Chaos::Laplacian(Mesh, IncidentElements, IncidentElementsLocalIndex, De_inverse, measure, u, Lu);
+		Energy = Chaos::LaplacianEnergy(Mesh, De_inverse, measure, u);
+		EXPECT_NEAR(Energy, .5 * Chaos::Utilities::DotProduct(u, Lu), 1.e-12);
+	}
+
+	TEST(PoissonTests, TestFiberField) 
+	{
+		//create regular grid
+		int32 N = 3;
+		double dx = 1. / N;
+		TVector<double, 3> Origin(0, 0, 0);
+		TVector<double, 3> MinCorner = Origin;
+		TVector<double, 3> MaxCorner = Origin + TVector<double, 3>(dx * N, dx * N, dx * N);
+		TUniformGrid<double, 3> Grid(MinCorner, MaxCorner, TVector<int32, 3>(N, N, N), 0);
+
+		//create mesh from grid
+		TArray<TVector<int, 4>> Mesh;
+		TArray<TVector<double, 3>> X;
+		Chaos::Utilities::TetMeshFromGrid<double>(Grid, Mesh, X);
+
+		TArray<TArray<int>> IncidentElementsLocalIndex;
+		TArray<TArray<int>> IncidentElements = Chaos::Utilities::ComputeIncidentElements(Mesh, &IncidentElementsLocalIndex);
+
+		TArray<int32> Origins;
+		TArray<int32> Insertions;
+		for (int32 i = 0; i < X.Num(); i++) 
+		{
+			if (X[i][0] < MinCorner[0] + .1 * dx) 
+			{
+				Origins.Add(i);
+			}
+			else if (X[i][0] > MaxCorner[0] - .1 * dx)
+			{
+				Insertions.Add(i);
+			}
+		}
+
+		TArray<Chaos::TVector<double, 3>> Directions;
+		Chaos::ComputeFiberField<double>(Mesh, X, IncidentElements, IncidentElementsLocalIndex, Origins, Insertions, Directions);
+
+		for(int32 e=0; e < Mesh.Num(); e++)
+		{
+			EXPECT_NEAR(Directions[e][0], 1., 1.e-12);
+			for (int32 alpha = 1; alpha < 3; alpha++) 
+			{
+				EXPECT_NEAR(Directions[e][alpha], 0., 1.e-12);
+			}
+		}
+
+	}
+
 }
diff --git a/Engine/Source/Runtime/Experimental/Chaos/Public/Chaos/Math/Poisson.h b/Engine/Source/Runtime/Experimental/Chaos/Public/Chaos/Math/Poisson.h
new file mode 100644
--- /dev/null
+++ b/Engine/Source/Runtime/Experimental/Chaos/Public/Chaos/Math/Poisson.h
@@ -0,0 +1,397 @@
+// Copyright Epic Games, Inc. All Rights Reserved.
+#pragma once
+
+//#include "ChaosFlesh/FleshCollectionUtility.h"
+#include "Chaos/Math/Krylov.h"
+#include "Chaos/Framework/Parallel.h"
+#include "Math/Matrix.h"
+#include "Math/Vector.h"
+
+namespace Chaos {
+
+template <class T>
+inline T RowMaj3x3Determinant(const T A0, const T A1, const T A2, const T A3, const T A4, const T A5, const T A6, const T A7, const T A8)
+{
+	return A0 * (A4 * A8 - A5 * A7) - A1 * (A3 * A8 - A5 * A6) + A2 * (A3 * A7 - A4 * A6);
+}
+
+template <class T>
+inline T RowMaj3x3Determinant(const T* A)
+{
+	return RowMaj3x3Determinant(A[0], A[1], A[2], A[3], A[4], A[5], A[6], A[7], A[8]);
+}
+
+template <class T>
+inline void RowMaj3x3Inverse(const T Det, const T A0, const T A1, const T A2, const T A3, const T A4, const T A5, const T A6, const T A7, const T A8, T* Inv)
+{
+	if (Det == T(0))
+	{
+		for (int i = 0; i < 9; i++)
+			Inv[i] = T(0);
+		Inv[0] = T(1);
+		Inv[4] = T(1);
+		Inv[8] = T(1);
+		return;
+	}
+
+	Inv[0] = (A4 * A8 - A5 * A7) / Det;
+	Inv[1] = (A2 * A7 - A1 * A8) / Det;
+	Inv[2] = (A1 * A5 - A2 * A4) / Det;
+	Inv[3] = (A5 * A6 - A3 * A8) / Det;
+	Inv[4] = (A0 * A8 - A2 * A6) / Det;
+	Inv[5] = (A2 * A3 - A0 * A5) / Det;
+	Inv[6] = (A3 * A7 - A4 * A6) / Det;
+	Inv[7] = (A1 * A6 - A0 * A7) / Det;
+	Inv[8] = (A0 * A4 - A1 * A3) / Det;
+}
+
+template <class T>
+inline void RowMaj3x3Inverse(const T Det, const T* A, T* Inv)
+{
+	RowMaj3x3Inverse(Det, A[0], A[1], A[2], A[3], A[4], A[5], A[6], A[7], A[8], Inv);
+}
+
+template <class T>
+inline void RowMaj3x3Transpose(const T* A, T* Transpose)
+{
+	for (int32 i = 0; i < 3; i++) 
+	{
+		for (int32 j = 0; j < 3; j++) 
+		{
+			Transpose[3 * j + i] = A[3 * i + j];
+		}
+	}
+}
+
+template <class T>
+inline Chaos::TVector<T,3> RowMaj3x3Multiply(const T* A, const Chaos::TVector<T,3>& x)
+{
+	TVector<T,3> Result(0,0,0);
+	for (int32 i = 0; i < 3; i++) 
+	{
+		for (int32 j = 0; j < 3; j++) 
+		{
+			Result[i] += A[3 * i + j] * x[j];
+		}
+	}
+	return Result;
+}
+
+
+
+template <class T, class TV=FVector3f, class TV_INT=FIntVector4, int d=3>
+void
+ComputeDeInverseAndElementMeasures(
+	const TArray<TV_INT>& Mesh, 
+	const TArray<TV>& X, 
+	TArray<T>& De_inverse, 
+	TArray<T>& measure)
+{
+	De_inverse.SetNumUninitialized(d * d * Mesh.Num());
+	measure.SetNumUninitialized(Mesh.Num());
+
+	const int32* MeshPtr = &Mesh[0][0]; // Assumes Mesh is a continuguous array of int32.
+
+	//PhysicsParallelForRange(Mesh.Num(), [&](const int32 eBegin, const int32 eEnd)
+	const int32 eBegin = 0; const int32 eEnd = Mesh.Num();
+	{
+		T Inv[d*d];
+		for(int32 e=eBegin; e < eEnd; e++)
+		{
+			const TV_INT& Elem = Mesh[e];
+			T De[9] =
+			{
+				X[Elem[1]][0] - X[Elem[0]][0], X[Elem[2]][0] - X[Elem[0]][0], X[Elem[3]][0] - X[Elem[0]][0],
+				X[Elem[1]][1] - X[Elem[0]][1], X[Elem[2]][1] - X[Elem[0]][1], X[Elem[3]][1] - X[Elem[0]][1],
+				X[Elem[1]][2] - X[Elem[0]][2], X[Elem[2]][2] - X[Elem[0]][2], X[Elem[3]][2] - X[Elem[0]][2]
+			};
+
+			T Det = RowMaj3x3Determinant(De);
+			measure[e] = Det / (T)6;
+
+			RowMaj3x3Inverse(Det, De, Inv);
+
+			for (int32 i = 0; i < d*d; i++) 
+			{
+				De_inverse[d*d*e+i] = Inv[i];
+			}
+		}
+	}
+	//}, 
+	//100,
+	//Mesh.Num() < 100);
+}
+
+template <class T>
+void Fill(TArray<T>& Array, const T Value)
+{ 
+	for (int32 Idx = 0; Idx < Array.Num(); ++Idx) 
+		Array[Idx] = Value; 
+}
+
+inline int32 MinFlatIndex(const TArray<int32>& ElemIdx, const TArray<int32>& LocalIdx)
+{
+	int32 MinIdx = 0;
+	for (int32 i = 1; i < ElemIdx.Num(); i++)
+	{
+		if (ElemIdx[i] < ElemIdx[MinIdx])
+			MinIdx = i;
+	}
+	return (ElemIdx[MinIdx] * 4) + LocalIdx[MinIdx];
+}
+
+template <class T, class TV_INT=FIntVector4, int d=3>
+void 
+Laplacian(
+	const TArray<TV_INT>& Mesh,
+	const TArray<TArray<int32>>& IncidentElements,
+	const TArray<TArray<int32>>& IncidentElementsLocalIndex,
+	const TArray<T>& De_inverse,
+	const TArray<T>& measure,
+	const TArray<T>& u,
+	TArray<T>& Lu)
+{
+	TArray<Chaos::TVector<T,d>> grad_Nie_hat; grad_Nie_hat.SetNumUninitialized(d + 1);
+	grad_Nie_hat[0] = { T(-1),T(-1),T(-1) };
+	grad_Nie_hat[1] = { T(1),T(0),T(0) };
+	grad_Nie_hat[2] = { T(0),T(1),T(0) };
+	grad_Nie_hat[3] = { T(0),T(0),T(1) };
+
+	Lu.SetNum(u.Num());
+	Fill(Lu, T(0));
+
+	//first compute element contributions
+	TArray<T> element_contributions; element_contributions.SetNum(Mesh.Num() * 4);
+	const int32* MeshPtr = &Mesh[0][0]; // Assumes Mesh is a contiguous array of int32.
+	//PhysicsParallelFor(int32(Mesh.Num()*4 / (d + 1)), [&](const int32 e)
+	for(int32 e=0; e < Mesh.Num(); e++)
+	{
+		T Deinv[d * d]{
+			De_inverse[d * d * e + 0],
+			De_inverse[d * d * e + 1],
+			De_inverse[d * d * e + 2],
+			De_inverse[d * d * e + 3],
+			De_inverse[d * d * e + 4],
+			De_inverse[d * d * e + 5],
+			De_inverse[d * d * e + 6],
+			De_inverse[d * d * e + 7],
+			De_inverse[d * d * e + 8],
+		};
+		T De_inverse_transpose[d * d];
+		RowMaj3x3Transpose(Deinv, De_inverse_transpose);
+
+		const TV_INT& Elem = Mesh[e];
+
+		for (int32 ie = 0; ie < d+1; ie++)
+		{
+			Chaos::TVector<T,3> grad_Nie = RowMaj3x3Multiply(De_inverse_transpose, grad_Nie_hat[ie]);
+
+			for (int32 je = 0; je < d+1; je++)
+			{
+				Chaos::TVector<T, 3> grad_Nje = RowMaj3x3Multiply(De_inverse_transpose, grad_Nie_hat[je]);
+				T Ae_ieje = Chaos::TVector<T,3>::DotProduct(grad_Nie, grad_Nje) * measure[e];
+				element_contributions[(d + 1) * e + ie] += Ae_ieje * u[Elem[je]];
+			}
+		}
+	//});
+	}
+
+	//now gather from elements
+	//PhysicsParallelFor(IncidentElements.Num(), [&](const int32 i)
+	for (int32 i = 0; i < IncidentElements.Num(); i++)
+	{
+		const TArray<int32>& IncidentElements_i = IncidentElements[i];
+		const TArray<int32>& IncidentElementsLocalIndex_i = IncidentElementsLocalIndex[i];
+		if (!IncidentElements_i.Num())
+			return;
+		const int32 p = MeshPtr[MinFlatIndex(IncidentElements_i, IncidentElementsLocalIndex_i)];
+		for (int32 e = 0; e < IncidentElements_i.Num(); e++)
+		{
+			const int32 Elem = IncidentElements_i[e];
+			const int32 Local = IncidentElementsLocalIndex_i[e];
+			Lu[p] += element_contributions[(Elem * 4) + Local];
+		}
+	}
+	//});
+}
+
+template<class T, class TV_INT, int d=3>
+T 
+LaplacianEnergy(
+	const TArray<TV_INT>& Mesh,
+	const TArray<T>& De_inverse,
+	const TArray<T>& measure,
+	const TArray<T>& u)
+{
+	TArray<TVector<T, 3>> grad_Nie_hat; grad_Nie_hat.SetNumUninitialized(d + 1);
+	grad_Nie_hat[0] = { T(-1),T(-1),T(-1) };
+	grad_Nie_hat[1] = { T(1),T(0),T(0) };
+	grad_Nie_hat[2] = { T(0),T(1),T(0) };
+	grad_Nie_hat[3] = { T(0),T(0),T(1) };
+
+	//first compute element contributions in parallel
+	TArray<T> element_contributions; element_contributions.SetNum(Mesh.Num() * 4);
+
+	//PhysicsParallelForRange(Mesh.Num(), [&](const int32 eBegin, const int32 eEnd)
+	const int32 eBegin = 0;
+	const int32 eEnd = Mesh.Num();
+	{
+		for(int32 e=eBegin; e < eEnd; e++)
+		{
+			T Deinv[d * d]{
+				De_inverse[d * d * e + 0],
+				De_inverse[d * d * e + 1],
+				De_inverse[d * d * e + 2],
+				De_inverse[d * d * e + 3],
+				De_inverse[d * d * e + 4],
+				De_inverse[d * d * e + 5],
+				De_inverse[d * d * e + 6],
+				De_inverse[d * d * e + 7],
+				De_inverse[d * d * e + 8],
+			};
+			T De_inverse_transpose[d * d];
+			RowMaj3x3Transpose(Deinv, De_inverse_transpose);
+
+			const TV_INT& Elem = Mesh[e];
+			for (int32 ie = 0; ie < d + 1; ie++) 
+			{
+				TVector<T, d> grad_Nie = RowMaj3x3Multiply(De_inverse_transpose, grad_Nie_hat[ie]);
+				for (int32 je = 0; je < d + 1; je++) 
+				{
+					Chaos::TVector<T, 3> grad_Nje = RowMaj3x3Multiply(De_inverse_transpose, grad_Nie_hat[je]);
+
+					T Ae_ieje = TVector<T,d>::DotProduct(grad_Nie, grad_Nje) * measure[e];
+					element_contributions[e] += u[Elem[ie]] * Ae_ieje * u[Elem[je]];
+				}
+			}
+		}
+	}/*,
+	100,
+	Mesh.Num() < 100);*/
+
+	//now gather from elements
+	T result = T(0);
+	for (int32 i = 0; i < Mesh.Num(); ++i) 
+	{
+		result += element_contributions[i];
+	}
+	return T(.5) * result;
+}
+
+
+
+
+template <class T, class TV, class TV_INT, int d=3>
+void 
+ComputeFiberField(
+	const TArray<TV_INT>& Mesh,
+	const TArray<TV>& Vertices,
+	const TArray<TArray<int32>>& IncidentElements,
+	const TArray<TArray<int32>>& IncidentElementsLocalIndex,
+	const TArray<int32>& Origins,
+	const TArray<int32>& Insertions,
+	TArray<TV>& Directions,
+	const int32 MaxIt=100,
+	const T Tol=T(1e-7))
+{
+	// Define bc lambda
+	auto set_bcs = [&Origins, &Insertions](TArray<T>& minus_u_bc, const T scale, bool zero_non_boundary = true)
+	{
+		if (zero_non_boundary)
+			Fill(minus_u_bc, T(0));
+		// origin points have bc = 0
+		for (int32 i = 0; i < Origins.Num(); i++)
+			minus_u_bc[Origins[i]] = T(0);
+		// insertion points have bc = scale
+		for (int32 i = 0; i < Insertions.Num(); i++)
+			minus_u_bc[Insertions[i]] = scale;
+	};
+
+	// Origin and insertion points are fixed
+	auto proj_bcs = [&Origins, &Insertions](TArray<T>& u)
+	{
+		for (int32 i = 0; i < Origins.Num(); i++)
+			u[Origins[i]] = T(0);
+		for (int32 i = 0; i < Insertions.Num(); i++)
+			u[Insertions[i]] = T(0);
+	};
+
+	TArray<T> u; u.SetNumZeroed(Vertices.Num());
+
+	// Initialize FEM meta data
+	TArray<T> De_inverse;
+	TArray<T> measure;
+	Chaos::ComputeDeInverseAndElementMeasures<T>(Mesh, Vertices, De_inverse, measure);
+
+	// Set BCs and RHS
+	TArray<T> negative_u_bc; negative_u_bc.SetNumZeroed(u.Num());
+	TArray<T> rhs; rhs.SetNumZeroed(u.Num());
+	set_bcs(negative_u_bc, T(-1), true);
+
+	Chaos::Laplacian(Mesh, IncidentElements, IncidentElementsLocalIndex, De_inverse, measure, negative_u_bc, rhs);
+	proj_bcs(rhs);
+
+	auto mult = [&](TArray<T>& out, const TArray<T>& in)
+	{
+		TArray<T> proj_in = in;
+		proj_bcs(proj_in);
+		Laplacian(Mesh, IncidentElements, IncidentElementsLocalIndex, De_inverse, measure, proj_in, out);
+		proj_bcs(out);
+	};
+
+	//solve system
+	Chaos::LanczosCG(mult, u, rhs, MaxIt, Tol, true);
+	set_bcs(u, T(1), false);
+
+	//compute directions from scalar gradient
+	Directions.SetNumUninitialized(Mesh.Num());
+	Fill(Directions, TV(0));
+
+	TArray<TV> grad_Nie_hat; grad_Nie_hat.SetNumUninitialized(d + 1);
+	grad_Nie_hat[0] = { T(-1),T(-1),T(-1) };
+	grad_Nie_hat[1] = { T(1),T(0),T(0) };
+	grad_Nie_hat[2] = { T(0),T(1),T(0) };
+	grad_Nie_hat[3] = { T(0),T(0),T(1) };
+
+	TV pseudo_direction = Vertices[Insertions[0]] - Vertices[Origins[0]];
+	pseudo_direction.Normalize();
+
+	//PhysicsParallelFor(Mesh.Num(), [&](const int32 e)
+	for(int32 e=0; e < Mesh.Num(); e++)
+	{
+		T Deinv[d * d]{
+			De_inverse[d * d * e + 0],
+			De_inverse[d * d * e + 1],
+			De_inverse[d * d * e + 2],
+			De_inverse[d * d * e + 3],
+			De_inverse[d * d * e + 4],
+			De_inverse[d * d * e + 5],
+			De_inverse[d * d * e + 6],
+			De_inverse[d * d * e + 7],
+			De_inverse[d * d * e + 8],
+		};
+		T De_inverse_transpose[d * d];
+		RowMaj3x3Transpose(Deinv, De_inverse_transpose);
+
+		const TV_INT& Elem = Mesh[e];
+		TV gradient(0);
+		for (size_t ie = 0; ie < d + 1; ie++) 
+		{
+			TVector<T, d> grad_Nie = RowMaj3x3Multiply(De_inverse_transpose, grad_Nie_hat[ie]);
+			gradient += grad_Nie * u[Elem[ie]];
+		}
+		const T Len = gradient.Length();
+		if (Len > T(1e-10)) 
+		{
+			Directions[e] = gradient * T(1) / Len;
+		}
+		else 
+		{
+			std::cout << "zero fiber gradient" << std::endl;
+			Directions[e] = pseudo_direction;
+		}
+	}//);
+}
+
+
+} // namespace ChaosFlesh
\ No newline at end of file
diff --git a/Engine/Source/Runtime/Experimental/Chaos/Public/Chaos/UniformGrid.h b/Engine/Source/Runtime/Experimental/Chaos/Public/Chaos/UniformGrid.h
--- a/Engine/Source/Runtime/Experimental/Chaos/Public/Chaos/UniformGrid.h
+++ b/Engine/Source/Runtime/Experimental/Chaos/Public/Chaos/UniformGrid.h
@@ -35,7 +35,6 @@ class CHAOS_API TUniformGridBase
 			check(MCells[Axis] != 0);
 		}
 
-
 		// Are corners valid?
 		bool bValidBounds = true;
 		if (MMaxCorner == TVector<T, d>(-TNumericLimits<T>::Max()) && MMinCorner == TVector<T, d>(TNumericLimits<T>::Max()))
@@ -170,16 +169,84 @@ class CHAOS_API TUniformGridBase
 	{
 		return (MMaxCorner - MMinCorner);
 	}
+
 	int32 GetNumCells() const
 	{
 		return MCells.Product();
 	}
+
+	int32 GetNumNodes() const
+	{
+		return (MCells + TVector<int32, 3>(1, 1, 1)).Product();
+	}
+
+	TVector<T, d> Node(const TVector<int32, d>& Index) const
+	{
+		TVector<T, d> Tmp = MMinCorner;
+		for (int32 i = 0; i < d; i++)
+			Tmp[i] += MDx[i] * Index[i];
+		return Tmp;
+	}
+
+	TVector<T, d> Node(const int32 FlatIndex) const
+	{
+		TVector<int32, d> Index; 
+		FlatToMultiIndex(FlatIndex, Index, true);
+		return Node(Index);
+	}
+
+	bool InteriorNode(const TVector<int32, d>& Index) const
+	{
+		bool Interior = true;
+		for (int32 i = 0; i < d && Interior; i++)
+		{
+			if (Index[i] <= 0 || Index[i] >= (MCells[i] - 1))
+				Interior = false;
+		}
+		return Interior;
+	}
+
+	int32 FlatIndex(const TVector<int32, 2>& MIndex, const bool NodeIndex=false) const
+	{
+		return MIndex[0] * (NodeIndex ? MCells[1]+1 : MCells[1]) + MIndex[1];
+	}
+
+	int32 FlatIndex(const TVector<int32, 3>& MIndex, const bool NodeIndex=false) const
+	{
+		return NodeIndex ?
+			MIndex[0] * (MCells[1]+1) * (MCells[2]+1) + MIndex[1] * (MCells[2]+1) + MIndex[2] :
+			MIndex[0] * MCells[1] * MCells[2] + MIndex[1] * MCells[2] + MIndex[2];
+	}
+
+	void FlatToMultiIndex(const int32 FlatIndex, TVector<int32, 2>& MIndex, const bool NodeIndex=false) const
+	{
+		MIndex[0] = FlatIndex / (NodeIndex?MCells[1]+1:MCells[1]);
+		MIndex[1] = FlatIndex % (NodeIndex?MCells[1]+1: MCells[1]);
+	}
+
+	void FlatToMultiIndex(const int32 FlatIndex, TVector<int32, 3>& MIndex, const bool NodeIndex=false) const
+	{
+		if (NodeIndex)
+		{
+			MIndex[0] = FlatIndex / ((MCells[1]+1) * (MCells[2]+1));
+			MIndex[1] = (FlatIndex / (MCells[2]+1)) % (MCells[1]+1);
+			MIndex[2] = FlatIndex % (MCells[2]+1);
+		}
+		else
+		{
+			MIndex[0] = FlatIndex / (MCells[1] * MCells[2]);
+			MIndex[1] = (FlatIndex / MCells[2]) % MCells[1];
+			MIndex[2] = FlatIndex % MCells[2];
+		}
+	}
+
 	template<class T_SCALAR>
 	T_SCALAR LinearlyInterpolate(const TArrayND<T_SCALAR, d>& ScalarN, const TVector<T, d>& X) const;
 	T LinearlyInterpolateComponent(const TArrayND<T, d>& ScalarNComponent, const TVector<T, d>& X, const int32 Axis) const;
 	TVector<T, d> LinearlyInterpolate(const TArrayFaceND<T, d>& ScalarN, const TVector<T, d>& X) const;
 	TVector<T, d> LinearlyInterpolate(const TArrayFaceND<T, d>& ScalarN, const TVector<T, d>& X, const Pair<int32, TVector<int32, d>> Index) const;
 	const TVector<int32, d>& Counts() const { return MCells; }
+	const TVector<int32, d> NodeCounts() const { return MCells + TVector<int32, d>(1); }
 	const TVector<T, d>& Dx() const { return MDx; }
 	const TVector<T, d>& MinCorner() const { return MMinCorner; }
 	const TVector<T, d>& MaxCorner() const { return MMaxCorner; }
diff --git a/Engine/Source/Runtime/Experimental/Chaos/Public/Chaos/Utilities.h b/Engine/Source/Runtime/Experimental/Chaos/Public/Chaos/Utilities.h
--- a/Engine/Source/Runtime/Experimental/Chaos/Public/Chaos/Utilities.h
+++ b/Engine/Source/Runtime/Experimental/Chaos/Public/Chaos/Utilities.h
@@ -5,6 +5,7 @@
 #include "Chaos/Core.h"
 #include "Chaos/Matrix.h"
 #include "Chaos/Transform.h"
+#include "Chaos/UniformGrid.h"
 #include "Chaos/Vector.h"
 
 #include "Math/NumericLimits.h"
@@ -445,6 +446,17 @@ namespace Chaos
 			return false;
 		}
 
+		template <class T>
+		inline T DotProduct(const TArray<T>& X, const TArray<T>& Y)
+		{
+			T Result = T(0);
+			for (int32 i = 0; i < X.Num(); i++)
+			{
+				Result += X[i] * Y[i];
+			}
+			return Result;
+		}
+
 		/**
 		 * Given the local-space diagonal inertia for an unscaled object, return an inertia as if generated from a non-uniformly scaled shape with the specified scale.
 		 * If bScaleMass is true, it also takes into account the fact that the mass would have changed by the increase in volume.
@@ -676,6 +688,133 @@ namespace Chaos
 			}
 		}
 
+		template <class T>
+		void TetMeshFromGrid(const TUniformGrid<T, 3>& Grid, TArray<TVector<int32, 4>>& Mesh, TArray<TVector<T, 3>>& X)
+		{
+			//mesh.resize(20 * (grid.grid_N[0] - 1) * (grid.grid_N[1] - 1) * (grid.grid_N[2] - 1));
+			Mesh.SetNum(
+				//20 * (Grid.Counts()[0] - 1) * (Grid.Counts()[1] - 1) * (Grid.Counts()[2] - 1) / 4);
+				20 * Grid.GetNumCells() / 4);
+			//TVector<int32,4>(INDEX_NONE,INDEX_NONE,INDEX_NONE,INDEX_NONE));
+			int32* MeshPtr = &Mesh[0][0];
+
+			const int32 NumNodes = Grid.GetNumNodes();
+			X.SetNum(NumNodes);
+			for(int32 ii=0; ii < NumNodes; ii++)
+			{
+				X[ii] = Grid.Node(ii);
+			}
+
+			int32 Count = 0;
+			for (int32 i = 0; i < Grid.Counts()[0]; i++) 
+			{
+				for (int32 j = 0; j < Grid.Counts()[1]; j++) 
+				{
+					for (int32 k = 0; k < Grid.Counts()[2]; k++) 
+					{
+						int32 ijk000 = Grid.FlatIndex(TVector<int32, 3>(i, j, k), true);
+						int32 ijk010 = Grid.FlatIndex(TVector<int32, 3>(i, j + 1, k), true);
+						int32 ijk001 = Grid.FlatIndex(TVector<int32, 3>(i, j, k + 1), true);
+						int32 ijk011 = Grid.FlatIndex(TVector<int32, 3>(i, j + 1, k + 1), true);
+						int32 ijk100 = Grid.FlatIndex(TVector<int32, 3>(i + 1, j, k), true);
+						int32 ijk110 = Grid.FlatIndex(TVector<int32, 3>(i + 1, j + 1, k), true);
+						int32 ijk101 = Grid.FlatIndex(TVector<int32, 3>(i + 1, j, k + 1), true);
+						int32 ijk111 = Grid.FlatIndex(TVector<int32, 3>(i + 1, j + 1, k + 1), true);
+						int32 ijk_index = i + j + k;
+						if (ijk_index % 2 == 0) 
+						{
+							MeshPtr[20 * Count] = ijk010;
+							MeshPtr[20 * Count + 1] = ijk000;
+							MeshPtr[20 * Count + 2] = ijk110;
+							MeshPtr[20 * Count + 3] = ijk011;
+
+							MeshPtr[20 * Count + 4] = ijk111;
+							MeshPtr[20 * Count + 5] = ijk110;
+							MeshPtr[20 * Count + 6] = ijk101;
+							MeshPtr[20 * Count + 7] = ijk011;
+
+							MeshPtr[20 * Count + 8] = ijk100;
+							MeshPtr[20 * Count + 9] = ijk101;
+							MeshPtr[20 * Count + 10] = ijk110;
+							MeshPtr[20 * Count + 11] = ijk000;
+
+							MeshPtr[20 * Count + 12] = ijk001;
+							MeshPtr[20 * Count + 13] = ijk000;
+							MeshPtr[20 * Count + 14] = ijk011;
+							MeshPtr[20 * Count + 15] = ijk101;
+
+							MeshPtr[20 * Count + 16] = ijk110;
+							MeshPtr[20 * Count + 17] = ijk011;
+							MeshPtr[20 * Count + 18] = ijk000;
+							MeshPtr[20 * Count + 19] = ijk101;
+						}
+						else 
+						{
+							MeshPtr[20 * Count] = ijk000;
+							MeshPtr[20 * Count + 1] = ijk100;
+							MeshPtr[20 * Count + 2] = ijk010;
+							MeshPtr[20 * Count + 3] = ijk001;
+
+							MeshPtr[20 * Count + 4] = ijk011;
+							MeshPtr[20 * Count + 5] = ijk010;
+							MeshPtr[20 * Count + 6] = ijk111;
+							MeshPtr[20 * Count + 7] = ijk001;
+
+							MeshPtr[20 * Count + 8] = ijk100;
+							MeshPtr[20 * Count + 9] = ijk111;
+							MeshPtr[20 * Count + 10] = ijk010;
+							MeshPtr[20 * Count + 11] = ijk001;
+
+							MeshPtr[20 * Count + 12] = ijk101;
+							MeshPtr[20 * Count + 13] = ijk111;
+							MeshPtr[20 * Count + 14] = ijk100;
+							MeshPtr[20 * Count + 15] = ijk001;
+
+							MeshPtr[20 * Count + 16] = ijk111;
+							MeshPtr[20 * Count + 17] = ijk010;
+							MeshPtr[20 * Count + 18] = ijk110;
+							MeshPtr[20 * Count + 19] = ijk100;
+						}
+						Count++;
+					}
+				}
+			}
+		}
+
+		template <int d>
+		TArray<TArray<int>> ComputeIncidentElements(const TArray<TVector<int32, d>>& Mesh, TArray<TArray<int32>>* LocalIndex=nullptr)
+		{
+			int32 MaxIdx = 0;
+			for(int32 i=0; i < Mesh.Num(); i++)
+			{
+				for (int32 j = 0; j < d; j++)
+				{
+					const int32 NodeIdx = Mesh[i][j];
+					MaxIdx = MaxIdx > NodeIdx ? MaxIdx : NodeIdx;
+				}
+			}
+
+			TArray<TArray<int>> IncidentElements;
+			IncidentElements.SetNum(MaxIdx + 1);
+			if (LocalIndex)
+				LocalIndex->SetNum(MaxIdx + 1);
+
+			for (int32 i = 0; i < Mesh.Num(); i++)
+			{
+				for (int32 j = 0; j < d; j++)
+				{
+					const int32 NodeIdx = Mesh[i][j];
+					if (NodeIdx >= 0)
+					{
+						IncidentElements[NodeIdx].Add(i);
+						if (LocalIndex)
+							(*LocalIndex)[NodeIdx].Add(j);
+					}
+				}
+			}
+
+			return IncidentElements;
+		}
 
 	} // namespace Utilities
 } // namespace Chaos
\ No newline at end of file
