Chaos destruction : add per component control of removal ( sleep and break )

diff --git a/Engine/Source/Runtime/Experimental/GeometryCollectionEngine/Private/GeometryCollection/GeometryCollectionComponent.cpp b/Engine/Source/Runtime/Experimental/GeometryCollectionEngine/Private/GeometryCollection/GeometryCollectionComponent.cpp
--- a/Engine/Source/Runtime/Experimental/GeometryCollectionEngine/Private/GeometryCollection/GeometryCollectionComponent.cpp
+++ b/Engine/Source/Runtime/Experimental/GeometryCollectionEngine/Private/GeometryCollection/GeometryCollectionComponent.cpp
@@ -227,6 +227,8 @@ UGeometryCollectionComponent::UGeometryCollectionComponent(const FObjectInitiali
 	, MaxClusterLevel(100)
 	, DamageThreshold({ 500000.f, 50000.f, 5000.f })
 	, bUseSizeSpecificDamageThreshold(false)
+	, bAllowRemovalOnSleep(true)
+	, bAllowRemovalOnBreak(true)
 	, ClusterConnectionType_DEPRECATED(EClusterConnectionTypeEnum::Chaos_MinimalSpanningSubsetDelaunayTriangulation)
 	, CollisionGroup(0)
 	, CollisionSampleFraction(1.0)
@@ -1946,10 +1948,7 @@ void UGeometryCollectionComponent::TickComponent(float DeltaTime, enum ELevelTic
 		// In editor mode we have no DynamicCollection so this test is necessary
 		if(DynamicCollection) //, TEXT("No dynamic collection available for component %s during tick."), *GetName()))
 		{
-			if (RestCollection->bRemoveOnMaxSleep)
-			{ 
-				IncrementSleepTimer(DeltaTime);
-			}
+			IncrementSleepTimer(DeltaTime);
 			IncrementBreakTimer(DeltaTime);
 
 			if (RestCollection->HasVisibleGeometry() || DynamicCollection->IsDirty())
@@ -2046,6 +2045,7 @@ void UGeometryCollectionComponent::ResetDynamicCollection()
 		GetSimulationTypeArrayCopyOnWrite();
 		GetStatusFlagsArrayCopyOnWrite();
 
+		// we are not testing for bAllowRemovalOnSleep, so that we can enable it at runtime if necessary
 		if (RestCollection->bRemoveOnMaxSleep)
 		{
 			if (!DynamicCollection->HasAttribute("SleepTimer", FGeometryCollection::TransformGroup))
@@ -3570,7 +3570,7 @@ void UGeometryCollectionComponent::InitializeEmbeddedGeometry()
 
 void UGeometryCollectionComponent::IncrementSleepTimer(float DeltaTime)
 {
-	if (DeltaTime <= 0)
+	if (DeltaTime <= 0 || !RestCollection->bRemoveOnMaxSleep || !bAllowRemovalOnSleep)
 	{
 		return;
 	}
@@ -3581,7 +3581,8 @@ void UGeometryCollectionComponent::IncrementSleepTimer(float DeltaTime)
 		&& DynamicCollection->HasAttribute("MaxSleepTime", FGeometryCollection::TransformGroup)
 		&& DynamicCollection->HasAttribute("SleepRemovalDuration", FGeometryCollection::TransformGroup)
 		&& DynamicCollection->HasAttribute("Decay", FGeometryCollection::TransformGroup)
-		&& DynamicCollection->HasAttribute("LastPosition", FGeometryCollection::TransformGroup))
+		&& DynamicCollection->HasAttribute("LastPosition", FGeometryCollection::TransformGroup)
+		&& DynamicCollection->HasAttribute("InternalClusterParentTypeArray", FGeometryCollection::TransformGroup))
 	{
 		TManagedArray<float>& SleepTimer = DynamicCollection->ModifyAttribute<float>("SleepTimer", FGeometryCollection::TransformGroup);
 		const TManagedArray<float>& SleepRemovalDuration = DynamicCollection->GetAttribute<float>("SleepRemovalDuration", FGeometryCollection::TransformGroup);
@@ -3657,7 +3658,7 @@ void UGeometryCollectionComponent::IncrementSleepTimer(float DeltaTime)
 
 void UGeometryCollectionComponent::IncrementBreakTimer(float DeltaTime)
 {
-	if (DeltaTime <= 0)
+	if (DeltaTime <= 0 || !bAllowRemovalOnBreak)
 	{
 		return;
 	}
diff --git a/Engine/Source/Runtime/Experimental/GeometryCollectionEngine/Public/GeometryCollection/GeometryCollectionComponent.h b/Engine/Source/Runtime/Experimental/GeometryCollectionEngine/Public/GeometryCollection/GeometryCollectionComponent.h
--- a/Engine/Source/Runtime/Experimental/GeometryCollectionEngine/Public/GeometryCollection/GeometryCollectionComponent.h
+++ b/Engine/Source/Runtime/Experimental/GeometryCollectionEngine/Public/GeometryCollection/GeometryCollectionComponent.h
@@ -592,6 +592,14 @@ public:
 	UPROPERTY(EditAnywhere, BlueprintReadWrite, Category = "ChaosPhysics|Clustering")
 	bool bUseSizeSpecificDamageThreshold;
 
+	/** Allow removal on sleep for the instance if the rest collection has it enabled */
+	UPROPERTY(EditAnywhere, BlueprintReadWrite, Category = "ChaosPhysics|Removal")
+	bool bAllowRemovalOnSleep;
+
+	/** Allow removal on break for the instance if the rest collection has it enabled */
+	UPROPERTY(EditAnywhere, BlueprintReadWrite, Category = "ChaosPhysics|Removal")
+	bool bAllowRemovalOnBreak;
+	
 	/** */
 	UPROPERTY(meta = (DeprecatedProperty, DeprecationMessage = "Connection types are defined on the asset now."))
 	EClusterConnectionTypeEnum ClusterConnectionType_DEPRECATED;
