Dataflow : API Clean up.

diff --git a/Engine/Plugins/Experimental/Dataflow/Source/DataflowNodes/Private/Dataflow/DataflowSkeletalMeshNodes.cpp b/Engine/Plugins/Experimental/Dataflow/Source/DataflowNodes/Private/Dataflow/DataflowSkeletalMeshNodes.cpp
--- a/Engine/Plugins/Experimental/Dataflow/Source/DataflowNodes/Private/Dataflow/DataflowSkeletalMeshNodes.cpp
+++ b/Engine/Plugins/Experimental/Dataflow/Source/DataflowNodes/Private/Dataflow/DataflowSkeletalMeshNodes.cpp
@@ -47,7 +47,7 @@ void FSkeletalMeshBoneDataflowNode::Evaluate(Dataflow::FContext& Context, const
 	{
 		SetValue<int>(Context, INDEX_NONE, &BoneIndexOut);
 
-		if( InDataType InSkeletalMesh = GetInput(&SkeletalMesh)->GetValue<InDataType>(Context, SkeletalMesh) )
+		if( InDataType InSkeletalMesh = GetValue<InDataType>(Context, &SkeletalMesh) )
 		{
 			int32 Index = InSkeletalMesh->GetRefSkeleton().FindBoneIndex(BoneName);
 			SetValue<int>(Context, Index, &BoneIndexOut);
diff --git a/Engine/Source/Runtime/Experimental/Dataflow/Core/Private/Dataflow/DataflowNode.cpp b/Engine/Source/Runtime/Experimental/Dataflow/Core/Private/Dataflow/DataflowNode.cpp
--- a/Engine/Source/Runtime/Experimental/Dataflow/Core/Private/Dataflow/DataflowNode.cpp
+++ b/Engine/Source/Runtime/Experimental/Dataflow/Core/Private/Dataflow/DataflowNode.cpp
@@ -72,21 +72,6 @@ FDataflowInput* FDataflowNode::FindInput(void* Reference)
 	return nullptr;
 }
 
-const FDataflowInput* FDataflowNode::GetInput(const void* Reference) const
-{
-	for (TPair<uint32, FDataflowConnection*> Elem : Inputs)
-	{
-		FDataflowConnection* Con = Elem.Value;
-		size_t AddressA = Con->RealAddress();
-		size_t AddressB = (size_t)Reference;
-		if (Con->RealAddress() == (size_t)Reference)
-		{
-			return (FDataflowInput*)Con;
-		}
-	}
-	return &FDataflowInput::NoOpInput;
-}
-
 TArray< FDataflowConnection* > FDataflowNode::GetInputs() const
 {
 	TArray< FDataflowConnection* > Result;
@@ -172,19 +157,6 @@ FDataflowOutput* FDataflowNode::FindOutput(void* Reference)
 	return nullptr;
 }
 
-const FDataflowOutput* FDataflowNode::GetOutput(const void* Reference) const
-{
-	for (TPair<uint32, FDataflowConnection*> Elem : Outputs)
-	{
-		FDataflowConnection* Con = Elem.Value;
-		if (Con->RealAddress() == (size_t)Reference)
-		{
-			return (FDataflowOutput*)Con;
-		}
-	}
-	return &FDataflowOutput::NoOpOutput;
-}
-
 TArray< FDataflowConnection* > FDataflowNode::GetOutputs() const
 {
 	TArray< FDataflowConnection* > Result;
diff --git a/Engine/Source/Runtime/Experimental/Dataflow/Core/Public/Dataflow/DataflowNode.h b/Engine/Source/Runtime/Experimental/Dataflow/Core/Public/Dataflow/DataflowNode.h
--- a/Engine/Source/Runtime/Experimental/Dataflow/Core/Public/Dataflow/DataflowNode.h
+++ b/Engine/Source/Runtime/Experimental/Dataflow/Core/Public/Dataflow/DataflowNode.h
@@ -78,21 +78,21 @@ struct DATAFLOWCORE_API FDataflowNode
 
 
 	void AddInput(FDataflowConnection* InPtr);
+	TArray< FDataflowConnection* > GetInputs() const;
+	void ClearInputs();
 
 	FDataflowInput* FindInput(FName Name);
 	FDataflowInput* FindInput(void* Reference);
 	const FDataflowInput* FindInput(const void* Reference) const;
-	const FDataflowInput* GetInput(const void* Reference) const;
-	TArray< FDataflowConnection* > GetInputs() const;
-	void ClearInputs();
+
 
 	void AddOutput(FDataflowConnection* InPtr);
+	TArray< FDataflowConnection* > GetOutputs() const;
+	void ClearOutputs();
+
 	FDataflowOutput* FindOutput(FName Name);
 	FDataflowOutput* FindOutput(void* Reference);
 	const FDataflowOutput* FindOutput(const void* Reference) const;
-	const FDataflowOutput* GetOutput(const void* Reference) const;
-	TArray< FDataflowConnection* > GetOutputs() const;
-	void ClearOutputs();
 
 
 
@@ -119,15 +119,32 @@ struct DATAFLOWCORE_API FDataflowNode
 	*   cached in the contexts data store. 
 	* 
 	*   @param Context : The evaluation context that holds the data store.
-	*   @param Value : The value to store in the contexts data store.
 	*   @param Reference : Pointer to a member of this node that corresponds with the output to set.
+	*						*Reference will be used as the default if the input is not connected. 
 	*/
 	template<class T> const T& GetValue(Dataflow::FContext& Context, const T* Reference) const
 	{
-		return GetInput(Reference)->template GetValueAsInput<T>(Context, *Reference);
+		checkSlow(FindInput(Reference));
+		return FindInput(Reference)->template GetValueAsInput<T>(Context, *Reference);
 	}
 
-	/** 
+	/**
+	*   GetValue(...)
+	*
+	*	Get the value of the Reference output, invoking up stream evaluations if not
+	*   cached in the contexts data store.
+	*
+	*   @param Context : The evaluation context that holds the data store.
+	*   @param Reference : Pointer to a member of this node that corresponds with the output to set.
+	*   @param Default : Default value if the input is not connected.
+	*/
+	template<class T> const T& GetValue(Dataflow::FContext& Context, const T* Reference, const T& Default) const
+	{
+		checkSlow(FindInput(Reference));
+		return FindInput(Reference)->template GetValueAsInput<T>(Context, Default);
+	}
+
+	/**
 	*   SetValue(...)
 	*
 	*	Set the value of the Reference output.
@@ -138,7 +155,8 @@ struct DATAFLOWCORE_API FDataflowNode
 	*/
 	template<class T> void SetValue(Dataflow::FContext& Context, const T& Value, const T* Reference) const
 	{
-		GetOutput(Reference)->template SetValue<T>(Value, Context);
+		checkSlow(FindOutput(Reference));
+		FindOutput(Reference)->template SetValue<T>(Value, Context);
 	}
 
 	void InvalidateOutputs();
