change 'remove overlaps' fracture convex generation option to be an enum w/ more options -- specifically to add option to only remove overlaps on clusters, or overlaps of clusters w/ other clusters

diff --git a/Engine/Plugins/Experimental/ChaosEditor/Source/FractureEditor/Private/FractureModeSettings.cpp b/Engine/Plugins/Experimental/ChaosEditor/Source/FractureEditor/Private/FractureModeSettings.cpp
--- a/Engine/Plugins/Experimental/ChaosEditor/Source/FractureEditor/Private/FractureModeSettings.cpp
+++ b/Engine/Plugins/Experimental/ChaosEditor/Source/FractureEditor/Private/FractureModeSettings.cpp
@@ -11,6 +11,6 @@ void UFractureModeSettings::ApplyDefaultConvexSettings(FGeometryCollection& Geom
 	Properties.FractionRemove = ConvexFractionAllowRemove;
 	Properties.SimplificationThreshold = ConvexSimplificationDistanceThreshold;
 	Properties.CanExceedFraction = ConvexCanExceedFraction;
-	Properties.bRemoveOverlaps = bConvexRemoveOverlaps;
+	Properties.RemoveOverlaps = ConvexRemoveOverlaps;
 	GeometryCollection.SetConvexProperties(Properties);
 }
\ No newline at end of file
diff --git a/Engine/Plugins/Experimental/ChaosEditor/Source/FractureEditor/Private/FractureToolConvex.cpp b/Engine/Plugins/Experimental/ChaosEditor/Source/FractureEditor/Private/FractureToolConvex.cpp
--- a/Engine/Plugins/Experimental/ChaosEditor/Source/FractureEditor/Private/FractureToolConvex.cpp
+++ b/Engine/Plugins/Experimental/ChaosEditor/Source/FractureEditor/Private/FractureToolConvex.cpp
@@ -42,7 +42,7 @@ void UFractureConvexSettings::SaveAsDefaults()
 	ModeSettings->ConvexFractionAllowRemove = FractionAllowRemove;
 	ModeSettings->ConvexCanExceedFraction = CanExceedFraction;
 	ModeSettings->ConvexSimplificationDistanceThreshold = SimplificationDistanceThreshold;
-	ModeSettings->bConvexRemoveOverlaps = bRemoveOverlaps;
+	ModeSettings->ConvexRemoveOverlaps = RemoveOverlaps;
 }
 
 void UFractureConvexSettings::SetFromDefaults()
@@ -51,7 +51,7 @@ void UFractureConvexSettings::SetFromDefaults()
 	FractionAllowRemove = ModeSettings->ConvexFractionAllowRemove;
 	CanExceedFraction = ModeSettings->ConvexCanExceedFraction;
 	SimplificationDistanceThreshold = ModeSettings->ConvexSimplificationDistanceThreshold;
-	bRemoveOverlaps = ModeSettings->bConvexRemoveOverlaps;
+	RemoveOverlaps = ModeSettings->ConvexRemoveOverlaps;
 
 	UFractureToolConvex* ConvexTool = Cast<UFractureToolConvex>(OwnerTool.Get());
 	ConvexTool->NotifyOfPropertyChangeByTool(this);
@@ -316,7 +316,7 @@ int32 UFractureToolConvex::ExecuteFracture(const FFractureToolContext& FractureC
 		Properties.FractionRemove = ConvexSettings->FractionAllowRemove;
 		Properties.SimplificationThreshold = ConvexSettings->SimplificationDistanceThreshold;
 		Properties.CanExceedFraction = ConvexSettings->CanExceedFraction;
-		Properties.bRemoveOverlaps = ConvexSettings->bRemoveOverlaps;
+		Properties.RemoveOverlaps = ConvexSettings->RemoveOverlaps;
 		FractureContext.GetGeometryCollection()->SetConvexProperties(Properties);
 	}
 
diff --git a/Engine/Plugins/Experimental/ChaosEditor/Source/FractureEditor/Private/FractureToolConvex.h b/Engine/Plugins/Experimental/ChaosEditor/Source/FractureEditor/Private/FractureToolConvex.h
--- a/Engine/Plugins/Experimental/ChaosEditor/Source/FractureEditor/Private/FractureToolConvex.h
+++ b/Engine/Plugins/Experimental/ChaosEditor/Source/FractureEditor/Private/FractureToolConvex.h
@@ -4,6 +4,8 @@
 
 #include "FractureTool.h"
 
+#include "GeometryCollection/GeometryCollectionConvexUtility.h"
+
 #include "FractureToolConvex.generated.h"
 
 class FFractureToolContext;
@@ -34,7 +36,7 @@ public:
 
 	/** Whether to automatically cut away overlapping parts of the convex hulls, to avoid the simulation 'popping' to fix the overlaps */
 	UPROPERTY(EditAnywhere, Category = Automatic)
-	bool bRemoveOverlaps = true;
+	EConvexOverlapRemoval RemoveOverlaps = EConvexOverlapRemoval::All;
 
 	/** Delete convex hulls from selected clusters.  Does not affect hulls on leaves. */
 	UFUNCTION(CallInEditor, Category = Custom, meta = (DisplayName = "Delete From Selected"))
diff --git a/Engine/Plugins/Experimental/ChaosEditor/Source/FractureEditor/Public/FractureModeSettings.h b/Engine/Plugins/Experimental/ChaosEditor/Source/FractureEditor/Public/FractureModeSettings.h
--- a/Engine/Plugins/Experimental/ChaosEditor/Source/FractureEditor/Public/FractureModeSettings.h
+++ b/Engine/Plugins/Experimental/ChaosEditor/Source/FractureEditor/Public/FractureModeSettings.h
@@ -4,6 +4,9 @@
 
 #include "Engine/DeveloperSettings.h"
 #include "Engine/EngineTypes.h"
+
+#include "GeometryCollection/GeometryCollectionConvexUtility.h"
+
 #include "FractureModeSettings.generated.h"
 
 class FGeometryCollection;
@@ -43,7 +46,7 @@ public:
 
 	/** Default overlap removal setting for convex hulls of new geometry collections */
 	UPROPERTY(config, EditAnywhere, Category = "Fracture Mode|Convex Generation Defaults")
-	bool bConvexRemoveOverlaps = true;
+	EConvexOverlapRemoval ConvexRemoveOverlaps = EConvexOverlapRemoval::All;
 
 	// Apply Convex Generation Defaults to a GeometryCollection
 	void ApplyDefaultConvexSettings(FGeometryCollection& GeometryCollection) const;
diff --git a/Engine/Source/Runtime/Experimental/Chaos/Private/GeometryCollection/GeometryCollectionConvexPropertiesInterface.cpp b/Engine/Source/Runtime/Experimental/Chaos/Private/GeometryCollection/GeometryCollectionConvexPropertiesInterface.cpp
--- a/Engine/Source/Runtime/Experimental/Chaos/Private/GeometryCollection/GeometryCollectionConvexPropertiesInterface.cpp
+++ b/Engine/Source/Runtime/Experimental/Chaos/Private/GeometryCollection/GeometryCollectionConvexPropertiesInterface.cpp
@@ -9,7 +9,7 @@ const FName FGeometryCollectionConvexPropertiesInterface::ConvexEnable = "Enable
 const FName FGeometryCollectionConvexPropertiesInterface::ConvexFractionRemoveAttribute = "FractionRemove";
 const FName FGeometryCollectionConvexPropertiesInterface::ConvexSimplificationThresholdAttribute = "SimplificationThreshold";
 const FName FGeometryCollectionConvexPropertiesInterface::ConvexCanExceedFractionAttribute = "CanExceedFraction";
-const FName FGeometryCollectionConvexPropertiesInterface::ConvexRemoveOverlapsAttribute = "RemoveOverlaps";
+const FName FGeometryCollectionConvexPropertiesInterface::ConvexRemoveOverlapsMethodAttribute = "RemoveOverlapsMethod";
 
 FGeometryCollectionConvexPropertiesInterface::FGeometryCollectionConvexPropertiesInterface(FGeometryCollection* InGeometryCollection)
 	: FManagedArrayInterface(InGeometryCollection)
@@ -49,9 +49,9 @@ FGeometryCollectionConvexPropertiesInterface::InitializeInterface()
 		ManagedCollection->AddAttribute<float>(ConvexCanExceedFractionAttribute, ConvexPropertiesGroup);
 	}
 
-	if (!ManagedCollection->HasAttribute(ConvexRemoveOverlapsAttribute, ConvexPropertiesGroup))
+	if (!ManagedCollection->HasAttribute(ConvexRemoveOverlapsMethodAttribute, ConvexPropertiesGroup))
 	{
-		ManagedCollection->AddAttribute<bool>(ConvexRemoveOverlapsAttribute, ConvexPropertiesGroup);
+		ManagedCollection->AddAttribute<int32>(ConvexRemoveOverlapsMethodAttribute, ConvexPropertiesGroup);
 	}
 
 	SetDefaultProperty();
@@ -80,7 +80,7 @@ FGeometryCollectionConvexPropertiesInterface::GetDefaultProperty() const
 		const TManagedArray<float>& FractionRemove = ManagedCollection->GetAttribute<float>(ConvexFractionRemoveAttribute, ConvexPropertiesGroup);
 		const TManagedArray<float>& SimplificationThreshold = ManagedCollection->GetAttribute<float>(ConvexSimplificationThresholdAttribute, ConvexPropertiesGroup);
 		const TManagedArray<float>& CanExceedFraction = ManagedCollection->GetAttribute<float>(ConvexCanExceedFractionAttribute, ConvexPropertiesGroup);
-		const TManagedArray<bool>& RemoveOverlaps = ManagedCollection->GetAttribute<bool>(ConvexRemoveOverlapsAttribute, ConvexPropertiesGroup);
+		const TManagedArray<int32>& RemoveOverlaps = ManagedCollection->GetAttribute<int32>(ConvexRemoveOverlapsMethodAttribute, ConvexPropertiesGroup);
 
 		int32 DefaultIndex = Index.Find(INDEX_NONE);
 		if (0 <= DefaultIndex && DefaultIndex < Index.Num())
@@ -89,7 +89,7 @@ FGeometryCollectionConvexPropertiesInterface::GetDefaultProperty() const
 			DefaultProperty.FractionRemove = FractionRemove[DefaultIndex];
 			DefaultProperty.SimplificationThreshold = SimplificationThreshold[DefaultIndex];
 			DefaultProperty.CanExceedFraction = CanExceedFraction[DefaultIndex];
-			DefaultProperty.bRemoveOverlaps = RemoveOverlaps[DefaultIndex];
+			DefaultProperty.RemoveOverlaps = (EConvexOverlapRemoval)RemoveOverlaps[DefaultIndex];
 		}
 	}
 	return DefaultProperty;
@@ -120,13 +120,13 @@ FGeometryCollectionConvexPropertiesInterface::GetConvexProperties(int TransformG
 		const TManagedArray<float>& FractionRemove = ManagedCollection->GetAttribute<float>(ConvexFractionRemoveAttribute, ConvexPropertiesGroup);
 		const TManagedArray<float>& SimplificationThreshold = ManagedCollection->GetAttribute<float>(ConvexSimplificationThresholdAttribute, ConvexPropertiesGroup);
 		const TManagedArray<float>& CanExceedFraction = ManagedCollection->GetAttribute<float>(ConvexCanExceedFractionAttribute, ConvexPropertiesGroup);
-		const TManagedArray<bool>& RemoveOverlaps = ManagedCollection->GetAttribute<bool>(ConvexRemoveOverlapsAttribute, ConvexPropertiesGroup);
+		const TManagedArray<int32>& RemoveOverlaps = ManagedCollection->GetAttribute<int32>(ConvexRemoveOverlapsMethodAttribute, ConvexPropertiesGroup);
 
 		ConvexProperty.Enable = Enable[PropIndex];
 		ConvexProperty.FractionRemove = FractionRemove[PropIndex];
 		ConvexProperty.SimplificationThreshold = SimplificationThreshold[PropIndex];
 		ConvexProperty.CanExceedFraction = CanExceedFraction[PropIndex];
-		ConvexProperty.bRemoveOverlaps = RemoveOverlaps[PropIndex];
+		ConvexProperty.RemoveOverlaps = (EConvexOverlapRemoval)RemoveOverlaps[PropIndex];
 	}
 	return ConvexProperty;
 }
@@ -140,7 +140,7 @@ FGeometryCollectionConvexPropertiesInterface::SetConvexProperties(const FConvexC
 	TManagedArray<float>& FractionRemove = ManagedCollection->ModifyAttribute<float>(ConvexFractionRemoveAttribute, ConvexPropertiesGroup);
 	TManagedArray<float>& SimplificationThreshold = ManagedCollection->ModifyAttribute<float>(ConvexSimplificationThresholdAttribute, ConvexPropertiesGroup);
 	TManagedArray<float>& CanExceedFraction = ManagedCollection->ModifyAttribute<float>(ConvexCanExceedFractionAttribute, ConvexPropertiesGroup);
-	TManagedArray<bool>& RemoveOverlaps = ManagedCollection->ModifyAttribute<bool>(ConvexRemoveOverlapsAttribute, ConvexPropertiesGroup);
+	TManagedArray<int32>& RemoveOverlaps = ManagedCollection->ModifyAttribute<int32>(ConvexRemoveOverlapsMethodAttribute, ConvexPropertiesGroup);
 
 	int32 AttributeIndex = INDEX_NONE;
 	if (!Index.Contains(TransformGroupIndex))
@@ -159,6 +159,6 @@ FGeometryCollectionConvexPropertiesInterface::SetConvexProperties(const FConvexC
 		FractionRemove[AttributeIndex] = InConvexAttributes.FractionRemove;
 		SimplificationThreshold[AttributeIndex] = InConvexAttributes.SimplificationThreshold;
 		CanExceedFraction[AttributeIndex] = InConvexAttributes.CanExceedFraction;
-		RemoveOverlaps[AttributeIndex] = InConvexAttributes.bRemoveOverlaps;
+		RemoveOverlaps[AttributeIndex] = (int32)InConvexAttributes.RemoveOverlaps;
 	}
 }
diff --git a/Engine/Source/Runtime/Experimental/Chaos/Private/GeometryCollection/GeometryCollectionConvexUtility.cpp b/Engine/Source/Runtime/Experimental/Chaos/Private/GeometryCollection/GeometryCollectionConvexUtility.cpp
--- a/Engine/Source/Runtime/Experimental/Chaos/Private/GeometryCollection/GeometryCollectionConvexUtility.cpp
+++ b/Engine/Source/Runtime/Experimental/Chaos/Private/GeometryCollection/GeometryCollectionConvexUtility.cpp
@@ -321,9 +321,12 @@ void CreateNonoverlappingConvexHulls(
 	double FracAllowRemove,
 	double SimplificationDistanceThreshold,
 	double CanExceedFraction,
-	bool bRemoveOverlaps
+	EConvexOverlapRemoval OverlapRemovalMethod
 )
 {
+	bool bRemoveOverlaps = OverlapRemovalMethod != EConvexOverlapRemoval::None;
+	bool bRemoveLeafOverlaps = OverlapRemovalMethod == EConvexOverlapRemoval::All;
+
 	int32 NumBones = TransformToConvexIndices.Num();
 	check(Parents.Num() == NumBones);
 
@@ -546,6 +549,11 @@ void CreateNonoverlappingConvexHulls(
 	Depths.SetNumZeroed(NumBones);
 	for (int HullIdx = 0; HullIdx < NumBones; HullIdx++)
 	{
+		if (SimulationType[HullIdx] == SkipType) // Skip any 'SkipType' elements (generally embedded geometry)
+		{
+			Depths[HullIdx] = -1;
+			continue;
+		}
 		if (Parents[HullIdx] != INDEX_NONE)
 		{
 			if (SimulationType[Parents[HullIdx]] != LeafType)
@@ -567,7 +575,7 @@ void CreateNonoverlappingConvexHulls(
 		Depths[HullIdx] = Depth;
 		MaxDepth = FMath::Max(Depth, MaxDepth);
 
-		if (bRemoveOverlaps && !HasCustomConvexFn(HullIdx) && TransformToConvexIndices[HullIdx].Num() > 0)
+		if (bRemoveOverlaps && bRemoveLeafOverlaps && !HasCustomConvexFn(HullIdx) && TransformToConvexIndices[HullIdx].Num() > 0)
 		{
 			const TSet<int32>& Neighbors = LeafProximity[HullIdx];
 			for (int32 NbrIdx : Neighbors)
@@ -941,17 +949,24 @@ void CreateNonoverlappingConvexHulls(
 					continue;
 				}
 				bool bNbrCustom = HasCustomConvexFn(Nbr);
+
+				// if the neighbor is less deep and not a leaf, skip processing this to favor processing the neighbor instead
+				if (Depths[Bone] > Depths[Nbr] && Children[Nbr].Num() > 0)
+				{
+					continue;
+				}
+				// If we only consider cluster-vs-cluster overlap, and the neighbor is a leaf, do not consider it
+				if (OverlapRemovalMethod == EConvexOverlapRemoval::OnlyClustersVsClusters && Children[Nbr].Num() == 0)
+				{
+					continue;
+				}
+
 				bool bAllOk = true;
 				for (int32 ConvexBone : TransformToConvexIndices[Bone])
 				{
 					for (int32 ConvexNbr : TransformToConvexIndices[Nbr])
 					{
 						bool bOneSidedCut = Depths[Bone] != Depths[Nbr] || Children[Nbr].Num() == 0 || bNbrCustom;
-						// if the neighbor is less deep and not a leaf, skip processing this to favor processing the neighbor instead
-						if (Depths[Bone] > Depths[Nbr] && Children[Nbr].Num() > 0)
-						{
-							continue;
-						}
 
 						bool bWasOk = CutIfOk(bOneSidedCut, ConvexBone, ConvexNbr);
 
@@ -1161,7 +1176,8 @@ double ComputeGeometryVolume(
 
 }
 
-FGeometryCollectionConvexUtility::FGeometryCollectionConvexData FGeometryCollectionConvexUtility::CreateNonOverlappingConvexHullData(FGeometryCollection* GeometryCollection, double FracAllowRemove, double SimplificationDistanceThreshold, double CanExceedFraction, bool bRemoveOverlaps)
+FGeometryCollectionConvexUtility::FGeometryCollectionConvexData FGeometryCollectionConvexUtility::CreateNonOverlappingConvexHullData(
+	FGeometryCollection* GeometryCollection, double FracAllowRemove, double SimplificationDistanceThreshold, double CanExceedFraction, EConvexOverlapRemoval OverlapRemovalMethod)
 {
 	check(GeometryCollection);
 
@@ -1186,7 +1202,8 @@ FGeometryCollectionConvexUtility::FGeometryCollectionConvexData FGeometryCollect
 	GeometryCollectionAlgo::GlobalMatrices(GeometryCollection->Transform, GeometryCollection->Parent, GlobalTransformArray);
 	HullsFromGeometry(*GeometryCollection, GlobalTransformArray, HasCustomConvexFn, Convexes, TransformToConvexIndexArr, GeometryCollection->SimulationType, FGeometryCollection::ESimulationTypes::FST_Rigid, SimplificationDistanceThreshold);
 
-	CreateNonoverlappingConvexHulls(Convexes, TransformToConvexIndexArr, HasCustomConvexFn, GeometryCollection->SimulationType, FGeometryCollection::ESimulationTypes::FST_Rigid, FGeometryCollection::ESimulationTypes::FST_None, GeometryCollection->Parent, GCProximity, GeometryCollection->TransformIndex, Volume, FracAllowRemove, SimplificationDistanceThreshold, CanExceedFraction, bRemoveOverlaps);
+	CreateNonoverlappingConvexHulls(Convexes, TransformToConvexIndexArr, HasCustomConvexFn, GeometryCollection->SimulationType, FGeometryCollection::ESimulationTypes::FST_Rigid, FGeometryCollection::ESimulationTypes::FST_None,
+		GeometryCollection->Parent, GCProximity, GeometryCollection->TransformIndex, Volume, FracAllowRemove, SimplificationDistanceThreshold, CanExceedFraction, OverlapRemovalMethod);
 
 	TransformHullsToLocal(GlobalTransformArray, Convexes, TransformToConvexIndexArr);
 
diff --git a/Engine/Source/Runtime/Experimental/Chaos/Public/GeometryCollection/GeometryCollectionConvexPropertiesInterface.h b/Engine/Source/Runtime/Experimental/Chaos/Public/GeometryCollection/GeometryCollectionConvexPropertiesInterface.h
--- a/Engine/Source/Runtime/Experimental/Chaos/Public/GeometryCollection/GeometryCollectionConvexPropertiesInterface.h
+++ b/Engine/Source/Runtime/Experimental/Chaos/Public/GeometryCollection/GeometryCollectionConvexPropertiesInterface.h
@@ -4,6 +4,7 @@
 
 #include "CoreMinimal.h"
 #include "GeometryCollection/ManagedArrayCollection.h"
+#include "GeometryCollection/GeometryCollectionConvexUtility.h"
 
 class FGeometryCollection;
 
@@ -26,7 +27,7 @@ public :
 	// Attribute
 	static const FName ConvexCanExceedFractionAttribute;
 	// Attribute
-	static const FName ConvexRemoveOverlapsAttribute;
+	static const FName ConvexRemoveOverlapsMethodAttribute;
 
 
 	struct FConvexCreationProperties {
@@ -34,7 +35,7 @@ public :
 		float FractionRemove = 0.5f;
 		float SimplificationThreshold = 10.0f;
 		float CanExceedFraction = 0.5f;
-		bool bRemoveOverlaps = true;
+		EConvexOverlapRemoval RemoveOverlaps = EConvexOverlapRemoval::All;
 	};
 
 	FGeometryCollectionConvexPropertiesInterface(FGeometryCollection* InGeometryCollection);
diff --git a/Engine/Source/Runtime/Experimental/Chaos/Public/GeometryCollection/GeometryCollectionConvexUtility.h b/Engine/Source/Runtime/Experimental/Chaos/Public/GeometryCollection/GeometryCollectionConvexUtility.h
--- a/Engine/Source/Runtime/Experimental/Chaos/Public/GeometryCollection/GeometryCollectionConvexUtility.h
+++ b/Engine/Source/Runtime/Experimental/Chaos/Public/GeometryCollection/GeometryCollectionConvexUtility.h
@@ -12,6 +12,19 @@ namespace Chaos
 	class FConvex;
 }
 
+UENUM()
+enum class EConvexOverlapRemoval : int32
+{
+	// Do not remove overlaps between convex hulls
+	None = 0,
+	// Remove all overlaps between neighboring convex hulls
+	All = 1,
+	// Only remove overlaps on convex hulls of clusters, ignoring leaf-leaf overlaps
+	OnlyClusters = 2,
+	// Only remove overlaps between overlapping clusters, ignoring leaf-leaf and cluster-leaf overlaps
+	OnlyClustersVsClusters = 3
+};
+
 class CHAOS_API FGeometryCollectionConvexUtility
 {
 public:
@@ -39,8 +52,10 @@ public:
 	 @param SimplificationDistanceThreshold		Approximate minimum distance between vertices, below which we remove vertices to generate a simpler convex shape.  If 0.0, no simplification will occur.
 	 @param CanExceedFraction					The fraction by which the convex body volume on a cluster can exceed the volume of the geometry under that cluster (a value of 1 == exceed by 100% == convex hull has 2x the volume of the geometry)
 	 @param bRemoveOverlaps						Whether to cut away overlaps between adjacent convex hulls
+	 @param OverlapRemovalMethod				If bRemoveOverlaps, control which overlaps are removed
 	 */
-	static FGeometryCollectionConvexData CreateNonOverlappingConvexHullData(FGeometryCollection* GeometryCollection, double FractionAllowRemove = .3, double SimplificationDistanceThreshold = 0.0, double CanExceedFraction = .5, bool bRemoveOverlaps = true);
+	static FGeometryCollectionConvexData CreateNonOverlappingConvexHullData(FGeometryCollection* GeometryCollection, double FractionAllowRemove = .3, double SimplificationDistanceThreshold = 0.0, double CanExceedFraction = .5, 
+		EConvexOverlapRemoval OverlapRemovalMethod = EConvexOverlapRemoval::All);
 
 	/** Returns the convex hull of the vertices contained in the specified geometry. */
 	static TUniquePtr<Chaos::FConvex> FindConvexHull(const FGeometryCollection* GeometryCollection, int32 GeometryIndex);
diff --git a/Engine/Source/Runtime/Experimental/GeometryCollectionEngine/Private/GeometryCollection/GeometryCollectionObject.cpp b/Engine/Source/Runtime/Experimental/GeometryCollectionEngine/Private/GeometryCollection/GeometryCollectionObject.cpp
--- a/Engine/Source/Runtime/Experimental/GeometryCollectionEngine/Private/GeometryCollection/GeometryCollectionObject.cpp
+++ b/Engine/Source/Runtime/Experimental/GeometryCollectionEngine/Private/GeometryCollection/GeometryCollectionObject.cpp
@@ -287,7 +287,7 @@ void UGeometryCollection::UpdateConvexGeometry()
 	{
 		FGeometryCollectionConvexPropertiesInterface::FConvexCreationProperties ConvexProperties = GeometryCollection->GetConvexProperties();
 		FGeometryCollectionConvexUtility::CreateNonOverlappingConvexHullData(GeometryCollection.Get(), ConvexProperties.FractionRemove, 
-			ConvexProperties.SimplificationThreshold, ConvexProperties.CanExceedFraction, ConvexProperties.bRemoveOverlaps);
+			ConvexProperties.SimplificationThreshold, ConvexProperties.CanExceedFraction, ConvexProperties.RemoveOverlaps);
 		InvalidateCollection();
 	}
 #endif
