Always defer probe narrow phase to the end.

diff --git a/Engine/Source/Runtime/Experimental/Chaos/Private/Chaos/Collision/ParticlePairMidPhase.cpp b/Engine/Source/Runtime/Experimental/Chaos/Private/Chaos/Collision/ParticlePairMidPhase.cpp
--- a/Engine/Source/Runtime/Experimental/Chaos/Private/Chaos/Collision/ParticlePairMidPhase.cpp
+++ b/Engine/Source/Runtime/Experimental/Chaos/Private/Chaos/Collision/ParticlePairMidPhase.cpp
@@ -135,6 +135,9 @@ namespace Chaos
 		// Do not try to reuse manifold points for capsules or spheres (against anything)
 		// NOTE: This can also be disabled for all shape types by the solver (see GenerateCollisionImpl and the Context)
 		Flags.bEnableManifoldUpdate = !bIsSphere0 && !bIsSphere1 && !bIsCapsule0 && !bIsCapsule1 && !bIsTriangle0 && !bIsTriangle1 && !bIsLevelSet;
+
+		// Mark probe flag now so we know which GenerateCollisions to use
+		Flags.bIsProbe = Shape0->GetIsProbe() || Shape1->GetIsProbe();
 	}
 
 	FSingleShapePairCollisionDetector::~FSingleShapePairCollisionDetector()
@@ -255,6 +258,11 @@ namespace Chaos
 		const FReal Dt,
 		FCollisionContext& Context)
 	{
+		if (Flags.bIsProbe)
+		{
+			return GenerateCollisionProbeImpl(CullDistance, Dt, Context);
+		}
+
 		if (Constraint == nullptr)
 		{
 			// Lazy creation of the constraint. If a shape pair never gets within CullDistance of each
@@ -311,7 +319,7 @@ namespace Chaos
 				// We will be updating the manifold, if only partially, so update the restore comparison transforms
 				Constraint->SetLastShapeWorldTransforms(ShapeWorldTransform0, ShapeWorldTransform1);
 
-				if (!(Context.GetSettings().bDeferNarrowPhase || Constraint->GetIsProbe()))
+				if (!Context.GetSettings().bDeferNarrowPhase)
 				{
 					// Run the narrow phase
 					Collisions::UpdateConstraint(*Constraint.Get(), ShapeWorldTransform0, ShapeWorldTransform1, Dt);
@@ -320,7 +328,7 @@ namespace Chaos
 
 			// If we have a valid contact, add it to the active list
 			// We also add it to the active list if collision detection is deferred (which is if per-iteration collision detection is enabled like with RBAN)
-			if ((Constraint->GetPhi() <= CullDistance) || Context.GetSettings().bDeferNarrowPhase)
+			if (Constraint->GetPhi() <= CullDistance || Context.GetSettings().bDeferNarrowPhase)
 			{
 				if (MidPhase.GetCollisionAllocator().ActivateConstraint(Constraint.Get()))
 				{
@@ -338,6 +346,11 @@ namespace Chaos
 		const FReal Dt,
 		FCollisionContext& Context)
 	{
+		if (Flags.bIsProbe)
+		{
+			return GenerateCollisionProbeImpl(CullDistance, Dt, Context);
+		}
+
 		if (Constraint == nullptr)
 		{
 			// Lazy creation of the constraint. 
@@ -348,11 +361,6 @@ namespace Chaos
 		{
 			PHYSICS_CSV_SCOPED_EXPENSIVE(PhysicsVerbose, NarrowPhase_UpdateConstraintCCD);
 
-			if (Constraint->GetIsProbe())
-			{
-				return 0;
-			}
-
 			const FImplicitObject* Implicit0 = Shape0->GetLeafGeometry();
 			const FImplicitObject* Implicit1 = Shape1->GetLeafGeometry();
 			const FRigidTransform3& ShapeWorldTransform0 = Shape0->GetLeafWorldTransform(Particle0);
@@ -399,6 +407,31 @@ namespace Chaos
 		return 0;
 	}
 
+	int32 FSingleShapePairCollisionDetector::GenerateCollisionProbeImpl(
+		const FReal CullDistance, 
+		const FReal Dt,
+		FCollisionContext& Context)
+	{
+		// Same as regular constraint generation, but always defer narrow phase.
+		// Don't do any initial constraint computations.
+
+		if (Constraint == nullptr)
+		{
+			CreateConstraint(CullDistance, Context);
+		}
+
+		if (Constraint != nullptr)
+		{
+			PHYSICS_CSV_SCOPED_EXPENSIVE(PhysicsVerbose, NarrowPhase_UpdateConstraintProbe);
+
+			MidPhase.GetCollisionAllocator().ActivateConstraint(Constraint.Get());
+			LastUsedEpoch = LastUsedEpoch = MidPhase.GetCollisionAllocator().GetCurrentEpoch();
+			return 1;
+		}
+
+		return 0;
+	}
+
 	void FSingleShapePairCollisionDetector::WakeCollision(const int32 SleepEpoch)
 	{
 		if ((Constraint != nullptr) && IsUsedSince(SleepEpoch))
diff --git a/Engine/Source/Runtime/Experimental/Chaos/Public/Chaos/Collision/ParticlePairMidPhase.h b/Engine/Source/Runtime/Experimental/Chaos/Public/Chaos/Collision/ParticlePairMidPhase.h
--- a/Engine/Source/Runtime/Experimental/Chaos/Public/Chaos/Collision/ParticlePairMidPhase.h
+++ b/Engine/Source/Runtime/Experimental/Chaos/Public/Chaos/Collision/ParticlePairMidPhase.h
@@ -103,6 +103,11 @@ namespace Chaos
 			const FReal Dt,
 			FCollisionContext& Context);
 
+		int32 GenerateCollisionProbeImpl(
+			const FReal CullDistance, 
+			const FReal Dt,
+			FCollisionContext& Context);
+
 		/**
 		 * @brief Whether the two shapes are separated by less than CullDistance (i.e., we should run the narrow phase).
 		 * Also returns true if bounds checking is disabled (globally or for this pair)
@@ -132,6 +137,7 @@ namespace Chaos
 				uint32 bEnableOBBCheck0 : 1;
 				uint32 bEnableOBBCheck1 : 1;
 				uint32 bEnableManifoldUpdate : 1;
+				uint32 bIsProbe : 1;
 			};
 			uint32 Bits;
 		} Flags;
