AABBTree optimization:
Don't compute raycast position if not required

diff --git a/Engine/Source/Runtime/Experimental/Chaos/Public/Chaos/AABBTree.h b/Engine/Source/Runtime/Experimental/Chaos/Public/Chaos/AABBTree.h
--- a/Engine/Source/Runtime/Experimental/Chaos/Public/Chaos/AABBTree.h
+++ b/Engine/Source/Runtime/Experimental/Chaos/Public/Chaos/AABBTree.h
@@ -147,7 +147,7 @@ static void UpdateElementHelper(T& InElem, const T& InFrom)
 template <typename TQueryFastData, EAABBQueryType Query>
 struct TAABBTreeIntersectionHelper
 {
-	static bool Intersects(const FVec3& Start, TQueryFastData& QueryFastData, FReal& TOI, FVec3& OutPosition,
+	static bool Intersects(const FVec3& Start, TQueryFastData& QueryFastData, FReal& TOI,
 		const FAABB3& Bounds, const FAABB3& QueryBounds, const FVec3& QueryHalfExtents, const FVec3& Dir, const FVec3 InvDir, const bool bParallel[3])
 	{
 		check(false);
@@ -158,10 +158,11 @@ struct TAABBTreeIntersectionHelper
 template<>
 struct TAABBTreeIntersectionHelper<FQueryFastData, EAABBQueryType::Raycast>
 {
-	FORCEINLINE_DEBUGGABLE static bool Intersects(const FVec3& Start, FQueryFastData& QueryFastData, FReal& TOI, FVec3& OutPosition,
+	FORCEINLINE_DEBUGGABLE static bool Intersects(const FVec3& Start, FQueryFastData& QueryFastData, FReal& TOI,
 		const FAABB3& Bounds, const FAABB3& QueryBounds, const FVec3& QueryHalfExtents, const FVec3& Dir, const FVec3 InvDir, const bool bParallel[3])
 	{
-		return Bounds.RaycastFast(Start, Dir, InvDir, bParallel, QueryFastData.CurrentLength, QueryFastData.InvCurrentLength, TOI, OutPosition);
+		FReal TmpExitTime;
+		return Bounds.RaycastFast(Start, Dir, InvDir, bParallel, QueryFastData.CurrentLength, QueryFastData.InvCurrentLength, TOI, TmpExitTime);
 	}
 
 };
@@ -169,11 +170,12 @@ struct TAABBTreeIntersectionHelper<FQueryFastData, EAABBQueryType::Raycast>
 template <>
 struct TAABBTreeIntersectionHelper<FQueryFastData, EAABBQueryType::Sweep>
 {
-	FORCEINLINE_DEBUGGABLE static bool Intersects(const FVec3& Start, FQueryFastData& QueryFastData, FReal& TOI, FVec3& OutPosition,
+	FORCEINLINE_DEBUGGABLE static bool Intersects(const FVec3& Start, FQueryFastData& QueryFastData, FReal& TOI,
 		const FAABB3& Bounds, const FAABB3& QueryBounds, const FVec3& QueryHalfExtents, const FVec3& Dir, const FVec3 InvDir, const bool bParallel[3])
 	{
 		FAABB3 SweepBounds(Bounds.Min() - QueryHalfExtents, Bounds.Max() + QueryHalfExtents);
-		return SweepBounds.RaycastFast(Start, Dir, InvDir, bParallel, QueryFastData.CurrentLength, QueryFastData.InvCurrentLength, TOI, OutPosition);
+		FReal TmpExitTime;
+		return SweepBounds.RaycastFast(Start, Dir, InvDir, bParallel, QueryFastData.CurrentLength, QueryFastData.InvCurrentLength, TOI, TmpExitTime);
 	}
 
 };
@@ -181,7 +183,7 @@ struct TAABBTreeIntersectionHelper<FQueryFastData, EAABBQueryType::Sweep>
 template <>
 struct TAABBTreeIntersectionHelper<FQueryFastDataVoid, EAABBQueryType::Overlap>
 {
-	FORCEINLINE_DEBUGGABLE static bool Intersects(const FVec3& Start, FQueryFastDataVoid& QueryFastData, FReal& TOI, FVec3& OutPosition,
+	FORCEINLINE_DEBUGGABLE static bool Intersects(const FVec3& Start, FQueryFastDataVoid& QueryFastData, FReal& TOI,
 		const FAABB3& Bounds, const FAABB3& QueryBounds, const FVec3& QueryHalfExtents, const FVec3& Dir, const FVec3 InvDir, const bool bParallel[3])
 	{
 		return QueryBounds.Intersects(Bounds);
@@ -315,7 +317,6 @@ struct TAABBTreeLeafArray : public TBoundsWrapperHelper<TPayloadType, T, bComput
 	FORCEINLINE_DEBUGGABLE bool RaycastSweepImp(const TVec3<T>& Start, TQueryFastData& QueryFastData, const TVec3<T>& QueryHalfExtents, TSQVisitor& Visitor, const TVec3<T>& Dir, const TVec3<T> InvDir, const bool bParallel[3]) const
 	{
 		PHYSICS_CSV_CUSTOM_VERY_EXPENSIVE(PhysicsCounters, MaxLeafSize, Elems.Num(), ECsvCustomStatOp::Max);
-		FVec3 TmpPosition;
 		FReal TOI;
 		for (const auto& Elem : Elems)
 		{
@@ -326,7 +327,7 @@ struct TAABBTreeLeafArray : public TBoundsWrapperHelper<TPayloadType, T, bComput
 
 			const FAABB3 InstanceBounds(Elem.Bounds.Min(), Elem.Bounds.Max());
 			if (TAABBTreeIntersectionHelper<TQueryFastData, bSweep ? EAABBQueryType::Sweep :
-				EAABBQueryType::Raycast>::Intersects(Start, QueryFastData, TOI, TmpPosition, InstanceBounds, FAABB3(), QueryHalfExtents, Dir, InvDir, bParallel))
+				EAABBQueryType::Raycast>::Intersects(Start, QueryFastData, TOI, InstanceBounds, FAABB3(), QueryHalfExtents, Dir, InvDir, bParallel))
 			{
 				TSpatialVisitorData<TPayloadType> VisitData(Elem.Payload, true, InstanceBounds);
 				const bool bContinue = (bSweep && Visitor.VisitSweep(VisitData, QueryFastData)) || (!bSweep && Visitor.VisitRaycast(VisitData, QueryFastData));
@@ -2378,7 +2379,6 @@ private:
 #if !WITH_EDITOR
 		//CSV_SCOPED_TIMING_STAT(ChaosPhysicsTimers, AABBTreeQuery)
 #endif
-		FVec3 TmpPosition;
 		FReal TOI = 0;
 		{
 			//QUICK_SCOPE_CYCLE_COUNTER(QueryGlobal);
@@ -2391,7 +2391,7 @@ private:
 				}
 
 				const FAABB3 InstanceBounds(Elem.Bounds.Min(), Elem.Bounds.Max());
-				if(TAABBTreeIntersectionHelper<TQueryFastData,Query>::Intersects(Start, CurData, TOI, TmpPosition, InstanceBounds, QueryBounds, QueryHalfExtents, Dir, InvDir, bParallel))
+				if(TAABBTreeIntersectionHelper<TQueryFastData,Query>::Intersects(Start, CurData, TOI, InstanceBounds, QueryBounds, QueryHalfExtents, Dir, InvDir, bParallel))
 				{
 					TSpatialVisitorData<TPayloadType> VisitData(Elem.Payload,true, InstanceBounds);
 					bool bContinue;
@@ -2421,7 +2421,7 @@ private:
 					return true;
 				}
 
-				if (TAABBTreeIntersectionHelper<TQueryFastData, Query>::Intersects(Start, CurData, TOI, TmpPosition, InstanceBounds, QueryBounds, QueryHalfExtents, Dir, InvDir, bParallel))
+				if (TAABBTreeIntersectionHelper<TQueryFastData, Query>::Intersects(Start, CurData, TOI, InstanceBounds, QueryBounds, QueryHalfExtents, Dir, InvDir, bParallel))
 				{
 					TSpatialVisitorData<TPayloadType> VisitData(Elem.Payload, true, InstanceBounds);
 					bool bContinue;
@@ -2610,7 +2610,7 @@ private:
 				int32 Idx = 0;
 				for (const TAABB<T, 3>& AABB : Node.ChildrenBounds)
 				{
-					if(TAABBTreeIntersectionHelper<TQueryFastData, Query>::Intersects(Start, CurData, TOI, TmpPosition, FAABB3(AABB.Min(), AABB.Max()), QueryBounds, QueryHalfExtents, Dir, InvDir, bParallel))
+					if(TAABBTreeIntersectionHelper<TQueryFastData, Query>::Intersects(Start, CurData, TOI, FAABB3(AABB.Min(), AABB.Max()), QueryBounds, QueryHalfExtents, Dir, InvDir, bParallel))
 					{
 						NodeStack.Add(FNodeQueueEntry{ Node.ChildrenNodes[Idx], TOI });
 					}
