Refactor GeometryCollectionEngineConversion functions to optionally have core FGeometryCollection conversion called separately from material-related conversion that requires the UObject/Component

diff --git a/Engine/Source/Runtime/Experimental/GeometryCollectionEngine/Private/GeometryCollection/GeometryCollectionEngineConversion.cpp b/Engine/Source/Runtime/Experimental/GeometryCollectionEngine/Private/GeometryCollection/GeometryCollectionEngineConversion.cpp
--- a/Engine/Source/Runtime/Experimental/GeometryCollectionEngine/Private/GeometryCollection/GeometryCollectionEngineConversion.cpp
+++ b/Engine/Source/Runtime/Experimental/GeometryCollectionEngine/Private/GeometryCollection/GeometryCollectionEngineConversion.cpp
@@ -24,6 +24,7 @@
 #include "StaticMeshAttributes.h"
 #include "StaticMeshOperations.h"
 #include "Physics/Experimental/ChaosInterfaceUtils.h"
+#include "PhysicsEngine/BodySetup.h"
 
 
 DEFINE_LOG_CATEGORY_STATIC(UGeometryCollectionConversionLogging, Log, All);
@@ -71,333 +72,377 @@ static bool IsImportableImplicitObjectType(Chaos::EImplicitObjectType Type)
 	return (InnerType == Chaos::ImplicitObjectType::Box || InnerType == Chaos::ImplicitObjectType::Sphere || InnerType == Chaos::ImplicitObjectType::Capsule || InnerType == Chaos::ImplicitObjectType::Convex);
 }
 
-void FGeometryCollectionEngineConversion::AppendStaticMesh(const UStaticMesh* StaticMesh, const TArray<UMaterialInterface*>& Materials, const FTransform& StaticMeshTransform, UGeometryCollection* GeometryCollectionObject, bool ReindexMaterials)
+void FGeometryCollectionEngineConversion::AppendMeshDescription(
+	const FMeshDescription* MeshDescription, const FString& Name, int32 MaterialStartIndex, const FTransform& StaticMeshTransform, 
+	FGeometryCollection* GeometryCollection, UBodySetup* BodySetup, bool ReindexMaterials)
 {
 #if WITH_EDITORONLY_DATA
 
-	if (StaticMesh == nullptr)
+	if (!MeshDescription)
 	{
 		return;
 	}
 
-	// Prefer the HiRes description, although this isn't always available.
-	FMeshDescription* MeshDescription = nullptr;
-	if (StaticMesh->IsHiResMeshDescriptionValid())
-	{
-		MeshDescription = StaticMesh->GetHiResMeshDescription();
-	}
-	else
-	{
-		MeshDescription = StaticMesh->GetMeshDescription(0);
-	}
-
-	check(GeometryCollectionObject);
-	TSharedPtr<FGeometryCollection, ESPMode::ThreadSafe> GeometryCollectionPtr = GeometryCollectionObject->GetGeometryCollection();
-	FGeometryCollection* GeometryCollection = GeometryCollectionPtr.Get();
 	check(GeometryCollection);
 
-	if (MeshDescription)
+	// source vertex information
+	FStaticMeshConstAttributes Attributes(*MeshDescription);
+	TArrayView<const FVector3f> SourcePosition = Attributes.GetVertexPositions().GetRawArray();
+	TArrayView<const FVector3f> SourceTangent = Attributes.GetVertexInstanceTangents().GetRawArray();
+	TArrayView<const float> SourceBinormalSign = Attributes.GetVertexInstanceBinormalSigns().GetRawArray();
+	TArrayView<const FVector3f> SourceNormal = Attributes.GetVertexInstanceNormals().GetRawArray();
+	TArrayView<const FVector4f> SourceColor = Attributes.GetVertexInstanceColors().GetRawArray();
+
+	TVertexInstanceAttributesConstRef<FVector2f> InstanceUVs = Attributes.GetVertexInstanceUVs();
+	const int32 NumUVLayers = InstanceUVs.GetNumChannels();
+	TArray<TArrayView<const FVector2f>> SourceUVArrays;
+	SourceUVArrays.SetNum(NumUVLayers);
+	for (int32 UVLayerIdx = 0; UVLayerIdx < NumUVLayers; ++UVLayerIdx)
 	{
-		FStaticMeshOperations::ComputeTriangleTangentsAndNormals(*MeshDescription);
-		FStaticMeshOperations::RecomputeNormalsAndTangentsIfNeeded(*MeshDescription, EComputeNTBsFlags::UseMikkTSpace);
-
-		// source vertex information
-		FStaticMeshAttributes Attributes(*MeshDescription);
-		TArrayView<const FVector3f> SourcePosition = Attributes.GetVertexPositions().GetRawArray();
-		TArrayView<const FVector3f> SourceTangent = Attributes.GetVertexInstanceTangents().GetRawArray();
-		TArrayView<const float> SourceBinormalSign = Attributes.GetVertexInstanceBinormalSigns().GetRawArray();
-		TArrayView<const FVector3f> SourceNormal = Attributes.GetVertexInstanceNormals().GetRawArray();
-		TArrayView<const FVector4f> SourceColor = Attributes.GetVertexInstanceColors().GetRawArray();
-
-		TVertexInstanceAttributesConstRef<FVector2f> InstanceUVs = Attributes.GetVertexInstanceUVs();
-		const int32 NumUVLayers = InstanceUVs.GetNumChannels();
-		TArray<TArrayView<const FVector2f>> SourceUVArrays;
-		SourceUVArrays.SetNum(NumUVLayers);
-		for (int32 UVLayerIdx = 0; UVLayerIdx < NumUVLayers; ++UVLayerIdx)
-		{
-			SourceUVArrays[UVLayerIdx] = InstanceUVs.GetRawArray(UVLayerIdx);
-		}
+		SourceUVArrays[UVLayerIdx] = InstanceUVs.GetRawArray(UVLayerIdx);
+	}
 		
-		// target vertex information
-		TManagedArray<FVector3f>& TargetVertex = GeometryCollection->Vertex;
-		TManagedArray<FVector3f>& TargetTangentU = GeometryCollection->TangentU;
-		TManagedArray<FVector3f>& TargetTangentV = GeometryCollection->TangentV;
-		TManagedArray<FVector3f>& TargetNormal = GeometryCollection->Normal;
-		TManagedArray<TArray<FVector2f>>& TargetUVs = GeometryCollection->UVs;
-		TManagedArray<FLinearColor>& TargetColor = GeometryCollection->Color;
-		TManagedArray<int32>& TargetBoneMap = GeometryCollection->BoneMap;
-		TManagedArray<FLinearColor>& TargetBoneColor = GeometryCollection->BoneColor;
-		TManagedArray<FString>& TargetBoneName = GeometryCollection->BoneName;
-
-		const int32 VertexStart = GeometryCollection->NumElements(FGeometryCollection::VerticesGroup);
-		int32 VertexCount = 0;
+	// target vertex information
+	TManagedArray<FVector3f>& TargetVertex = GeometryCollection->Vertex;
+	TManagedArray<FVector3f>& TargetTangentU = GeometryCollection->TangentU;
+	TManagedArray<FVector3f>& TargetTangentV = GeometryCollection->TangentV;
+	TManagedArray<FVector3f>& TargetNormal = GeometryCollection->Normal;
+	TManagedArray<TArray<FVector2f>>& TargetUVs = GeometryCollection->UVs;
+	TManagedArray<FLinearColor>& TargetColor = GeometryCollection->Color;
+	TManagedArray<int32>& TargetBoneMap = GeometryCollection->BoneMap;
+	TManagedArray<FLinearColor>& TargetBoneColor = GeometryCollection->BoneColor;
+	TManagedArray<FString>& TargetBoneName = GeometryCollection->BoneName;
+
+	const int32 VertexStart = GeometryCollection->NumElements(FGeometryCollection::VerticesGroup);
+	int32 VertexCount = 0;
 		
-		FVector Scale = StaticMeshTransform.GetScale3D();
+	FVector Scale = StaticMeshTransform.GetScale3D();
 		
-		// We'll need to re-introduce UV seams, etc. by splitting vertices.
-		// A new mapping of MeshDescription vertex instances to the split vertices is maintained.
-		TMap<FVertexInstanceID, int32> VertexInstanceToGeometryCollectionVertex;
-		VertexInstanceToGeometryCollectionVertex.Reserve(Attributes.GetVertexInstanceNormals().GetNumElements());
+	// We'll need to re-introduce UV seams, etc. by splitting vertices.
+	// A new mapping of MeshDescription vertex instances to the split vertices is maintained.
+	TMap<FVertexInstanceID, int32> VertexInstanceToGeometryCollectionVertex;
+	VertexInstanceToGeometryCollectionVertex.Reserve(Attributes.GetVertexInstanceNormals().GetNumElements());
 		
-		for (const FVertexID VertexIndex : MeshDescription->Vertices().GetElementIDs())
-		{		
-			TArrayView<const FVertexInstanceID> ReferencingVertexInstances = MeshDescription->GetVertexVertexInstanceIDs(VertexIndex);
+	for (const FVertexID VertexIndex : MeshDescription->Vertices().GetElementIDs())
+	{		
+		TArrayView<const FVertexInstanceID> ReferencingVertexInstances = MeshDescription->GetVertexVertexInstanceIDs(VertexIndex);
 
-			// Generate per instance hash of splittable attributes.
-			TMap<FUniqueVertex, TArray<FVertexInstanceID>> SplitVertices;
-			for (const FVertexInstanceID& InstanceID : ReferencingVertexInstances)
+		// Generate per instance hash of splittable attributes.
+		TMap<FUniqueVertex, TArray<FVertexInstanceID>> SplitVertices;
+		for (const FVertexInstanceID& InstanceID : ReferencingVertexInstances)
+		{
+			TArray<FVector2f> SourceUVs;
+			SourceUVs.SetNum(NumUVLayers);
+			for (int32 UVLayerIdx = 0; UVLayerIdx < NumUVLayers; ++UVLayerIdx)
 			{
-				TArray<FVector2f> SourceUVs;
-				SourceUVs.SetNum(NumUVLayers);
-				for (int32 UVLayerIdx = 0; UVLayerIdx < NumUVLayers; ++UVLayerIdx)
-				{
-					SourceUVs[UVLayerIdx] = SourceUVArrays[UVLayerIdx][InstanceID];
-				}
-				
-				FUniqueVertex UniqueVertex{ SourceNormal[InstanceID], SourceTangent[InstanceID], SourceUVs };
-				TArray<FVertexInstanceID>& SplitVertex = SplitVertices.FindOrAdd(UniqueVertex);
-				SplitVertex.Add(InstanceID);
+				SourceUVs[UVLayerIdx] = SourceUVArrays[UVLayerIdx][InstanceID];
 			}
+				
+			FUniqueVertex UniqueVertex{ SourceNormal[InstanceID], SourceTangent[InstanceID], SourceUVs };
+			TArray<FVertexInstanceID>& SplitVertex = SplitVertices.FindOrAdd(UniqueVertex);
+			SplitVertex.Add(InstanceID);
+		}
 
-			int32 CurrentVertex = GeometryCollection->AddElements(SplitVertices.Num(), FGeometryCollection::VerticesGroup);
+		int32 CurrentVertex = GeometryCollection->AddElements(SplitVertices.Num(), FGeometryCollection::VerticesGroup);
 			
-			// Create a new vertex for each split vertex and map the mesh description instance to it.
-			for (const TTuple<FUniqueVertex,TArray<FVertexInstanceID>>& SplitVertex : SplitVertices)
-			{
-				const TArray<FVertexInstanceID>& InstanceIDs = SplitVertex.Value;
-				const FVertexInstanceID& ExemplarInstanceID = InstanceIDs[0];
-
-				TargetVertex[CurrentVertex] = SourcePosition[VertexIndex] * (FVector3f)Scale;
-				TargetBoneMap[CurrentVertex] = GeometryCollection->NumElements(FGeometryCollection::TransformGroup);
-
-				TargetNormal[CurrentVertex] = SourceNormal[ExemplarInstanceID];
-				TargetTangentU[CurrentVertex] = SourceTangent[ExemplarInstanceID];
-				TargetTangentV[CurrentVertex] = (FVector3f)SourceBinormalSign[ExemplarInstanceID] * FVector3f::CrossProduct(TargetNormal[CurrentVertex], TargetTangentU[CurrentVertex]);
-
-				TargetUVs[CurrentVertex] = SplitVertex.Key.UVs;
-
-				if (SourceColor.Num() > 0)
-				{
-					TargetColor[CurrentVertex] = FLinearColor(SourceColor[ExemplarInstanceID]);
-				}
-				else
-				{
-					TargetColor[CurrentVertex] = FLinearColor::White;
-				}
-
-				for (const FVertexInstanceID& InstanceID : InstanceIDs)
-				{
-					VertexInstanceToGeometryCollectionVertex.Add(InstanceID, CurrentVertex);
-				}
-
-				++CurrentVertex;
-				++VertexCount;
-			}
-		}
+		// Create a new vertex for each split vertex and map the mesh description instance to it.
+		for (const TTuple<FUniqueVertex,TArray<FVertexInstanceID>>& SplitVertex : SplitVertices)
+		{
+			const TArray<FVertexInstanceID>& InstanceIDs = SplitVertex.Value;
+			const FVertexInstanceID& ExemplarInstanceID = InstanceIDs[0];
 
+			TargetVertex[CurrentVertex] = SourcePosition[VertexIndex] * (FVector3f)Scale;
+			TargetBoneMap[CurrentVertex] = GeometryCollection->NumElements(FGeometryCollection::TransformGroup);
 
-		// for each material, add a reference in our GeometryCollectionObject
-		const int32 MaterialStart = GeometryCollectionObject->Materials.Num();
-		const int32 NumMeshMaterials = Materials.Num();
-		GeometryCollectionObject->Materials.Reserve(MaterialStart + NumMeshMaterials);
+			TargetNormal[CurrentVertex] = SourceNormal[ExemplarInstanceID];
+			TargetTangentU[CurrentVertex] = SourceTangent[ExemplarInstanceID];
+			TargetTangentV[CurrentVertex] = (FVector3f)SourceBinormalSign[ExemplarInstanceID] * FVector3f::CrossProduct(TargetNormal[CurrentVertex], TargetTangentU[CurrentVertex]);
 
-		for (int32 Index = 0; Index < NumMeshMaterials; ++Index)
-		{
-			UMaterialInterface* CurrMaterial = Materials[Index];
+			TargetUVs[CurrentVertex] = SplitVertex.Key.UVs;
+
+			if (SourceColor.Num() > 0)
+			{
+				TargetColor[CurrentVertex] = FLinearColor(SourceColor[ExemplarInstanceID]);
+			}
+			else
+			{
+				TargetColor[CurrentVertex] = FLinearColor::White;
+			}
 
-			// Possible we have a null entry - replace with default
-			if (CurrMaterial == nullptr)
+			for (const FVertexInstanceID& InstanceID : InstanceIDs)
 			{
-				CurrMaterial = UMaterial::GetDefaultMaterial(MD_Surface);
+				VertexInstanceToGeometryCollectionVertex.Add(InstanceID, CurrentVertex);
 			}
 
-			// We add the material twice, once for interior and again for exterior.
-			GeometryCollectionObject->Materials.Add(CurrMaterial);
-			GeometryCollectionObject->Materials.Add(CurrMaterial);
+			++CurrentVertex;
+			++VertexCount;
 		}
+	}
 
-		// target triangle indices
-		TManagedArray<FIntVector>& TargetIndices = GeometryCollection->Indices;
-		TManagedArray<bool>& TargetVisible = GeometryCollection->Visible;
-		TManagedArray<int32>& TargetMaterialID = GeometryCollection->MaterialID;
-		TManagedArray<int32>& TargetMaterialIndex = GeometryCollection->MaterialIndex;
-
-		const int32 IndicesCount = MeshDescription->Triangles().Num();
-		const int32 InitialNumIndices = GeometryCollection->NumElements(FGeometryCollection::FacesGroup);
-		const int32 IndicesStart = GeometryCollection->AddElements(IndicesCount, FGeometryCollection::FacesGroup);
-		int32 TargetIndex = IndicesStart;
-		for (const int32 TriangleIndex : MeshDescription->Triangles().GetElementIDs())
-		{
-			TArrayView<const FVertexInstanceID> TriangleVertices = MeshDescription->GetTriangleVertexInstances(TriangleIndex);
+	// target triangle indices
+	TManagedArray<FIntVector>& TargetIndices = GeometryCollection->Indices;
+	TManagedArray<bool>& TargetVisible = GeometryCollection->Visible;
+	TManagedArray<int32>& TargetMaterialID = GeometryCollection->MaterialID;
+	TManagedArray<int32>& TargetMaterialIndex = GeometryCollection->MaterialIndex;
 
-			TargetIndices[TargetIndex] = FIntVector(
-				VertexInstanceToGeometryCollectionVertex[TriangleVertices[0]],
-				VertexInstanceToGeometryCollectionVertex[TriangleVertices[1]],
-				VertexInstanceToGeometryCollectionVertex[TriangleVertices[2]]
-			);
+	const int32 IndicesCount = MeshDescription->Triangles().Num();
+	const int32 InitialNumIndices = GeometryCollection->NumElements(FGeometryCollection::FacesGroup);
+	const int32 IndicesStart = GeometryCollection->AddElements(IndicesCount, FGeometryCollection::FacesGroup);
+	int32 TargetIndex = IndicesStart;
+	for (const int32 TriangleIndex : MeshDescription->Triangles().GetElementIDs())
+	{
+		TArrayView<const FVertexInstanceID> TriangleVertices = MeshDescription->GetTriangleVertexInstances(TriangleIndex);
 
-			TargetVisible[TargetIndex] = true;
+		TargetIndices[TargetIndex] = FIntVector(
+			VertexInstanceToGeometryCollectionVertex[TriangleVertices[0]],
+			VertexInstanceToGeometryCollectionVertex[TriangleVertices[1]],
+			VertexInstanceToGeometryCollectionVertex[TriangleVertices[2]]
+		);
 
-			// Materials are ganged in pairs and we want the id to associate with the first of each pair.
-			TargetMaterialID[TargetIndex] = MaterialStart + (MeshDescription->GetTrianglePolygonGroup(TriangleIndex) * 2);
+		TargetVisible[TargetIndex] = true;
 
-			// Is this right?
-			TargetMaterialIndex[TargetIndex] = TargetIndex;
+		// Materials are ganged in pairs and we want the id to associate with the first of each pair.
+		TargetMaterialID[TargetIndex] = MaterialStartIndex + (MeshDescription->GetTrianglePolygonGroup(TriangleIndex) * 2);
 
-			++TargetIndex;
-		}
+		// Is this right?
+		TargetMaterialIndex[TargetIndex] = TargetIndex;
+
+		++TargetIndex;
+	}
 
-		// Geometry transform
-		TManagedArray<FTransform>& Transform = GeometryCollection->Transform;
+	// Geometry transform
+	TManagedArray<FTransform>& Transform = GeometryCollection->Transform;
 
-		int32 TransformIndex1 = GeometryCollection->AddElements(1, FGeometryCollection::TransformGroup);
-		Transform[TransformIndex1] = StaticMeshTransform;
-		Transform[TransformIndex1].SetScale3D(FVector(1.f, 1.f, 1.f));
+	int32 TransformIndex1 = GeometryCollection->AddElements(1, FGeometryCollection::TransformGroup);
+	Transform[TransformIndex1] = StaticMeshTransform;
+	Transform[TransformIndex1].SetScale3D(FVector(1.f, 1.f, 1.f));
 
-		// collisions
-		if (UBodySetup* BodySetup = StaticMesh->GetBodySetup())
+	// collisions
+	if (BodySetup)
+	{
+		TArray<TUniquePtr<Chaos::FImplicitObject>> Geoms;
+		Chaos::FShapesArray Shapes;
+
+		FGeometryAddParams CreateGeometryParams;
+		CreateGeometryParams.bDoubleSided = false;
+		CreateGeometryParams.CollisionData.CollisionFlags.bEnableQueryCollision = true;
+		CreateGeometryParams.CollisionData.CollisionFlags.bEnableSimCollisionComplex = false; // no support for trimesh in destruction
+		CreateGeometryParams.CollisionData.CollisionFlags.bEnableSimCollisionSimple = true;
+		CreateGeometryParams.CollisionTraceType = ECollisionTraceFlag::CTF_UseSimpleAsComplex;
+		CreateGeometryParams.Scale = Scale;
+		CreateGeometryParams.LocalTransform = Chaos::FRigidTransform3::Identity;
+		CreateGeometryParams.WorldTransform = Chaos::FRigidTransform3::Identity;
+		CreateGeometryParams.Geometry = &BodySetup->AggGeom;
+		CreateGeometryParams.ChaosTriMeshes = MakeArrayView(BodySetup->ChaosTriMeshes);
+
+		// todo(chaos) : this currently also create the shape array which is unnecessary ,this could be optimized by having a common function to create only the implicits 
+		ChaosInterface::CreateGeometry(CreateGeometryParams, Geoms, Shapes);
+
+		using FCollisionType = FGeometryDynamicCollection::FSharedImplicit;
+		TManagedArray<FCollisionType>& ExternaCollisions = GeometryCollection->AddAttribute<FCollisionType>("ExternalCollisions", FGeometryCollection::TransformGroup);
+
+		ExternaCollisions[TransformIndex1] = nullptr;
+		for (int32 GeomIndex = 0; GeomIndex < Geoms.Num();)
 		{
-			TArray<TUniquePtr<Chaos::FImplicitObject>> Geoms;
-			Chaos::FShapesArray Shapes;
-
-			FGeometryAddParams CreateGeometryParams;
-			CreateGeometryParams.bDoubleSided = false;
-			CreateGeometryParams.CollisionData.CollisionFlags.bEnableQueryCollision = true;
-			CreateGeometryParams.CollisionData.CollisionFlags.bEnableSimCollisionComplex = false; // no support for trimesh in destruction
-			CreateGeometryParams.CollisionData.CollisionFlags.bEnableSimCollisionSimple = true;
-			CreateGeometryParams.CollisionTraceType = ECollisionTraceFlag::CTF_UseSimpleAsComplex;
-			CreateGeometryParams.Scale = Scale;
-			CreateGeometryParams.LocalTransform = Chaos::FRigidTransform3::Identity;
-			CreateGeometryParams.WorldTransform = Chaos::FRigidTransform3::Identity;
-			CreateGeometryParams.Geometry = &BodySetup->AggGeom;
-			CreateGeometryParams.ChaosTriMeshes = MakeArrayView(BodySetup->ChaosTriMeshes);
-
-			// todo(chaos) : this currently also create the shape array which is unnecessary ,this could be optimized by having a common function to create only the implicits 
-			ChaosInterface::CreateGeometry(CreateGeometryParams, Geoms, Shapes);
-
-			using FCollisionType = FGeometryDynamicCollection::FSharedImplicit;
-			TManagedArray<FCollisionType>& ExternaCollisions = GeometryCollection->AddAttribute<FCollisionType>("ExternalCollisions", FGeometryCollection::TransformGroup);
-
-			ExternaCollisions[TransformIndex1] = nullptr;
-			for (int32 GeomIndex = 0; GeomIndex < Geoms.Num();)
+			// make sure we only import box, sphere, capsule or convex
+			if (IsImportableImplicitObjectType(Geoms[GeomIndex]->GetType()))
 			{
-				// make sure we only import box, sphere, capsule or convex
-				if (IsImportableImplicitObjectType(Geoms[GeomIndex]->GetType()))
-				{
-					GeomIndex++;
-				}
-				else
-				{
-					Geoms.RemoveAtSwap(GeomIndex);
-				}
+				GeomIndex++;
 			}
-			if (Geoms.Num() > 0)
+			else
 			{
-				ExternaCollisions[TransformIndex1] = MakeShared<Chaos::FImplicitObjectUnion>(MoveTemp(Geoms));
+				Geoms.RemoveAtSwap(GeomIndex);
 			}
 		}
+		if (Geoms.Num() > 0)
+		{
+			ExternaCollisions[TransformIndex1] = MakeShared<Chaos::FImplicitObjectUnion>(MoveTemp(Geoms));
+		}
+	}
 		
-		// Bone Hierarchy - Added at root with no common parent
-		TManagedArray<int32>& Parent = GeometryCollection->Parent;
-		TManagedArray<int32>& SimulationType = GeometryCollection->SimulationType;
-		Parent[TransformIndex1] = FGeometryCollection::Invalid;
-		SimulationType[TransformIndex1] = FGeometryCollection::ESimulationTypes::FST_Rigid;
+	// Bone Hierarchy - Added at root with no common parent
+	TManagedArray<int32>& Parent = GeometryCollection->Parent;
+	TManagedArray<int32>& SimulationType = GeometryCollection->SimulationType;
+	Parent[TransformIndex1] = FGeometryCollection::Invalid;
+	SimulationType[TransformIndex1] = FGeometryCollection::ESimulationTypes::FST_Rigid;
+
+	const FColor RandBoneColor(FMath::Rand() % 100 + 5, FMath::Rand() % 100 + 5, FMath::Rand() % 100 + 5, 255);
+	TargetBoneColor[TransformIndex1] = FLinearColor(RandBoneColor);
+	TargetBoneName[TransformIndex1] = Name;
+
+	// GeometryGroup
+	int GeometryIndex = GeometryCollection->AddElements(1, FGeometryCollection::GeometryGroup);
+
+	TManagedArray<int32>& TransformIndex = GeometryCollection->TransformIndex;
+	TManagedArray<FBox>& BoundingBox = GeometryCollection->BoundingBox;
+	TManagedArray<float>& InnerRadius = GeometryCollection->InnerRadius;
+	TManagedArray<float>& OuterRadius = GeometryCollection->OuterRadius;
+	TManagedArray<int32>& VertexStartArray = GeometryCollection->VertexStart;
+	TManagedArray<int32>& VertexCountArray = GeometryCollection->VertexCount;
+	TManagedArray<int32>& FaceStartArray = GeometryCollection->FaceStart;
+	TManagedArray<int32>& FaceCountArray = GeometryCollection->FaceCount;
+
+	TransformIndex[GeometryIndex] = TargetBoneMap[VertexStart];
+	VertexStartArray[GeometryIndex] = VertexStart;
+	VertexCountArray[GeometryIndex] = VertexCount;
+	FaceStartArray[GeometryIndex] = InitialNumIndices;
+	FaceCountArray[GeometryIndex] = IndicesCount;
+
+	// TransformGroup
+	TManagedArray<int32>& TransformToGeometryIndexArray = GeometryCollection->TransformToGeometryIndex;
+	TransformToGeometryIndexArray[TransformIndex1] = GeometryIndex;
+
+	FVector Center(0);
+	for (int32 VertexIndex = VertexStart; VertexIndex < VertexStart + VertexCount; VertexIndex++)
+	{
+		Center += (FVector)TargetVertex[VertexIndex];
+	}
+	if (VertexCount) Center /= VertexCount;
 
-		const FColor RandBoneColor(FMath::Rand() % 100 + 5, FMath::Rand() % 100 + 5, FMath::Rand() % 100 + 5, 255);
-		TargetBoneColor[TransformIndex1] = FLinearColor(RandBoneColor);
-		TargetBoneName[TransformIndex1] = StaticMesh->GetName();
-
-		// GeometryGroup
-		int GeometryIndex = GeometryCollection->AddElements(1, FGeometryCollection::GeometryGroup);
-
-		TManagedArray<int32>& TransformIndex = GeometryCollection->TransformIndex;
-		TManagedArray<FBox>& BoundingBox = GeometryCollection->BoundingBox;
-		TManagedArray<float>& InnerRadius = GeometryCollection->InnerRadius;
-		TManagedArray<float>& OuterRadius = GeometryCollection->OuterRadius;
-		TManagedArray<int32>& VertexStartArray = GeometryCollection->VertexStart;
-		TManagedArray<int32>& VertexCountArray = GeometryCollection->VertexCount;
-		TManagedArray<int32>& FaceStartArray = GeometryCollection->FaceStart;
-		TManagedArray<int32>& FaceCountArray = GeometryCollection->FaceCount;
-
-		TransformIndex[GeometryIndex] = TargetBoneMap[VertexStart];
-		VertexStartArray[GeometryIndex] = VertexStart;
-		VertexCountArray[GeometryIndex] = VertexCount;
-		FaceStartArray[GeometryIndex] = InitialNumIndices;
-		FaceCountArray[GeometryIndex] = IndicesCount;
-
-		// TransformGroup
-		TManagedArray<int32>& TransformToGeometryIndexArray = GeometryCollection->TransformToGeometryIndex;
-		TransformToGeometryIndexArray[TransformIndex1] = GeometryIndex;
+	// Inner/Outer edges, bounding box
+	BoundingBox[GeometryIndex] = FBox(ForceInitToZero);
+	InnerRadius[GeometryIndex] = FLT_MAX;
+	OuterRadius[GeometryIndex] = -FLT_MAX;
+	for (int32 VertexIndex = VertexStart; VertexIndex < VertexStart + VertexCount; VertexIndex++)
+	{
+		BoundingBox[GeometryIndex] += (FVector)TargetVertex[VertexIndex];
 
-		FVector Center(0);
-		for (int32 VertexIndex = VertexStart; VertexIndex < VertexStart + VertexCount; VertexIndex++)
+		float Delta = (Center - (FVector)TargetVertex[VertexIndex]).Size();
+		InnerRadius[GeometryIndex] = FMath::Min(InnerRadius[GeometryIndex], Delta);
+		OuterRadius[GeometryIndex] = FMath::Max(OuterRadius[GeometryIndex], Delta);
+	}
+
+	// Inner/Outer centroid
+	for (int fdx = IndicesStart; fdx < IndicesStart + IndicesCount; fdx++)
+	{
+		FVector Centroid(0);
+		for (int e = 0; e < 3; e++)
 		{
-			Center += (FVector)TargetVertex[VertexIndex];
+			Centroid += (FVector)TargetVertex[TargetIndices[fdx][e]];
 		}
-		if (VertexCount) Center /= VertexCount;
+		Centroid /= 3;
 
-		// Inner/Outer edges, bounding box
-		BoundingBox[GeometryIndex] = FBox(ForceInitToZero);
-		InnerRadius[GeometryIndex] = FLT_MAX;
-		OuterRadius[GeometryIndex] = -FLT_MAX;
-		for (int32 VertexIndex = VertexStart; VertexIndex < VertexStart + VertexCount; VertexIndex++)
-		{
-			BoundingBox[GeometryIndex] += (FVector)TargetVertex[VertexIndex];
+		float Delta = (Center - Centroid).Size();
+		InnerRadius[GeometryIndex] = FMath::Min(InnerRadius[GeometryIndex], Delta);
+		OuterRadius[GeometryIndex] = FMath::Max(OuterRadius[GeometryIndex], Delta);
+	}
 
-			float Delta = (Center - (FVector)TargetVertex[VertexIndex]).Size();
+	// Inner/Outer edges
+	for (int fdx = IndicesStart; fdx < IndicesStart + IndicesCount; fdx++)
+	{
+		for (int e = 0; e < 3; e++)
+		{
+			int i = e, j = (e + 1) % 3;
+			FVector Edge = (FVector)TargetVertex[TargetIndices[fdx][i]] + 0.5 * FVector(TargetVertex[TargetIndices[fdx][j]] - TargetVertex[TargetIndices[fdx][i]]);
+			float Delta = (Center - Edge).Size();
 			InnerRadius[GeometryIndex] = FMath::Min(InnerRadius[GeometryIndex], Delta);
 			OuterRadius[GeometryIndex] = FMath::Max(OuterRadius[GeometryIndex], Delta);
 		}
+	}
 
-		// Inner/Outer centroid
-		for (int fdx = IndicesStart; fdx < IndicesStart + IndicesCount; fdx++)
-		{
-			FVector Centroid(0);
-			for (int e = 0; e < 3; e++)
-			{
-				Centroid += (FVector)TargetVertex[TargetIndices[fdx][e]];
-			}
-			Centroid /= 3;
+	if (ReindexMaterials) {
+		GeometryCollection->ReindexMaterials();
+	}
+#endif //WITH_EDITORONLY_DATA
 
-			float Delta = (Center - Centroid).Size();
-			InnerRadius[GeometryIndex] = FMath::Min(InnerRadius[GeometryIndex], Delta);
-			OuterRadius[GeometryIndex] = FMath::Max(OuterRadius[GeometryIndex], Delta);
-		}
+}
 
-		// Inner/Outer edges
-		for (int fdx = IndicesStart; fdx < IndicesStart + IndicesCount; fdx++)
-		{
-			for (int e = 0; e < 3; e++)
-			{
-				int i = e, j = (e + 1) % 3;
-				FVector Edge = (FVector)TargetVertex[TargetIndices[fdx][i]] + 0.5 * FVector(TargetVertex[TargetIndices[fdx][j]] - TargetVertex[TargetIndices[fdx][i]]);
-				float Delta = (Center - Edge).Size();
-				InnerRadius[GeometryIndex] = FMath::Min(InnerRadius[GeometryIndex], Delta);
-				OuterRadius[GeometryIndex] = FMath::Max(OuterRadius[GeometryIndex], Delta);
-			}
-		}
+FMeshDescription* FGeometryCollectionEngineConversion::GetMaxResMeshDescriptionWithNormalsAndTangents(const UStaticMesh* StaticMesh)
+{
+	if (StaticMesh == nullptr)
+	{
+		return nullptr;
+	}
 
-		if (ReindexMaterials) {
-			GeometryCollection->ReindexMaterials();
-		}
+	FMeshDescription* MeshDescription = nullptr;
+#if WITH_EDITORONLY_DATA
+	// Prefer the HiRes description, although this isn't always available.
+	if (StaticMesh->IsHiResMeshDescriptionValid())
+	{
+		MeshDescription = StaticMesh->GetHiResMeshDescription();
+	}
+	else
+	{
+		MeshDescription = StaticMesh->GetMeshDescription(0);
 	}
-#endif //WITH_EDITORONLY_DATA
 
+	FStaticMeshOperations::ComputeTriangleTangentsAndNormals(*MeshDescription);
+	FStaticMeshOperations::RecomputeNormalsAndTangentsIfNeeded(*MeshDescription, EComputeNTBsFlags::UseMikkTSpace);
+#endif //WITH_EDITORONLY_DATA
+	return MeshDescription;
 }
 
-void FGeometryCollectionEngineConversion::AppendGeometryCollection(const UGeometryCollection* SourceGeometryCollection, const TArray<UMaterialInterface*>& Materials, const FTransform& GeometryCollectionTransform, UGeometryCollection* TargetGeometryCollectionObject, bool ReindexMaterials)
+int32 FGeometryCollectionEngineConversion::AppendMaterials(const TArray<UMaterialInterface*>& Materials, UGeometryCollection* GeometryCollectionObject)
 {
-	if (SourceGeometryCollection == nullptr)
+	// for each material, add a reference in our GeometryCollectionObject
+	const int32 MaterialStart = GeometryCollectionObject->Materials.Num();
+	const int32 NumMeshMaterials = Materials.Num();
+	GeometryCollectionObject->Materials.Reserve(MaterialStart + NumMeshMaterials);
+
+	for (int32 Index = 0; Index < NumMeshMaterials; ++Index)
 	{
-		return;
+		UMaterialInterface* CurrMaterial = Materials[Index];
+
+		// Possible we have a null entry - replace with default
+		if (CurrMaterial == nullptr)
+		{
+			CurrMaterial = UMaterial::GetDefaultMaterial(MD_Surface);
+		}
+
+		// We add the material twice, once for interior and again for exterior.
+		GeometryCollectionObject->Materials.Add(CurrMaterial);
+		GeometryCollectionObject->Materials.Add(CurrMaterial);
 	}
-	const TSharedPtr<FGeometryCollection, ESPMode::ThreadSafe> SourceGeometryCollectionPtr = SourceGeometryCollection->GetGeometryCollection();
+	return MaterialStart;
+}
 
-	check(TargetGeometryCollectionObject);
-	TSharedPtr<FGeometryCollection, ESPMode::ThreadSafe> GeometryCollectionPtr = TargetGeometryCollectionObject->GetGeometryCollection();
+bool FGeometryCollectionEngineConversion::AppendStaticMesh(const UStaticMesh* StaticMesh, const TArray<UMaterialInterface*>& Materials, const FTransform& StaticMeshTransform, UGeometryCollection* GeometryCollectionObject, bool bReindexMaterials)
+{
+#if WITH_EDITORONLY_DATA
+
+	int32 StartMaterialIndex = GeometryCollectionObject->Materials.Num();
+
+	check(GeometryCollectionObject);
+	TSharedPtr<FGeometryCollection, ESPMode::ThreadSafe> GeometryCollectionPtr = GeometryCollectionObject->GetGeometryCollection();
 	FGeometryCollection* GeometryCollection = GeometryCollectionPtr.Get();
 	check(GeometryCollection);
 
+	if (AppendStaticMesh(StaticMesh, StartMaterialIndex, StaticMeshTransform, GeometryCollection, bReindexMaterials))
+	{
+		AppendMaterials(Materials, GeometryCollectionObject);
+		return true;
+	}
+
+#endif //WITH_EDITORONLY_DATA
+	return false;
+}
+
+bool FGeometryCollectionEngineConversion::AppendStaticMesh(const UStaticMesh* StaticMesh, int32 StartMaterialIndex, const FTransform& StaticMeshTransform, FGeometryCollection* GeometryCollection, bool bReindexMaterials)
+{
+#if WITH_EDITORONLY_DATA
+
+	FMeshDescription* MeshDescription = GetMaxResMeshDescriptionWithNormalsAndTangents(StaticMesh);
+
+	check(GeometryCollection);
+
+	if (MeshDescription)
+	{
+		AppendMeshDescription(MeshDescription, StaticMesh->GetName(), StartMaterialIndex, StaticMeshTransform, GeometryCollection, StaticMesh->GetBodySetup(), bReindexMaterials);
+		return true;
+	}
+#endif //WITH_EDITORONLY_DATA
+
+	return false;
+}
+
+
+bool FGeometryCollectionEngineConversion::AppendGeometryCollection(const FGeometryCollection* SourceGeometryCollectionPtr, int32 AssetMaterialStart, const FTransform& GeometryCollectionTransform, FGeometryCollection* TargetGeometryCollection, bool bReindexMaterials)
+{
+	if (SourceGeometryCollectionPtr == nullptr)
+	{
+		return false;
+	}
+
 	// Assemble offsets and add elements
 	const int32 VertexCount = SourceGeometryCollectionPtr->Vertex.Num();
 	const int32 FaceCount = SourceGeometryCollectionPtr->Indices.Num();
@@ -405,17 +450,15 @@ void FGeometryCollectionEngineConversion::AppendGeometryCollection(const UGeomet
 	const int32 GeometryCount = SourceGeometryCollectionPtr->TransformIndex.Num();
 	const int32 SectionCount = SourceGeometryCollectionPtr->Sections.Num();
 
-	const int32 AssetMaterialStart = TargetGeometryCollectionObject->Materials.Num();
-
 	FVector Scale = GeometryCollectionTransform.GetScale3D();
 	FTransform AppliedTransform = GeometryCollectionTransform;
 	AppliedTransform.RemoveScaling();
 
-	const int32 VertexStart = GeometryCollection->AddElements(VertexCount, FGeometryCollection::VerticesGroup);
-	const int32 FaceStart = GeometryCollection->AddElements(FaceCount, FGeometryCollection::FacesGroup);
-	const int32 TransformStart = GeometryCollection->AddElements(TransformCount, FGeometryCollection::TransformGroup);
-	const int32 GeometryStart = GeometryCollection->AddElements(GeometryCount, FGeometryCollection::GeometryGroup);
-	const int32 SectionStart = GeometryCollection->AddElements(SectionCount, FGeometryCollection::MaterialGroup);
+	const int32 VertexStart = TargetGeometryCollection->AddElements(VertexCount, FGeometryCollection::VerticesGroup);
+	const int32 FaceStart = TargetGeometryCollection->AddElements(FaceCount, FGeometryCollection::FacesGroup);
+	const int32 TransformStart = TargetGeometryCollection->AddElements(TransformCount, FGeometryCollection::TransformGroup);
+	const int32 GeometryStart = TargetGeometryCollection->AddElements(GeometryCount, FGeometryCollection::GeometryGroup);
+	const int32 SectionStart = TargetGeometryCollection->AddElements(SectionCount, FGeometryCollection::MaterialGroup);
 
 	// source vertex information
 	const TManagedArray<FVector3f>& SourceVertex = SourceGeometryCollectionPtr->Vertex;
@@ -427,20 +470,20 @@ void FGeometryCollectionEngineConversion::AppendGeometryCollection(const UGeomet
 	const TManagedArray<int32>& SourceBoneMap = SourceGeometryCollectionPtr->BoneMap;
 
 	// target vertex information
-	TManagedArray<FVector3f>& TargetVertex = GeometryCollection->Vertex;
-	TManagedArray<FVector3f>& TargetTangentU = GeometryCollection->TangentU;
-	TManagedArray<FVector3f>& TargetTangentV = GeometryCollection->TangentV;
-	TManagedArray<FVector3f>& TargetNormal = GeometryCollection->Normal;
-	TManagedArray<TArray<FVector2f>>& TargetUVs = GeometryCollection->UVs;
-	TManagedArray<FLinearColor>& TargetColor = GeometryCollection->Color;
-	TManagedArray<int32>& TargetBoneMap = GeometryCollection->BoneMap;
+	TManagedArray<FVector3f>& TargetVertex = TargetGeometryCollection->Vertex;
+	TManagedArray<FVector3f>& TargetTangentU = TargetGeometryCollection->TangentU;
+	TManagedArray<FVector3f>& TargetTangentV = TargetGeometryCollection->TangentV;
+	TManagedArray<FVector3f>& TargetNormal = TargetGeometryCollection->Normal;
+	TManagedArray<TArray<FVector2f>>& TargetUVs = TargetGeometryCollection->UVs;
+	TManagedArray<FLinearColor>& TargetColor = TargetGeometryCollection->Color;
+	TManagedArray<int32>& TargetBoneMap = TargetGeometryCollection->BoneMap;
 
 	// append vertices
 	for (int32 VertexIndex = 0; VertexIndex < VertexCount; VertexIndex++)
 	{
 		const int32 VertexOffset = VertexStart + VertexIndex;
 		TargetVertex[VertexOffset] = SourceVertex[VertexIndex] * (FVector3f)Scale;
-		
+
 		TargetTangentU[VertexOffset] = SourceTangentU[VertexIndex];
 		TargetTangentV[VertexOffset] = SourceTangentV[VertexIndex];
 		TargetNormal[VertexOffset] = SourceNormal[VertexIndex];
@@ -455,12 +498,12 @@ void FGeometryCollectionEngineConversion::AppendGeometryCollection(const UGeomet
 	const TManagedArray<bool>& SourceVisible = SourceGeometryCollectionPtr->Visible;
 	const TManagedArray<int32>& SourceMaterialID = SourceGeometryCollectionPtr->MaterialID;
 	const TManagedArray<int32>& SourceMaterialIndex = SourceGeometryCollectionPtr->MaterialIndex;
-	
+
 	// target face information
-	TManagedArray<FIntVector>& TargetIndices = GeometryCollection->Indices;
-	TManagedArray<bool>& TargetVisible = GeometryCollection->Visible;
-	TManagedArray<int32>& TargetMaterialID = GeometryCollection->MaterialID;
-	TManagedArray<int32>& TargetMaterialIndex = GeometryCollection->MaterialIndex;
+	TManagedArray<FIntVector>& TargetIndices = TargetGeometryCollection->Indices;
+	TManagedArray<bool>& TargetVisible = TargetGeometryCollection->Visible;
+	TManagedArray<int32>& TargetMaterialID = TargetGeometryCollection->MaterialID;
+	TManagedArray<int32>& TargetMaterialIndex = TargetGeometryCollection->MaterialIndex;
 
 	// append faces
 	for (int32 FaceIndex = 0; FaceIndex < FaceCount; ++FaceIndex)
@@ -472,13 +515,13 @@ void FGeometryCollectionEngineConversion::AppendGeometryCollection(const UGeomet
 			SourceFace[1] + VertexStart,
 			SourceFace[2] + VertexStart);
 		TargetVisible[FaceOffset] = SourceVisible[FaceIndex];
-		
+
 		TargetMaterialID[FaceOffset] = AssetMaterialStart + SourceMaterialID[FaceIndex];
 		TargetMaterialIndex[FaceOffset] = FaceOffset;
 	}
 
 	using FCollisionType = FGeometryDynamicCollection::FSharedImplicit;
-	
+
 	// source transform information
 	const TManagedArray<FTransform>& SourceTransform = SourceGeometryCollectionPtr->Transform;
 	const TManagedArray<FString>& SourceBoneName = SourceGeometryCollectionPtr->BoneName;
@@ -490,18 +533,18 @@ void FGeometryCollectionEngineConversion::AppendGeometryCollection(const UGeomet
 	const TManagedArray<int32>& SourceStatusFlags = SourceGeometryCollectionPtr->StatusFlags;
 	const TManagedArray<int32>& SourceInitialDynamicState = SourceGeometryCollectionPtr->InitialDynamicState;
 	const TManagedArray<FCollisionType>* SourceExternalCollisions = SourceGeometryCollectionPtr->FindAttribute<FCollisionType>("ExternalCollisions", FGeometryCollection::TransformGroup);
-	
+
 	// target transform information
-	TManagedArray<FTransform>& TargetTransform = GeometryCollection->Transform;
-	TManagedArray<FString>& TargetBoneName = GeometryCollection->BoneName;
-	TManagedArray<FLinearColor>& TargetBoneColor = GeometryCollection->BoneColor;
-	TManagedArray<int32>& TargetParent = GeometryCollection->Parent;
-	TManagedArray<TSet<int32>>& TargetChildren = GeometryCollection->Children;
-	TManagedArray<int32>& TargetTransformToGeometryIndex = GeometryCollection->TransformToGeometryIndex;
-	TManagedArray<int32>& TargetSimulationType = GeometryCollection->SimulationType;
-	TManagedArray<int32>& TargetStatusFlags = GeometryCollection->StatusFlags;
-	TManagedArray<int32>& TargetInitialDynamicState = GeometryCollection->InitialDynamicState;
-	TManagedArray<FCollisionType>& TargetExternalCollisions = GeometryCollection->AddAttribute<FCollisionType>("ExternalCollisions", FGeometryCollection::TransformGroup);
+	TManagedArray<FTransform>& TargetTransform = TargetGeometryCollection->Transform;
+	TManagedArray<FString>& TargetBoneName = TargetGeometryCollection->BoneName;
+	TManagedArray<FLinearColor>& TargetBoneColor = TargetGeometryCollection->BoneColor;
+	TManagedArray<int32>& TargetParent = TargetGeometryCollection->Parent;
+	TManagedArray<TSet<int32>>& TargetChildren = TargetGeometryCollection->Children;
+	TManagedArray<int32>& TargetTransformToGeometryIndex = TargetGeometryCollection->TransformToGeometryIndex;
+	TManagedArray<int32>& TargetSimulationType = TargetGeometryCollection->SimulationType;
+	TManagedArray<int32>& TargetStatusFlags = TargetGeometryCollection->StatusFlags;
+	TManagedArray<int32>& TargetInitialDynamicState = TargetGeometryCollection->InitialDynamicState;
+	TManagedArray<FCollisionType>& TargetExternalCollisions = TargetGeometryCollection->AddAttribute<FCollisionType>("ExternalCollisions", FGeometryCollection::TransformGroup);
 
 	// append transform hierarchy
 	for (int32 TransformIndex = 0; TransformIndex < TransformCount; ++TransformIndex)
@@ -519,14 +562,14 @@ void FGeometryCollectionEngineConversion::AppendGeometryCollection(const UGeomet
 			ScaledTranslation.ScaleTranslation(Scale);
 			TargetTransform[TransformOffset] = ScaledTranslation;
 		}
-	
+
 		// #todo Get this Bone name to be unique
 		TargetBoneName[TransformOffset] = SourceBoneName[TransformIndex];
 
 		const FColor RandBoneColor(FMath::Rand() % 100 + 5, FMath::Rand() % 100 + 5, FMath::Rand() % 100 + 5, 255);
 		TargetBoneColor[TransformOffset] = FLinearColor(RandBoneColor);
 
-		TargetParent[TransformOffset] = (SourceParent[TransformIndex]==INDEX_NONE) ? INDEX_NONE: SourceParent[TransformIndex] + TransformStart;
+		TargetParent[TransformOffset] = (SourceParent[TransformIndex] == INDEX_NONE) ? INDEX_NONE : SourceParent[TransformIndex] + TransformStart;
 
 		const TSet<int32>& SourceChildrenSet = SourceChildren[TransformIndex];
 		for (int32 ChildIndex : SourceChildrenSet)
@@ -552,16 +595,16 @@ void FGeometryCollectionEngineConversion::AppendGeometryCollection(const UGeomet
 	const TManagedArray<int32>& SourceVertexCount = SourceGeometryCollectionPtr->VertexCount;
 	const TManagedArray<int32>& SourceFaceStart = SourceGeometryCollectionPtr->FaceStart;
 	const TManagedArray<int32>& SourceFaceCount = SourceGeometryCollectionPtr->FaceCount;
-	
+
 	// target geometry information
-	TManagedArray<int32>& TargetTransformIndex = GeometryCollection->TransformIndex;
-	TManagedArray<FBox>& TargetBoundingBox = GeometryCollection->BoundingBox;
-	TManagedArray<float>& TargetInnerRadius = GeometryCollection->InnerRadius;
-	TManagedArray<float>& TargetOuterRadius = GeometryCollection->OuterRadius;
-	TManagedArray<int32>& TargetVertexStart = GeometryCollection->VertexStart;
-	TManagedArray<int32>& TargetVertexCount = GeometryCollection->VertexCount;
-	TManagedArray<int32>& TargetFaceStart = GeometryCollection->FaceStart;
-	TManagedArray<int32>& TargetFaceCount = GeometryCollection->FaceCount;
+	TManagedArray<int32>& TargetTransformIndex = TargetGeometryCollection->TransformIndex;
+	TManagedArray<FBox>& TargetBoundingBox = TargetGeometryCollection->BoundingBox;
+	TManagedArray<float>& TargetInnerRadius = TargetGeometryCollection->InnerRadius;
+	TManagedArray<float>& TargetOuterRadius = TargetGeometryCollection->OuterRadius;
+	TManagedArray<int32>& TargetVertexStart = TargetGeometryCollection->VertexStart;
+	TManagedArray<int32>& TargetVertexCount = TargetGeometryCollection->VertexCount;
+	TManagedArray<int32>& TargetFaceStart = TargetGeometryCollection->FaceStart;
+	TManagedArray<int32>& TargetFaceCount = TargetGeometryCollection->FaceCount;
 
 	// append geometry
 	for (int32 GeometryIndex = 0; GeometryIndex < GeometryCount; ++GeometryIndex)
@@ -569,7 +612,7 @@ void FGeometryCollectionEngineConversion::AppendGeometryCollection(const UGeomet
 		const int32 GeometryOffset = GeometryStart + GeometryIndex;
 
 		TargetTransformIndex[GeometryOffset] = SourceTransformIndex[GeometryIndex] + TransformStart;
-		
+
 		TargetVertexStart[GeometryOffset] = SourceVertexStart[GeometryIndex] + VertexStart;
 		TargetVertexCount[GeometryOffset] = SourceVertexCount[GeometryIndex];
 		TargetFaceStart[GeometryOffset] = SourceFaceStart[GeometryIndex] + FaceStart;
@@ -582,7 +625,7 @@ void FGeometryCollectionEngineConversion::AppendGeometryCollection(const UGeomet
 			Center += (FVector)TargetVertex[VertexIndex];
 		}
 		if (TargetVertexCount[GeometryOffset]) Center /= TargetVertexCount[GeometryOffset];
-		
+
 		TargetBoundingBox[GeometryOffset] = FBox(ForceInitToZero);
 		TargetInnerRadius[GeometryOffset] = FLT_MAX;
 		TargetOuterRadius[GeometryOffset] = -FLT_MAX;
@@ -596,28 +639,11 @@ void FGeometryCollectionEngineConversion::AppendGeometryCollection(const UGeomet
 		}
 	}
 
-	// for each material, add a reference in our GeometryCollectionObject
-	const int32 NumMeshMaterials = Materials.Num();
-	TargetGeometryCollectionObject->Materials.Reserve(AssetMaterialStart + NumMeshMaterials);
-
-	for (int32 Index = 0; Index < NumMeshMaterials; ++Index)
-	{
-		UMaterialInterface* CurrMaterial = Materials[Index];
-
-		// Possible we have a null entry - replace with default
-		if (CurrMaterial == nullptr)
-		{
-			CurrMaterial = UMaterial::GetDefaultMaterial(MD_Surface);
-		}
-
-		TargetGeometryCollectionObject->Materials.Add(CurrMaterial);
-	}
-
 	// source material information
 	const TManagedArray<FGeometryCollectionSection>& SourceSections = SourceGeometryCollectionPtr->Sections;
 
 	// target material information
-	TManagedArray<FGeometryCollectionSection>& TargetSections = GeometryCollection->Sections;
+	TManagedArray<FGeometryCollectionSection>& TargetSections = TargetGeometryCollection->Sections;
 
 	// append sections
 	for (int32 SectionIndex = 0; SectionIndex < SectionCount; ++SectionIndex)
@@ -633,10 +659,31 @@ void FGeometryCollectionEngineConversion::AppendGeometryCollection(const UGeomet
 		TargetSections[SectionOffset].MaxVertexIndex = VertexStart + SourceSections[SectionIndex].MaxVertexIndex;
 	}
 
-	if (ReindexMaterials) {
-		GeometryCollection->ReindexMaterials();
+	if (bReindexMaterials)
+	{
+		TargetGeometryCollection->ReindexMaterials();
+	}
+
+	return true;
+
+}
+
+void FGeometryCollectionEngineConversion::AppendGeometryCollection(const UGeometryCollection* SourceGeometryCollection, const TArray<UMaterialInterface*>& Materials, const FTransform& GeometryCollectionTransform, UGeometryCollection* TargetGeometryCollectionObject, bool bReindexMaterials)
+{
+	if (SourceGeometryCollection == nullptr)
+	{
+		return;
 	}
+	const TSharedPtr<FGeometryCollection, ESPMode::ThreadSafe> SourceGeometryCollectionPtr = SourceGeometryCollection->GetGeometryCollection();
+
+	check(TargetGeometryCollectionObject);
+	TSharedPtr<FGeometryCollection, ESPMode::ThreadSafe> GeometryCollectionPtr = TargetGeometryCollectionObject->GetGeometryCollection();
+	FGeometryCollection* GeometryCollection = GeometryCollectionPtr.Get();
+	check(GeometryCollection);
+
+	int32 MaterialStart = AppendMaterials(Materials, TargetGeometryCollectionObject);
 
+	AppendGeometryCollection(SourceGeometryCollectionPtr.Get(), MaterialStart, GeometryCollectionTransform, GeometryCollection, bReindexMaterials);
 }
 
 
@@ -665,12 +712,11 @@ void FGeometryCollectionEngineConversion::AppendStaticMesh(const UStaticMesh* St
 }
 
 
-void FGeometryCollectionEngineConversion::AppendGeometryCollection(const UGeometryCollection* SourceGeometryCollection, const UGeometryCollectionComponent* GeometryCollectionComponent, const FTransform& GeometryCollectionTransform, UGeometryCollection* TargetGeometryCollectionObject, bool ReindexMaterials)
+int32 FGeometryCollectionEngineConversion::AppendGeometryCollectionMaterials(const UGeometryCollection* SourceGeometryCollection, const UGeometryCollectionComponent* GeometryCollectionComponent, UGeometryCollection* TargetGeometryCollectionObject)
 {
-	if (SourceGeometryCollection == nullptr)
-	{
-		return;
-	}
+	check(SourceGeometryCollection);
+	check(GeometryCollectionComponent);
+	check(TargetGeometryCollectionObject);
 
 	TArray<UMaterialInterface*> Materials;
 	Materials.Reserve(SourceGeometryCollection->Materials.Num());
@@ -686,222 +732,276 @@ void FGeometryCollectionEngineConversion::AppendGeometryCollection(const UGeomet
 	TargetGeometryCollectionObject->Materials.Remove(BoneSelectedMaterial);
 	Materials.Remove(BoneSelectedMaterial);
 
-	AppendGeometryCollection(SourceGeometryCollection, Materials, GeometryCollectionTransform, TargetGeometryCollectionObject, ReindexMaterials);
+	return AppendMaterials(Materials, TargetGeometryCollectionObject);
+}
+
+
+void FGeometryCollectionEngineConversion::AppendGeometryCollection(const UGeometryCollection* SourceGeometryCollection, const UGeometryCollectionComponent* GeometryCollectionComponent, const FTransform& GeometryCollectionTransform, UGeometryCollection* TargetGeometryCollectionObject, bool bReindexMaterials)
+{
+	if (SourceGeometryCollection == nullptr)
+	{
+		return;
+	}
+
+	int32 MaterialStartIndex = AppendGeometryCollectionMaterials(SourceGeometryCollection, GeometryCollectionComponent, TargetGeometryCollectionObject);
+
+	const TSharedPtr<FGeometryCollection, ESPMode::ThreadSafe> SourceGeometryCollectionPtr = SourceGeometryCollection->GetGeometryCollection();
+
+	check(TargetGeometryCollectionObject);
+	TSharedPtr<FGeometryCollection, ESPMode::ThreadSafe> GeometryCollectionPtr = TargetGeometryCollectionObject->GetGeometryCollection();
+	FGeometryCollection* GeometryCollection = GeometryCollectionPtr.Get();
+	check(GeometryCollection);
+
+	AppendGeometryCollection(SourceGeometryCollectionPtr.Get(), MaterialStartIndex, GeometryCollectionTransform, GeometryCollection, bReindexMaterials);
 }
 
 
-void FGeometryCollectionEngineConversion::AppendSkeletalMesh(const USkeletalMesh* SkeletalMesh, const USkeletalMeshComponent *SkeletalMeshComponent, const FTransform & SkeletalMeshTransform, UGeometryCollection * GeometryCollectionObject, bool ReindexMaterials)
+bool FGeometryCollectionEngineConversion::AppendSkeletalMesh(const USkeletalMesh* SkeletalMesh, int32 MaterialStartIndex, const FTransform& SkeletalMeshTransform, FGeometryCollection* GeometryCollection, bool bReindexMaterials)
 {
 	//UE_LOG(UGeometryCollectionConversionLogging, Log, TEXT("FGeometryCollectionEngineConversion::AppendSkeletalMesh()"));
-	check(SkeletalMesh);
-	if (GeometryCollectionObject)
+	if (!GeometryCollection)
 	{
-		TSharedPtr<FGeometryCollection, ESPMode::ThreadSafe> GeometryCollectionPtr = GeometryCollectionObject->GetGeometryCollection();
-		if (FGeometryCollection* GeometryCollection = GeometryCollectionPtr.Get())
+		return false;
+	}
+
+	const FSkeletalMeshLODRenderData* MeshLODData = GetSkeletalMeshLOD(SkeletalMesh, 0);
+	if (!MeshLODData)
+	{
+		return false;
+	}
+		
+	const FSkeletalMeshLODRenderData& SkeletalMeshLODRenderData = *MeshLODData;
+	const FSkinWeightVertexBuffer& SkinWeightVertexBuffer = *SkeletalMeshLODRenderData.GetSkinWeightVertexBuffer();
+	const FStaticMeshVertexBuffers& VertexBuffers = SkeletalMeshLODRenderData.StaticVertexBuffers;
+	const FPositionVertexBuffer& PositionVertexBuffer = VertexBuffers.PositionVertexBuffer;
+	const int32 VertexCount = PositionVertexBuffer.GetNumVertices();
+	// Check that all vertex weightings are rigid. 
+	for (int32 VertexIndex = 0; VertexIndex < VertexCount; VertexIndex++)
+	{
+		int32 SkeletalBoneIndex = -1;
+		if (!SkinWeightVertexBuffer.GetRigidWeightBone(VertexIndex, SkeletalBoneIndex))
 		{
+			UE_LOG(UGeometryCollectionConversionLogging, Error, TEXT("Non-rigid weighting found on vertex %d: Cannot convert to GeometryCollection."), VertexIndex);
+			return false;
+		}
+	}
+	const FSkelMeshRenderSection& RenderSection = SkeletalMeshLODRenderData.RenderSections[0];
+	const TArray<FBoneIndexType>& SkeletalBoneMap = RenderSection.BoneMap;
+
+	//
+	// The Component transform for each Mesh will become the FTransform that drives
+	// its associated VerticesGroup. The Skeleton will contain a nested transform hierarchy
+	// that is evaluated using the GetComponentSpaceTransformRefPose. The resulting
+	// Transforms array stored in the GeometryCollection will be the same size as
+	// the SkeletalBoneMap. Note the @todo: the SkeletalBoneMap is pulled from only
+	// the first render section, this will need to be expanded to include all render
+	// sections.
+	//
+	const USkeleton* Skeleton = SkeletalMesh->GetSkeleton();
+	TManagedArray<FTransform>& Transform = GeometryCollection->Transform;
+	int32 TransformBaseIndex = GeometryCollection->AddElements(SkeletalBoneMap.Num(), FGeometryCollection::TransformGroup);
+	const FReferenceSkeleton & ReferenceSkeletion = Skeleton->GetReferenceSkeleton();
+	const TArray<FTransform> & RestArray = Skeleton->GetRefLocalPoses();
+	for (int32 BoneIndex = 0; BoneIndex < SkeletalBoneMap.Num(); BoneIndex++)
+	{
+		FTransform BoneTransform = FAnimationRuntime::GetComponentSpaceTransformRefPose(ReferenceSkeletion, SkeletalBoneMap[BoneIndex]);
+		Transform[TransformBaseIndex + BoneIndex] = BoneTransform;
+	}
 
-			if (const USkeleton* Skeleton = SkeletalMesh->GetSkeleton())
-			{
-				if (const FSkeletalMeshRenderData * SkelMeshRenderData = SkeletalMesh->GetResourceForRendering())
-				{
-					if (SkelMeshRenderData->LODRenderData.Num())
-					{
-						const FSkeletalMeshLODRenderData& SkeletalMeshLODRenderData = SkelMeshRenderData->LODRenderData[0];
-						const FSkinWeightVertexBuffer& SkinWeightVertexBuffer = *SkeletalMeshLODRenderData.GetSkinWeightVertexBuffer();
-						const FStaticMeshVertexBuffers& VertexBuffers = SkeletalMeshLODRenderData.StaticVertexBuffers;
-						const FPositionVertexBuffer& PositionVertexBuffer = VertexBuffers.PositionVertexBuffer;
-						const int32 VertexCount = PositionVertexBuffer.GetNumVertices();
-						// Check that all vertex weightings are rigid. 
-						for (int32 VertexIndex = 0; VertexIndex < VertexCount; VertexIndex++)
-						{
-							int32 SkeletalBoneIndex = -1;
-							if (!SkinWeightVertexBuffer.GetRigidWeightBone(VertexIndex, SkeletalBoneIndex))
-							{
-								UE_LOG(UGeometryCollectionConversionLogging, Error, TEXT("Non-rigid weighting found on vertex %d: Cannot convert to GeometryCollection."), VertexIndex);
-								return;
-							}
-						}
-						const FSkelMeshRenderSection& RenderSection = SkeletalMeshLODRenderData.RenderSections[0];
-						const TArray<FBoneIndexType>& SkeletalBoneMap = RenderSection.BoneMap;
-
-						//
-						// The Component transform for each Mesh will become the FTransform that drives
-						// its associated VerticesGroup. The Skeleton will contain a nested transform hierarchy
-						// that is evaluated using the GetComponentSpaceTransformRefPose. The resulting
-						// Transforms array stored in the GeometryCollection will be the same size as
-						// the SkeletalBoneMap. Note the @todo: the SkeletalBoneMap is pulled from only
-						// the first render section, this will need to be expanded to include all render
-						// sections.
-						//
-						TManagedArray<FTransform>& Transform = GeometryCollection->Transform;
-						int32 TransformBaseIndex = GeometryCollection->AddElements(SkeletalBoneMap.Num(), FGeometryCollection::TransformGroup);
-						const FReferenceSkeleton & ReferenceSkeletion = Skeleton->GetReferenceSkeleton();
-						const TArray<FTransform> & RestArray = Skeleton->GetRefLocalPoses();
-						for (int32 BoneIndex = 0; BoneIndex < SkeletalBoneMap.Num(); BoneIndex++)
-						{
-							FTransform BoneTransform = FAnimationRuntime::GetComponentSpaceTransformRefPose(ReferenceSkeletion, SkeletalBoneMap[BoneIndex]);
-							Transform[TransformBaseIndex + BoneIndex] = BoneTransform;
-						}
-
-
-						//
-						// The Triangle Indices
-						//
-						TManagedArray<FIntVector>& Indices = GeometryCollection->Indices;
-						TManagedArray<bool>& Visible = GeometryCollection->Visible;
-						TManagedArray<int32>& MaterialID = GeometryCollection->MaterialID;
-						TManagedArray<int32>& MaterialIndex = GeometryCollection->MaterialIndex;
-
-						TArray<uint32> IndexBuffer;
-						SkeletalMeshLODRenderData.MultiSizeIndexContainer.GetIndexBuffer(IndexBuffer);
-
-						const int32 IndicesCount = IndexBuffer.Num() / 3;
-						int NumVertices = GeometryCollection->NumElements(FGeometryCollection::VerticesGroup);
-						int InitialNumIndices = GeometryCollection->NumElements(FGeometryCollection::FacesGroup);
-						int IndicesBaseIndex = GeometryCollection->AddElements(IndicesCount, FGeometryCollection::FacesGroup);
-						for (int32 IndicesIndex = 0, StaticIndex = 0; IndicesIndex < IndicesCount; IndicesIndex++, StaticIndex += 3)
-						{
-							int32 IndicesOffset = IndicesBaseIndex + IndicesIndex;
-							Indices[IndicesOffset] = FIntVector(
-								IndexBuffer[StaticIndex] + NumVertices,
-								IndexBuffer[StaticIndex + 1] + NumVertices,
-								IndexBuffer[StaticIndex + 2] + NumVertices);
-							Visible[IndicesOffset] = true;
-							MaterialID[IndicesOffset] = 0;
-							MaterialIndex[IndicesOffset] = IndicesOffset;
-						}
-
-						//
-						// Vertex Attributes
-						//
-						TManagedArray<FVector3f>& Vertex = GeometryCollection->Vertex;
-						TManagedArray<FVector3f>& TangentU = GeometryCollection->TangentU;
-						TManagedArray<FVector3f>& TangentV = GeometryCollection->TangentV;
-						TManagedArray<FVector3f>& Normal = GeometryCollection->Normal;
-						TManagedArray<TArray<FVector2f>>& UVs = GeometryCollection->UVs;
-						TManagedArray<FLinearColor>& Color = GeometryCollection->Color;
-						TManagedArray<int32>& BoneMap = GeometryCollection->BoneMap;
-						TManagedArray<FLinearColor>& BoneColor = GeometryCollection->BoneColor;
-						TManagedArray<FString>& BoneName = GeometryCollection->BoneName;
-
-						// 
-						// Transform Attributes 
-						// 
-						TManagedArray<int32>& Parent = GeometryCollection->Parent;
-						TManagedArray<int32>& SimulationType = GeometryCollection->SimulationType;
-						int InitialNumVertices = GeometryCollection->NumElements(FGeometryCollection::VerticesGroup);
-						int VertexBaseIndex = GeometryCollection->AddElements(VertexCount, FGeometryCollection::VerticesGroup);
-						const int32 NumUVLayers = VertexBuffers.StaticMeshVertexBuffer.GetNumTexCoords();
-						for (int32 VertexIndex = 0; VertexIndex < VertexCount; VertexIndex++)
-						{
-							int VertexOffset = VertexBaseIndex + VertexIndex;
-							BoneMap[VertexOffset] = -1;
-							int32 SkeletalBoneIndex = -1;
-							SkinWeightVertexBuffer.GetRigidWeightBone(VertexIndex, SkeletalBoneIndex);
-							if (SkeletalBoneIndex > -1)
-							{
-								BoneMap[VertexOffset] = SkeletalBoneIndex + TransformBaseIndex;
-								Vertex[VertexOffset] = (FVector4f)Transform[BoneMap[VertexOffset]].ToInverseMatrixWithScale().TransformPosition((FVector)PositionVertexBuffer.VertexPosition(VertexIndex));
-							}
-							check(BoneMap[VertexOffset] != -1);
-							TangentU[VertexOffset] = VertexBuffers.StaticMeshVertexBuffer.VertexTangentX(VertexIndex);
-							TangentV[VertexOffset] = VertexBuffers.StaticMeshVertexBuffer.VertexTangentY(VertexIndex);
-							Normal[VertexOffset] = VertexBuffers.StaticMeshVertexBuffer.VertexTangentZ(VertexIndex);
+
+	//
+	// The Triangle Indices
+	//
+	TManagedArray<FIntVector>& Indices = GeometryCollection->Indices;
+	TManagedArray<bool>& Visible = GeometryCollection->Visible;
+	TManagedArray<int32>& MaterialID = GeometryCollection->MaterialID;
+	TManagedArray<int32>& MaterialIndex = GeometryCollection->MaterialIndex;
+
+	TArray<uint32> IndexBuffer;
+	SkeletalMeshLODRenderData.MultiSizeIndexContainer.GetIndexBuffer(IndexBuffer);
+
+	const int32 IndicesCount = IndexBuffer.Num() / 3;
+	int NumVertices = GeometryCollection->NumElements(FGeometryCollection::VerticesGroup);
+	int InitialNumIndices = GeometryCollection->NumElements(FGeometryCollection::FacesGroup);
+	int IndicesBaseIndex = GeometryCollection->AddElements(IndicesCount, FGeometryCollection::FacesGroup);
+	for (int32 IndicesIndex = 0, StaticIndex = 0; IndicesIndex < IndicesCount; IndicesIndex++, StaticIndex += 3)
+	{
+		int32 IndicesOffset = IndicesBaseIndex + IndicesIndex;
+		Indices[IndicesOffset] = FIntVector(
+			IndexBuffer[StaticIndex] + NumVertices,
+			IndexBuffer[StaticIndex + 1] + NumVertices,
+			IndexBuffer[StaticIndex + 2] + NumVertices);
+		Visible[IndicesOffset] = true;
+		MaterialID[IndicesOffset] = 0;
+		MaterialIndex[IndicesOffset] = IndicesOffset;
+	}
+
+	//
+	// Vertex Attributes
+	//
+	TManagedArray<FVector3f>& Vertex = GeometryCollection->Vertex;
+	TManagedArray<FVector3f>& TangentU = GeometryCollection->TangentU;
+	TManagedArray<FVector3f>& TangentV = GeometryCollection->TangentV;
+	TManagedArray<FVector3f>& Normal = GeometryCollection->Normal;
+	TManagedArray<TArray<FVector2f>>& UVs = GeometryCollection->UVs;
+	TManagedArray<FLinearColor>& Color = GeometryCollection->Color;
+	TManagedArray<int32>& BoneMap = GeometryCollection->BoneMap;
+	TManagedArray<FLinearColor>& BoneColor = GeometryCollection->BoneColor;
+	TManagedArray<FString>& BoneName = GeometryCollection->BoneName;
+
+	// 
+	// Transform Attributes 
+	// 
+	TManagedArray<int32>& Parent = GeometryCollection->Parent;
+	TManagedArray<int32>& SimulationType = GeometryCollection->SimulationType;
+	int InitialNumVertices = GeometryCollection->NumElements(FGeometryCollection::VerticesGroup);
+	int VertexBaseIndex = GeometryCollection->AddElements(VertexCount, FGeometryCollection::VerticesGroup);
+	const int32 NumUVLayers = VertexBuffers.StaticMeshVertexBuffer.GetNumTexCoords();
+	for (int32 VertexIndex = 0; VertexIndex < VertexCount; VertexIndex++)
+	{
+		int VertexOffset = VertexBaseIndex + VertexIndex;
+		BoneMap[VertexOffset] = -1;
+		int32 SkeletalBoneIndex = -1;
+		SkinWeightVertexBuffer.GetRigidWeightBone(VertexIndex, SkeletalBoneIndex);
+		if (SkeletalBoneIndex > -1)
+		{
+			BoneMap[VertexOffset] = SkeletalBoneIndex + TransformBaseIndex;
+			Vertex[VertexOffset] = (FVector4f)Transform[BoneMap[VertexOffset]].ToInverseMatrixWithScale().TransformPosition((FVector)PositionVertexBuffer.VertexPosition(VertexIndex));
+		}
+		check(BoneMap[VertexOffset] != -1);
+		TangentU[VertexOffset] = VertexBuffers.StaticMeshVertexBuffer.VertexTangentX(VertexIndex);
+		TangentV[VertexOffset] = VertexBuffers.StaticMeshVertexBuffer.VertexTangentY(VertexIndex);
+		Normal[VertexOffset] = VertexBuffers.StaticMeshVertexBuffer.VertexTangentZ(VertexIndex);
 							
-							UVs[VertexOffset].SetNum(NumUVLayers);
-							for (int32 UVLayerIdx = 0; UVLayerIdx < NumUVLayers; ++UVLayerIdx)
-							{
-								UVs[VertexOffset][UVLayerIdx] = VertexBuffers.StaticMeshVertexBuffer.GetVertexUV(VertexIndex, UVLayerIdx);
-							}
+		UVs[VertexOffset].SetNum(NumUVLayers);
+		for (int32 UVLayerIdx = 0; UVLayerIdx < NumUVLayers; ++UVLayerIdx)
+		{
+			UVs[VertexOffset][UVLayerIdx] = VertexBuffers.StaticMeshVertexBuffer.GetVertexUV(VertexIndex, UVLayerIdx);
+		}
 							
-							if (VertexBuffers.ColorVertexBuffer.GetNumVertices() == VertexCount)
-								Color[VertexOffset] = VertexBuffers.ColorVertexBuffer.VertexColor(VertexIndex);
-							else
-								Color[VertexOffset] = FLinearColor::White;
-						}
-
-						int32 InitialIndex = -1;
-						int32 LastParentIndex = -1;
-						int32 CurrentLevel = 0;
-						for (int32 BoneIndex = 0; BoneIndex < SkeletalBoneMap.Num(); BoneIndex++)
-						{
-							// transform based on position of the actor. 
-							Transform[TransformBaseIndex + BoneIndex] = SkeletalMeshTransform * Transform[TransformBaseIndex + BoneIndex];
-
-							// bone attributes
-							BoneName[TransformBaseIndex + BoneIndex] = ReferenceSkeletion.GetBoneName(SkeletalBoneMap[BoneIndex]).ToString();
-							const FColor RandBoneColor(FMath::Rand() % 100 + 5, FMath::Rand() % 100 + 5, FMath::Rand() % 100 + 5, 255);
-							BoneColor[TransformBaseIndex + BoneIndex] = FLinearColor(RandBoneColor);
-
-							// Bone Hierarchy - Added at root with no common parent
-							int32 ParentIndex = ReferenceSkeletion.GetParentIndex(SkeletalBoneMap[BoneIndex]);
-							int32 UseParentIndex = ParentIndex + InitialIndex;
-							if (LastParentIndex != UseParentIndex)
-							{
-								LastParentIndex = UseParentIndex;
-							}
-							Parent[TransformBaseIndex + BoneIndex] = UseParentIndex;
-							SimulationType[TransformBaseIndex + BoneIndex] = FGeometryCollection::ESimulationTypes::FST_Rigid;
-						}
-
-						// Geometry Group
-						TArray<int32> GeometryIndices;
-						GeometryCollectionAlgo::ContiguousArray(GeometryIndices, GeometryCollection->NumElements(FGeometryCollection::GeometryGroup));
-						GeometryCollection->RemoveDependencyFor(FGeometryCollection::GeometryGroup);
-						GeometryCollection->RemoveElements(FGeometryCollection::GeometryGroup, GeometryIndices);
-						::GeometryCollection::AddGeometryProperties(GeometryCollection);
-
-						// for each material, add a reference in our GeometryCollectionObject
-						int CurrIdx = 0;
-						
-						const TArray<FSkeletalMaterial>& SkeletalMeshMaterials = SkeletalMesh->GetMaterials();
+		if (VertexBuffers.ColorVertexBuffer.GetNumVertices() == VertexCount)
+			Color[VertexOffset] = VertexBuffers.ColorVertexBuffer.VertexColor(VertexIndex);
+		else
+			Color[VertexOffset] = FLinearColor::White;
+	}
 
-						UMaterialInterface* CurrMaterial = SkeletalMeshComponent ? SkeletalMeshComponent->GetMaterial(CurrIdx) : ToRawPtr(SkeletalMeshMaterials[CurrIdx].MaterialInterface);
+	int32 InitialIndex = -1;
+	int32 LastParentIndex = -1;
+	int32 CurrentLevel = 0;
+	for (int32 BoneIndex = 0; BoneIndex < SkeletalBoneMap.Num(); BoneIndex++)
+	{
+		// transform based on position of the actor. 
+		Transform[TransformBaseIndex + BoneIndex] = SkeletalMeshTransform * Transform[TransformBaseIndex + BoneIndex];
 
+		// bone attributes
+		BoneName[TransformBaseIndex + BoneIndex] = ReferenceSkeletion.GetBoneName(SkeletalBoneMap[BoneIndex]).ToString();
+		const FColor RandBoneColor(FMath::Rand() % 100 + 5, FMath::Rand() % 100 + 5, FMath::Rand() % 100 + 5, 255);
+		BoneColor[TransformBaseIndex + BoneIndex] = FLinearColor(RandBoneColor);
 
-						int MaterialStart = GeometryCollectionObject->Materials.Num();
-						while (CurrMaterial)
-						{
-							GeometryCollectionObject->Materials.Add(CurrMaterial);
-							CurrMaterial = SkeletalMeshComponent ? SkeletalMeshComponent->GetMaterial(++CurrIdx) : ToRawPtr(SkeletalMeshMaterials[++CurrIdx].MaterialInterface);
+		// Bone Hierarchy - Added at root with no common parent
+		int32 ParentIndex = ReferenceSkeletion.GetParentIndex(SkeletalBoneMap[BoneIndex]);
+		int32 UseParentIndex = ParentIndex + InitialIndex;
+		if (LastParentIndex != UseParentIndex)
+		{
+			LastParentIndex = UseParentIndex;
+		}
+		Parent[TransformBaseIndex + BoneIndex] = UseParentIndex;
+		SimulationType[TransformBaseIndex + BoneIndex] = FGeometryCollection::ESimulationTypes::FST_Rigid;
+	}
 
-						}
+	// Geometry Group
+	TArray<int32> GeometryIndices;
+	GeometryCollectionAlgo::ContiguousArray(GeometryIndices, GeometryCollection->NumElements(FGeometryCollection::GeometryGroup));
+	GeometryCollection->RemoveDependencyFor(FGeometryCollection::GeometryGroup);
+	GeometryCollection->RemoveElements(FGeometryCollection::GeometryGroup, GeometryIndices);
+	::GeometryCollection::AddGeometryProperties(GeometryCollection);
 
-						const TArray<FSkelMeshRenderSection> &StaticMeshSections = SkeletalMesh->GetResourceForRendering()->LODRenderData[0].RenderSections;
+	const TArray<FSkelMeshRenderSection> &StaticMeshSections = SkeletalMesh->GetResourceForRendering()->LODRenderData[0].RenderSections;
 
-						TManagedArray<FGeometryCollectionSection> & Sections = GeometryCollection->Sections;
+	TManagedArray<FGeometryCollectionSection> & Sections = GeometryCollection->Sections;
 
-						for (const FSkelMeshRenderSection &CurrSection : StaticMeshSections)
-						{
-							// create new section
-							int32 SectionIndex = GeometryCollection->AddElements(1, FGeometryCollection::MaterialGroup);
+	for (const FSkelMeshRenderSection &CurrSection : StaticMeshSections)
+	{
+		// create new section
+		int32 SectionIndex = GeometryCollection->AddElements(1, FGeometryCollection::MaterialGroup);
 						
-							Sections[SectionIndex].MaterialID = MaterialStart + CurrSection.MaterialIndex;
+		Sections[SectionIndex].MaterialID = MaterialStartIndex + CurrSection.MaterialIndex;
+
+		Sections[SectionIndex].FirstIndex = IndicesBaseIndex * 3 + CurrSection.BaseIndex;
+		Sections[SectionIndex].MinVertexIndex = VertexBaseIndex + CurrSection.BaseVertexIndex;
 
-							Sections[SectionIndex].FirstIndex = IndicesBaseIndex * 3 + CurrSection.BaseIndex;
-							Sections[SectionIndex].MinVertexIndex = VertexBaseIndex + CurrSection.BaseVertexIndex;
+		Sections[SectionIndex].NumTriangles = CurrSection.NumTriangles;
 
-							Sections[SectionIndex].NumTriangles = CurrSection.NumTriangles;
+		// #todo(dmp): what should we set this to?  SkeletalMesh sections are different
+		// but we are resetting this when the re indexing happens
+		Sections[SectionIndex].MaxVertexIndex = VertexBaseIndex + CurrSection.NumVertices;
 
-							// #todo(dmp): what should we set this to?  SkeletalMesh sections are different
-							// but we are resetting this when the re indexing happens
-							Sections[SectionIndex].MaxVertexIndex = VertexBaseIndex + CurrSection.NumVertices;
+		// set the materialid for all of the faces
+		for (int32 i = Sections[SectionIndex].FirstIndex / 3; i < Sections[SectionIndex].FirstIndex / 3 + Sections[SectionIndex].NumTriangles; ++i)
+		{
+			MaterialID[i] = SectionIndex;
+		}
+	}
+
+	if (bReindexMaterials)
+	{
+		GeometryCollection->ReindexMaterials();
+	}
 
-							// set the materialid for all of the faces
-							for (int32 i = Sections[SectionIndex].FirstIndex / 3; i < Sections[SectionIndex].FirstIndex / 3 + Sections[SectionIndex].NumTriangles; ++i)
-							{
-								MaterialID[i] = SectionIndex;
-							}
-						}
+	return true;
+}
 
-					}
-				}
+const FSkeletalMeshLODRenderData* FGeometryCollectionEngineConversion::GetSkeletalMeshLOD(const USkeletalMesh* SkeletalMesh, int32 LOD)
+{
+	if (const USkeleton* Skeleton = SkeletalMesh->GetSkeleton())
+	{
+		if (const FSkeletalMeshRenderData* SkelMeshRenderData = SkeletalMesh->GetResourceForRendering())
+		{
+			if (SkelMeshRenderData->LODRenderData.IsValidIndex(LOD))
+			{
+				return &SkelMeshRenderData->LODRenderData[LOD];
 			}
+		}
+	}
+	return nullptr;
+}
 
-			if (ReindexMaterials) {
-				GeometryCollection->ReindexMaterials();
+void FGeometryCollectionEngineConversion::AppendSkeletalMesh(const USkeletalMesh* SkeletalMesh, const USkeletalMeshComponent* SkeletalMeshComponent, const FTransform& SkeletalMeshTransform, UGeometryCollection* GeometryCollectionObject, bool bReindexMaterials)
+{
+	//UE_LOG(UGeometryCollectionConversionLogging, Log, TEXT("FGeometryCollectionEngineConversion::AppendSkeletalMesh()"));
+	check(SkeletalMesh);
+	if (GeometryCollectionObject)
+	{
+		TSharedPtr<FGeometryCollection, ESPMode::ThreadSafe> GeometryCollectionPtr = GeometryCollectionObject->GetGeometryCollection();
+		if (FGeometryCollection* GeometryCollection = GeometryCollectionPtr.Get())
+		{
+			int32 MaterialStart = GeometryCollectionObject->Materials.Num();
+			if (AppendSkeletalMesh(SkeletalMesh, MaterialStart, SkeletalMeshTransform, GeometryCollection, bReindexMaterials))
+			{
+				AppendSkeletalMeshMaterials(SkeletalMesh, SkeletalMeshComponent, GeometryCollectionObject);
 			}
 		}
 	}
 }
 
+int32 FGeometryCollectionEngineConversion::AppendSkeletalMeshMaterials(const USkeletalMesh* SkeletalMesh, const USkeletalMeshComponent* SkeletalMeshComponent, UGeometryCollection* GeometryCollectionObject)
+{
+	check(SkeletalMesh);
+	check(SkeletalMeshComponent);
+	check(GeometryCollectionObject);
+
+	const TArray<FSkeletalMaterial>& SkeletalMeshMaterials = SkeletalMesh->GetMaterials();
+
+	int32 CurrIdx = 0;
+	UMaterialInterface* CurrMaterial = SkeletalMeshComponent ? SkeletalMeshComponent->GetMaterial(CurrIdx) : ToRawPtr(SkeletalMeshMaterials[CurrIdx].MaterialInterface);
+
+	int MaterialStart = GeometryCollectionObject->Materials.Num();
+	while (CurrMaterial)
+	{
+		GeometryCollectionObject->Materials.Add(CurrMaterial);
+		CurrMaterial = SkeletalMeshComponent ? SkeletalMeshComponent->GetMaterial(++CurrIdx) : ToRawPtr(SkeletalMeshMaterials[++CurrIdx].MaterialInterface);
+	}
+
+	return MaterialStart;
+}
+
diff --git a/Engine/Source/Runtime/Experimental/GeometryCollectionEngine/Public/GeometryCollection/GeometryCollectionEngineConversion.h b/Engine/Source/Runtime/Experimental/GeometryCollectionEngine/Public/GeometryCollection/GeometryCollectionEngineConversion.h
--- a/Engine/Source/Runtime/Experimental/GeometryCollectionEngine/Public/GeometryCollection/GeometryCollectionEngineConversion.h
+++ b/Engine/Source/Runtime/Experimental/GeometryCollectionEngine/Public/GeometryCollection/GeometryCollectionEngineConversion.h
@@ -5,28 +5,76 @@
 
 class UStaticMesh;
 class USkeletalMesh;
+class USkeletalMeshComponent;
 class UGeometryCollection;
 class UMaterialInterface;
 class UGeometryCollectionComponent;
+class FGeometryCollection;
+class UBodySetup;
+class FSkeletalMeshLODRenderData;
+struct FMeshDescription;
 
 typedef TTuple<const UStaticMesh *, const UStaticMeshComponent *, FTransform> GeometryCollectionStaticMeshConversionTuple;
 typedef TTuple<const USkeletalMesh *, const USkeletalMeshComponent *, FTransform> GeometryCollectionSkeletalMeshConversionTuple;
 
 /**
-* The public interface to this module
-*/
+ * The public interface to this module
+ */
 class GEOMETRYCOLLECTIONENGINE_API FGeometryCollectionEngineConversion
 {
 public:
 
 	/**
-	*  Appends a static mesh to a GeometryCollectionComponent.
-	*  @param StaticMesh : Const mesh to read vertex/normals/index data from
-	*  @param Materials : Materials fetched from the StaticMeshComponent used to configure this geometry
-	*  @param StaticMeshTransform : Mesh transform.
-	*  @param GeometryCollection  : Collection to append the mesh into.
-	*/
-	static void AppendStaticMesh(const UStaticMesh* StaticMesh, const TArray<UMaterialInterface*>& Materials, const FTransform& StaticMeshTransform, UGeometryCollection* GeometryCollectionObject, bool ReindexMaterials = true);
+	 * Appends materials to a GeometryCollectionComponent.
+	 * @param Materials : Materials fetched from the StaticMeshComponent used to configure this geometry
+	 * @param GeometryCollection  : Collection to append the mesh into.
+	 */
+	static int32 AppendMaterials(const TArray<UMaterialInterface*>& Materials, UGeometryCollection* GeometryCollectionObject);
+
+	/**
+	 * Appends a MeshDescription to a GeometryCollection.
+	 * @param MeshDescription : Const mesh description to read vertex/normals/index data from
+	 * @param Name : A name to use for the Mesh, e.g. from the source UObject's GetName()
+	 * @param StartMaterialIndex : Index of materials to be used, e.g. as returned by AppendMaterials() (see below)
+	 * @param StaticMeshTransform : Mesh transform.
+	 * @param GeometryCollection  : Collection to append the mesh into.
+	 * @param BodySetup : Optional collision setup to transfer to the geometry collection
+	 * @param bReindexMaterials	: Whether to reindex materials -- if appending multiple meshes, pass false and call ReindexMaterials afterwards
+	 */
+	static void AppendMeshDescription(const FMeshDescription* MeshDescription, const FString& Name, int32 StartMaterialIndex, const FTransform& StaticMeshTransform, FGeometryCollection* GeometryCollection, UBodySetup* BodySetup = nullptr, bool bReindexMaterials = true);
+
+	/**
+	 * Get a HiRes (or LOD 0 if no HiRes available) MeshDescription for the given static mesh, and make sure it includes normals and tangents.
+	 * (i.e., to pass to AppendMeshDescription)
+	 */
+	static FMeshDescription* GetMaxResMeshDescriptionWithNormalsAndTangents(const UStaticMesh* StaticMesh);
+
+	/**
+	 * Appends a static mesh to a GeometryCollectionComponent.
+	 * @param StaticMesh : Const mesh to read vertex/normals/index data from
+	 * @param Materials : Materials fetched from the StaticMeshComponent used to configure this geometry
+	 * @param StaticMeshTransform : Mesh transform.
+	 * @param GeometryCollection  : Collection to append the mesh into.
+	 * @param bReindexMaterials	: Whether to reindex materials -- if appending multiple meshes, pass false and call ReindexMaterials afterwards
+	 * @return true if succeeded in appending a mesh
+	 */
+	static bool AppendStaticMesh(const UStaticMesh* StaticMesh, const TArray<UMaterialInterface*>& Materials, const FTransform& StaticMeshTransform, UGeometryCollection* GeometryCollectionObject, bool bReindexMaterials = true);
+
+	// TODO: consider having the below fn return bool, and have a helper that returns the StartMaterialIndex from UGeometryCollection
+	//  before any materials are actually added, so we can call the AppendMaterial function after, and not need to use its return value ...
+	//  This would let the is-there-a-valid-mesh logic live inside this function
+	// NOTE: and similar change will need to propagate to similar places ...
+
+	/**
+	 * Appends a static mesh to a GeometryCollectionComponent.
+	 * @param StaticMesh : Const mesh to read vertex/normals/index data from
+	 * @param StartMaterialIndex : Index of materials to be used, e.g. as returned by AppendMaterials() (see below)
+	 * @param StaticMeshTransform : Mesh transform.
+	 * @param GeometryCollection  : Collection to append the mesh into.
+	 * @param bReindexMaterials	: Whether to reindex materials -- if appending multiple meshes, pass false and call ReindexMaterials afterwards
+	 * @return true if succeeded in appending a mesh
+	 */
+	static bool AppendStaticMesh(const UStaticMesh* StaticMesh, int32 StartMaterialIndex, const FTransform& StaticMeshTransform, FGeometryCollection* GeometryCollection, bool bReindexMaterials = true);
 
 	/**
 	*  Appends a static mesh to a GeometryCollectionComponent.
@@ -34,7 +82,7 @@ public:
 	*  @param StaticMeshTransform : Mesh transform.
 	*  @param GeometryCollection  : Collection to append the mesh into.
 	*/
-	static void AppendStaticMesh(const UStaticMesh * StaticMesh, const UStaticMeshComponent *StaticMeshComponent, const FTransform & StaticMeshTransform, UGeometryCollection * GeometryCollection, bool ReindexMaterials = true);
+	static void AppendStaticMesh(const UStaticMesh* StaticMesh, const UStaticMeshComponent *StaticMeshComponent, const FTransform& StaticMeshTransform, UGeometryCollection* GeometryCollection, bool bReindexMaterials = true);
 
 	/**
 	*  Appends a skeletal mesh to a GeometryCollectionComponent.
@@ -42,7 +90,39 @@ public:
 	*  @param SkeletalMeshTransform : Mesh transform.
 	*  @param GeometryCollection    : Collection to append the mesh into.
 	*/
-	static void AppendSkeletalMesh(const USkeletalMesh* SkeletalMesh, const USkeletalMeshComponent *SkeletalMeshComponent, const FTransform & SkeletalMeshTransform, UGeometryCollection * GeometryCollection, bool ReindexMaterials = true);
+	static void AppendSkeletalMesh(const USkeletalMesh* SkeletalMesh, const USkeletalMeshComponent *SkeletalMeshComponent, const FTransform& SkeletalMeshTransform, UGeometryCollection* GeometryCollection, bool bReindexMaterials = true);
+
+	/**
+	*  Appends a skeletal mesh to a GeometryCollection.
+	*  @param SkeletalMesh : Const mesh to read vertex/normals/index data from
+	*  @param MaterialStartIndex : First index of materials to use, as returned by AppendSkeletalMeshMaterials
+	*  @param SkeletalMeshTransform : Mesh transform.
+	*  @param GeometryCollection    : Collection to append the mesh into.
+	*/
+	static bool AppendSkeletalMesh(const USkeletalMesh* SkeletalMesh, int32 MaterialStartIndex, const FTransform& SkeletalMeshTransform, FGeometryCollection* GeometryCollection, bool bReindexMaterials = true);
+
+	/**
+	*  Appends a skeletal mesh to a GeometryCollectionComponent.
+	*  @param SkeletalMeshComponent : Const mesh to read vertex/normals/index data from
+	*  @param SkeletalMeshTransform : Mesh transform.
+	*  @param GeometryCollection    : Collection to append the mesh into.
+	*/
+	static int32 AppendSkeletalMeshMaterials(const USkeletalMesh* SkeletalMesh, const USkeletalMeshComponent* SkeletalMeshComponent, UGeometryCollection* GeometryCollectionObject);
+
+	static const FSkeletalMeshLODRenderData* GetSkeletalMeshLOD(const USkeletalMesh* SkeletalMesh, int32 LOD);
+
+	static int32 AppendGeometryCollectionMaterials(const UGeometryCollection* SourceGeometryCollection, const UGeometryCollectionComponent* GeometryCollectionComponent, UGeometryCollection* TargetGeometryCollectionObject);
+
+	/**
+	 * Appends a GeometryCollection to another GeometryCollection.
+	 * @param SourceGeometryCollection : Const GeometryCollection to read vertex/normals/index data from
+	 * @param MaterialStartIndex : First index of materials to use
+	 * @param GeometryCollectionTransform : GeometryCollection transform.
+	 * @param TargetGeometryCollection  : Collection to append the GeometryCollection into.
+	 * @param bReindexMaterials	: Whether to reindex materials -- if appending multiple meshes, pass false and call ReindexMaterials afterwards
+	 * @return true if succeeded in appending a geometry collection
+	 */
+	static bool AppendGeometryCollection(const FGeometryCollection* SourceGeometryCollection, int32 MaterialStartIndex, const FTransform& GeometryCollectionTransform, FGeometryCollection* TargetGeometryCollection, bool bReindexMaterials = true);
 
 	/**
 	*  Appends a GeometryCollection to a GeometryCollectionComponent.
