Add missing null checks in CollisionConstraint::Setup. Constraint shapes may be null in some unit tests.
#horde https://horde.devtools.epicgames.com/job/62b3a3b7300290d55aed1712?step=f34c

diff --git a/Engine/Source/Runtime/Experimental/Chaos/Private/Chaos/Collision/PBDCollisionConstraint.cpp b/Engine/Source/Runtime/Experimental/Chaos/Private/Chaos/Collision/PBDCollisionConstraint.cpp
--- a/Engine/Source/Runtime/Experimental/Chaos/Private/Chaos/Collision/PBDCollisionConstraint.cpp
+++ b/Engine/Source/Runtime/Experimental/Chaos/Private/Chaos/Collision/PBDCollisionConstraint.cpp
@@ -281,7 +281,7 @@ namespace Chaos
 		// Contact modification may change bIsProbe but we want to store the unmodified value
 		// so that it can be reset between frames.
 		Flags.bIsProbe = 
-		Flags.bIsProbeUnmodified = Shape[0]->GetIsProbe() || Shape[1]->GetIsProbe();
+		Flags.bIsProbeUnmodified = (Shape[0] && Shape[0]->GetIsProbe()) || (Shape[1] && Shape[1]->GetIsProbe());
 
 		const FReal Margin0 = GetImplicit0()->GetMargin();
 		const FReal Margin1 = GetImplicit1()->GetMargin();
