GeometryCollection : HasAttributes method

diff --git a/Engine/Source/Runtime/Experimental/Chaos/Private/GeometryCollection/ManagedArrayCollection.cpp b/Engine/Source/Runtime/Experimental/Chaos/Private/GeometryCollection/ManagedArrayCollection.cpp
--- a/Engine/Source/Runtime/Experimental/Chaos/Private/GeometryCollection/ManagedArrayCollection.cpp
+++ b/Engine/Source/Runtime/Experimental/Chaos/Private/GeometryCollection/ManagedArrayCollection.cpp
@@ -88,6 +88,25 @@ bool FManagedArrayCollection::HasAttribute(FName Name, FName Group) const
 	return Map.Contains(Key);
 }
 
+bool FManagedArrayCollection::HasAttributes(const TArray<FManagedArrayCollection::FManagedType>& Types) const
+{
+	for (const FManagedType& ManagedType : Types)
+	{
+		FKeyType Key = FManagedArrayCollection::MakeMapKey(ManagedType.Name, ManagedType.Group);
+		if(Map.Contains(Key))
+		{
+			const FValueType& FoundValue = Map[Key];
+
+			if (FoundValue.ArrayType != ManagedType.Type)
+			{
+				return false;
+			}
+		}
+	}
+	return true;
+}
+
+
 bool FManagedArrayCollection::IsAttributeDirty(FName Name, FName Group) const
 {
 	const FKeyType Key = FManagedArrayCollection::MakeMapKey(Name, Group);
diff --git a/Engine/Source/Runtime/Experimental/Chaos/Public/GeometryCollection/ManagedArrayCollection.h b/Engine/Source/Runtime/Experimental/Chaos/Public/GeometryCollection/ManagedArrayCollection.h
--- a/Engine/Source/Runtime/Experimental/Chaos/Public/GeometryCollection/ManagedArrayCollection.h
+++ b/Engine/Source/Runtime/Experimental/Chaos/Public/GeometryCollection/ManagedArrayCollection.h
@@ -81,6 +81,22 @@ public:
 		bool bReindexDependentAttibutes;
 	};
 
+	struct FManagedType
+	{
+		FManagedType(EManagedArrayType InType, FName InName, FName InGroup)
+			: Type(InType), Name(InName), Group(InGroup) {}
+		EManagedArrayType Type = EManagedArrayType::FNoneType;
+		FName Name = "";
+		FName Group = "";
+	};
+
+	template<class T>
+	struct TManagedType : FManagedType
+	{
+		TManagedType(FName InName, FName InGroup) :
+			FManagedType(ManagedArrayType<T>(), InName, InGroup) {}			
+	};
+
 	/**
 	* Add an attribute of Type(T) to the group
 	* @param Name - The name of the attribute
@@ -363,6 +379,12 @@ public:
 	*/
 	bool HasAttribute(FName Name, FName Group) const;
 
+	/**
+	* Check for the existence of a attribute
+	* @param TArray<ManagedType> - Attributes to search for
+	*/
+	bool HasAttributes(const TArray<FManagedType>& Types) const;
+
 	/**
 	* Check for the existence of a group
 	* @param Group - The group name
