CIS/SA - Chaos: Fixed double assign to same variable, picked up in SA run

diff --git a/Engine/Source/Runtime/Experimental/Chaos/Private/Chaos/Collision/ParticlePairMidPhase.cpp b/Engine/Source/Runtime/Experimental/Chaos/Private/Chaos/Collision/ParticlePairMidPhase.cpp
--- a/Engine/Source/Runtime/Experimental/Chaos/Private/Chaos/Collision/ParticlePairMidPhase.cpp
+++ b/Engine/Source/Runtime/Experimental/Chaos/Private/Chaos/Collision/ParticlePairMidPhase.cpp
@@ -425,7 +425,7 @@ namespace Chaos
 			PHYSICS_CSV_SCOPED_EXPENSIVE(PhysicsVerbose, NarrowPhase_UpdateConstraintProbe);
 
 			MidPhase.GetCollisionAllocator().ActivateConstraint(Constraint.Get());
-			LastUsedEpoch = LastUsedEpoch = MidPhase.GetCollisionAllocator().GetCurrentEpoch();
+			LastUsedEpoch = MidPhase.GetCollisionAllocator().GetCurrentEpoch();
 			return 1;
 		}
 
