Chaos : add _External postfix to geometry collection  physics proxy to make sure we clearly communicate what can be ran on the gamethread

diff --git a/Engine/Source/Runtime/Experimental/Chaos/Private/PhysicsProxy/GeometryCollectionPhysicsProxy.cpp b/Engine/Source/Runtime/Experimental/Chaos/Private/PhysicsProxy/GeometryCollectionPhysicsProxy.cpp
--- a/Engine/Source/Runtime/Experimental/Chaos/Private/PhysicsProxy/GeometryCollectionPhysicsProxy.cpp
+++ b/Engine/Source/Runtime/Experimental/Chaos/Private/PhysicsProxy/GeometryCollectionPhysicsProxy.cpp
@@ -1503,7 +1503,7 @@ void FGeometryCollectionPhysicsProxy::GetRelevantParticleHandles(
 #endif
 }
 
-void FGeometryCollectionPhysicsProxy::DisableParticles(TArray<int32>&& TransformGroupIndices)
+void FGeometryCollectionPhysicsProxy::DisableParticles_External(TArray<int32>&& TransformGroupIndices)
 {
 	check(IsInGameThread());
 
@@ -1533,7 +1533,7 @@ static void ApplyStrainToAllClusterChildren(Chaos::FRigidClustering& Clustering,
 }
 
 
-void FGeometryCollectionPhysicsProxy::BreakInternalClusterParents(TArray<int32>&& TransformGroupIndices)
+void FGeometryCollectionPhysicsProxy::BreakInternalClusterParents_External(TArray<int32>&& TransformGroupIndices)
 {
 	check(IsInGameThread());
 
@@ -1575,7 +1575,7 @@ void FGeometryCollectionPhysicsProxy::BreakInternalClusterParents(TArray<int32>&
 	}
 }
 
-void FGeometryCollectionPhysicsProxy::BreakClusters(TArray<int32>&& TransformGroupIndices)
+void FGeometryCollectionPhysicsProxy::BreakClusters_External(TArray<int32>&& TransformGroupIndices)
 {
 	check(IsInGameThread());
 
@@ -1605,7 +1605,7 @@ void FGeometryCollectionPhysicsProxy::BreakClusters(TArray<int32>&& TransformGro
 	}
 }
 
-void FGeometryCollectionPhysicsProxy::ApplyStrain(int32 TransformGroupIndex, const FVector& Location, float StrainValue)
+void FGeometryCollectionPhysicsProxy::ApplyStrain_External(int32 TransformGroupIndex, const FVector& Location, float StrainValue)
 {
 	check(IsInGameThread());
 
@@ -1789,7 +1789,7 @@ void FGeometryCollectionPhysicsProxy::BufferGameState()
 }
 
 
-void FGeometryCollectionPhysicsProxy::SetWorldTransform(const FTransform& WorldTransform)
+void FGeometryCollectionPhysicsProxy::SetWorldTransform_External(const FTransform& WorldTransform)
 {
 	check(IsInGameThread());
 	GameThreadPerFrameData.SetWorldTransform(WorldTransform);
diff --git a/Engine/Source/Runtime/Experimental/Chaos/Public/PhysicsProxy/GeometryCollectionPhysicsProxy.h b/Engine/Source/Runtime/Experimental/Chaos/Public/PhysicsProxy/GeometryCollectionPhysicsProxy.h
--- a/Engine/Source/Runtime/Experimental/Chaos/Public/PhysicsProxy/GeometryCollectionPhysicsProxy.h
+++ b/Engine/Source/Runtime/Experimental/Chaos/Public/PhysicsProxy/GeometryCollectionPhysicsProxy.h
@@ -180,7 +180,7 @@ public:
 	bool IsGTCollectionDirty() const { return GameThreadCollection.IsDirty(); }
 
 	// set the world transform ( this needs to be called on the game thread ) 
-	void SetWorldTransform(const FTransform& WorldTransform);
+	void SetWorldTransform_External(const FTransform& WorldTransform);
 
 	const TArray<FClusterHandle*> GetParticles() const
 	{
@@ -247,11 +247,11 @@ public:
 
 	bool GetIsObjectDynamic() const { return IsObjectDynamic; }
 
-	void DisableParticles(TArray<int32>&& TransformGroupIndices);
+	void DisableParticles_External(TArray<int32>&& TransformGroupIndices);
 
-	void BreakInternalClusterParents(TArray<int32>&& TransformGroupIndices);
-	void BreakClusters(TArray<int32>&& TransformGroupIndices);
-	void ApplyStrain(int32 TransformGroupIndex, const FVector& Location, float StrainValue);
+	void BreakInternalClusterParents_External(TArray<int32>&& TransformGroupIndices);
+	void BreakClusters_External(TArray<int32>&& TransformGroupIndices);
+	void ApplyStrain_External(int32 TransformGroupIndex, const FVector& Location, float StrainValue);
 
 	FProxyInterpolationData& GetInterpolationData() { return InterpolationData; }
 	const FProxyInterpolationData& GetInterpolationData() const { return InterpolationData; }
diff --git a/Engine/Source/Runtime/Experimental/GeometryCollectionEngine/Private/GeometryCollection/GeometryCollectionComponent.cpp b/Engine/Source/Runtime/Experimental/GeometryCollectionEngine/Private/GeometryCollection/GeometryCollectionComponent.cpp
--- a/Engine/Source/Runtime/Experimental/GeometryCollectionEngine/Private/GeometryCollection/GeometryCollectionComponent.cpp
+++ b/Engine/Source/Runtime/Experimental/GeometryCollectionEngine/Private/GeometryCollection/GeometryCollectionComponent.cpp
@@ -1916,7 +1916,7 @@ void UGeometryCollectionComponent::OnUpdateTransform(EUpdateTransformFlags Updat
 
 	if (PhysicsProxy)
 	{
-		PhysicsProxy->SetWorldTransform(GetComponentTransform());
+		PhysicsProxy->SetWorldTransform_External(GetComponentTransform());
 	}
 }
 
@@ -3647,11 +3647,11 @@ void UGeometryCollectionComponent::IncrementSleepTimer(float DeltaTime)
 
 		if (ToBreakParent.Num())
 		{
-			PhysicsProxy->BreakInternalClusterParents(MoveTemp(ToBreakParent));
+			PhysicsProxy->BreakInternalClusterParents_External(MoveTemp(ToBreakParent));
 		}
 		if (ToDisable.Num())
 		{
-			PhysicsProxy->DisableParticles(MoveTemp(ToDisable));
+			PhysicsProxy->DisableParticles_External(MoveTemp(ToDisable));
 		}
 	}
 }
@@ -3721,11 +3721,11 @@ void UGeometryCollectionComponent::IncrementBreakTimer(float DeltaTime)
 
 		if (ToCrumble.Num())
 		{
-			PhysicsProxy->BreakClusters(MoveTemp(ToCrumble));
+			PhysicsProxy->BreakClusters_External(MoveTemp(ToCrumble));
 		}
 		if (ToDisable.Num())
 		{
-			PhysicsProxy->DisableParticles(MoveTemp(ToDisable));
+			PhysicsProxy->DisableParticles_External(MoveTemp(ToDisable));
 		}
 	}
 }
@@ -3737,7 +3737,7 @@ void UGeometryCollectionComponent::ApplyExternalStrain(int32 Index, const FVecto
 		const int32 NumTransform = DynamicCollection->NumElements(FGeometryCollection::TransformGroup);
 		if (Index >= 0 && Index < NumTransform)
 		{
-			PhysicsProxy->ApplyStrain(Index, Location, Strain);
+			PhysicsProxy->ApplyStrain_External(Index, Location, Strain);
 		}
 	}
 }
@@ -3755,12 +3755,12 @@ bool UGeometryCollectionComponent::CrumbleCluster(int32 Index)
 			const bool bIsCluster = (DynamicCollection->Children[Index].Num() > 0);
 			if (bIsCluster)
 			{
-				PhysicsProxy->BreakClusters({Index});
+				PhysicsProxy->BreakClusters_External({Index});
 				Result = true;
 			}
 			else if (bHasInternalClusterParent)
 			{
-				PhysicsProxy->BreakInternalClusterParents({Index});
+				PhysicsProxy->BreakInternalClusterParents_External({Index});
 				Result = true;
 			}
 		}
