Add a Convex OverlapRemovalShrinkPercent variable to the geometry collection convex generation settings, to compute overlap removal on scaled-down convex hulls (without baking that scale into the hull).

diff --git a/Engine/Plugins/Experimental/ChaosEditor/Source/FractureEditor/Private/FractureModeSettings.cpp b/Engine/Plugins/Experimental/ChaosEditor/Source/FractureEditor/Private/FractureModeSettings.cpp
--- a/Engine/Plugins/Experimental/ChaosEditor/Source/FractureEditor/Private/FractureModeSettings.cpp
+++ b/Engine/Plugins/Experimental/ChaosEditor/Source/FractureEditor/Private/FractureModeSettings.cpp
@@ -12,5 +12,6 @@ void UFractureModeSettings::ApplyDefaultConvexSettings(FGeometryCollection& Geom
 	Properties.SimplificationThreshold = ConvexSimplificationDistanceThreshold;
 	Properties.CanExceedFraction = ConvexCanExceedFraction;
 	Properties.RemoveOverlaps = ConvexRemoveOverlaps;
+	Properties.OverlapRemovalShrinkPercent = ConvexOverlapRemovalShrinkPercent;
 	GeometryCollection.SetConvexProperties(Properties);
 }
\ No newline at end of file
diff --git a/Engine/Plugins/Experimental/ChaosEditor/Source/FractureEditor/Private/FractureToolConvex.cpp b/Engine/Plugins/Experimental/ChaosEditor/Source/FractureEditor/Private/FractureToolConvex.cpp
--- a/Engine/Plugins/Experimental/ChaosEditor/Source/FractureEditor/Private/FractureToolConvex.cpp
+++ b/Engine/Plugins/Experimental/ChaosEditor/Source/FractureEditor/Private/FractureToolConvex.cpp
@@ -43,6 +43,7 @@ void UFractureConvexSettings::SaveAsDefaults()
 	ModeSettings->ConvexCanExceedFraction = CanExceedFraction;
 	ModeSettings->ConvexSimplificationDistanceThreshold = SimplificationDistanceThreshold;
 	ModeSettings->ConvexRemoveOverlaps = RemoveOverlaps;
+	ModeSettings->ConvexOverlapRemovalShrinkPercent = OverlapRemovalShrinkPercent;
 }
 
 void UFractureConvexSettings::SetFromDefaults()
@@ -52,6 +53,7 @@ void UFractureConvexSettings::SetFromDefaults()
 	CanExceedFraction = ModeSettings->ConvexCanExceedFraction;
 	SimplificationDistanceThreshold = ModeSettings->ConvexSimplificationDistanceThreshold;
 	RemoveOverlaps = ModeSettings->ConvexRemoveOverlaps;
+	OverlapRemovalShrinkPercent = ModeSettings->ConvexOverlapRemovalShrinkPercent;
 
 	UFractureToolConvex* ConvexTool = Cast<UFractureToolConvex>(OwnerTool.Get());
 	ConvexTool->NotifyOfPropertyChangeByTool(this);
@@ -317,6 +319,7 @@ int32 UFractureToolConvex::ExecuteFracture(const FFractureToolContext& FractureC
 		Properties.SimplificationThreshold = ConvexSettings->SimplificationDistanceThreshold;
 		Properties.CanExceedFraction = ConvexSettings->CanExceedFraction;
 		Properties.RemoveOverlaps = ConvexSettings->RemoveOverlaps;
+		Properties.OverlapRemovalShrinkPercent = ConvexSettings->OverlapRemovalShrinkPercent;
 		FractureContext.GetGeometryCollection()->SetConvexProperties(Properties);
 	}
 
diff --git a/Engine/Plugins/Experimental/ChaosEditor/Source/FractureEditor/Private/FractureToolConvex.h b/Engine/Plugins/Experimental/ChaosEditor/Source/FractureEditor/Private/FractureToolConvex.h
--- a/Engine/Plugins/Experimental/ChaosEditor/Source/FractureEditor/Private/FractureToolConvex.h
+++ b/Engine/Plugins/Experimental/ChaosEditor/Source/FractureEditor/Private/FractureToolConvex.h
@@ -38,6 +38,10 @@ public:
 	UPROPERTY(EditAnywhere, Category = Automatic)
 	EConvexOverlapRemoval RemoveOverlaps = EConvexOverlapRemoval::All;
 
+	/** Overlap removal will be computed as if convex hulls were this percentage smaller (in range 0-100) */
+	UPROPERTY(EditAnywhere, Category = Automatic, meta = (UIMin = "0", ClampMax = "99.9"))
+	double OverlapRemovalShrinkPercent = 0.0;
+
 	/** Delete convex hulls from selected clusters.  Does not affect hulls on leaves. */
 	UFUNCTION(CallInEditor, Category = Custom, meta = (DisplayName = "Delete From Selected"))
 	void DeleteFromSelected();
diff --git a/Engine/Plugins/Experimental/ChaosEditor/Source/FractureEditor/Public/FractureModeSettings.h b/Engine/Plugins/Experimental/ChaosEditor/Source/FractureEditor/Public/FractureModeSettings.h
--- a/Engine/Plugins/Experimental/ChaosEditor/Source/FractureEditor/Public/FractureModeSettings.h
+++ b/Engine/Plugins/Experimental/ChaosEditor/Source/FractureEditor/Public/FractureModeSettings.h
@@ -48,6 +48,10 @@ public:
 	UPROPERTY(config, EditAnywhere, Category = "Fracture Mode|Convex Generation Defaults")
 	EConvexOverlapRemoval ConvexRemoveOverlaps = EConvexOverlapRemoval::All;
 
+	/** Default overlap removal shrink percent (in range 0-100) for convex hulls of new geometry collections. Overlap removal will be computed assuming convex shapes will be scaled down by this percentage. */
+	UPROPERTY(config, EditAnywhere, Category = "Fracture Mode|Convex Generation Defaults", meta = (UIMin = "0", ClampMax = "99.9"))
+	float ConvexOverlapRemovalShrinkPercent = 0;
+
 	// Apply Convex Generation Defaults to a GeometryCollection
 	void ApplyDefaultConvexSettings(FGeometryCollection& GeometryCollection) const;
 
diff --git a/Engine/Source/Runtime/Experimental/Chaos/Private/GeometryCollection/GeometryCollectionConvexPropertiesInterface.cpp b/Engine/Source/Runtime/Experimental/Chaos/Private/GeometryCollection/GeometryCollectionConvexPropertiesInterface.cpp
--- a/Engine/Source/Runtime/Experimental/Chaos/Private/GeometryCollection/GeometryCollectionConvexPropertiesInterface.cpp
+++ b/Engine/Source/Runtime/Experimental/Chaos/Private/GeometryCollection/GeometryCollectionConvexPropertiesInterface.cpp
@@ -10,6 +10,7 @@ const FName FGeometryCollectionConvexPropertiesInterface::ConvexFractionRemoveAt
 const FName FGeometryCollectionConvexPropertiesInterface::ConvexSimplificationThresholdAttribute = "SimplificationThreshold";
 const FName FGeometryCollectionConvexPropertiesInterface::ConvexCanExceedFractionAttribute = "CanExceedFraction";
 const FName FGeometryCollectionConvexPropertiesInterface::ConvexRemoveOverlapsMethodAttribute = "RemoveOverlapsMethod";
+const FName FGeometryCollectionConvexPropertiesInterface::ConvexRemoveOverlapsShrinkAttribute = "RemoveOverlapsShrinkPercent";
 
 FGeometryCollectionConvexPropertiesInterface::FGeometryCollectionConvexPropertiesInterface(FGeometryCollection* InGeometryCollection)
 	: FManagedArrayInterface(InGeometryCollection)
@@ -54,6 +55,11 @@ FGeometryCollectionConvexPropertiesInterface::InitializeInterface()
 		ManagedCollection->AddAttribute<int32>(ConvexRemoveOverlapsMethodAttribute, ConvexPropertiesGroup);
 	}
 
+	if (!ManagedCollection->HasAttribute(ConvexRemoveOverlapsShrinkAttribute, ConvexPropertiesGroup))
+	{
+		ManagedCollection->AddAttribute<float>(ConvexRemoveOverlapsShrinkAttribute, ConvexPropertiesGroup);
+	}
+
 	SetDefaultProperty();
 }
 
@@ -81,6 +87,7 @@ FGeometryCollectionConvexPropertiesInterface::GetDefaultProperty() const
 		const TManagedArray<float>& SimplificationThreshold = ManagedCollection->GetAttribute<float>(ConvexSimplificationThresholdAttribute, ConvexPropertiesGroup);
 		const TManagedArray<float>& CanExceedFraction = ManagedCollection->GetAttribute<float>(ConvexCanExceedFractionAttribute, ConvexPropertiesGroup);
 		const TManagedArray<int32>& RemoveOverlaps = ManagedCollection->GetAttribute<int32>(ConvexRemoveOverlapsMethodAttribute, ConvexPropertiesGroup);
+		const TManagedArray<float>& RemoveOverlapsShrinkPercent = ManagedCollection->GetAttribute<float>(ConvexRemoveOverlapsShrinkAttribute, ConvexPropertiesGroup);
 
 		int32 DefaultIndex = Index.Find(INDEX_NONE);
 		if (0 <= DefaultIndex && DefaultIndex < Index.Num())
@@ -90,6 +97,7 @@ FGeometryCollectionConvexPropertiesInterface::GetDefaultProperty() const
 			DefaultProperty.SimplificationThreshold = SimplificationThreshold[DefaultIndex];
 			DefaultProperty.CanExceedFraction = CanExceedFraction[DefaultIndex];
 			DefaultProperty.RemoveOverlaps = (EConvexOverlapRemoval)RemoveOverlaps[DefaultIndex];
+			DefaultProperty.OverlapRemovalShrinkPercent = RemoveOverlapsShrinkPercent[DefaultIndex];
 		}
 	}
 	return DefaultProperty;
@@ -121,12 +129,14 @@ FGeometryCollectionConvexPropertiesInterface::GetConvexProperties(int TransformG
 		const TManagedArray<float>& SimplificationThreshold = ManagedCollection->GetAttribute<float>(ConvexSimplificationThresholdAttribute, ConvexPropertiesGroup);
 		const TManagedArray<float>& CanExceedFraction = ManagedCollection->GetAttribute<float>(ConvexCanExceedFractionAttribute, ConvexPropertiesGroup);
 		const TManagedArray<int32>& RemoveOverlaps = ManagedCollection->GetAttribute<int32>(ConvexRemoveOverlapsMethodAttribute, ConvexPropertiesGroup);
+		const TManagedArray<float>& RemoveOverlapsShrinkPercent = ManagedCollection->GetAttribute<float>(ConvexRemoveOverlapsShrinkAttribute, ConvexPropertiesGroup);
 
 		ConvexProperty.Enable = Enable[PropIndex];
 		ConvexProperty.FractionRemove = FractionRemove[PropIndex];
 		ConvexProperty.SimplificationThreshold = SimplificationThreshold[PropIndex];
 		ConvexProperty.CanExceedFraction = CanExceedFraction[PropIndex];
 		ConvexProperty.RemoveOverlaps = (EConvexOverlapRemoval)RemoveOverlaps[PropIndex];
+		ConvexProperty.OverlapRemovalShrinkPercent = RemoveOverlapsShrinkPercent[PropIndex];
 	}
 	return ConvexProperty;
 }
@@ -141,6 +151,7 @@ FGeometryCollectionConvexPropertiesInterface::SetConvexProperties(const FConvexC
 	TManagedArray<float>& SimplificationThreshold = ManagedCollection->ModifyAttribute<float>(ConvexSimplificationThresholdAttribute, ConvexPropertiesGroup);
 	TManagedArray<float>& CanExceedFraction = ManagedCollection->ModifyAttribute<float>(ConvexCanExceedFractionAttribute, ConvexPropertiesGroup);
 	TManagedArray<int32>& RemoveOverlaps = ManagedCollection->ModifyAttribute<int32>(ConvexRemoveOverlapsMethodAttribute, ConvexPropertiesGroup);
+	TManagedArray<float>& RemoveOverlapsShrinkPercent = ManagedCollection->ModifyAttribute<float>(ConvexRemoveOverlapsShrinkAttribute, ConvexPropertiesGroup);
 
 	int32 AttributeIndex = INDEX_NONE;
 	if (!Index.Contains(TransformGroupIndex))
@@ -160,5 +171,6 @@ FGeometryCollectionConvexPropertiesInterface::SetConvexProperties(const FConvexC
 		SimplificationThreshold[AttributeIndex] = InConvexAttributes.SimplificationThreshold;
 		CanExceedFraction[AttributeIndex] = InConvexAttributes.CanExceedFraction;
 		RemoveOverlaps[AttributeIndex] = (int32)InConvexAttributes.RemoveOverlaps;
+		RemoveOverlapsShrinkPercent[AttributeIndex] = InConvexAttributes.OverlapRemovalShrinkPercent;
 	}
 }
diff --git a/Engine/Source/Runtime/Experimental/Chaos/Private/GeometryCollection/GeometryCollectionConvexUtility.cpp b/Engine/Source/Runtime/Experimental/Chaos/Private/GeometryCollection/GeometryCollectionConvexUtility.cpp
--- a/Engine/Source/Runtime/Experimental/Chaos/Private/GeometryCollection/GeometryCollectionConvexUtility.cpp
+++ b/Engine/Source/Runtime/Experimental/Chaos/Private/GeometryCollection/GeometryCollectionConvexUtility.cpp
@@ -188,6 +188,32 @@ void FilterHullPoints(TArray<Chaos::FConvex::FVec3Type>& InOutPts, double Simpli
 	}
 }
 
+FVector ScaleHullPoints(TArray<Chaos::FConvex::FVec3Type>& InOutPts, double ShrinkPercentage)
+{
+	FVector Pivot = FVector::ZeroVector;
+	if (ShrinkPercentage == 0.0 || InOutPts.IsEmpty())
+	{
+		return Pivot;
+	}
+
+	double ScaleFactor = 1.0 - ShrinkPercentage / 100.0;
+	ensure(ScaleFactor != 0.0);
+
+	FVector Sum(0,0,0);
+	for (Chaos::FConvex::FVec3Type& Pt : InOutPts)
+	{
+		Sum += (FVector)Pt;
+	}
+
+	Pivot = Sum / (double)InOutPts.Num();
+	for (Chaos::FConvex::FVec3Type& Pt : InOutPts)
+	{
+		Pt = Chaos::FConvex::FVec3Type((FVector(Pt) - Pivot) * ScaleFactor + Pivot);
+	}
+
+	return Pivot;
+}
+
 Chaos::FConvex MakeHull(const TArray<Chaos::FConvex::FVec3Type>& Pts, double SimplificationDistanceThreshold)
 {
 	if (SimplificationDistanceThreshold > 0)
@@ -302,6 +328,12 @@ bool SplitHull(const Chaos::FConvex& HullIn, FChaosPlane Plane, bool KeepSide, T
 	return true;
 }
 
+// Add a pivot and verify (in debug) that the pivot array and convexes array are in sync
+void AddPivot(TArray<TUniquePtr<Chaos::FConvex>>& Convexes, TArray<FVector>& ConvexPivots, FVector Pivot)
+{
+	ConvexPivots.Add(Pivot);
+	checkSlow(Convexes.Num() == ConvexPivots.Num());
+}
 
 /// Assumptions: 
 ///		Convexes is initialized to one convex hull for each leaf geometry in a SHARED coordinate space
@@ -309,6 +341,7 @@ bool SplitHull(const Chaos::FConvex& HullIn, FChaosPlane Plane, bool KeepSide, T
 ///		Parents, GeoProximity, and GeometryToTransformIndex are all initialized from the geometry collection
 void CreateNonoverlappingConvexHulls(
 	TArray<TUniquePtr<Chaos::FConvex>>& Convexes,
+	TArray<FVector>& ConvexPivots,
 	TArray<TSet<int32>>& TransformToConvexIndices,
 	TFunctionRef<bool(int32)> HasCustomConvexFn,
 	const TManagedArray<int32>& SimulationType,
@@ -321,7 +354,8 @@ void CreateNonoverlappingConvexHulls(
 	double FracAllowRemove,
 	double SimplificationDistanceThreshold,
 	double CanExceedFraction,
-	EConvexOverlapRemoval OverlapRemovalMethod
+	EConvexOverlapRemoval OverlapRemovalMethod,
+	double ShrinkPercentage
 )
 {
 	bool bRemoveOverlaps = OverlapRemovalMethod != EConvexOverlapRemoval::None;
@@ -764,6 +798,7 @@ void CreateNonoverlappingConvexHulls(
 					if (JoinedHullPts.Num() > 0)
 					{
 						FilterHullPoints(JoinedHullPts, SimplificationDistanceThreshold);
+						FVector Pivot = ScaleHullPoints(JoinedHullPts, ShrinkPercentage);
 						TUniquePtr<Chaos::FConvex> Hull = MakeUnique<Chaos::FConvex>(JoinedHullPts, UE_KINDA_SMALL_NUMBER);
 						bool bIsTooBig = false;
 						if (Volume)
@@ -778,6 +813,7 @@ void CreateNonoverlappingConvexHulls(
 						{
 							ConvexesCS.Lock();
 							int32 ConvexIdx = Convexes.Add(MoveTemp(Hull));
+							AddPivot(Convexes, ConvexPivots, Pivot);
 							ConvexesCS.Unlock();
 							TransformToConvexIndices[Bone].Add(ConvexIdx);
 						}
@@ -1014,10 +1050,12 @@ void HullsFromGeometry(
 	TArray<FTransform>& GlobalTransformArray,
 	TFunctionRef<bool(int32)> HasCustomConvexFn,
 	TArray<TUniquePtr<Chaos::FConvex>>& Convexes,
+	TArray<FVector>& ConvexPivots,
 	TArray<TSet<int32>>& TransformToConvexIndices,
 	const TManagedArray<int32>& SimulationType,
 	int32 RigidType,
-	double SimplificationDistanceThreshold
+	double SimplificationDistanceThreshold,
+	double OverlapRemovalShrinkPercent
 )
 {
 	TArray<FVector> GlobalVertices;
@@ -1030,6 +1068,8 @@ void HullsFromGeometry(
 		GlobalVertices[Idx] = GlobalTransformArray[Geometry.BoneMap[Idx]].TransformPosition(FVector(Geometry.Vertex[Idx]));
 	}
 
+	double ScaleFactor = 1 - OverlapRemovalShrinkPercent / 100.0;
+
 	int32 NumBones = Geometry.TransformToGeometryIndex.Num();
 	TransformToConvexIndices.SetNum(NumBones);
 	FCriticalSection HullCS;
@@ -1043,15 +1083,20 @@ void HullsFromGeometry(
 			for (int32 OrigConvexIdx : OrigConvexData->TransformToConvexIndices[Idx])
 			{
 				TArray<Chaos::FConvex::FVec3Type> HullPts;
+				// As we already have a hull, use the center of mass as the pivot
+				FVector COM = (FVector)OrigConvexData->ConvexHull[OrigConvexIdx]->GetCenterOfMass();
 				HullPts.Reserve(OrigConvexData->ConvexHull[OrigConvexIdx]->NumVertices());
 				for (const Chaos::FConvex::FVec3Type& P : OrigConvexData->ConvexHull[OrigConvexIdx]->GetVertices())
 				{
-					HullPts.Add(Chaos::FConvex::FVec3Type(Transform.TransformPosition((FVector)P)));
+					HullPts.Add(Chaos::FConvex::FVec3Type(Transform.TransformPosition( // transform to the global space
+						(FVector(P) - COM) * ScaleFactor + COM // scale in the original coordinate frame
+						)));
 				}
 				// Do not simplify hulls when we're just trying to transform them
 				TUniquePtr Hull = MakeUnique<Chaos::FConvex>(HullPts, UE_KINDA_SMALL_NUMBER);
 				HullCS.Lock();
 				int32 NewConvexIdx = Convexes.Add(MoveTemp(Hull));
+				AddPivot(Convexes, ConvexPivots, COM);
 				HullCS.Unlock();
 				TransformToConvexIndices[Idx].Add(NewConvexIdx);
 			}
@@ -1069,9 +1114,11 @@ void HullsFromGeometry(
 			}
 			ensure(HullPts.Num() > 0);
 			FilterHullPoints(HullPts, SimplificationDistanceThreshold);
+			FVector Pivot = ScaleHullPoints(HullPts, OverlapRemovalShrinkPercent);
 			TUniquePtr Hull = MakeUnique<Chaos::FConvex>(HullPts, UE_KINDA_SMALL_NUMBER);
 			HullCS.Lock();
 			int32 ConvexIdx = Convexes.Add(MoveTemp(Hull));
+			AddPivot(Convexes, ConvexPivots, Pivot);
 			HullCS.Unlock();
 			TransformToConvexIndices[Idx].Add(ConvexIdx);
 		}
@@ -1081,9 +1128,20 @@ void HullsFromGeometry(
 void TransformHullsToLocal(
 	TArray<FTransform>& GlobalTransformArray,
 	TArray<TUniquePtr<Chaos::FConvex>>& Convexes,
-	TArray<TSet<int32>>& TransformToConvexIndices
+	TArray<FVector>& ConvexPivots,
+	TArray<TSet<int32>>& TransformToConvexIndices,
+	double OverlapRemovalShrinkPercent
 )
 {
+	checkSlow(Convexes.Num() == ConvexPivots.Num());
+	
+	double ScaleFactor = 1.0 - OverlapRemovalShrinkPercent / 100.0;
+	double InvScaleFactor = 1.0;
+	if (ensure(ScaleFactor != 0.0))
+	{
+		InvScaleFactor = 1.0 / ScaleFactor;
+	}
+
 	ParallelFor(TransformToConvexIndices.Num(), [&](int32 Bone)
 	{
 		FTransform& Transform = GlobalTransformArray[Bone];
@@ -1093,7 +1151,13 @@ void TransformHullsToLocal(
 			HullPts.Reset();
 			for (const Chaos::FConvex::FVec3Type& P : Convexes[ConvexIdx]->GetVertices())
 			{
-				HullPts.Add(FVector(Transform.InverseTransformPosition((FVector)P)));
+				FVector PVec(P);
+				if (OverlapRemovalShrinkPercent != 0.0)
+				{
+
+					PVec = (PVec - ConvexPivots[ConvexIdx]) * InvScaleFactor + ConvexPivots[ConvexIdx];
+				}
+				HullPts.Add(FVector(Transform.InverseTransformPosition(PVec)));
 			}
 			// Do not simplify hulls when we're just trying to transform them
 			*Convexes[ConvexIdx] = Chaos::FConvex(HullPts, UE_KINDA_SMALL_NUMBER);
@@ -1177,10 +1241,15 @@ double ComputeGeometryVolume(
 }
 
 FGeometryCollectionConvexUtility::FGeometryCollectionConvexData FGeometryCollectionConvexUtility::CreateNonOverlappingConvexHullData(
-	FGeometryCollection* GeometryCollection, double FracAllowRemove, double SimplificationDistanceThreshold, double CanExceedFraction, EConvexOverlapRemoval OverlapRemovalMethod)
+	FGeometryCollection* GeometryCollection, double FracAllowRemove, double SimplificationDistanceThreshold, double CanExceedFraction, EConvexOverlapRemoval OverlapRemovalMethod,
+	double OverlapRemovalShrinkPercent)
 {
 	check(GeometryCollection);
 
+	// Prevent ~100% shrink percentage, because it would not be reversible ...
+	// Note: Could alternatively just disable convex overlap removal in this case
+	OverlapRemovalShrinkPercent = FMath::Min(99.9, OverlapRemovalShrinkPercent);
+
 	TManagedArray<int32>* CustomConvexFlags = GetCustomConvexFlags(GeometryCollection, false);
 	auto ConvexFlagsAlwaysFalse = [](int32) -> bool { return false; };
 	auto ConvexFlagsFromArray = [&CustomConvexFlags](int32 TransformIdx) -> bool { return (bool)(*CustomConvexFlags)[TransformIdx]; };
@@ -1197,15 +1266,17 @@ FGeometryCollectionConvexUtility::FGeometryCollectionConvexData FGeometryCollect
 	const TManagedArray<TSet<int32>>* GCProximity = GeometryCollection->FindAttribute<TSet<int32>>("Proximity", FGeometryCollection::GeometryGroup);
 	const TManagedArray<float>* Volume = GeometryCollection->FindAttribute<float>("Volume", FGeometryCollection::TransformGroup);
 	TArray<TUniquePtr<Chaos::FConvex>> Convexes;
+	TArray<FVector> ConvexPivots;
 	TArray<TSet<int32>> TransformToConvexIndexArr;
 	TArray<FTransform> GlobalTransformArray;
 	GeometryCollectionAlgo::GlobalMatrices(GeometryCollection->Transform, GeometryCollection->Parent, GlobalTransformArray);
-	HullsFromGeometry(*GeometryCollection, GlobalTransformArray, HasCustomConvexFn, Convexes, TransformToConvexIndexArr, GeometryCollection->SimulationType, FGeometryCollection::ESimulationTypes::FST_Rigid, SimplificationDistanceThreshold);
+	HullsFromGeometry(*GeometryCollection, GlobalTransformArray, HasCustomConvexFn, Convexes, ConvexPivots, TransformToConvexIndexArr, GeometryCollection->SimulationType, 
+		FGeometryCollection::ESimulationTypes::FST_Rigid, SimplificationDistanceThreshold, OverlapRemovalShrinkPercent);
 
-	CreateNonoverlappingConvexHulls(Convexes, TransformToConvexIndexArr, HasCustomConvexFn, GeometryCollection->SimulationType, FGeometryCollection::ESimulationTypes::FST_Rigid, FGeometryCollection::ESimulationTypes::FST_None,
-		GeometryCollection->Parent, GCProximity, GeometryCollection->TransformIndex, Volume, FracAllowRemove, SimplificationDistanceThreshold, CanExceedFraction, OverlapRemovalMethod);
+	CreateNonoverlappingConvexHulls(Convexes, ConvexPivots, TransformToConvexIndexArr, HasCustomConvexFn, GeometryCollection->SimulationType, FGeometryCollection::ESimulationTypes::FST_Rigid, FGeometryCollection::ESimulationTypes::FST_None,
+		GeometryCollection->Parent, GCProximity, GeometryCollection->TransformIndex, Volume, FracAllowRemove, SimplificationDistanceThreshold, CanExceedFraction, OverlapRemovalMethod, OverlapRemovalShrinkPercent);
 
-	TransformHullsToLocal(GlobalTransformArray, Convexes, TransformToConvexIndexArr);
+	TransformHullsToLocal(GlobalTransformArray, Convexes, ConvexPivots, TransformToConvexIndexArr, OverlapRemovalShrinkPercent);
 
 	if (!GeometryCollection->HasGroup("Convex"))
 	{
diff --git a/Engine/Source/Runtime/Experimental/Chaos/Public/GeometryCollection/GeometryCollectionConvexPropertiesInterface.h b/Engine/Source/Runtime/Experimental/Chaos/Public/GeometryCollection/GeometryCollectionConvexPropertiesInterface.h
--- a/Engine/Source/Runtime/Experimental/Chaos/Public/GeometryCollection/GeometryCollectionConvexPropertiesInterface.h
+++ b/Engine/Source/Runtime/Experimental/Chaos/Public/GeometryCollection/GeometryCollectionConvexPropertiesInterface.h
@@ -28,6 +28,8 @@ public :
 	static const FName ConvexCanExceedFractionAttribute;
 	// Attribute
 	static const FName ConvexRemoveOverlapsMethodAttribute;
+	// Attribute
+	static const FName ConvexRemoveOverlapsShrinkAttribute;
 
 
 	struct FConvexCreationProperties {
@@ -36,6 +38,7 @@ public :
 		float SimplificationThreshold = 10.0f;
 		float CanExceedFraction = 0.5f;
 		EConvexOverlapRemoval RemoveOverlaps = EConvexOverlapRemoval::All;
+		float OverlapRemovalShrinkPercent = 0.0f;
 	};
 
 	FGeometryCollectionConvexPropertiesInterface(FGeometryCollection* InGeometryCollection);
diff --git a/Engine/Source/Runtime/Experimental/Chaos/Public/GeometryCollection/GeometryCollectionConvexUtility.h b/Engine/Source/Runtime/Experimental/Chaos/Public/GeometryCollection/GeometryCollectionConvexUtility.h
--- a/Engine/Source/Runtime/Experimental/Chaos/Public/GeometryCollection/GeometryCollectionConvexUtility.h
+++ b/Engine/Source/Runtime/Experimental/Chaos/Public/GeometryCollection/GeometryCollectionConvexUtility.h
@@ -51,11 +51,11 @@ public:
 	 @param FractionAllowRemove					The fraction of a convex body we can cut away to remove overlaps with neighbors, before we fall back to using the hulls of the children directly.  (Does not affect leaves of hierarchy)
 	 @param SimplificationDistanceThreshold		Approximate minimum distance between vertices, below which we remove vertices to generate a simpler convex shape.  If 0.0, no simplification will occur.
 	 @param CanExceedFraction					The fraction by which the convex body volume on a cluster can exceed the volume of the geometry under that cluster (a value of 1 == exceed by 100% == convex hull has 2x the volume of the geometry)
-	 @param bRemoveOverlaps						Whether to cut away overlaps between adjacent convex hulls
 	 @param OverlapRemovalMethod				If bRemoveOverlaps, control which overlaps are removed
+	 @param OverlapRemovalShrinkPercent			Compute overlaps based on objects shrunk by this percentage, so objects that would not overlap with this value set as their 'Collision Object Reduction Percentage' will not be cut
 	 */
 	static FGeometryCollectionConvexData CreateNonOverlappingConvexHullData(FGeometryCollection* GeometryCollection, double FractionAllowRemove = .3, double SimplificationDistanceThreshold = 0.0, double CanExceedFraction = .5, 
-		EConvexOverlapRemoval OverlapRemovalMethod = EConvexOverlapRemoval::All);
+		EConvexOverlapRemoval OverlapRemovalMethod = EConvexOverlapRemoval::All, double OverlapRemovalShrinkPercent = 0.0);
 
 	/** Returns the convex hull of the vertices contained in the specified geometry. */
 	static TUniquePtr<Chaos::FConvex> FindConvexHull(const FGeometryCollection* GeometryCollection, int32 GeometryIndex);
diff --git a/Engine/Source/Runtime/Experimental/GeometryCollectionEngine/Private/GeometryCollection/GeometryCollectionObject.cpp b/Engine/Source/Runtime/Experimental/GeometryCollectionEngine/Private/GeometryCollection/GeometryCollectionObject.cpp
--- a/Engine/Source/Runtime/Experimental/GeometryCollectionEngine/Private/GeometryCollection/GeometryCollectionObject.cpp
+++ b/Engine/Source/Runtime/Experimental/GeometryCollectionEngine/Private/GeometryCollection/GeometryCollectionObject.cpp
@@ -287,7 +287,7 @@ void UGeometryCollection::UpdateConvexGeometry()
 	{
 		FGeometryCollectionConvexPropertiesInterface::FConvexCreationProperties ConvexProperties = GeometryCollection->GetConvexProperties();
 		FGeometryCollectionConvexUtility::CreateNonOverlappingConvexHullData(GeometryCollection.Get(), ConvexProperties.FractionRemove, 
-			ConvexProperties.SimplificationThreshold, ConvexProperties.CanExceedFraction, ConvexProperties.RemoveOverlaps);
+			ConvexProperties.SimplificationThreshold, ConvexProperties.CanExceedFraction, ConvexProperties.RemoveOverlaps, ConvexProperties.OverlapRemovalShrinkPercent);
 		InvalidateCollection();
 	}
 #endif
