Chaos geometry collection : add support for static mesh buildScale parameter

diff --git a/Engine/Source/Runtime/Experimental/GeometryCollectionEngine/Private/GeometryCollection/GeometryCollectionEngineConversion.cpp b/Engine/Source/Runtime/Experimental/GeometryCollectionEngine/Private/GeometryCollection/GeometryCollectionEngineConversion.cpp
--- a/Engine/Source/Runtime/Experimental/GeometryCollectionEngine/Private/GeometryCollection/GeometryCollectionEngineConversion.cpp
+++ b/Engine/Source/Runtime/Experimental/GeometryCollectionEngine/Private/GeometryCollection/GeometryCollectionEngineConversion.cpp
@@ -72,6 +72,18 @@ static bool IsImportableImplicitObjectType(Chaos::EImplicitObjectType Type)
 	return (InnerType == Chaos::ImplicitObjectType::Box || InnerType == Chaos::ImplicitObjectType::Sphere || InnerType == Chaos::ImplicitObjectType::Capsule || InnerType == Chaos::ImplicitObjectType::Convex);
 }
 
+static FVector GetMeshBuildScale3D(const UStaticMesh& StaticMesh)
+{
+#if WITH_EDITOR
+	const TArray<FStaticMeshSourceModel>& SourceModels = StaticMesh.GetSourceModels();
+	if (SourceModels.Num() > 0)
+	{
+		return SourceModels[0].BuildSettings.BuildScale3D;
+	}
+#endif
+	return FVector::One();
+}
+
 void FGeometryCollectionEngineConversion::AppendMeshDescription(
 	const FMeshDescription* MeshDescription, const FString& Name, int32 MaterialStartIndex, const FTransform& StaticMeshTransform, 
 	FGeometryCollection* GeometryCollection, UBodySetup* BodySetup, bool ReindexMaterials)
@@ -420,18 +432,26 @@ bool FGeometryCollectionEngineConversion::AppendStaticMesh(const UStaticMesh* St
 bool FGeometryCollectionEngineConversion::AppendStaticMesh(const UStaticMesh* StaticMesh, int32 StartMaterialIndex, const FTransform& StaticMeshTransform, FGeometryCollection* GeometryCollection, bool bReindexMaterials)
 {
 #if WITH_EDITORONLY_DATA
+	if (StaticMesh)
+	{
+		FMeshDescription* MeshDescription = GetMaxResMeshDescriptionWithNormalsAndTangents(StaticMesh);
 
-	FMeshDescription* MeshDescription = GetMaxResMeshDescriptionWithNormalsAndTangents(StaticMesh);
-
-	check(GeometryCollection);
+		check(GeometryCollection);
 
-	if (MeshDescription)
-	{
-		AppendMeshDescription(MeshDescription, StaticMesh->GetName(), StartMaterialIndex, StaticMeshTransform, GeometryCollection, StaticMesh->GetBodySetup(), bReindexMaterials);
-		return true;
+		if (MeshDescription)
+		{
+			const FVector MeshBuildScale3D = GetMeshBuildScale3D(*StaticMesh);
+			const FTransform MeshTransform(
+					StaticMeshTransform.GetRotation(),
+					StaticMeshTransform.GetTranslation(),
+					StaticMeshTransform.GetScale3D() * MeshBuildScale3D
+			);
+			
+			AppendMeshDescription(MeshDescription, StaticMesh->GetName(), StartMaterialIndex, MeshTransform, GeometryCollection, StaticMesh->GetBodySetup(), bReindexMaterials);
+			return true;
+		}
 	}
 #endif //WITH_EDITORONLY_DATA
-
 	return false;
 }
 
