Dataflow - Validation of node instantiations

diff --git a/Engine/Plugins/Experimental/Dataflow/Source/DataflowEditor/Private/Dataflow/DataflowEditorActions.cpp b/Engine/Plugins/Experimental/Dataflow/Source/DataflowEditor/Private/Dataflow/DataflowEditorActions.cpp
--- a/Engine/Plugins/Experimental/Dataflow/Source/DataflowEditor/Private/Dataflow/DataflowEditorActions.cpp
+++ b/Engine/Plugins/Experimental/Dataflow/Source/DataflowEditor/Private/Dataflow/DataflowEditorActions.cpp
@@ -116,7 +116,7 @@ void FDataflowEditorCommands::OnSelectedNodesChanged(TSharedPtr<IStructureDetail
 					{
 						if (TSharedPtr<FDataflowNode> DataflowNode = DataflowGraph->FindBaseNode(EdNode->GetDataflowNodeGuid()))
 						{
-							TSharedPtr<FStructOnScope> Struct(DataflowNode->NewScructOnScope());
+							TSharedPtr<FStructOnScope> Struct(DataflowNode->NewStructOnScope());
 							PropertiesEditor->SetStructureData(Struct);
 						}
 					}
diff --git a/Engine/Plugins/Experimental/Dataflow/Source/DataflowEditor/Private/Dataflow/DataflowSNode.cpp b/Engine/Plugins/Experimental/Dataflow/Source/DataflowEditor/Private/Dataflow/DataflowSNode.cpp
--- a/Engine/Plugins/Experimental/Dataflow/Source/DataflowEditor/Private/Dataflow/DataflowSNode.cpp
+++ b/Engine/Plugins/Experimental/Dataflow/Source/DataflowEditor/Private/Dataflow/DataflowSNode.cpp
@@ -57,38 +57,35 @@ UEdGraphNode* FAssetSchemaAction_Dataflow_CreateNode_DataflowEdNode::PerformActi
 	if (UDataflow* Dataflow = Cast<UDataflow>(ParentGraph))
 	{
 		const FName NodeName = MakeUniqueObjectName(Dataflow, UDataflowEdNode::StaticClass(), FName(GetMenuDescription().ToString()));
-		if (UDataflowEdNode* EdNode = NewObject<UDataflowEdNode>(Dataflow, UDataflowEdNode::StaticClass(), NodeName))
+		if (Dataflow::FNodeFactory* Factory = Dataflow::FNodeFactory::GetInstance())
 		{
-			//const FScopedTransaction Transaction(LOCTEXT("DataflowNewNode", "Generic Graph Editor: New Node"));
-			Dataflow->Modify();
-			if (FromPin != nullptr)
-				FromPin->Modify();
-
-			Dataflow->AddNode(EdNode, true, bSelectNewNode);
-
-			EdNode->CreateNewGuid();
-			EdNode->PostPlacedNewNode();
-
-			if (Dataflow::FNodeFactory* Factory = Dataflow::FNodeFactory::GetInstance())
+			if (TSharedPtr<FDataflowNode> DataflowNode = Factory->NewNodeFromRegisteredType(*Dataflow->GetDataflow(), { FGuid::NewGuid(),NodeTypeName,NodeName}))
 			{
-				if (TSharedPtr<FDataflowNode> DataflowNode = Factory->NewNodeFromRegisteredType(*Dataflow->GetDataflow(), { FGuid::NewGuid(),NodeTypeName,NodeName}))
+				if (UDataflowEdNode* EdNode = NewObject<UDataflowEdNode>(Dataflow, UDataflowEdNode::StaticClass(), NodeName))
 				{
+					Dataflow->Modify();
+					if (FromPin != nullptr)
+						FromPin->Modify();
+		
+					Dataflow->AddNode(EdNode, true, bSelectNewNode);
+		
+					EdNode->CreateNewGuid();
+					EdNode->PostPlacedNewNode();
+		
 					EdNode->SetDataflowGraph(Dataflow->GetDataflow());
 					EdNode->SetDataflowNodeGuid(DataflowNode->GetGuid());
 					EdNode->AllocateDefaultPins();
+		
+					EdNode->AutowireNewNode(FromPin);
+		
+					EdNode->NodePosX = Location.X;
+					EdNode->NodePosY = Location.Y;
+		
+					EdNode->SetFlags(RF_Transactional);
+		
+					return EdNode;
 				}
 			}
-
-			//EdNode->AllocateDefaultPins();
-			EdNode->AutowireNewNode(FromPin);
-
-			EdNode->NodePosX = Location.X;
-			EdNode->NodePosY = Location.Y;
-
-			//EdNode->GenericGraphNode->SetFlags(RF_Transactional);
-			EdNode->SetFlags(RF_Transactional);
-
-			return EdNode;
 		}
 	}
 	return nullptr;
diff --git a/Engine/Source/Runtime/Experimental/Dataflow/Core/Private/Dataflow/DataflowNode.cpp b/Engine/Source/Runtime/Experimental/Dataflow/Core/Private/Dataflow/DataflowNode.cpp
--- a/Engine/Source/Runtime/Experimental/Dataflow/Core/Private/Dataflow/DataflowNode.cpp
+++ b/Engine/Source/Runtime/Experimental/Dataflow/Core/Private/Dataflow/DataflowNode.cpp
@@ -2,6 +2,7 @@
 
 #include "Dataflow/DataflowNode.h"
 
+#include "ChaosLog.h"
 #include "Dataflow/DataflowInputOutput.h"
 #include "Dataflow/DataflowArchive.h"
 
@@ -203,7 +204,7 @@ void FDataflowNode::InvalidateOutputs()
 
 void FDataflowNode::RegisterInputConnection(const void* Data)
 {
-	if (TUniquePtr<FStructOnScope> ScriptOnStruct = TUniquePtr<FStructOnScope>(NewScructOnScope()))
+	if (TUniquePtr<FStructOnScope> ScriptOnStruct = TUniquePtr<FStructOnScope>(NewStructOnScope()))
 	{
 		if (const UStruct* Struct = ScriptOnStruct->GetStruct())
 		{
@@ -224,7 +225,7 @@ void FDataflowNode::RegisterInputConnection(const void* Data)
 
 void FDataflowNode::RegisterOutputConnection(const void* Data)
 {
-	if (TUniquePtr<FStructOnScope> ScriptOnStruct = TUniquePtr<FStructOnScope>(NewScructOnScope()))
+	if (TUniquePtr<FStructOnScope> ScriptOnStruct = TUniquePtr<FStructOnScope>(NewStructOnScope()))
 	{
 		if (const UStruct* Struct = ScriptOnStruct->GetStruct())
 		{
@@ -246,38 +247,37 @@ void FDataflowNode::RegisterOutputConnection(const void* Data)
 
 bool FDataflowNode::ValidateConnections()
 {
-	bool bResult = true;
-	if (TUniquePtr<FStructOnScope> ScriptOnStruct = TUniquePtr<FStructOnScope>(NewScructOnScope()))
+	bValid = true;
+	if (const TUniquePtr<FStructOnScope> ScriptOnStruct = TUniquePtr<FStructOnScope>(NewStructOnScope()))
 	{
 		if (const UStruct* Struct = ScriptOnStruct->GetStruct())
 		{
 			for (TFieldIterator<FProperty> PropertyIt(Struct); PropertyIt; ++PropertyIt)
 			{
-				FProperty* Property = *PropertyIt;
+				const FProperty* Property = *PropertyIt;
 				FName PropName(Property->GetName());
-				FName PropType(Property->GetCPPType());
 #if WITH_EDITORONLY_DATA
 				if (Property->HasMetaData(TEXT("DataflowInput")))
 				{
-					FDataflowInput* Input = FindInput(PropName);
-					if (ensureMsgf(Input, TEXT("Missing dataflow RegisterInputConnection in constructor for (%s:%s)"), *GetName().ToString(), *PropName.ToString()))
+					if (!FindInput(PropName))
 					{
-						bResult = false;
+						UE_LOG(LogChaos, Warning, TEXT("Missing dataflow RegisterInputConnection in constructor for (%s:%s)"), *GetName().ToString(), *PropName.ToString())
+						bValid = false;
 					}
 				}
 				if (Property->HasMetaData(TEXT("DataflowOutput")))
 				{
-					FDataflowOutput* Output = FindOutput(PropName);
-					if(ensureMsgf(Output,TEXT("Missing dataflow RegisterOutputConnection in constructor for (%s:%s)"), *GetName().ToString(),*PropName.ToString()))
+					if(!FindOutput(PropName))
 					{
-						bResult = false;
+						UE_LOG(LogChaos, Warning, TEXT("Missing dataflow RegisterOutputConnection in constructor for (%s:%s)"), *GetName().ToString(),*PropName.ToString());
+						bValid = false;
 					}
 				}
 #endif
 			}
 		}
 	}
-	return bResult;
+	return bValid;
 }
 
 
diff --git a/Engine/Source/Runtime/Experimental/Dataflow/Core/Private/Dataflow/DataflowNodeFactory.cpp b/Engine/Source/Runtime/Experimental/Dataflow/Core/Private/Dataflow/DataflowNodeFactory.cpp
--- a/Engine/Source/Runtime/Experimental/Dataflow/Core/Private/Dataflow/DataflowNodeFactory.cpp
+++ b/Engine/Source/Runtime/Experimental/Dataflow/Core/Private/Dataflow/DataflowNodeFactory.cpp
@@ -3,6 +3,7 @@
 #include "Dataflow/DataflowNodeFactory.h"
 
 #include "Dataflow/DataflowNode.h"
+#include "Misc/MessageDialog.h"
 
 namespace Dataflow
 {
@@ -12,7 +13,16 @@ namespace Dataflow
 	{ 
 		if (ClassMap.Contains(Param.Type))
 		{
-			return Graph.AddNode(ClassMap[Param.Type](Param));
+			TUniquePtr<FDataflowNode> Node = ClassMap[Param.Type](Param);
+			if(Node->IsValid())
+			{
+				return Graph.AddNode(MoveTemp(Node));
+			}
+			
+			const FText ErrorTitle = FText::FromString("Node Factory");
+			const FString ErrorMessageString = FString::Printf(TEXT("Cannot create Node %s. Node Type %s is not well defined."), *Node->GetName().ToString(), *Node->GetDisplayName().ToString());
+			const FText ErrorMessage = FText::FromString(ErrorMessageString);
+			FMessageDialog::Debugf(ErrorMessage, &ErrorTitle);
 		}
 		return TSharedPtr<FDataflowNode>(nullptr);
 	}
diff --git a/Engine/Source/Runtime/Experimental/Dataflow/Core/Public/Dataflow/DataflowGraph.h b/Engine/Source/Runtime/Experimental/Dataflow/Core/Public/Dataflow/DataflowGraph.h
--- a/Engine/Source/Runtime/Experimental/Dataflow/Core/Public/Dataflow/DataflowGraph.h
+++ b/Engine/Source/Runtime/Experimental/Dataflow/Core/Public/Dataflow/DataflowGraph.h
@@ -66,6 +66,13 @@ namespace Dataflow
 			return NewNode;
 		}
 
+		template<class T> TSharedPtr<T> AddNode(TUniquePtr<T> &&InNode)
+		{
+			TSharedPtr<T> NewNode(InNode.Release());
+			Nodes.AddUnique(NewNode);
+			return NewNode;
+		}
+
 		TSharedPtr<FDataflowNode> FindBaseNode(FGuid InGuid)
 		{
 			for (TSharedPtr<FDataflowNode> Node : Nodes)
diff --git a/Engine/Source/Runtime/Experimental/Dataflow/Core/Public/Dataflow/DataflowNode.h b/Engine/Source/Runtime/Experimental/Dataflow/Core/Public/Dataflow/DataflowNode.h
--- a/Engine/Source/Runtime/Experimental/Dataflow/Core/Public/Dataflow/DataflowNode.h
+++ b/Engine/Source/Runtime/Experimental/Dataflow/Core/Public/Dataflow/DataflowNode.h
@@ -101,7 +101,7 @@ struct DATAFLOWCORE_API FDataflowNode
 	//
 
 	virtual void SerializeInternal(FArchive& Ar) { check(false); }
-	virtual FStructOnScope* NewScructOnScope() { return nullptr; }
+	virtual FStructOnScope* NewStructOnScope() { return nullptr; }
 
 	/** Register the Input and Outputs after the creation in the factory */
 	void RegisterInputConnection(const void*);
@@ -163,6 +163,12 @@ struct DATAFLOWCORE_API FDataflowNode
 
 	bool ValidateConnections();
 
+	bool IsValid() const { return bValid; }
+
+private:
+
+	bool bValid = true;
+
 };
 
 namespace Dataflow
@@ -176,7 +182,7 @@ namespace Dataflow
 		{A::StaticType(),A::StaticDisplay(),A::StaticCategory(),					\
 			A::StaticTags(),A::StaticToolTip()},									\
 		[](const FNewNodeParameters& InParam){										\
-				A* Val=new A({InParam.Name}, InParam.Guid);							\
+				TUniquePtr<A> Val = MakeUnique<A>(FNodeParameters{InParam.Name}, InParam.Guid);    \
 				Val->ValidateConnections(); return Val;});
 
 #define DATAFLOW_NODE_DEFINE_INTERNAL(TYPE, DISPLAY_NAME, CATEGORY, TAGS)			\
@@ -187,7 +193,7 @@ public:																				\
 	static FString StaticTags() {return TAGS;}										\
 	static FString StaticToolTip() {return FString("Create a dataflow node.");}		\
 	virtual FName GetType() const { return #TYPE; }									\
-	virtual FStructOnScope* NewScructOnScope() override {							\
+	virtual FStructOnScope* NewStructOnScope() override {							\
 	   return new FStructOnScope(TYPE::StaticStruct(), (uint8*)this);}				\
 	virtual void SerializeInternal(FArchive& Ar) override {							\
 		UScriptStruct* const Struct = TYPE::StaticStruct();							\
diff --git a/Engine/Source/Runtime/Experimental/Dataflow/Core/Public/Dataflow/DataflowNodeFactory.h b/Engine/Source/Runtime/Experimental/Dataflow/Core/Public/Dataflow/DataflowNodeFactory.h
--- a/Engine/Source/Runtime/Experimental/Dataflow/Core/Public/Dataflow/DataflowNodeFactory.h
+++ b/Engine/Source/Runtime/Experimental/Dataflow/Core/Public/Dataflow/DataflowNodeFactory.h
@@ -40,7 +40,7 @@ namespace Dataflow
 	//
 	class DATAFLOWCORE_API FNodeFactory
 	{
-		typedef TFunction<FDataflowNode* (const FNewNodeParameters&)> FNewNodeFunction;
+		typedef TFunction<TUniquePtr<FDataflowNode> (const FNewNodeParameters&)> FNewNodeFunction;
 
 		// All Maps indexed by TypeName
 		TMap<FName, FNewNodeFunction > ClassMap;		// [TypeName] -> NewNodeFunction
