DataFlow node for computing fiber field.

diff --git a/Engine/Source/Runtime/Experimental/Chaos/Public/Chaos/Math/Poisson.h b/Engine/Source/Runtime/Experimental/Chaos/Public/Chaos/Math/Poisson.h
--- a/Engine/Source/Runtime/Experimental/Chaos/Public/Chaos/Math/Poisson.h
+++ b/Engine/Source/Runtime/Experimental/Chaos/Public/Chaos/Math/Poisson.h
@@ -2,8 +2,9 @@
 #pragma once
 
 //#include "ChaosFlesh/FleshCollectionUtility.h"
-#include "Chaos/Math/Krylov.h"
 #include "Chaos/Framework/Parallel.h"
+#include "Chaos/Math/Krylov.h"
+#include "Chaos/Vector.h"
 #include "Math/Matrix.h"
 #include "Math/Vector.h"
 
@@ -63,10 +64,10 @@ inline void RowMaj3x3Transpose(const T* A, T* Transpose)
 	}
 }
 
-template <class T>
-inline Chaos::TVector<T,3> RowMaj3x3Multiply(const T* A, const Chaos::TVector<T,3>& x)
+template <class T,class TV>
+inline TV RowMaj3x3Multiply(const T* A, const TV& x)
 {
-	TVector<T,3> Result(0,0,0);
+	TV Result(0,0,0);
 	for (int32 i = 0; i < 3; i++) 
 	{
 		for (int32 j = 0; j < 3; j++) 
@@ -79,7 +80,7 @@ inline Chaos::TVector<T,3> RowMaj3x3Multiply(const T* A, const Chaos::TVector<T,
 
 
 
-template <class T, class TV=FVector3f, class TV_INT=FIntVector4, int d=3>
+template <class T, class TV=FVector3f, class TV_INT=FIntVector4, int32 d=3>
 void
 ComputeDeInverseAndElementMeasures(
 	const TArray<TV_INT>& Mesh, 
@@ -182,14 +183,14 @@ Laplacian(
 
 		const TV_INT& Elem = Mesh[e];
 
+		TVector<T, d> grad_Nie;
 		for (int32 ie = 0; ie < d+1; ie++)
 		{
-			Chaos::TVector<T,3> grad_Nie = RowMaj3x3Multiply(De_inverse_transpose, grad_Nie_hat[ie]);
-
+			grad_Nie = RowMaj3x3Multiply(De_inverse_transpose, grad_Nie_hat[ie]);
 			for (int32 je = 0; je < d+1; je++)
 			{
-				Chaos::TVector<T, 3> grad_Nje = RowMaj3x3Multiply(De_inverse_transpose, grad_Nie_hat[je]);
-				T Ae_ieje = Chaos::TVector<T,3>::DotProduct(grad_Nie, grad_Nje) * measure[e];
+				Chaos::TVector<T, d> grad_Nje = RowMaj3x3Multiply(De_inverse_transpose, grad_Nie_hat[je]);
+				T Ae_ieje = Chaos::TVector<T, d>::DotProduct(grad_Nie, grad_Nje) * measure[e];
 				element_contributions[(d + 1) * e + ie] += Ae_ieje * u[Elem[je]];
 			}
 		}
@@ -253,9 +254,10 @@ LaplacianEnergy(
 			RowMaj3x3Transpose(Deinv, De_inverse_transpose);
 
 			const TV_INT& Elem = Mesh[e];
+			TVector<T, d> grad_Nie;
 			for (int32 ie = 0; ie < d + 1; ie++) 
 			{
-				TVector<T, d> grad_Nie = RowMaj3x3Multiply(De_inverse_transpose, grad_Nie_hat[ie]);
+				grad_Nie = RowMaj3x3Multiply(De_inverse_transpose, grad_Nie_hat[ie]);
 				for (int32 je = 0; je < d + 1; je++) 
 				{
 					Chaos::TVector<T, 3> grad_Nje = RowMaj3x3Multiply(De_inverse_transpose, grad_Nie_hat[je]);
@@ -375,9 +377,10 @@ ComputeFiberField(
 
 		const TV_INT& Elem = Mesh[e];
 		TV gradient(0);
+		TVector<T, d> grad_Nie;
 		for (size_t ie = 0; ie < d + 1; ie++) 
 		{
-			TVector<T, d> grad_Nie = RowMaj3x3Multiply(De_inverse_transpose, grad_Nie_hat[ie]);
+			grad_Nie = RowMaj3x3Multiply(De_inverse_transpose, grad_Nie_hat[ie]);
 			gradient += grad_Nie * u[Elem[ie]];
 		}
 		const T Len = gradient.Length();
