Add node for generating a tetrahedral mesh from a uniform grid.

diff --git a/Engine/Source/Runtime/Experimental/Chaos/Public/Chaos/Utilities.h b/Engine/Source/Runtime/Experimental/Chaos/Public/Chaos/Utilities.h
--- a/Engine/Source/Runtime/Experimental/Chaos/Public/Chaos/Utilities.h
+++ b/Engine/Source/Runtime/Experimental/Chaos/Public/Chaos/Utilities.h
@@ -688,8 +688,8 @@ namespace Chaos
 			}
 		}
 
-		template <class T>
-		void TetMeshFromGrid(const TUniformGrid<T, 3>& Grid, TArray<TVector<int32, 4>>& Mesh, TArray<TVector<T, 3>>& X)
+		template <class T, class TV, class TV_INT4>
+		void TetMeshFromGrid(const TUniformGrid<T, 3>& Grid, TArray<TV_INT4>& Mesh, TArray<TV>& X)
 		{
 			//mesh.resize(20 * (grid.grid_N[0] - 1) * (grid.grid_N[1] - 1) * (grid.grid_N[2] - 1));
 			Mesh.SetNum(
