Fix a bug on the HeightField walking. Walking was disabled if it was done in a negative direction.

diff --git a/Engine/Source/Runtime/Experimental/Chaos/Private/Chaos/HeightField.cpp b/Engine/Source/Runtime/Experimental/Chaos/Private/Chaos/HeightField.cpp
--- a/Engine/Source/Runtime/Experimental/Chaos/Private/Chaos/HeightField.cpp
+++ b/Engine/Source/Runtime/Experimental/Chaos/Private/Chaos/HeightField.cpp
@@ -1077,7 +1077,7 @@ namespace Chaos
 			const FVec3 ScaleSign = GeomData.Scale.GetSignVector();
 
 			const FVec3 DirScaled = Dir / ScaledDx;
-			const FReal SumPlaneAxis = DirScaled[0] + DirScaled[1];
+			const FReal SumPlaneAxis = FMath::Abs(DirScaled[0]) + FMath::Abs(DirScaled[1]);
 			const bool bCanWalk = SumPlaneAxis > UE_SMALL_NUMBER;
 			if (bCanWalk)
 			{
