HeightField optimizations :
- cache the Scale in Simd
- Remove a hot unnecessary if when iterating through the cells

diff --git a/Engine/Source/Runtime/Experimental/Chaos/Private/Chaos/HeightField.cpp b/Engine/Source/Runtime/Experimental/Chaos/Private/Chaos/HeightField.cpp
--- a/Engine/Source/Runtime/Experimental/Chaos/Private/Chaos/HeightField.cpp
+++ b/Engine/Source/Runtime/Experimental/Chaos/Private/Chaos/HeightField.cpp
@@ -619,6 +619,7 @@ namespace Chaos
 		OutData.MinValue = ToRealFunc(BufferView[0]);
 		OutData.MaxValue = ToRealFunc(BufferView[0]);
 		OutData.Scale = InScale;
+		OutData.ScaleSimd = MakeVectorRegisterFloatFromDouble(MakeVectorRegisterDouble(InScale.X, InScale.Y, InScale.Z, 0.0));
 
 		for(int32 HeightIndex = 1; HeightIndex < NumHeights; ++HeightIndex)
 		{
diff --git a/Engine/Source/Runtime/Experimental/Chaos/Public/Chaos/HeightField.h b/Engine/Source/Runtime/Experimental/Chaos/Public/Chaos/HeightField.h
--- a/Engine/Source/Runtime/Experimental/Chaos/Public/Chaos/HeightField.h
+++ b/Engine/Source/Runtime/Experimental/Chaos/Public/Chaos/HeightField.h
@@ -246,6 +246,7 @@ namespace Chaos
 		void SetScale(const FVec3& InScale)
 		{
 			GeomData.Scale = InScale;
+			GeomData.ScaleSimd = MakeVectorRegisterFloatFromDouble(MakeVectorRegisterDouble(InScale.X, InScale.Y, InScale.Z, 0.0));
 		}
 
 		template<typename InStorageType>
@@ -276,6 +277,7 @@ namespace Chaos
 			TArray<StorageType> Heights;
 			TArray<uint8> MaterialIndices;
 			FVec3 Scale;
+			VectorRegister4Float ScaleSimd;
 			FReal MinValue;
 			FReal MaxValue;
 			uint16 NumRows;
@@ -400,25 +402,25 @@ namespace Chaos
 				int32 LastIndexX = FMath::Max<int32>(CellIdx[0], EndIndexX);
 				// Increment LastIndex to look into the four corners of a cell 
 				LastIndexX++;
-				LastIndexX = FMath::Min<int32>(LastIndexX, NumCols);
+				LastIndexX = FMath::Min<int32>(LastIndexX, NumCols-1);
 				
 				const int32 EndIndexY = CellIdx[1] + Area[1];
 				int32 FirstIndexY = FMath::Min<int32>(CellIdx[1], EndIndexY);
 				FirstIndexY = FMath::Max<int32>(FirstIndexY, 0);
 				int32 LastIndexY = FMath::Max<int32>(CellIdx[1], EndIndexY);
 				LastIndexY++;
-				LastIndexY = FMath::Min<int32>(LastIndexY, NumRows);
+				LastIndexY = FMath::Min<int32>(LastIndexY, NumRows-1);
 
 				for (int IndexY = FirstIndexY; IndexY <= LastIndexY; IndexY++)
 				{
 					for (int IndexX = FirstIndexX; IndexX <= LastIndexX; IndexX++)
 					{
 						const int32 Index = IndexY * (NumCols - 1) + IndexX + IndexY;
-						if (Index < Heights.Num())
-						{
-							MinHeight = FMath::Min<FRealSingle>(Heights[Index], MinHeight);
-							MaxHeight = FMath::Max<FRealSingle>(Heights[Index], MaxHeight);
-						}
+						check(Index < Heights.Num());
+						const FRealSingle CurrHeight = Heights[Index];
+						MinHeight = FMath::Min<FRealSingle>(CurrHeight, MinHeight);
+						MaxHeight = FMath::Max<FRealSingle>(CurrHeight, MaxHeight);
+						
 					}
 				}
 
@@ -457,8 +459,6 @@ namespace Chaos
 			{
 				GetPointsAndBoundsSimd(Index, OutPts, OutBounds);
 
-				VectorRegister4Float ScaleSimd = MakeVectorRegisterFloatFromDouble(MakeVectorRegisterDouble(Scale.X, Scale.Y, Scale.Z, 0.0));
-
 				OutPts[0] = VectorMultiply(OutPts[0], ScaleSimd);
 				OutPts[1] = VectorMultiply(OutPts[1], ScaleSimd);
 				OutPts[2] = VectorMultiply(OutPts[2], ScaleSimd);
@@ -476,8 +476,6 @@ namespace Chaos
 			{
 				GetBounds(CellIdx, Area, OutBounds);
 
-				VectorRegister4Float ScaleSimd = MakeVectorRegisterFloatFromDouble(MakeVectorRegisterDouble(Scale.X, Scale.Y, Scale.Z, 0.0));
-
 				VectorRegister4Float P0 = VectorMultiply(OutBounds.GetMin(), ScaleSimd);
 				VectorRegister4Float P1 = VectorMultiply(OutBounds.GetMax(), ScaleSimd);
 				VectorRegister4Float Min = VectorMin(P0, P1);
@@ -523,6 +521,8 @@ namespace Chaos
 				Ar << NumRows;
 				Ar << NumCols;
 
+				ScaleSimd = MakeVectorRegisterFloatFromDouble(MakeVectorRegisterDouble(Scale.X, Scale.Y, Scale.Z, 0.0));
+
 				Ar.UsingCustomVersion(FExternalPhysicsCustomObjectVersion::GUID);
 				if (Ar.CustomVer(FExternalPhysicsCustomObjectVersion::GUID) >= FExternalPhysicsCustomObjectVersion::HeightfieldData)
 				{
