Chaos : geometry collection : fix generation of collision for clusters using their children collision
- Cluster transform was wrongly using Transform array , causing the shape space to be wrong. Worked most of the time because the transform was identity

diff --git a/Engine/Source/Runtime/Experimental/Chaos/Private/PhysicsProxy/GeometryCollectionPhysicsProxy.cpp b/Engine/Source/Runtime/Experimental/Chaos/Private/PhysicsProxy/GeometryCollectionPhysicsProxy.cpp
--- a/Engine/Source/Runtime/Experimental/Chaos/Private/PhysicsProxy/GeometryCollectionPhysicsProxy.cpp
+++ b/Engine/Source/Runtime/Experimental/Chaos/Private/PhysicsProxy/GeometryCollectionPhysicsProxy.cpp
@@ -518,7 +518,7 @@ void FGeometryCollectionPhysicsProxy::Initialize(Chaos::FPBDRigidsEvolutionBase
 				{
 					if (GameThreadCollection.Implicits[ParentToFixIndex] == nullptr)
 					{
-						const Chaos::FRigidTransform3 ParentShapeTransform =  GameThreadCollection.MassToLocal[ParentToFixIndex] * GameThreadCollection.Transform[ParentToFixIndex];
+						const Chaos::FRigidTransform3 ParentShapeTransform =  GameThreadCollection.MassToLocal[ParentToFixIndex];
 				
 						// Make a union of the children geometry
 						TArray<TUniquePtr<Chaos::FImplicitObject>> ChildImplicits;
