fix issues with fracture convex tool's 'overlap removal shrink percent' setting double-applying scaling for clusters, and filtering points without accounting for the scaling

diff --git a/Engine/Source/Runtime/Experimental/Chaos/Private/GeometryCollection/GeometryCollectionConvexUtility.cpp b/Engine/Source/Runtime/Experimental/Chaos/Private/GeometryCollection/GeometryCollectionConvexUtility.cpp
--- a/Engine/Source/Runtime/Experimental/Chaos/Private/GeometryCollection/GeometryCollectionConvexUtility.cpp
+++ b/Engine/Source/Runtime/Experimental/Chaos/Private/GeometryCollection/GeometryCollectionConvexUtility.cpp
@@ -214,6 +214,23 @@ FVector ScaleHullPoints(TArray<Chaos::FConvex::FVec3Type>& InOutPts, double Shri
 	return Pivot;
 }
 
+void AddUnscaledPts(TArray<Chaos::FConvex::FVec3Type>& OutHullPts, const TArray<Chaos::FConvex::FVec3Type>& Vertices, const FVector& Pivot, double ShrinkPercentage)
+{
+	double ScaleFactor = 1.0 - ShrinkPercentage / 100.0;
+	double InvScaleFactor = 1.0;
+	if (ensure(ScaleFactor != 0.0))
+	{
+		InvScaleFactor = 1.0 / ScaleFactor;
+	}
+
+	int32 StartIdx = OutHullPts.AddUninitialized(Vertices.Num());
+	for (int32 VertIdx = 0, Num = Vertices.Num(); VertIdx < Num; VertIdx++)
+	{
+		FVector Vert = (FVector)Vertices[VertIdx];
+		OutHullPts[StartIdx + VertIdx] = Chaos::FConvex::FVec3Type((Vert - Pivot) * InvScaleFactor + Pivot);
+	}
+}
+
 Chaos::FConvex MakeHull(const TArray<Chaos::FConvex::FVec3Type>& Pts, double SimplificationDistanceThreshold)
 {
 	if (SimplificationDistanceThreshold > 0)
@@ -361,6 +378,9 @@ void CreateNonoverlappingConvexHulls(
 	bool bRemoveOverlaps = OverlapRemovalMethod != EConvexOverlapRemoval::None;
 	bool bRemoveLeafOverlaps = OverlapRemovalMethod == EConvexOverlapRemoval::All;
 
+	double ScaleFactor = 1.0 - ShrinkPercentage / 100.0;
+	ensure(ScaleFactor != 0.0);
+
 	int32 NumBones = TransformToConvexIndices.Num();
 	check(Parents.Num() == NumBones);
 
@@ -539,7 +559,7 @@ void CreateNonoverlappingConvexHulls(
 		return true;
 	};
 
-	auto FixCollisionWithCut = [&Convexes, &FindCutPlane, &SimplificationDistanceThreshold](int32 ConvexA, int32 ConvexB)
+	auto FixCollisionWithCut = [&Convexes, &FindCutPlane, &SimplificationDistanceThreshold, &ScaleFactor](int32 ConvexA, int32 ConvexB)
 	{
 		if (Convexes[ConvexA]->NumVertices() == 0 || Convexes[ConvexB]->NumVertices() == 0)
 		{
@@ -562,12 +582,12 @@ void CreateNonoverlappingConvexHulls(
 			TArray<Chaos::FConvex::FVec3Type> CutHullPts;
 			if (CutHull(*Convexes[ConvexA], CutPlane, true, CutHullPts))
 			{
-				*Convexes[ConvexA] = MakeHull(CutHullPts, SimplificationDistanceThreshold);
+				*Convexes[ConvexA] = MakeHull(CutHullPts, SimplificationDistanceThreshold * ScaleFactor);
 			}
 			CutHullPts.Reset();
 			if (CutHull(*Convexes[ConvexB], CutPlane, false, CutHullPts))
 			{
-				*Convexes[ConvexB] = MakeHull(CutHullPts, SimplificationDistanceThreshold);
+				*Convexes[ConvexB] = MakeHull(CutHullPts, SimplificationDistanceThreshold * ScaleFactor);
 			}
 		}
 		return bCollide;
@@ -786,15 +806,26 @@ void CreateNonoverlappingConvexHulls(
 					TArray<Chaos::FConvex::FVec3Type> JoinedHullPts;
 					TArray<int32> ChildrenWithHulls;
 					AddDescendentsWithHulls(Bone, ChildrenWithHulls);
-					ConvexesCS.Lock();
 					for (int32 Child : ChildrenWithHulls)
 					{
 						for (int32 ConvexIdx : TransformToConvexIndices[Child])
 						{
-							JoinedHullPts.Append(Convexes[ConvexIdx]->GetVertices());
+							FVector Pivot;
+							ConvexesCS.Lock();
+							Chaos::FConvex* Convex = Convexes[ConvexIdx].Get();
+							Pivot = ConvexPivots[ConvexIdx];
+							ConvexesCS.Unlock();
+							if (ShrinkPercentage != 0.0)
+							{
+								AddUnscaledPts(JoinedHullPts, Convex->GetVertices(), Pivot, ShrinkPercentage);
+							}
+							else
+							{
+								JoinedHullPts.Append(Convex->GetVertices());
+							}
 						}
 					}
-					ConvexesCS.Unlock();
+					
 					if (JoinedHullPts.Num() > 0)
 					{
 						FilterHullPoints(JoinedHullPts, SimplificationDistanceThreshold);
@@ -889,7 +920,7 @@ void CreateNonoverlappingConvexHulls(
 	});
 
 	// if bOneSidedCut, then only ConvexA is cut; ConvexB is left unchanged
-	auto CutIfOk = [&Convexes, &NonLeafVolumes, &FindCutPlane, &FracAllowRemove, &SimplificationDistanceThreshold](bool bOneSidedCut, int32 ConvexA, int32 ConvexB) -> bool
+	auto CutIfOk = [&Convexes, &NonLeafVolumes, &FindCutPlane, &FracAllowRemove, &SimplificationDistanceThreshold, &ScaleFactor](bool bOneSidedCut, int32 ConvexA, int32 ConvexB) -> bool
 	{
 		Chaos::FReal Depth;
 		Chaos::FVec3 CloseA, CloseB, Normal;
@@ -912,7 +943,7 @@ void CreateNonoverlappingConvexHulls(
 				{
 					return false;
 				}
-				CutHullA = MakeHull(CutHullPts, SimplificationDistanceThreshold);
+				CutHullA = MakeHull(CutHullPts, SimplificationDistanceThreshold * ScaleFactor);
 				bCreatedA = true;
 			}
 			if (!bOneSidedCut)
@@ -920,7 +951,7 @@ void CreateNonoverlappingConvexHulls(
 				CutHullPts.Reset();
 				if (CutHull(*Convexes[ConvexB], CutPlane, false, CutHullPts))
 				{
-					CutHullB = MakeHull(CutHullPts, SimplificationDistanceThreshold);
+					CutHullB = MakeHull(CutHullPts, SimplificationDistanceThreshold * ScaleFactor);
 					bCreatedB = true;
 				}
 			}
