Constraints: Add PIE, serialization, and BP/Python support.
For PIE got rid of just using the level editor world, and removed it as a UPROPERTY. We now get it dynamically.
For constrol rig serialization just needed to make the pointer that's held a soft object ptr. It will resolve automagically when the level sequence is loaded. Not sure if we should store a soft object ptr to that also so we can tell the user which level sequence(s) they need to load.
For BP/Python added a library to handle the  add/create/get functions. and then exposed BP to the handles and the tickable constraints.  Also added function for the Control Rig to dynamically create the handle so we can script it.

diff --git a/Engine/Plugins/Animation/ControlRig/Source/ControlRig/Private/ControlRig.cpp b/Engine/Plugins/Animation/ControlRig/Source/ControlRig/Private/ControlRig.cpp
--- a/Engine/Plugins/Animation/ControlRig/Source/ControlRig/Private/ControlRig.cpp
+++ b/Engine/Plugins/Animation/ControlRig/Source/ControlRig/Private/ControlRig.cpp
@@ -28,6 +28,7 @@
 #include "RigVMDeveloperTypeUtils.h"
 #endif// WITH_EDITOR
 #include "ControlRigComponent.h"
+#include "Constraints/ControlRigTransformableHandle.h"
 
 #define LOCTEXT_NAMESPACE "ControlRig"
 
@@ -3352,6 +3353,16 @@ uint32 UControlRig::GetHashForInitializeVMSnapShot()
 	return Hash;
 }
 
+UTransformableControlHandle* UControlRig::CreateTransformableControlHandle(UObject* InOuter, const FName& InControlName)
+{
+	UTransformableControlHandle* CtrlHandle = NewObject<UTransformableControlHandle>(InOuter);
+	ensure(CtrlHandle);
+	CtrlHandle->ControlRig = this;
+	CtrlHandle->ControlName = InControlName;
+	CtrlHandle->RegisterDelegates();
+	return CtrlHandle;
+}
+
 void UControlRig::OnHierarchyTransformUndoRedo(URigHierarchy* InHierarchy, const FRigElementKey& InKey, ERigTransformType::Type InTransformType, const FTransform& InTransform, bool bIsUndo)
 {
 	if(InKey.Type == ERigElementType::Control)
diff --git a/Engine/Plugins/Animation/ControlRig/Source/ControlRig/Private/ControlRigModule.cpp b/Engine/Plugins/Animation/ControlRig/Source/ControlRig/Private/ControlRigModule.cpp
--- a/Engine/Plugins/Animation/ControlRig/Source/ControlRig/Private/ControlRigModule.cpp
+++ b/Engine/Plugins/Animation/ControlRig/Source/ControlRig/Private/ControlRigModule.cpp
@@ -10,6 +10,7 @@
 #include "ControlRigObjectVersion.h"
 #include "Constraints/ControlRigTransformableHandle.h"
 #include "UObject/DevObjectVersion.h"
+#include "ControlRig.h"
 
 // Unique Control Rig Object version id
 const FGuid FControlRigObjectVersion::GUID(0xA7820CFB, 0x20A74359, 0x8C542C14, 0x9623CF50);
@@ -43,11 +44,11 @@ void FControlRigModule::RegisterTransformableCustomization() const
 	{
 		if (const AControlRigShapeActor* ControlActor = Cast<AControlRigShapeActor>(InObject))
 		{
-			UTransformableControlHandle* CtrlHandle = NewObject<UTransformableControlHandle>(Outer);
-			CtrlHandle->ControlRig = ControlActor->ControlRig;
-			CtrlHandle->ControlName = ControlActor->ControlName;
-			CtrlHandle->RegisterDelegates();
-			return CtrlHandle;
+			if (UControlRig* ControlRig = ControlActor->ControlRig.Get())
+			{
+				UTransformableControlHandle* CtrlHandle = ControlRig->CreateTransformableControlHandle(Outer, ControlActor->ControlName);
+				return CtrlHandle;
+			}
 		}
 		return nullptr;
 	};
diff --git a/Engine/Plugins/Animation/ControlRig/Source/ControlRig/Public/Constraints/ControlRigTransformableHandle.h b/Engine/Plugins/Animation/ControlRig/Source/ControlRig/Public/Constraints/ControlRigTransformableHandle.h
--- a/Engine/Plugins/Animation/ControlRig/Source/ControlRig/Public/Constraints/ControlRigTransformableHandle.h
+++ b/Engine/Plugins/Animation/ControlRig/Source/ControlRig/Public/Constraints/ControlRigTransformableHandle.h
@@ -18,7 +18,7 @@ struct FRigControlElement;
  * UTransformableControlHandle
  */
 
-UCLASS()
+UCLASS(Blueprintable)
 class CONTROLRIG_API UTransformableControlHandle : public UTransformableHandle 
 {
 	GENERATED_BODY()
@@ -60,11 +60,11 @@ public:
 #endif
 
 	/** @todo document */
-	UPROPERTY()
-	TWeakObjectPtr<UControlRig> ControlRig;
+	UPROPERTY(BlueprintReadOnly, Category = "Object")
+	TSoftObjectPtr<UControlRig> ControlRig;
 
 	/** @todo document */
-	UPROPERTY()
+	UPROPERTY(BlueprintReadOnly, Category = "Object")
 	FName ControlName;
 	
 private:
diff --git a/Engine/Plugins/Animation/ControlRig/Source/ControlRig/Public/ControlRig.h b/Engine/Plugins/Animation/ControlRig/Source/ControlRig/Public/ControlRig.h
--- a/Engine/Plugins/Animation/ControlRig/Source/ControlRig/Public/ControlRig.h
+++ b/Engine/Plugins/Animation/ControlRig/Source/ControlRig/Public/ControlRig.h
@@ -39,6 +39,7 @@ class UScriptStruct;
 class USkeletalMesh;
 class USkeletalMeshComponent;
 class AActor;
+class UTransformableControlHandle;
 
 struct FReferenceSkeleton;
 struct FRigUnit;
@@ -91,6 +92,7 @@ public:
 	UFUNCTION(BlueprintCallable, Category = ControlRig)
 	static TArray<UControlRig*> FindControlRigs(UObject* Outer, TSubclassOf<UControlRig> OptionalClass);
 
+
 private:
 	/** Current delta time */
 	float DeltaTime;
@@ -153,6 +155,12 @@ public:
 	UFUNCTION(BlueprintPure, Category = "Control Rig")
 	float GetCurrentFramesPerSecond() const;
 
+	/** Creates a transformable control handle for the specified control to be used by the constraints system. Should use the UObject from 
+	ConstraintsScriptingLibrary::GetManager(UWorld* InWorld)*/
+	UFUNCTION(BlueprintPure, Category = "Control Rig | Constraints")
+	UTransformableControlHandle* CreateTransformableControlHandle(UObject* Outer, const FName& ControlName);
+
+
 #if WITH_EDITOR
 	/** Get the category of this ControlRig (for display in menus) */
 	virtual FText GetCategory() const;
diff --git a/Engine/Plugins/Animation/ControlRig/Source/ControlRigEditor/Private/Editor/Constraints/SConstraintsWidget.cpp b/Engine/Plugins/Animation/ControlRig/Source/ControlRigEditor/Private/Editor/Constraints/SConstraintsWidget.cpp
--- a/Engine/Plugins/Animation/ControlRig/Source/ControlRigEditor/Private/Editor/Constraints/SConstraintsWidget.cpp
+++ b/Engine/Plugins/Animation/ControlRig/Source/ControlRigEditor/Private/Editor/Constraints/SConstraintsWidget.cpp
@@ -278,11 +278,14 @@ void SDroppableConstraintItem::CreateConstraint(
 	// create constraints
 	UWorld* World = GCurrentLevelEditingViewportClient->GetWorld();
 	bool bCreated = false;
-	for (const AActor* Child: Selection)
+	for (AActor* Child: Selection)
 	{
 		if (Child != InParent)
 		{
-			bCreated |= TransformConstraintUtils::Create(World, InParent, Child, InConstraintType);
+			if (FTransformConstraintUtils::CreateAndAddFromActors(World, InParent, Child, InConstraintType) != nullptr)
+			{
+				bCreated = true;
+			}
 		}
 	}
 
@@ -723,13 +726,13 @@ void SConstraintsEditionWidget::RefreshConstraintList()
 	if (bIsConstraintsActor)
 	{
 		const FConstraintsManagerController& Controller = FConstraintsManagerController::Get(World);
-		Constraints = Controller.GetConstrainsArray();
+		Constraints = Controller.GetConstraintsArray();
 	}
 	else
 	{
 		for (const AActor* Actor : Selection)
 		{
-			TransformConstraintUtils::GetParentConstraints(World, Actor, Constraints);
+			FTransformConstraintUtils::GetParentConstraints(World, Actor, Constraints);
 		}
 	}
 	
diff --git a/Engine/Plugins/Animation/ControlRig/Source/ControlRigEditor/Private/Sequencer/ConstraintChannelHelper.cpp b/Engine/Plugins/Animation/ControlRig/Source/ControlRigEditor/Private/Sequencer/ConstraintChannelHelper.cpp
--- a/Engine/Plugins/Animation/ControlRig/Source/ControlRigEditor/Private/Sequencer/ConstraintChannelHelper.cpp
+++ b/Engine/Plugins/Animation/ControlRig/Source/ControlRigEditor/Private/Sequencer/ConstraintChannelHelper.cpp
@@ -188,16 +188,18 @@ UMovieSceneControlRigParameterSection* FConstraintChannelHelper::GetControlSecti
 		return nullptr;
 	}
 
-	const TWeakObjectPtr<UControlRig> ControlRig = InHandle->ControlRig;
-	
-	const TArray<FMovieSceneBinding>& Bindings = MovieScene->GetBindings();
-	for (const FMovieSceneBinding& Binding : Bindings)
-	{
-		UMovieSceneTrack* Track = MovieScene->FindTrack(UMovieSceneControlRigParameterTrack::StaticClass(), Binding.GetObjectGuid());
-		UMovieSceneControlRigParameterTrack* ControlRigTrack = Cast<UMovieSceneControlRigParameterTrack>(Track);
-		if (ControlRigTrack && ControlRigTrack->GetControlRig() == ControlRig)
+	const TWeakObjectPtr<UControlRig> ControlRig = InHandle->ControlRig.LoadSynchronous();
+	if(ControlRig.IsValid())
+	{	
+		const TArray<FMovieSceneBinding>& Bindings = MovieScene->GetBindings();
+		for (const FMovieSceneBinding& Binding : Bindings)
 		{
-			return Cast<UMovieSceneControlRigParameterSection>(ControlRigTrack->FindSection(0));
+			UMovieSceneTrack* Track = MovieScene->FindTrack(UMovieSceneControlRigParameterTrack::StaticClass(), Binding.GetObjectGuid());
+			UMovieSceneControlRigParameterTrack* ControlRigTrack = Cast<UMovieSceneControlRigParameterTrack>(Track);
+			if (ControlRigTrack && ControlRigTrack->GetControlRig() == ControlRig)
+			{
+				return Cast<UMovieSceneControlRigParameterSection>(ControlRigTrack->FindSection(0));
+			}
 		}
 	}
 
diff --git a/Engine/Source/Runtime/AnimationCore/Public/Constraint.h b/Engine/Source/Runtime/AnimationCore/Public/Constraint.h
--- a/Engine/Source/Runtime/AnimationCore/Public/Constraint.h
+++ b/Engine/Source/Runtime/AnimationCore/Public/Constraint.h
@@ -306,7 +306,7 @@ struct ANIMATIONCORE_API FTransformConstraint
 /// new changes of constraints
 
 /** Constraint Types*/
-UENUM()
+UENUM(BlueprintType)
 enum class EConstraintType : uint8
 {
 	/** Transform Constraint */
@@ -384,7 +384,7 @@ struct TStructOpsTypeTraits<FConstraintDescriptionEx> : public TStructOpsTypeTra
 };
 
 /** Transform Constraint Types*/
-UENUM()
+UENUM(Blueprintable)
 enum class ETransformConstraintType : uint8
 {
 	Translation,
diff --git a/Engine/Source/Runtime/Experimental/Animation/Constraints/Private/ConstraintsActor.cpp b/Engine/Source/Runtime/Experimental/Animation/Constraints/Private/ConstraintsActor.cpp
--- a/Engine/Source/Runtime/Experimental/Animation/Constraints/Private/ConstraintsActor.cpp
+++ b/Engine/Source/Runtime/Experimental/Animation/Constraints/Private/ConstraintsActor.cpp
@@ -23,16 +23,16 @@ AConstraintsActor::~AConstraintsActor()
 
 void AConstraintsActor::BeginDestroy()
 {
+	if (ConstraintsManager)
+	{
+		ConstraintsManager->Clear(GetWorld());
+	}
+
 	Super::BeginDestroy();
 }
 
 void AConstraintsActor::Destroyed()
 {
-	if(ConstraintsManager)
-	{
-		ConstraintsManager->Clear();
-	}
-	
 	Super::Destroyed();
 }
 
@@ -48,7 +48,13 @@ void AConstraintsActor::Tick(float DeltaTime)
 
 void AConstraintsActor::PostRegisterAllComponents()
 {
-	Super::PostRegisterAllComponents();	
+	Super::PostRegisterAllComponents();
+
+	if (ConstraintsManager == nullptr)
+	{
+		ConstraintsManager = NewObject<UConstraintsManager>(this);
+	}
+	ConstraintsManager->Init(GetWorld());
 	RegisterConstraintsTickFunctions();
 }
 
diff --git a/Engine/Source/Runtime/Experimental/Animation/Constraints/Private/ConstraintsManager.cpp b/Engine/Source/Runtime/Experimental/Animation/Constraints/Private/ConstraintsManager.cpp
--- a/Engine/Source/Runtime/Experimental/Animation/Constraints/Private/ConstraintsManager.cpp
+++ b/Engine/Source/Runtime/Experimental/Animation/Constraints/Private/ConstraintsManager.cpp
@@ -110,21 +110,39 @@ UConstraintsManager::UConstraintsManager()
 
 UConstraintsManager::~UConstraintsManager()
 {
-	UnregisterDelegates();
+
 }
 
+void UConstraintsManager::PostLoad()
+{
+	Super::PostLoad();
+	for (TObjectPtr<UTickableConstraint> ConstPtr : Constraints)
+	{
+		if (ConstPtr)
+		{
+			ConstPtr->ConstraintTick.Constraint = ConstPtr;
+		}
+	}
+}
+
+
 void UConstraintsManager::OnActorDestroyed(AActor* InActor)
 {
 	if (USceneComponent* SceneComponent = InActor->GetRootComponent())
 	{
-		Constraints.RemoveAll([SceneComponent](const TObjectPtr<UTickableConstraint>& Constraint)
+		Constraints.RemoveAll([this, SceneComponent](const TObjectPtr<UTickableConstraint>& Constraint)
 		{
-			return IsValid(Constraint) ? Constraint->ReferencesObject(SceneComponent) : false;
+			const bool bIsRemoved = IsValid(Constraint) && Constraint->ReferencesObject(SceneComponent);
+			if (bIsRemoved)
+			{
+				OnConstraintRemoved_BP.Broadcast(this, Constraint);
+			}
+			return bIsRemoved;
 		} );
 	}
 }
 
-void UConstraintsManager::RegisterDelegates()
+void UConstraintsManager::RegisterDelegates(UWorld* World)
 {
 	if (!OnActorDestroyedHandle.IsValid())
 	{
@@ -134,7 +152,7 @@ void UConstraintsManager::RegisterDelegates()
 	}
 }
 
-void UConstraintsManager::UnregisterDelegates()
+void UConstraintsManager::UnregisterDelegates(UWorld* World)
 {
 	if (World)
 	{
@@ -143,15 +161,13 @@ void UConstraintsManager::UnregisterDelegates()
 	OnActorDestroyedHandle.Reset();
 }
 
-void UConstraintsManager::Init(UWorld* InWorld)
+void UConstraintsManager::Init(UWorld* World)
 {
-	if (InWorld)
+	if (World)
 	{
-		UnregisterDelegates();
-
-		World = InWorld;
+		UnregisterDelegates(World);
 
-		RegisterDelegates();
+		RegisterDelegates(World);
 	}
 }
 
@@ -168,11 +184,6 @@ UConstraintsManager* UConstraintsManager::Get(UWorld* InWorld)
 #if WITH_EDITOR
 	ConstraintsActor->SetActorLabel(TEXT("Constraints Manager"));
 #endif // WITH_EDITOR
-	ConstraintsActor->ConstraintsManager = NewObject<UConstraintsManager>(ConstraintsActor);
-	
-	// ULevel* Level = InWorld->GetCurrentLevel();
-	ConstraintsActor->ConstraintsManager->Init(InWorld);
-	// ConstraintsActor->ConstraintsManager->Level = Level;
 
 	return ConstraintsActor->ConstraintsManager;
 }
@@ -233,8 +244,13 @@ void UConstraintsManager::SetConstraintDependencies(
 	InFunctionToTickAfter->AddPrerequisite(this, *InFunctionToTickBefore);
 }
 
-void UConstraintsManager::Clear()
+void UConstraintsManager::Clear(UWorld* World)
 {
+	UnregisterDelegates(World);
+	for (UTickableConstraint* Constraint : Constraints)
+	{
+		OnConstraintRemoved_BP.Broadcast(this, Constraint);
+	}
 	Constraints.Empty();
 }
 
@@ -328,29 +344,32 @@ void FConstraintsManagerController::DestroyManager() const
 
 	for (AActor* ConstraintsActor: ConstraintsActorsToRemove)
 	{
-		World->EditorDestroyActor(ConstraintsActor, true);
+		World->DestroyActor(ConstraintsActor, true);
 	}
 }
 
-void FConstraintsManagerController::AddConstraint(UTickableConstraint* InConstraint) const
+bool FConstraintsManagerController::AddConstraint(UTickableConstraint* InConstraint) const
 {
 	if (!InConstraint)
 	{
-		return;
+		return false;
 	}
 	
-	UConstraintsManager* Manager = GetManager();
+	UConstraintsManager* Manager = GetManager(); //this will allocate if doesn't exist
 	if (!Manager)
 	{
-		return;
+		return false;
 	}
 
 	// TODO handle transaction
 	Manager->Constraints.Emplace(InConstraint);
 
-	InConstraint->ConstraintTick.Constraint = InConstraint;
 	InConstraint->ConstraintTick.RegisterFunction(InConstraint->GetFunction());
 	InConstraint->ConstraintTick.RegisterTickFunction(World->GetCurrentLevel());
+
+	Manager->OnConstraintAdded_BP.Broadcast(Manager, InConstraint);
+
+	return true;
 }
 
 int32 FConstraintsManagerController::GetConstraintIndex(const FName& InConstraintName) const
@@ -367,44 +386,46 @@ int32 FConstraintsManagerController::GetConstraintIndex(const FName& InConstrain
 	} );
 }
 	
-void FConstraintsManagerController::RemoveConstraint(const FName& InConstraintName) const
+bool  FConstraintsManagerController::RemoveConstraint(const FName& InConstraintName) const
 {
 	const int32 Index = GetConstraintIndex(InConstraintName);
 	if (Index == INDEX_NONE)
 	{
-		return;	
+		return false;
 	}
 	
-	RemoveConstraint(Index);
+	return RemoveConstraint(Index);
 }
 
-void FConstraintsManagerController::RemoveConstraint(const int32 InConstraintIndex) const
+bool FConstraintsManagerController::RemoveConstraint(const int32 InConstraintIndex) const
 {
 	UConstraintsManager* Manager = FindManager();
 	if (!Manager)
 	{
-		return;
+		return false;
 	}
 	
 	if (!Manager->Constraints.IsValidIndex(InConstraintIndex))
 	{
-		return;
+		return false;
 	}
 
 	// TODO handle transaction
 	const FName ConstraintName = Manager->Constraints[InConstraintIndex]->GetFName(); 
-	
+	UTickableConstraint* Constraint = Manager->Constraints[InConstraintIndex];
 	Manager->Constraints[InConstraintIndex]->ConstraintTick.UnRegisterTickFunction();
 	Manager->Constraints.RemoveAt(InConstraintIndex);
 
 	// notify deletion
 	ConstraintRemoved.Broadcast(ConstraintName);
-	
+	Manager->OnConstraintRemoved_BP.Broadcast(Manager, Constraint);
+
 	// destroy constraints actor if no constraints left
 	if (Manager->Constraints.IsEmpty())
 	{
 		DestroyManager();
 	}
+	return true;
 }
 
 UTickableConstraint* FConstraintsManagerController::GetConstraint(const FName& InConstraintName) const
@@ -504,7 +525,7 @@ void FConstraintsManagerController::SetConstraintsDependencies(
 	Manager->SetConstraintDependencies( &FunctionToTickBefore, &FunctionToTickAfter);
 }
 
-const TArray< TObjectPtr<UTickableConstraint> >& FConstraintsManagerController::GetConstrainsArray() const
+const TArray< TObjectPtr<UTickableConstraint> >& FConstraintsManagerController::GetConstraintsArray() const
 {
 	UConstraintsManager* Manager = FindManager();
 	if (!Manager)
diff --git a/Engine/Source/Runtime/Experimental/Animation/Constraints/Private/ConstraintsScripting.cpp b/Engine/Source/Runtime/Experimental/Animation/Constraints/Private/ConstraintsScripting.cpp
new file mode 100644
--- /dev/null
+++ b/Engine/Source/Runtime/Experimental/Animation/Constraints/Private/ConstraintsScripting.cpp
@@ -0,0 +1,59 @@
+// Copyright Epic Games, Inc. All Rights Reserved.
+
+
+#include "ConstraintsScripting.h"
+#include "TransformConstraint.h"
+#include "Animation/Constraints/Public/TransformableHandle.h"
+#include "ConstraintsManager.h"
+
+UConstraintsManager* UConstraintsScriptingLibrary::GetManager(UWorld* InWorld)
+{
+	UConstraintsManager* ConstraintsManager = UConstraintsManager::Get(InWorld);
+	//add ue_log
+	return ConstraintsManager;
+}
+
+UTransformableComponentHandle* UConstraintsScriptingLibrary::CreateTransformableComponentHandle(UWorld* InWorld,USceneComponent* InSceneComponent)
+{
+	UConstraintsManager* ConstraintsManager = UConstraintsManager::Get(InWorld);
+	if (ConstraintsManager)
+	{
+		UTransformableComponentHandle* Handle = FTransformConstraintUtils::CreateHandleForSceneComponent(InSceneComponent, ConstraintsManager);
+		return Handle;
+
+	}
+	return nullptr;
+}
+
+UTickableTransformConstraint* UConstraintsScriptingLibrary::CreateFromType(
+	UWorld* InWorld,
+	const ETransformConstraintType InType)
+{
+	UTickableTransformConstraint* Constraint = FTransformConstraintUtils::CreateFromType(InWorld, InType);
+	return Constraint;
+}
+
+bool UConstraintsScriptingLibrary::AddConstraint(UWorld* InWorld, UTransformableHandle* InParentHandle, UTransformableHandle* InChildHandle,
+	UTickableTransformConstraint *InConstraint, const bool bMaintainOffset)
+{
+	bool Val =  FTransformConstraintUtils::AddConstraint(InWorld, InParentHandle, InChildHandle,InConstraint, bMaintainOffset);
+	return Val;
+}
+
+TArray<TWeakObjectPtr<UTickableConstraint>> UConstraintsScriptingLibrary::GetConstraintsArray(UWorld* InWorld)
+{
+	TArray<TWeakObjectPtr<UTickableConstraint>> Constraints;
+	const FConstraintsManagerController& Controller = FConstraintsManagerController::Get(InWorld);
+	const TArray< TObjectPtr<UTickableConstraint> >& ConstraintsArray = Controller.GetConstraintsArray();
+	for (const TObjectPtr<UTickableConstraint>& Constraint : ConstraintsArray)
+	{
+		Constraints.Add(Constraint.Get());
+	}
+	return Constraints;
+}
+
+bool UConstraintsScriptingLibrary::RemoveConstraint(UWorld* InWorld, int32 InIndex)
+{
+	const FConstraintsManagerController& Controller = FConstraintsManagerController::Get(InWorld);
+	return Controller.RemoveConstraint(InIndex);
+}
\ No newline at end of file
diff --git a/Engine/Source/Runtime/Experimental/Animation/Constraints/Private/TransformConstraint.cpp b/Engine/Source/Runtime/Experimental/Animation/Constraints/Private/TransformConstraint.cpp
--- a/Engine/Source/Runtime/Experimental/Animation/Constraints/Private/TransformConstraint.cpp
+++ b/Engine/Source/Runtime/Experimental/Animation/Constraints/Private/TransformConstraint.cpp
@@ -1,7 +1,6 @@
 ﻿// Copyright Epic Games, Inc. All Rights Reserved.
 
 #include "TransformConstraint.h"
-
 #include "Animation/Constraints/Public/TransformableHandle.h"
 #include "ConstraintsManager.h"
 #include "TransformableRegistry.h"
@@ -82,13 +81,13 @@
 void UTickableTransformConstraint::PostLoad()
 {
 	Super::PostLoad();
-
 	ConstraintTick.RegisterFunction(GetFunction() );
 }
 
 void UTickableTransformConstraint::PostDuplicate(bool bDuplicateForPIE)
 {
 	Super::PostDuplicate(bDuplicateForPIE);
+	ConstraintTick.RegisterFunction(GetFunction());
 }
 
 uint32 UTickableTransformConstraint::GetTargetHash() const
@@ -452,13 +451,13 @@
 }
 
 /** 
- * TransformConstraintUtils
+ * FTransformConstraintUtils
  **/
 
 namespace
 {
 
-UTransformableHandle* GetHandle(const AActor* InActor, UObject* Outer)
+UTransformableHandle* GetHandle(AActor* InActor, UObject* Outer)
 {
 	// look for customized transform handle
 	const FTransformableRegistry& Registry = FTransformableRegistry::Get();
@@ -467,11 +466,13 @@
 		return CreateFunction(InActor, Outer);
 	}
 
-	// scene component handle
-	UTransformableComponentHandle* ComponentHandle = NewObject<UTransformableComponentHandle>(Outer);
-	ComponentHandle->Component = InActor->GetRootComponent();
-	
-	return ComponentHandle;
+	// need to make sure it's moveable
+	if (InActor->GetRootComponent())
+	{
+		return FTransformConstraintUtils::CreateHandleForSceneComponent(InActor->GetRootComponent(), Outer);
+
+	}
+	return nullptr;
 }
 	
 uint32 GetConstrainableHash(const AActor* InActor)
@@ -490,7 +491,19 @@
 	
 }
 
-void TransformConstraintUtils::GetParentConstraints(
+UTransformableComponentHandle* FTransformConstraintUtils::CreateHandleForSceneComponent(USceneComponent* InSceneComponent, UObject* Outer)
+{
+	UTransformableComponentHandle* ComponentHandle = nullptr;
+	if (InSceneComponent)
+	{
+		ComponentHandle = NewObject<UTransformableComponentHandle>(Outer);
+		ComponentHandle->Component = InSceneComponent;
+		InSceneComponent->SetMobility(EComponentMobility::Movable);
+	}
+	return ComponentHandle;
+}
+
+void FTransformConstraintUtils::GetParentConstraints(
 	UWorld* InWorld,
 	const AActor* InChild,
 	TArray< TObjectPtr<UTickableConstraint> >& OutConstraints)
@@ -511,64 +524,113 @@
 	OutConstraints.Append(Controller.GetParentConstraints(ChildHash, bSorted));
 }
 
-bool TransformConstraintUtils::Create(
+UTickableTransformConstraint* FTransformConstraintUtils::CreateFromType(
 	UWorld* InWorld,
-	const AActor* InParent,
-	const AActor* InChild,
+	const ETransformConstraintType InType)
+{
+	if (!InWorld)
+	{
+		UE_LOG(LogTemp, Error, TEXT("FTransformConstraintUtils::CreateFromType sanity check failed."));
+		return nullptr;
+	}
+	const UEnum* ETransformConstraintTypeEnum = StaticEnum<ETransformConstraintType>();
+	if (!ETransformConstraintTypeEnum->IsValidEnumValue(static_cast<int64>(InType)))
+	{
+		UE_LOG(LogTemp, Error, TEXT("Constraint Type %d not recognized"), InType);
+		return nullptr;
+	}
+
+
+	// unique name (we may want to use another approach here to manage uniqueness)
+	const FString ConstraintTypeStr = ETransformConstraintTypeEnum->GetNameStringByValue((uint8)InType);
+	const FName BaseName(*FString::Printf(TEXT("%sConstraint"), *ConstraintTypeStr));
+
+	const FConstraintsManagerController& Controller = FConstraintsManagerController::Get(InWorld);
+	UTickableTransformConstraint* Constraint = nullptr;
+
+	switch (InType)
+	{
+	case ETransformConstraintType::Translation:
+		Constraint = Controller.AllocateConstraintT<UTickableTranslationConstraint>(BaseName);
+		break;
+	case ETransformConstraintType::Rotation:
+		Constraint = Controller.AllocateConstraintT<UTickableRotationConstraint>(BaseName);
+		break;
+	case ETransformConstraintType::Scale:
+		Constraint = Controller.AllocateConstraintT<UTickableScaleConstraint>(BaseName);
+		break;
+	case ETransformConstraintType::Parent:
+		Constraint = Controller.AllocateConstraintT<UTickableParentConstraint>(BaseName);
+		break;
+	case ETransformConstraintType::LookAt:
+		Constraint = Controller.AllocateConstraintT<UTickableLookAtConstraint>(BaseName);
+		break;
+	default:
+		ensure(false);
+		break;
+	}
+	return Constraint;
+}
+
+UTickableTransformConstraint* FTransformConstraintUtils::CreateAndAddFromActors(
+	UWorld* InWorld,
+	AActor* InParent,
+	AActor* InChild,
 	const ETransformConstraintType InType,
 	const bool bMaintainOffset)
 {
 	// SANITY CHECK
 	if (!InWorld || !InParent || !InChild)
 	{
-		UE_LOG(LogTemp, Error, TEXT("TransformConstraintUtils::Create sanity check failed."));
-		return false;
+		UE_LOG(LogTemp, Error, TEXT("FTransformConstraintUtils::CreateAndAddFromActors sanity check failed."));
+		return nullptr;
 	}
 	
 	UConstraintsManager* ConstraintsManager = UConstraintsManager::Get(InWorld);
 	if (!ConstraintsManager)
 	{
-		UE_LOG(LogTemp, Error, TEXT("TransformConstraintUtils::Create constraint manager is null."));
-		return false;
+		UE_LOG(LogTemp, Error, TEXT("FTransformConstraintUtils::CreateAndAddFromActors constraint manager is null."));
+		return nullptr;
 	}
 
 	UTransformableHandle* ParentHandle = GetHandle(InParent, ConstraintsManager);
 	UTransformableHandle* ChildHandle = GetHandle(InChild, ConstraintsManager);
 	
-	// DO IT
-	if (ParentHandle->IsValid() && ChildHandle->IsValid())
+	UTickableTransformConstraint* Constraint = FTransformConstraintUtils::CreateFromType(InWorld, InType);
+
+
+	if (Constraint && (ParentHandle->IsValid() && ChildHandle->IsValid()))
 	{
-		return AddConst(InWorld, ParentHandle, ChildHandle, InType, bMaintainOffset);
+		if (AddConstraint(InWorld, ParentHandle, ChildHandle, Constraint, bMaintainOffset) == false)
+		{
+			Constraint->MarkAsGarbage();
+			Constraint = nullptr;
+		}
 	}
-	
-	return false;
+	return Constraint;
 }
 
-bool TransformConstraintUtils::AddConst(
+bool FTransformConstraintUtils::AddConstraint(
 	UWorld* InWorld,
 	UTransformableHandle* InParentHandle,
 	UTransformableHandle* InChildHandle,
-	const ETransformConstraintType InType,
+	UTickableTransformConstraint* Constraint,
 	const bool bMaintainOffset)
 {
 	const bool bIsValidParent = InParentHandle && InParentHandle->IsValid();
 	const bool bIsValidChild = InChildHandle && InChildHandle->IsValid();
 	if (!bIsValidParent || !bIsValidChild)
 	{
-		UE_LOG(LogTemp, Error, TEXT("TransformConstraintUtils::AddConst error creating constraint"));
+		UE_LOG(LogTemp, Error, TEXT("FTransformConstraintUtils::AddConst error adding constraint"));
 		return false;
 	}
 
-	const UEnum* ETransformConstraintTypeEnum = StaticEnum<ETransformConstraintType>();
-	if (!ETransformConstraintTypeEnum->IsValidEnumValue(static_cast<int64>(InType)))
+	if (Constraint == nullptr)
 	{
-		UE_LOG(LogTemp, Error, TEXT("Constraint Type %d not recognized"), InType);
+		UE_LOG(LogTemp, Error, TEXT("FTransformConstraintUtils::AddConst error creating constraint"));
 		return false;
 	}
-	
-	// unique name (we may want to use another approach here to manage uniqueness)
-	const FString ConstraintTypeStr = ETransformConstraintTypeEnum->GetNameStringByValue((uint8)InType);
-	const FName BaseName(*FString::Printf(TEXT("%sConstraint"), *ConstraintTypeStr));
 
+
 	auto SetupConstraint = [InParentHandle, InChildHandle, bMaintainOffset](UTickableTransformConstraint* InConstraint)
 	{
@@ -580,28 +642,7 @@
 	
 	const FConstraintsManagerController& Controller = FConstraintsManagerController::Get(InWorld);
 	
-	UTickableTransformConstraint* Constraint = nullptr;
-	switch (InType)
-	{
-	case ETransformConstraintType::Translation:
-		Constraint = Controller.AllocateConstraintT<UTickableTranslationConstraint>(BaseName);
-		break;
-	case ETransformConstraintType::Rotation:
-		Constraint = Controller.AllocateConstraintT<UTickableRotationConstraint>(BaseName);
-		break;
-	case ETransformConstraintType::Scale:
-		Constraint = Controller.AllocateConstraintT<UTickableScaleConstraint>(BaseName);
-		break;
-	case ETransformConstraintType::Parent:
-		Constraint = Controller.AllocateConstraintT<UTickableParentConstraint>(BaseName);
-		break;
-	case ETransformConstraintType::LookAt:
-		Constraint = Controller.AllocateConstraintT<UTickableLookAtConstraint>(BaseName);
-		break;
-	default:
-		ensure(false);
-		return false;
-	}
+	Controller.AddConstraint(Constraint);
 
 	SetupConstraint(Constraint);
 
diff --git a/Engine/Source/Runtime/Experimental/Animation/Constraints/Private/TransformableHandle.cpp b/Engine/Source/Runtime/Experimental/Animation/Constraints/Private/TransformableHandle.cpp
--- a/Engine/Source/Runtime/Experimental/Animation/Constraints/Private/TransformableHandle.cpp
+++ b/Engine/Source/Runtime/Experimental/Animation/Constraints/Private/TransformableHandle.cpp
@@ -56,6 +56,11 @@ FTransform UTransformableComponentHandle::GetGlobalTransform() const
 	return Component.IsValid() ? Component->GetComponentTransform() : FTransform::Identity;
 }
 
+FTransform UTransformableComponentHandle::GetParentTransform() const
+{
+	return (Component.IsValid() && Component->GetAttachParent()) ? Component->GetAttachParent()->GetComponentTransform() : FTransform::Identity;
+}
+
 UObject* UTransformableComponentHandle::GetPrerequisiteObject() const
 {
 	return Component.Get(); 
diff --git a/Engine/Source/Runtime/Experimental/Animation/Constraints/Public/ConstraintsManager.h b/Engine/Source/Runtime/Experimental/Animation/Constraints/Public/ConstraintsManager.h
--- a/Engine/Source/Runtime/Experimental/Animation/Constraints/Public/ConstraintsManager.h
+++ b/Engine/Source/Runtime/Experimental/Animation/Constraints/Public/ConstraintsManager.h
@@ -65,7 +65,7 @@ struct TStructOpsTypeTraits<FConstraintTickFunction> : public TStructOpsTypeTrai
  * Represents the basic interface of constraint within the constraints manager.
  **/
 
-UCLASS(Abstract)
+UCLASS(Abstract, Blueprintable)
 class CONSTRAINTS_API UTickableConstraint : public UObject
 {
 	GENERATED_BODY()
@@ -100,28 +100,35 @@ public:
 	virtual void PostEditChangeProperty(FPropertyChangedEvent& PropertyChangedEvent) override;
 	// End of UObject interface
 #endif
-
 	/** @todo documentation. */
-	UPROPERTY(EditAnywhere, DisplayName="Active State", Category="Constraint")
+	UPROPERTY(EditAnywhere, BlueprintReadWrite, DisplayName="Active State", Category="Constraint")
 	bool Active = true;
 };
 
+
 /** 
  * UConstraintsManager
  * This object gathers the different constraints of the level and is held by the ConstraintsActor (unique in the level)
  **/
 
-UCLASS()
+UCLASS(BLUEPRINTABLE)
 class CONSTRAINTS_API UConstraintsManager : public UObject
 {
 	GENERATED_BODY()
 public:
 
 	using ConstraintPtr = TObjectPtr<UTickableConstraint>;
+
+	/** Dynamic blueprintable delegates for knowing when a constraints are added or deleted*/
+	DECLARE_DYNAMIC_MULTICAST_SPARSE_DELEGATE_TwoParams(FOnConstraintAdded, UConstraintsManager, OnConstraintAdded_BP, UConstraintsManager*, Mananger, UTickableConstraint*, Constraint);
+	DECLARE_DYNAMIC_MULTICAST_SPARSE_DELEGATE_TwoParams(FOnConstraintRemoved, UConstraintsManager, OnConstraintRemoved_BP, UConstraintsManager*, Mananger, UTickableConstraint*, Constraint);
+
 	
 	UConstraintsManager();
 	virtual ~UConstraintsManager();
 	
+	//UObjects
+	virtual void PostLoad() override;
 	/** Get the existing Constraints Manager if existing or create a new one. */
 	static UConstraintsManager* Get(UWorld* InWorld);
 
@@ -137,8 +144,17 @@ public:
 		FConstraintTickFunction* InFunctionToTickAfter);
 
 	/** @todo document */
-	void Clear();
+	void Clear(UWorld* World);
+
+
+	/** BP Delegate fired when constraints are added*/
+	UPROPERTY(BlueprintAssignable, Category = Constraints, meta = (DisplayName = "OnConstraintAdded"))
+	FOnConstraintAdded OnConstraintAdded_BP;
 	
+	/** BP Delegate fired when constraints are removed*/
+	UPROPERTY(BlueprintAssignable, Category = Constraints, meta = (DisplayName = "OnConstraintRemoved"))
+	FOnConstraintAdded OnConstraintRemoved_BP;
+
 private:
 
 	/** @todo document */
@@ -146,16 +162,12 @@ private:
 	
 	void OnActorDestroyed(AActor* InActor);
 
-	void RegisterDelegates();
-	void UnregisterDelegates();
+	void RegisterDelegates(UWorld* World);
+	void UnregisterDelegates(UWorld* World);
 	
 	/** @todo document */
 	void Dump() const;
 
-	/** @todo document */
-	UPROPERTY()
-	TObjectPtr<UWorld> World = nullptr;
-
 	/** @todo document */
 	UPROPERTY()
 	TArray< TObjectPtr<UTickableConstraint> > Constraints;
@@ -188,23 +200,21 @@ public:
 		// unique name (we may want to use another approach here to manage uniqueness)
 		const FName Name = MakeUniqueObjectName(Manager, TConstraint::StaticClass(), InBaseName);
 
-		TConstraint* NewConstraint = NewObject<TConstraint>(Manager, Name);
-		AddConstraint(NewConstraint);
-		
+		TConstraint* NewConstraint = NewObject<TConstraint>(Manager, Name);		
 		return NewConstraint;
 	}
 
 	/** @todo document */
-	void AddConstraint(UTickableConstraint* InConstraint) const;
+	bool AddConstraint(UTickableConstraint* InConstraint) const;
 	
 	/** Get the index of the given constraint's name. */
 	int32 GetConstraintIndex(const FName& InConstraintName) const;
 	
 	/** Remove the constraint by name. */
-	void RemoveConstraint(const FName& InConstraintName) const;
+	bool RemoveConstraint(const FName& InConstraintName) const;
 
 	/** Remove the constraint at the given index. */
-	void RemoveConstraint(const int32 InConstraintIndex) const;
+	bool RemoveConstraint(const int32 InConstraintIndex) const;
 
 	/** @todo document */
 	UTickableConstraint* GetConstraint(const FName& InConstraintName) const;
@@ -213,8 +223,8 @@ public:
 	UTickableConstraint* GetConstraint(const int32 InConstraintIndex) const;
 	
 	/** Get read-only access to the array of constraints. */
-	const TArray< TObjectPtr<UTickableConstraint> >& GetConstrainsArray() const;
-	
+	const TArray< TObjectPtr<UTickableConstraint> >& GetConstraintsArray() const;
+
 	/** Get parent constraints of the specified child. If bSorted is true, then the constraints will be sorted by dependency. */
 	TArray< TObjectPtr<UTickableConstraint> > GetParentConstraints(
 		const uint32 InTargetHash,
diff --git a/Engine/Source/Runtime/Experimental/Animation/Constraints/Public/ConstraintsScripting.h b/Engine/Source/Runtime/Experimental/Animation/Constraints/Public/ConstraintsScripting.h
new file mode 100644
--- /dev/null
+++ b/Engine/Source/Runtime/Experimental/Animation/Constraints/Public/ConstraintsScripting.h
@@ -0,0 +1,84 @@
+// Copyright Epic Games, Inc. All Rights Reserved.
+
+#pragma once
+
+#include "CoreMinimal.h"
+#include "UObject/ObjectMacros.h"
+#include "Kismet/BlueprintFunctionLibrary.h"
+#include "Constraint.h"
+#include "ConstraintsScripting.generated.h"
+
+class UConstraintsManager;
+class UTransformableComponentHandle;
+class UTickableConstraint;
+class UTickableTransformConstraint;
+
+/**
+* This is a set of helper functions to access various parts of the Sequencer and Control Rig API via Python and Blueprints.
+*/
+UCLASS(meta = (Transient, ScriptName = "ConstraintsScriptingLibrary"))
+class CONSTRAINTS_API UConstraintsScriptingLibrary : public UBlueprintFunctionLibrary
+{
+
+public:
+	GENERATED_BODY()
+
+public:
+
+	/**
+	* Get the manager of the constraints. This object contains delegates to listen to for when constraints are added,deleted,
+	* and is also the outer used when creating custom transformable handles, for example this is used to create control rig transformable handles
+	* Note this function will create the mananager and it's actor if one doesn't exist.
+	* @param InWorld, the world you are in
+	* @return Returns the mananger
+	*/
+	UFUNCTION(BlueprintCallable, Category = "Editor Scripting | Sequencer Tools | Control Rig | Constraints")
+	static UConstraintsManager* GetManager(UWorld* InWorld);
+
+	/**
+	* Create the transformable handle that deals with getting and setting transforms on this scene component
+	* @param InWorld, the world you are in
+	* @param InSceneComponent World to create the constraint
+	* @return returns the handle for this scene component
+	**/
+	UFUNCTION(BlueprintCallable, Category = "Editor Scripting | Sequencer Tools | Control Rig | Constraints")
+	static UTransformableComponentHandle* CreateTransformableComponentHandle(UWorld* InWorld, USceneComponent* InSceneComponent);
+
+	/**
+	* Create Constraint based on the specified type.
+	* @param InWorld World to create the constraint
+	* @param InType The type of constraint
+	* @return return The constraint object
+	*/
+	UFUNCTION(BlueprintCallable, Category = "Editor Scripting | Sequencer Tools | Control Rig | Constraints")
+	static UTickableTransformConstraint* CreateFromType(UWorld* InWorld, const ETransformConstraintType InType);
+
+	/**
+	* Add Constraint to the system using the incoming parent and child handles with the specified type.
+	* @param InWorld World to create the constraint
+	* @param InParentHandle The parent handle
+	* @param InChildHandle The child handle
+	* @param InConsrtaint The constraint
+	* @return return If constraint added correctly
+	*/
+	UFUNCTION(BlueprintCallable, Category = "Editor Scripting | Sequencer Tools | Control Rig | Constraints")
+	static bool AddConstraint(UWorld* InWorld, UTransformableHandle* InParentHandle, UTransformableHandle* InChildHandle, UTickableTransformConstraint *InConstraint,
+			const bool bMaintainOffset);
+
+
+	/* Get a copy of the constraints in the current world
+	@param InWorld World we are in
+	@return Copy of the constraints in the level
+	*/
+	static TArray<TWeakObjectPtr<UTickableConstraint>> GetConstraintsArray(UWorld* InWorld);
+
+	/**
+	* Remove constraint at specified index
+	* @param InWorld World to create the constraint
+	* @param InIndex Index to remove from
+	* @return return If constraint removed correctly
+	*/
+	UFUNCTION(BlueprintCallable, Category = "Editor Scripting | Sequencer Tools | Control Rig | Constraints")
+	static bool RemoveConstraint(UWorld* InWorld, int32 InIndex);
+
+};
\ No newline at end of file
diff --git a/Engine/Source/Runtime/Experimental/Animation/Constraints/Public/TransformConstraint.h b/Engine/Source/Runtime/Experimental/Animation/Constraints/Public/TransformConstraint.h
--- a/Engine/Source/Runtime/Experimental/Animation/Constraints/Public/TransformConstraint.h
+++ b/Engine/Source/Runtime/Experimental/Animation/Constraints/Public/TransformConstraint.h
@@ -4,10 +4,13 @@
 
 #include "Constraint.h"
 #include "ConstraintsManager.h"
+#include "Kismet/BlueprintFunctionLibrary.h"
 
 #include "TransformConstraint.generated.h"
 
 class UTransformableHandle;
+class UTransformableComponentHandle;
+class USceneComponent;
 
 using SetTransformFunc = TFunction<void(const FTransform&)>;
 using GetTransformFunc = TFunction<FTransform()>;
@@ -16,7 +19,7 @@ using GetTransformFunc = TFunction<FTransform()>;
  * UTickableTransformConstraint
  **/
 
-UCLASS(Abstract)
+UCLASS(Abstract, Blueprintable)
 class CONSTRAINTS_API UTickableTransformConstraint : public UTickableConstraint
 {
 	GENERATED_BODY()
@@ -36,22 +39,23 @@ public:
 	virtual bool ReferencesObject(TWeakObjectPtr<UObject> InObject) const override;
 	
 	/** @todo document */
-	UPROPERTY()
+	UPROPERTY(BlueprintReadWrite, Category = "Handle")
 	TObjectPtr<UTransformableHandle> ParentTRSHandle;
 	/** @todo document */
-	UPROPERTY()
+	UPROPERTY(BlueprintReadWrite, Category = "Handle")
 	TObjectPtr<UTransformableHandle> ChildTRSHandle;
 
 	/** @todo document */
-	UPROPERTY(EditAnywhere, Category="Offset")
+	UPROPERTY(EditAnywhere, BlueprintReadWrite, Category="Offset")
 	bool bMaintainOffset = true;
 
 	/** @todo document */
-	// UPROPERTY(EditAnywhere, Category="Weight", meta = (Input, ClampMin = "0.0", ClampMax = "1.0", UIMin = "0.0", UIMax = "1.0"))
+	// UPROPERTY(EditAnywhere, BlueprintReadWrite,Category="Weight", meta = (Input, ClampMin = "0.0", ClampMax = "1.0", UIMin = "0.0", UIMax = "1.0"))
+	//@benoit when not EditAnywhere?
 	UPROPERTY(meta = (Input, ClampMin = "0.0", ClampMax = "1.0", UIMin = "0.0", UIMax = "1.0"))
 	float Weight = 1.f;
 
-	UPROPERTY(EditAnywhere, Category="Offset")
+	UPROPERTY(EditAnywhere, BlueprintReadWrite, Category="Offset")
 	bool bDynamicOffset = false;
 
 	/** @todo document. */
@@ -98,7 +102,7 @@ public:
  * UTickableTranslationConstraint
  **/
 
-UCLASS()
+UCLASS(Blueprintable)
 class CONSTRAINTS_API UTickableTranslationConstraint : public UTickableTransformConstraint
 {
 	GENERATED_BODY()
@@ -114,7 +118,7 @@ protected:
 	virtual void ComputeOffset() override;
 
 	/** @todo document */
-	UPROPERTY(EditAnywhere, Category="Offset")
+	UPROPERTY(EditAnywhere, BlueprintReadWrite, Category="Offset")
 	FVector OffsetTranslation = FVector::ZeroVector;
 };
 
@@ -122,7 +126,7 @@ protected:
  * UTickableRotationConstraint
  **/
 
-UCLASS()
+UCLASS(Blueprintable)
 class CONSTRAINTS_API UTickableRotationConstraint : public UTickableTransformConstraint
 {
 	GENERATED_BODY()
@@ -138,7 +142,7 @@ protected:
 	virtual void ComputeOffset() override;
 	
 	/** @todo document */
-	UPROPERTY(EditAnywhere, Category="Offset")
+	UPROPERTY(EditAnywhere, BlueprintReadWrite,Category="Offset")
 	FQuat OffsetRotation = FQuat::Identity;
 };
 
@@ -146,7 +150,7 @@ protected:
  * UTickableScaleConstraint
  **/
 
-UCLASS()
+UCLASS(Blueprintable)
 class CONSTRAINTS_API UTickableScaleConstraint : public UTickableTransformConstraint
 {
 	GENERATED_BODY()
@@ -162,7 +166,7 @@ protected:
 	virtual void ComputeOffset() override;
 
 	/** @todo document */
-	UPROPERTY(EditAnywhere, Category="Offset")
+	UPROPERTY(EditAnywhere, BlueprintReadWrite,Category="Offset")
 	FVector OffsetScale = FVector::OneVector;
 };
 
@@ -170,7 +174,7 @@ protected:
  * UTickableParentConstraint
  **/
 
-UCLASS()
+UCLASS(Blueprintable)
 class CONSTRAINTS_API UTickableParentConstraint : public UTickableTransformConstraint
 {
 	GENERATED_BODY()
@@ -186,7 +190,7 @@ protected:
 	virtual void ComputeOffset() override;
 
 	/** @todo document */
-	UPROPERTY(EditAnywhere, Category="Offset")
+	UPROPERTY(EditAnywhere, BlueprintReadWrite, Category="Offset")
 	FTransform OffsetTransform = FTransform::Identity;
 
 	/** @todo document */
@@ -206,7 +210,7 @@ protected:
  * UTickableLookAtConstraint
  **/
 
-UCLASS()
+UCLASS(Blueprintable)
 class CONSTRAINTS_API UTickableLookAtConstraint : public UTickableTransformConstraint
 {
 	GENERATED_BODY()
@@ -222,7 +226,7 @@ protected:
 	virtual void ComputeOffset() override;
 
 	/** @todo document */
-	UPROPERTY(EditAnywhere, Category="Offset")
+	UPROPERTY(EditAnywhere, BlueprintReadWrite,Category="Offset")
 	FTransform OffsetTransform = FTransform::Identity;
 };
 
@@ -230,29 +234,38 @@ protected:
  * TransformConstraintUtils
  **/
 
-struct CONSTRAINTS_API TransformConstraintUtils
+struct CONSTRAINTS_API FTransformConstraintUtils
 {
-	/** @todo document */
-	static bool Create(
-		UWorld* World,
-		const AActor* InParent,
-		const AActor* InChild,
-		const ETransformConstraintType InType = ETransformConstraintType::Parent,
-		const bool bMaintainOffset = true);
 
 	/** @todo document */
 	static void GetParentConstraints(
 		UWorld* World,
 		const AActor* InChild,
 		TArray< TObjectPtr<UTickableConstraint> >& OutConstraints);
-	
-private:
-	
+		
+	/** Create a handle for the scene component.*/
+	static UTransformableComponentHandle* CreateHandleForSceneComponent(
+		USceneComponent* InSceneComponent, 
+		UObject* Outer);
+
+	/** @todo document */
+	static UTickableTransformConstraint* CreateFromType(
+		UWorld* InWorld,
+		const ETransformConstraintType InType);
+
+	/** @todo document */
+	static UTickableTransformConstraint* CreateAndAddFromActors(
+		UWorld* InWorld,
+		AActor* InParent,
+		AActor* InChild,
+		const ETransformConstraintType InType,
+		const bool bMaintainOffset = true);
+
 	/** Registers a new transform constraint within the constraints manager. */	
-	static bool AddConst(
+	static bool AddConstraint(
 		UWorld* InWorld,
 		UTransformableHandle* InParentHandle,
 		UTransformableHandle* InChildHandle,
-		const ETransformConstraintType InType,
+		UTickableTransformConstraint* Constraint,
 		const bool bMaintainOffset = true);
 };
diff --git a/Engine/Source/Runtime/Experimental/Animation/Constraints/Public/TransformableHandle.h b/Engine/Source/Runtime/Experimental/Animation/Constraints/Public/TransformableHandle.h
--- a/Engine/Source/Runtime/Experimental/Animation/Constraints/Public/TransformableHandle.h
+++ b/Engine/Source/Runtime/Experimental/Animation/Constraints/Public/TransformableHandle.h
@@ -13,7 +13,7 @@ class USceneComponent;
  * UTransformableHandle
  */
 
-UCLASS(Abstract)
+UCLASS(Abstract, Blueprintable)
 class CONSTRAINTS_API UTransformableHandle : public UObject 
 {
 	GENERATED_BODY()
@@ -34,6 +34,8 @@ public:
 	virtual FTransform GetGlobalTransform() const PURE_VIRTUAL(GetGlobalTransform, return FTransform::Identity;);
 	/** Gets the local transform of the underlying transformable object in it's parent space. */
 	virtual FTransform GetLocalTransform() const PURE_VIRTUAL(GetLocalTransform, return FTransform::Identity;);
+	/** Gets the parent transform of the object*/
+	virtual FTransform GetParentTransform() const PURE_VIRTUAL(GetParentTransform, return FTransform::Identity;);
 
 	/**
 	 * Returns the target object containing the tick function (returned in GetTickFunction).
@@ -63,7 +65,7 @@ public:
  * UTransformableComponentHandle
  */
 
-UCLASS()
+UCLASS(Blueprintable)
 class CONSTRAINTS_API UTransformableComponentHandle : public UTransformableHandle 
 {
 	GENERATED_BODY()
@@ -83,6 +85,8 @@ public:
 	virtual FTransform GetGlobalTransform() const override;
 	/** @todo document */
 	virtual FTransform GetLocalTransform() const override;
+	/** @todo document */
+	virtual FTransform GetParentTransform() const override;
 
 	/** @todo document */
 	virtual UObject* GetPrerequisiteObject() const override;
@@ -101,6 +105,6 @@ public:
 #endif
 	
 	/** @todo document */
-	UPROPERTY()
+	UPROPERTY(BlueprintReadOnly, Category = "Object")
 	TWeakObjectPtr<USceneComponent> Component;
 };
\ No newline at end of file
