Update CCD settings for NextChapter. It shouldn't be necessary to have such a low value for boundsscale anymore since shapes are bigger and less issue-prone now. For the same reason, it should be fine to set CCD threshold bounds mode to 2, which uses the thinnest bound on each axis.

diff --git a/Engine/Source/Runtime/Experimental/Chaos/Private/Chaos/CCDUtilities.cpp b/Engine/Source/Runtime/Experimental/Chaos/Private/Chaos/CCDUtilities.cpp
--- a/Engine/Source/Runtime/Experimental/Chaos/Private/Chaos/CCDUtilities.cpp
+++ b/Engine/Source/Runtime/Experimental/Chaos/Private/Chaos/CCDUtilities.cpp
@@ -43,7 +43,7 @@ namespace Chaos
 		bool bCCDNewTargetDepthMode = true;
 		FAutoConsoleVariableRef  CVarCCDNewTargetDepthMode(TEXT("p.Chaos.CCD.NewTargetDepthMode"), bCCDNewTargetDepthMode, TEXT("Find the first contact with that results in a penetration of (CCDAllowedDepthBoundsScale*Size) as opposed to the first contact"));
 
-		int32 CCDAxisThresholdMode = 1;
+		int32 CCDAxisThresholdMode = 2;
 		FAutoConsoleVariableRef  CVarCCDAxisThresholdMode(TEXT("p.Chaos.CCD.AxisThresholdMode"), CCDAxisThresholdMode , TEXT("Change the mode used to generate CCD axis threshold bounds for particle geometries.\n0: Use object bounds\n1: Find the thinnest object bound on any axis and use it for all CCD axes\n2: On each axis, use the thinnest shape bound on that axis\n3: Find the thinnest shape bound on any axis and use this for all axes"));
 
 		bool bCCDAxisThresholdUsesProbeShapes = false;
