CIS Fix, remove unused virtual functions.
#trivial

diff --git a/Engine/Source/Runtime/Experimental/Animation/Constraints/Private/TransformableHandle.cpp b/Engine/Source/Runtime/Experimental/Animation/Constraints/Private/TransformableHandle.cpp
--- a/Engine/Source/Runtime/Experimental/Animation/Constraints/Private/TransformableHandle.cpp
+++ b/Engine/Source/Runtime/Experimental/Animation/Constraints/Private/TransformableHandle.cpp
@@ -56,11 +56,6 @@ FTransform UTransformableComponentHandle::GetGlobalTransform() const
 	return Component.IsValid() ? Component->GetComponentTransform() : FTransform::Identity;
 }
 
-FTransform UTransformableComponentHandle::GetParentTransform() const
-{
-	return (Component.IsValid() && Component->GetAttachParent()) ? Component->GetAttachParent()->GetComponentTransform() : FTransform::Identity;
-}
-
 UObject* UTransformableComponentHandle::GetPrerequisiteObject() const
 {
 	return Component.Get(); 
diff --git a/Engine/Source/Runtime/Experimental/Animation/Constraints/Public/TransformableHandle.h b/Engine/Source/Runtime/Experimental/Animation/Constraints/Public/TransformableHandle.h
--- a/Engine/Source/Runtime/Experimental/Animation/Constraints/Public/TransformableHandle.h
+++ b/Engine/Source/Runtime/Experimental/Animation/Constraints/Public/TransformableHandle.h
@@ -34,8 +34,6 @@ public:
 	virtual FTransform GetGlobalTransform() const PURE_VIRTUAL(GetGlobalTransform, return FTransform::Identity;);
 	/** Gets the local transform of the underlying transformable object in it's parent space. */
 	virtual FTransform GetLocalTransform() const PURE_VIRTUAL(GetLocalTransform, return FTransform::Identity;);
-	/** Gets the parent transform of the object*/
-	virtual FTransform GetParentTransform() const PURE_VIRTUAL(GetParentTransform, return FTransform::Identity;);
 
 	/**
 	 * Returns the target object containing the tick function (returned in GetTickFunction).
@@ -85,8 +83,6 @@ public:
 	virtual FTransform GetGlobalTransform() const override;
 	/** @todo document */
 	virtual FTransform GetLocalTransform() const override;
-	/** @todo document */
-	virtual FTransform GetParentTransform() const override;
 
 	/** @todo document */
 	virtual UObject* GetPrerequisiteObject() const override;
