Chaos - optimize non-linear joints
- remove joint IsActive tracking - results were not used
- non-linear joint projection doesn't call CorrectedP - P/Q gets updated each projection iteration anyway and there's usually only 1 projection iteration
- non-linear joint velocity solve does not use CorrectedP either because P does not change in the velocity phase
- removed some unnecessary position/mass updates in joint loop
- remove padding in joints. This is used for resitution in the Standard PBD mode which is no longer used and will be removed

diff --git a/Engine/Source/Runtime/Experimental/Chaos/Private/Chaos/Joint/PBDJointCachedSolverGaussSeidel.cpp b/Engine/Source/Runtime/Experimental/Chaos/Private/Chaos/Joint/PBDJointCachedSolverGaussSeidel.cpp
--- a/Engine/Source/Runtime/Experimental/Chaos/Private/Chaos/Joint/PBDJointCachedSolverGaussSeidel.cpp
+++ b/Engine/Source/Runtime/Experimental/Chaos/Private/Chaos/Joint/PBDJointCachedSolverGaussSeidel.cpp
@@ -184,9 +184,6 @@ void FPBDJointCachedSolver::Init(
 	NetLinearImpulse = FVec3(0);
 	NetAngularImpulse = FVec3(0);
 
-	LinearConstraintPadding = FVec3(-1);
-	AngularConstraintPadding = FVec3(-1);
-
 	// Tolerances are positional errors below visible detection. But in PBD the errors
 	// we leave behind get converted to velocity, so we need to ensure that the resultant
 	// movement from that erroneous velocity is less than the desired position tolerance.
@@ -464,8 +461,7 @@ void FPBDJointCachedSolver::InitPositionConstraintDatas(
 	}
 	else if(JointType == EJointMotionType::Limited)
 	{
-		PositionConstraints.ConstraintLimits[ConstraintIndex] =
-			FMath::Max(ConstraintLimit - GetLinearConstraintPadding(ConstraintIndex), (FReal)0.);
+		PositionConstraints.ConstraintLimits[ConstraintIndex] = ConstraintLimit;
 		PositionConstraints.UpdateDatas(ConstraintIndex, LocalAxis, LocalDelta,
 			ConstraintRestitution, true, ConstraintArm0, ConstraintArm1);
 	}
@@ -689,10 +685,6 @@ void FPBDJointCachedSolver::ApplyAxisPositionConstraint(
 		}
 		else if (PositionConstraints.MotionType[ConstraintIndex] != EJointMotionType::Free)
 		{
-			if (PositionConstraints.ConstraintRestitution[ConstraintIndex] > 0.0f)
-			{
-				CalculateLinearConstraintPadding(ConstraintIndex, Dt, PositionConstraints.ConstraintRestitution[ConstraintIndex], DeltaPosition);
-			}
 			SolvePositionConstraintHard(ConstraintIndex, DeltaPosition);
 		}
 	}
@@ -920,8 +912,7 @@ void FPBDJointCachedSolver::InitRotationConstraintDatas(
 
 	RotationConstraints.UpdateDatas(ConstraintIndex, LocalAxis, LocalAngle, ConstraintRestitution, bCheckLimit);
 
-	RotationConstraints.ConstraintLimits[ConstraintIndex] = FMath::Max(
-		JointSettings.AngularLimits[ConstraintIndex] - GetAngularConstraintPadding(ConstraintIndex), (FReal)0.);
+	RotationConstraints.ConstraintLimits[ConstraintIndex] = JointSettings.AngularLimits[ConstraintIndex];
 
 	InitConstraintAxisAngularVelocities[ConstraintIndex] = FVec3::DotProduct(W(1) - W(0), LocalAxis);
 
@@ -934,8 +925,7 @@ void FPBDJointCachedSolver::CorrectAxisAngleConstraint(
 		FVec3& ConstraintAxis,
 		FReal& ConstraintAngle) const
 {
-	const FReal AngleMax = FMath::Max(JointSettings.AngularLimits[ConstraintIndex] -
-		GetAngularConstraintPadding(ConstraintIndex), (FReal)0.);
+	const FReal AngleMax = JointSettings.AngularLimits[ConstraintIndex];
 
 	if (ConstraintAngle > AngleMax)
 	{
@@ -1183,10 +1173,6 @@ void FPBDJointCachedSolver::ApplyRotationConstraint(
 		}
 		else
 		{
-			if (RotationConstraints.ConstraintRestitution[ConstraintIndex] > 0.0f)
-			{
-				CalculateAngularConstraintPadding(ConstraintIndex, RotationConstraints.ConstraintRestitution[ConstraintIndex], DeltaAngle);
-			}
 			SolveRotationConstraintHard(ConstraintIndex, DeltaAngle);
 		}
 	}
@@ -1778,66 +1764,6 @@ void FPBDJointCachedSolver::ApplyRotationDelta(
 	Body(BodyIndex).ApplyRotationDelta(DR);
 }
 
-// Used for non-zero restitution. We pad constraints by an amount such that the velocity
-// calculated after solving constraint positions will as required for the restitution.
-void FPBDJointCachedSolver::CalculateLinearConstraintPadding(
-	const int32 ConstraintIndex,
-	const FReal Dt,
-	const FReal Restitution,
-	FReal& InOutPos)
-{
-	// NOTE: We only calculate the padding after the constraint is first violated, and after
-	// that the padding is fixed for the rest of the iterations in the current step.
-	if ((Restitution > 0.0f) && (InOutPos > 0.0f) && !HasLinearConstraintPadding(ConstraintIndex))
-	{
-		SetLinearConstraintPadding(ConstraintIndex, 0.0f);
-
-		// Calculate the velocity we want to match
-    
-		const FVec3 V0Dt = FVec3::CalculateVelocity(InitConnectorXs[0], ConnectorXs[0]+Body(0).DP() + FVec3::CrossProduct(Body(0).DQ(), PositionConstraints.ConstraintArms[PointPositionConstraintIndex][0]), 1.0f);
-		const FVec3 V1Dt = FVec3::CalculateVelocity(InitConnectorXs[1], ConnectorXs[1]+Body(1).DP() + FVec3::CrossProduct(Body(1).DQ(), PositionConstraints.ConstraintArms[PointPositionConstraintIndex][1]), 1.0f);
-		const FReal AxisVDt = FVec3::DotProduct(V1Dt - V0Dt, PositionConstraints.ConstraintAxis[ConstraintIndex]);
-
-		// Calculate the padding to apply to the constraint that will result in the
-		// desired outward velocity (assuming the constraint is fully resolved)
-		const FReal Padding = (1.0f + Restitution) * AxisVDt - InOutPos;
-		if (Padding > 0.0f)
-		{
-			SetLinearConstraintPadding(ConstraintIndex, Padding);
-			InOutPos += Padding;
-		}
-	}
-}
-
-// Used for non-zero restitution. We pad constraints by an amount such that the velocity
-// calculated after solving constraint positions will as required for the restitution.
-void FPBDJointCachedSolver::CalculateAngularConstraintPadding(
-	const int32 ConstraintIndex,
-	const FReal Restitution,
-	FReal& InOutAngle)
-{
-	// NOTE: We only calculate the padding after the constraint is first violated, and after
-	// that the padding is fixed for the rest of the iterations in the current step.
-	if ((Restitution > 0.0f) && (InOutAngle > 0.0f) && !HasAngularConstraintPadding(ConstraintIndex))
-	{
-		SetAngularConstraintPadding(ConstraintIndex, 0.0f);
-
-		// Calculate the velocity we want to match
-		const FVec3 W0Dt = FVec3(Body(0).DQ()) + ConnectorWDts[0];
-		const FVec3 W1Dt = FVec3(Body(1).DQ()) + ConnectorWDts[1];
-		const FReal AxisWDt = FVec3::DotProduct(W1Dt - W0Dt, RotationConstraints.ConstraintAxis[(int32)ConstraintIndex]);
-
-		// Calculate the padding to apply to the constraint that will result in the
-		// desired outward velocity (assuming the constraint is fully resolved)
-		const FReal Padding = (1.0f + Restitution) * AxisWDt - InOutAngle;
-		if (Padding > 0.0f)
-		{
-			SetAngularConstraintPadding(ConstraintIndex, Padding);
-			InOutAngle += Padding;
-		}
-	}
-}
-
 void FAxisConstraintDatas::InitDatas(
 	const int32 ConstraintIndex,
 	const bool bHasSoftLimits,
diff --git a/Engine/Source/Runtime/Experimental/Chaos/Private/Chaos/Joint/PBDJointSolverGaussSeidel.cpp b/Engine/Source/Runtime/Experimental/Chaos/Private/Chaos/Joint/PBDJointSolverGaussSeidel.cpp
--- a/Engine/Source/Runtime/Experimental/Chaos/Private/Chaos/Joint/PBDJointSolverGaussSeidel.cpp
+++ b/Engine/Source/Runtime/Experimental/Chaos/Private/Chaos/Joint/PBDJointSolverGaussSeidel.cpp
@@ -84,32 +84,6 @@ namespace Chaos
 		ConnectorRs[1].EnforceShortestArcWith(ConnectorRs[0]);
 	}
 
-	bool FPBDJointSolver::UpdateIsActive(
-		const FReal Dt,
-		const FPBDJointSolverSettings& SolverSettings,
-		const FPBDJointSettings& JointSettings)
-	{
-		// NumActiveConstraints is initialized to -1, so there's no danger of getting invalid LastPs/Qs
-		// We also check SolverStiffness mainly for testing when solver stiffness is 0 (so we don't exit immediately)
-		if ((NumActiveConstraints >= 0) && (SolverStiffness > 0.0f))
-		{
-			bool bIsSolved =
-				FVec3::IsNearlyEqual(P(0), LastPs[0], PositionTolerance)
-				&& FVec3::IsNearlyEqual(P(1), LastPs[1], PositionTolerance)
-				&& FRotation3::IsNearlyEqual(Q(0), LastQs[0], 0.5f * AngleTolerance)
-				&& FRotation3::IsNearlyEqual(Q(1), LastQs[1], 0.5f * AngleTolerance);
-			bIsActive = !bIsSolved;
-		}
-
-		LastPs[0] = P(0);
-		LastPs[1] = P(1);
-		LastQs[0] = Q(0);
-		LastQs[1] = Q(1);
-
-		return bIsActive;
-	}
-
-
 	void FPBDJointSolver::Init(
 		const FReal Dt,
 		const FSolverBodyPtrPair& SolverBodyPair,
@@ -143,9 +117,6 @@ namespace Chaos
 		LinearDriveLambdas = FVec3(0);
 		RotationDriveLambdas = FVec3(0);
 
-		LinearConstraintPadding = FVec3(-1);
-		AngularConstraintPadding = FVec3(-1);
-
 		// Tolerances are positional errors below visible detection. But in PBD the errors
 		// we leave behind get converted to velocity, so we need to ensure that the resultant
 		// movement from that erroneous velocity is less than the desired position tolerance.
@@ -157,9 +128,6 @@ namespace Chaos
 		PositionTolerance = ToleranceScale * SolverSettings.PositionTolerance;
 		AngleTolerance = ToleranceScale * SolverSettings.AngleTolerance;
 
-		NumActiveConstraints = -1;
-		bIsActive = true;
-
 		SolverStiffness = 1.0f;
 
 		LinearHardLambda = FVec3(0);
@@ -190,8 +158,6 @@ namespace Chaos
 		const FPBDJointSettings& JointSettings)
 	{
 		UpdateDerivedState();
-
-		UpdateIsActive(Dt, SolverSettings, JointSettings);
 	}
 
 	void FPBDJointSolver::UpdateMasses(
@@ -252,7 +218,6 @@ namespace Chaos
 		const FPBDJointSolverSettings& SolverSettings,
 		const FPBDJointSettings& JointSettings)
 	{
-		NumActiveConstraints = 0;
 		SolverStiffness = InSolverStiffness;
 
 		if (SolverSettings.bSolvePositionLast)
@@ -271,8 +236,6 @@ namespace Chaos
 			ApplyPositionDrives(Dt, SolverSettings, JointSettings);
 			ApplyRotationDrives(Dt, SolverSettings, JointSettings);
 		}
-
-		UpdateIsActive(Dt, SolverSettings, JointSettings);
 	}
 
 
@@ -286,7 +249,7 @@ namespace Chaos
 
 		// This is used for the QuasiPbd solver. If the Pbd step applied impulses to
 		// correct position errors, it will have introduced a velocity equal to the 
-		// corection divided by the timestep. We ensure that the velocity constraints
+		// correction divided by the timestep. We ensure that the velocity constraints
 		// (including restitution) are also enforced. This also prevents any position
 		// errors from the previous frame getting converted into energy.
 
@@ -447,8 +410,6 @@ namespace Chaos
 				ApplyVelocityProjection(Dt, SolverSettings, JointSettings, Chaos_Joint_VelProjectionAlpha, DP1, DR1);
 			}
 		}
-	
-		UpdateIsActive(Dt, SolverSettings, JointSettings);
 	}
 
 
@@ -996,9 +957,8 @@ namespace Chaos
 		{
 			FVec3 SwingAxisLocal;
 			FReal DSwingAngle = 0.0f;
-			const FReal LimitPadding = GetAngularConstraintPadding(EJointAngularConstraintIndex::Swing1);
-			const FReal Swing1Limit = FMath::Max(JointSettings.AngularLimits[(int32)EJointAngularConstraintIndex::Swing1] - LimitPadding, 0.0f);
-			const FReal Swing2Limit = FMath::Max(JointSettings.AngularLimits[(int32)EJointAngularConstraintIndex::Swing2] - LimitPadding, 0.0f);
+			const FReal Swing1Limit = JointSettings.AngularLimits[(int32)EJointAngularConstraintIndex::Swing1];
+			const FReal Swing2Limit = JointSettings.AngularLimits[(int32)EJointAngularConstraintIndex::Swing2];
 			FPBDJointUtilities::GetEllipticalConeAxisErrorLocal(ConnectorRs[0], ConnectorRs[1], Swing2Limit, Swing1Limit, SwingAxisLocal, DSwingAngle);
 			// SwingAxisLocal has the size of Sine(SwingAngle), thus we need to normalize it
 			SwingAxisLocal.SafeNormalize();
@@ -1285,8 +1245,6 @@ namespace Chaos
 		{
 			LinearHardLambda[LinearHardLambdaIndex] += DLambda;
 		}
-
-		++NumActiveConstraints;
 	}
 
 
@@ -1344,8 +1302,6 @@ namespace Chaos
 			Lambda += DLambda;
 			NetLinearImpulse += DLambda * Axis;
 		}
-
-		++NumActiveConstraints;
 	}
 	
 
@@ -1419,8 +1375,6 @@ namespace Chaos
 		{
 			ApplyRotationConstraintDD(JointStiffness, Axis, Angle, AngularHardLambdaIndex);
 		}
-
-		++NumActiveConstraints;
 	}
 
 
@@ -1564,74 +1518,6 @@ namespace Chaos
 		{
 			ApplyRotationConstraintSoftDD(Dt, JointStiffness, JointDamping, bAccelerationMode, Axis, Angle, AngVelTarget, Lambda);
 		}
-
-		++NumActiveConstraints;
-	}
-
-	// Used for non-zero restitution. We pad constraints by an amount such that the velocity
-	// calculated after solving constraint positions will as required for the restitution.
-	void FPBDJointSolver::CalculateLinearConstraintPadding(
-		const FReal Dt,
-		const FPBDJointSolverSettings& SolverSettings,
-		const FPBDJointSettings& JointSettings,
-		const FReal Restitution,
-		const int32 AxisIndex,
-		const FVec3 Axis,
-		FReal& InOutPos)
-	{
-		// NOTE: We only calculate the padding after the constraint is first violated, and after
-		// that the padding is fixed for the rest of the iterations in the current step.
-		if ((Restitution > 0.0f) && (InOutPos > 0.0f) && !HasLinearConstraintPadding(AxisIndex))
-		{
-			SetLinearConstraintPadding(AxisIndex, 0.0f);
-
-			// Calculate the velocity we want to match
-			const FVec3 V0Dt = FVec3::CalculateVelocity(InitConnectorXs[0], ConnectorXs[0], 1.0f);
-			const FVec3 V1Dt = FVec3::CalculateVelocity(InitConnectorXs[1], ConnectorXs[1], 1.0f);
-			const FReal AxisVDt = FVec3::DotProduct(V1Dt - V0Dt, Axis);
-
-			// Calculate the padding to apply to the constraint that will result in the
-			// desired outward velocity (assuming the constraint is fully resolved)
-			const FReal Padding = (1.0f + Restitution) * AxisVDt - InOutPos;
-			if (Padding > 0.0f)
-			{
-				SetLinearConstraintPadding(AxisIndex, Padding);
-				InOutPos += Padding;
-			}
-		}
-	}
-
-	// Used for non-zero restitution. We pad constraints by an amount such that the velocity
-	// calculated after solving constraint positions will as required for the restitution.
-	void FPBDJointSolver::CalculateAngularConstraintPadding(
-		const FReal Dt,
-		const FPBDJointSolverSettings& SolverSettings,
-		const FPBDJointSettings& JointSettings,
-		const FReal Restitution,
-		const EJointAngularConstraintIndex ConstraintIndex,
-		const FVec3 Axis,
-		FReal& InOutAngle)
-	{
-		// NOTE: We only calculate the padding after the constraint is first violated, and after
-		// that the padding is fixed for the rest of the iterations in the current step.
-		if ((Restitution > 0.0f) && (InOutAngle > 0.0f) && !HasAngularConstraintPadding(ConstraintIndex))
-		{
-			SetAngularConstraintPadding(ConstraintIndex, 0.0f);
-
-			// Calculate the velocity we want to match
-			const FVec3 W0Dt = FRotation3::CalculateAngularVelocity(InitConnectorRs[0], ConnectorRs[0], 1.0f);
-			const FVec3 W1Dt = FRotation3::CalculateAngularVelocity(InitConnectorRs[1], ConnectorRs[1], 1.0f);
-			const FReal AxisWDt = FVec3::DotProduct(W1Dt - W0Dt, Axis);
-
-			// Calculate the padding to apply to the constraint that will result in the
-			// desired outward velocity (assuming the constraint is fully resolved)
-			const FReal Padding = (1.0f + Restitution) * AxisWDt - InOutAngle;
-			if (Padding > 0.0f)
-			{
-				SetAngularConstraintPadding(ConstraintIndex, Padding);
-				InOutAngle += Padding;
-			}
-		}
 	}
 
 	//
@@ -1677,9 +1563,8 @@ namespace Chaos
 		FPBDJointUtilities::GetTwistAxisAngle(ConnectorRs[0], ConnectorRs[1], TwistAxis, TwistAngle);
 
 		// Calculate the twist correction to apply to each body
-		const FReal LimitPadding = GetAngularConstraintPadding(EJointAngularConstraintIndex::Twist);
 		FReal DTwistAngle = 0;
-		FReal TwistAngleMax = FMath::Max(JointSettings.AngularLimits[(int32)EJointAngularConstraintIndex::Twist] - LimitPadding, (FReal)0.);
+		FReal TwistAngleMax = JointSettings.AngularLimits[(int32)EJointAngularConstraintIndex::Twist];
 		if (TwistAngle > TwistAngleMax)
 		{
 			DTwistAngle = TwistAngle - TwistAngleMax;
@@ -1705,11 +1590,6 @@ namespace Chaos
 			}
 			else
 			{
-				if (JointSettings.TwistRestitution > 0.0f)
-				{
-					CalculateAngularConstraintPadding(Dt, SolverSettings, JointSettings, JointSettings.TwistRestitution, EJointAngularConstraintIndex::Twist, TwistAxis, DTwistAngle);
-				}
-
 				FReal TwistStiffness = FPBDJointUtilities::GetTwistStiffness(SolverSettings, JointSettings);
 				ApplyRotationConstraint(TwistStiffness, TwistAxis, DTwistAngle, (int32)EJointAngularConstraintIndex::Twist);
 			}
@@ -1725,9 +1605,8 @@ namespace Chaos
 		FVec3 SwingAxisLocal;
 		FReal DSwingAngle = 0.0f;
 
-		const FReal LimitPadding = GetAngularConstraintPadding(EJointAngularConstraintIndex::Swing1);
-		const FReal Swing1Limit = FMath::Max(JointSettings.AngularLimits[(int32)EJointAngularConstraintIndex::Swing1] - LimitPadding, (FReal)0.);
-		const FReal Swing2Limit = FMath::Max(JointSettings.AngularLimits[(int32)EJointAngularConstraintIndex::Swing2] - LimitPadding, (FReal)0.);
+		const FReal Swing1Limit = JointSettings.AngularLimits[(int32)EJointAngularConstraintIndex::Swing1];
+		const FReal Swing2Limit = JointSettings.AngularLimits[(int32)EJointAngularConstraintIndex::Swing2];
 		FPBDJointUtilities::GetEllipticalConeAxisErrorLocal(ConnectorRs[0], ConnectorRs[1], Swing2Limit, Swing1Limit, SwingAxisLocal, DSwingAngle);
 
 		UE_LOG(LogChaosJoint, VeryVerbose, TEXT("    Cone Error %f [Limits %f %f]"), FMath::RadiansToDegrees(DSwingAngle), FMath::RadiansToDegrees(Swing2Limit), FMath::RadiansToDegrees(Swing1Limit));
@@ -1746,11 +1625,6 @@ namespace Chaos
 			}
 			else
 			{
-				if (JointSettings.SwingRestitution > 0.0f)
-				{
-					CalculateAngularConstraintPadding(Dt, SolverSettings, JointSettings, JointSettings.SwingRestitution, EJointAngularConstraintIndex::Swing1, SwingAxis, DSwingAngle);
-				}
-
 				FReal SwingStiffness = FPBDJointUtilities::GetSwingStiffness(SolverSettings, JointSettings);
 				// For cone constraints, the lambda are all accumulated in Swing2
 				ApplyRotationConstraint(SwingStiffness, SwingAxis, DSwingAngle, (int32)EJointAngularConstraintIndex::Swing2);
@@ -1806,9 +1680,8 @@ namespace Chaos
 		FPBDJointUtilities::GetDualConeSwingAxisAngle(ConnectorRs[0], ConnectorRs[1], SwingConstraintIndex, SwingAxis, SwingAngle);
 
 		// Calculate swing error we need to correct
-		const FReal LimitPadding = GetAngularConstraintPadding(SwingConstraintIndex);
 		FReal DSwingAngle = 0;
-		const FReal SwingAngleMax = FMath::Max(JointSettings.AngularLimits[(int32)SwingConstraintIndex] - LimitPadding, (FReal)0.);
+		const FReal SwingAngleMax = JointSettings.AngularLimits[(int32)SwingConstraintIndex];
 		if (SwingAngle > SwingAngleMax)
 		{
 			DSwingAngle = SwingAngle - SwingAngleMax;
@@ -1834,11 +1707,6 @@ namespace Chaos
 			}
 			else
 			{
-				if (JointSettings.SwingRestitution > 0.0f)
-				{
-					CalculateAngularConstraintPadding(Dt, SolverSettings, JointSettings, JointSettings.SwingRestitution, SwingConstraintIndex, SwingAxis, DSwingAngle);
-				}
-
 				const FReal SwingStiffness = FPBDJointUtilities::GetSwingStiffness(SolverSettings, JointSettings);
 				ApplyRotationConstraint(SwingStiffness, SwingAxis, DSwingAngle, (int32)SwingConstraintIndex);
 			}
@@ -1858,9 +1726,8 @@ namespace Chaos
 		FPBDJointUtilities::GetSwingAxisAngle(ConnectorRs[0], ConnectorRs[1], SolverSettings.SwingTwistAngleTolerance, SwingConstraintIndex, SwingAxis, SwingAngle);
 
 		// Calculate swing error we need to correct
-		const FReal LimitPadding = GetAngularConstraintPadding(SwingConstraintIndex);
 		FReal DSwingAngle = 0;
-		const FReal SwingAngleMax = FMath::Max(JointSettings.AngularLimits[(int32)SwingConstraintIndex] - LimitPadding, (FReal)0.);
+		const FReal SwingAngleMax = JointSettings.AngularLimits[(int32)SwingConstraintIndex];
 		if (SwingAngle > SwingAngleMax)
 		{
 			DSwingAngle = SwingAngle - SwingAngleMax;
@@ -1886,11 +1753,6 @@ namespace Chaos
 			}
 			else
 			{
-				if (JointSettings.SwingRestitution > 0.0f)
-				{
-					CalculateAngularConstraintPadding(Dt, SolverSettings, JointSettings, JointSettings.SwingRestitution, SwingConstraintIndex, SwingAxis, DSwingAngle);
-				}
-
 				const FReal SwingStiffness = FPBDJointUtilities::GetSwingStiffness(SolverSettings, JointSettings);
 				ApplyRotationConstraint(SwingStiffness, SwingAxis, DSwingAngle, (int32)SwingConstraintIndex);
 			}
@@ -2028,8 +1890,6 @@ namespace Chaos
 
 				NetLinearImpulse += (KIndex == 0) ? DX : -DX;
 			}
-
-			++NumActiveConstraints;
 		}
 	}
 
@@ -2075,8 +1935,6 @@ namespace Chaos
 
 				NetLinearImpulse += DX;
 			}
-			
-			++NumActiveConstraints;
 		}
 	}
 
@@ -2090,19 +1948,13 @@ namespace Chaos
 		FReal Delta;
 		FPBDJointUtilities::GetSphericalAxisDelta(ConnectorXs[0], ConnectorXs[1], Axis, Delta);
 
-		const FReal LimitPadding = GetLinearConstraintPadding(0);
-		const FReal Limit = FMath::Max(JointSettings.LinearLimit - LimitPadding, (FReal)0.);
+		const FReal Limit = JointSettings.LinearLimit;
 
 		FReal Error = Delta - Limit;
 		if (Error > PositionTolerance)
 		{
 			if (!FPBDJointUtilities::GetSoftLinearLimitEnabled(SolverSettings, JointSettings))
 			{
-				if (JointSettings.LinearRestitution > 0.0f)
-				{
-					CalculateLinearConstraintPadding(Dt, SolverSettings, JointSettings, JointSettings.LinearRestitution, 0, Axis, Error);
-				}
-
 				const FReal JointStiffness = FPBDJointUtilities::GetLinearStiffness(SolverSettings, JointSettings);
 				const FVec3 Connector0Correction = Axis * Limit;
 				ApplyPositionConstraint(JointStiffness, Axis, Error, Connector0Correction);
@@ -2138,8 +1990,7 @@ namespace Chaos
 			Axis = -Axis;
 		}
 		
-		const FReal AxialLimitPadding = GetLinearConstraintPadding(0);
-		const FReal AxialLimit = (AxialMotion == EJointMotionType::Locked) ? 0.0f : FMath::Max(JointSettings.LinearLimit - AxialLimitPadding, (FReal)0.);
+		const FReal AxialLimit = (AxialMotion == EJointMotionType::Locked) ? 0.0f : JointSettings.LinearLimit;
 		FReal AxialError = AxialDelta - AxialLimit;
 
 		if (AxialError > PositionTolerance)
@@ -2154,20 +2005,15 @@ namespace Chaos
 			}
 			else if (AxialMotion != EJointMotionType::Free)
 			{
-			// Hard Axial constraint
-				if (JointSettings.LinearRestitution > 0.0f)
-				{
-					CalculateLinearConstraintPadding(Dt, SolverSettings, JointSettings, JointSettings.LinearRestitution, 0, Axis, AxialError);
-				}
-
+				// Hard Axial constraint
+	
 				const FReal JointStiffness = FPBDJointUtilities::GetLinearStiffness(SolverSettings, JointSettings);
 				const FVec3 Connector0Correction = ConnectorXs[1] - ConnectorXs[0] - AxialError * Axis;
 				ApplyPositionConstraint(JointStiffness, Axis, AxialError, Connector0Correction, (int32)EJointCylindricalPositionConstraintType::Axial);
 			}
 		}
 
-		const FReal RadialLimitPadding = GetLinearConstraintPadding(1);
-		const FReal RadialLimit = (RadialMotion == EJointMotionType::Locked) ? 0.0f : FMath::Max(JointSettings.LinearLimit - RadialLimitPadding, (FReal)0.);
+		const FReal RadialLimit = (RadialMotion == EJointMotionType::Locked) ? 0.0f : JointSettings.LinearLimit;
 		FReal RadialError = RadialDelta - RadialLimit;
 
 		if (RadialError > PositionTolerance)
@@ -2183,11 +2029,6 @@ namespace Chaos
 			else if (RadialMotion != EJointMotionType::Free)
 			{
 				// Hard Radial constraint
-				if (JointSettings.LinearRestitution > 0.0f)
-				{
-					CalculateLinearConstraintPadding(Dt, SolverSettings, JointSettings, JointSettings.LinearRestitution, 1, RadialAxis, RadialError);
-				}
-
 				const FReal JointStiffness = FPBDJointUtilities::GetLinearStiffness(SolverSettings, JointSettings);
 				const FVec3 Connector0Correction = ConnectorXs[1] - ConnectorXs[0] - RadialError * RadialAxis;
 				ApplyPositionConstraint(JointStiffness, RadialAxis, RadialError, Connector0Correction, (int32)EJointCylindricalPositionConstraintType::Radial);
@@ -2213,8 +2054,7 @@ namespace Chaos
 			Axis = -Axis;
 		}
 
-		const FReal LimitPadding = GetLinearConstraintPadding(0);
-		const FReal Limit = (AxialMotion == EJointMotionType::Locked) ? 0 : FMath::Max(JointSettings.LinearLimit - LimitPadding, (FReal)0.);
+		const FReal Limit = (AxialMotion == EJointMotionType::Locked) ? 0 : JointSettings.LinearLimit;
 		FReal Error = Delta - Limit;
 		if (Error > PositionTolerance)
 		{
@@ -2227,11 +2067,6 @@ namespace Chaos
 			}
 			else
 			{
-				if (JointSettings.LinearRestitution > 0.0f)
-				{
-					CalculateLinearConstraintPadding(Dt, SolverSettings, JointSettings, JointSettings.LinearRestitution, 0, Axis, Error);
-				}
-
 				const FReal JointStiffness = FPBDJointUtilities::GetLinearStiffness(SolverSettings, JointSettings);
 				const FVec3 Connector0Correction = ConnectorXs[1] - ConnectorXs[0] - Axis * Error;
 				ApplyPositionConstraint(JointStiffness, Axis, Error, Connector0Correction);
@@ -2275,12 +2110,12 @@ namespace Chaos
 		const FVec3 CX = ConnectorXs[1] - ConnectorXs[0];
 		if (CX.Size() > ProjectionPositionTolerance)
 		{
-			FMatrix33 J = Utilities::ComputeJointFactorMatrix(ConnectorXs[1] - P(1), InvI(1), InvM(1));
+			FMatrix33 J = Utilities::ComputeJointFactorMatrix(ConnectorXs[1] - Body1().P(), InvI(1), InvM(1));
 			const FMatrix33 IJ = J.Inverse();
 			const FVec3 DX = Utilities::Multiply(IJ, CX);
 
 			const FVec3 DP1 = -Alpha * InvM(1) * DX;
-			const FVec3 DR1 = -Alpha * Utilities::Multiply(InvI(1), FVec3::CrossProduct(ConnectorXs[1] - P(1), DX));
+			const FVec3 DR1 = -Alpha * Utilities::Multiply(InvI(1), FVec3::CrossProduct(ConnectorXs[1] - Body1().P(), DX));
 			ApplyDelta(1, DP1, DR1);
 
 			NetDP1 += DP1;
@@ -2304,14 +2139,14 @@ namespace Chaos
 		const FReal Error = FMath::Max((FReal)0, Delta - JointSettings.LinearLimit);
 		if (FMath::Abs(Error) > ProjectionPositionTolerance)
 		{
-			const FVec3 AngularAxis1 = FVec3::CrossProduct(ConnectorXs[1] - P(1), Axis);
+			const FVec3 AngularAxis1 = FVec3::CrossProduct(ConnectorXs[1] - Body1().P(), Axis);
 			const FVec3 IA1 = Utilities::Multiply(InvI(1), AngularAxis1);
 			const FReal II1 = FVec3::DotProduct(AngularAxis1, IA1);
 			const FReal IM = InvM(1) + II1;
 			const FVec3 DX = Axis * Error / IM;
 
 			const FVec3 DP1 = -Alpha * InvM(1) * DX;
-			const FVec3 DR1 = -Alpha * Utilities::Multiply(InvI(1), FVec3::CrossProduct(ConnectorXs[1] - P(1), DX));
+			const FVec3 DR1 = -Alpha * Utilities::Multiply(InvI(1), FVec3::CrossProduct(ConnectorXs[1] - Body1().P(), DX));
 			ApplyDelta(1, DP1, DR1);
 
 			NetDP1 += DP1;
@@ -2365,7 +2200,7 @@ namespace Chaos
 			FVec3 DP1 = FVec3(0);
 			if (bPositionLocked)
 			{
-				DP1 = -Alpha * FVec3::CrossProduct(DR1, ConnectorXs[1] - P(1));
+				DP1 = -Alpha * FVec3::CrossProduct(DR1, ConnectorXs[1] - Body1().P());
 			}
 			ApplyDelta(1, DP1, DR1);
 
@@ -2411,7 +2246,7 @@ namespace Chaos
 			FVec3 DP1 = FVec3(0);
 			if (bPositionLocked)
 			{
-				DP1 = -Alpha * FVec3::CrossProduct(DR1, ConnectorXs[1] - P(1));
+				DP1 = -Alpha * FVec3::CrossProduct(DR1, ConnectorXs[1] - Body1().P());
 			}
 			ApplyDelta(1, DP1, DR1);
 
@@ -2454,7 +2289,7 @@ namespace Chaos
 			FVec3 DP1 = FVec3(0);
 			if (bPositionLocked)
 			{
-				DP1 = -Alpha * FVec3::CrossProduct(DR1, ConnectorXs[1] - P(1));
+				DP1 = -Alpha * FVec3::CrossProduct(DR1, ConnectorXs[1] - Body1().P());
 			}
 			ApplyDelta(1, DP1, DR1);
 
@@ -2498,7 +2333,7 @@ namespace Chaos
 			FVec3 DP1 = FVec3(0);
 			if (bPositionLocked)
 			{
-				DP1 = -Alpha * FVec3::CrossProduct(DR1, ConnectorXs[1] - P(1));
+				DP1 = -Alpha * FVec3::CrossProduct(DR1, ConnectorXs[1] - Body1().P());
 			}
 			ApplyDelta(1, DP1, DR1);
 
@@ -2538,7 +2373,7 @@ namespace Chaos
 			FVec3 DP1 = FVec3(0);
 			if (bPositionLocked)
 			{
-				DP1 = -Alpha * FVec3::CrossProduct(DR1, ConnectorXs[1] - P(1));
+				DP1 = -Alpha * FVec3::CrossProduct(DR1, ConnectorXs[1] - Body1().P());
 			}
 			ApplyDelta(1, DP1, DR1);
 
@@ -2566,8 +2401,8 @@ namespace Chaos
 		const FVec3& Connector0Correction,
 		const FReal TargetVel)
 	{
-		const FVec3 Arm0 = ConnectorXs[0] + Connector0Correction - P(0);
-		const FVec3 Arm1 = ConnectorXs[1] - P(1);
+		const FVec3 Arm0 = ConnectorXs[0] + Connector0Correction - Body0().P();
+		const FVec3 Arm1 = ConnectorXs[1] - Body1().P();
 		const FVec3 CV0 = V(0) + FVec3::CrossProduct(W(0), Arm0);
 		const FVec3 CV1 = V(1) + FVec3::CrossProduct(W(1), Arm1);
 		const FVec3 CV = CV1 - CV0;
@@ -2598,8 +2433,8 @@ namespace Chaos
 		if (!NetLinearImpulse.IsNearlyZero() && !FPBDJointUtilities::GetSoftLinearLimitEnabled(SolverSettings, JointSettings))
 		{
 			FReal Stiffness = SolverStiffness * FPBDJointUtilities::GetLinearStiffness(SolverSettings, JointSettings);
-			const FVec3 CV0 = V(0) + FVec3::CrossProduct(W(0), ConnectorXs[0] - P(0));
-			const FVec3 CV1 = V(1) + FVec3::CrossProduct(W(1), ConnectorXs[1] - P(1));
+			const FVec3 CV0 = V(0) + FVec3::CrossProduct(W(0), ConnectorXs[0] - Body0().P());
+			const FVec3 CV1 = V(1) + FVec3::CrossProduct(W(1), ConnectorXs[1] - Body1().P());
 			const FVec3 CV = CV1 - CV0;
 
 			UE_LOG(LogChaosJoint, VeryVerbose, TEXT("    PointVel Delta %f"), CV.Size());
@@ -2608,11 +2443,11 @@ namespace Chaos
 			FMatrix33 M1 = FMatrix33::Zero;
 			if (InvM(0) > UE_SMALL_NUMBER)
 			{
-				M0 = Utilities::ComputeJointFactorMatrix(ConnectorXs[0] - P(0), InvI(0), InvM(0));
+				M0 = Utilities::ComputeJointFactorMatrix(ConnectorXs[0] - Body0().P(), InvI(0), InvM(0));
 			}
 			if (InvM(1) > UE_SMALL_NUMBER)
 			{
-				M1 = Utilities::ComputeJointFactorMatrix(ConnectorXs[1] - P(1), InvI(1), InvM(1));
+				M1 = Utilities::ComputeJointFactorMatrix(ConnectorXs[1] - Body1().P(), InvI(1), InvM(1));
 			}
 			FMatrix33 MI = (M0 + M1).Inverse();
 			const FVec3 Impulse = Stiffness * Utilities::Multiply(MI, CV);
@@ -2620,8 +2455,8 @@ namespace Chaos
 			// Apply constraint correction
 			const FVec3 DV0 = InvM(0) * Impulse;
 			const FVec3 DV1 = -InvM(1) * Impulse;
-			const FVec3 DW0 = Utilities::Multiply(InvI(0), FVec3::CrossProduct(ConnectorXs[0] - P(0), Impulse));
-			const FVec3 DW1 = Utilities::Multiply(InvI(1), FVec3::CrossProduct(ConnectorXs[1] - P(1), -Impulse));
+			const FVec3 DW0 = Utilities::Multiply(InvI(0), FVec3::CrossProduct(ConnectorXs[0] - Body0().P(), Impulse));
+			const FVec3 DW1 = Utilities::Multiply(InvI(1), FVec3::CrossProduct(ConnectorXs[1] - Body1().P(), -Impulse));
 
 			ApplyVelocityDelta(DV0, DW0, DV1, DW1);
 		}
@@ -2840,8 +2675,8 @@ namespace Chaos
 		const FVec3 DConnectorX = ConnectorXs[1] - ConnectorXs[0];
 		const FVec3 Connector0Correction = Delta > JointSettings.LinearLimit ? Axis * JointSettings.LinearLimit : DConnectorX;
 
-		const FVec3 Arm0 = ConnectorXs[0] + Connector0Correction - P(0);
-		const FVec3 Arm1 = ConnectorXs[1] - P(1);
+		const FVec3 Arm0 = ConnectorXs[0] + Connector0Correction - Body0().P();
+		const FVec3 Arm1 = ConnectorXs[1] - Body1().P();
 		const FVec3 CV0 = V(0) + FVec3::CrossProduct(W(0), Arm0);
 		const FVec3 CV1 = V(1) + FVec3::CrossProduct(W(1), Arm1);
 		const FVec3 CV = CV1 - CV0;
@@ -2920,8 +2755,8 @@ namespace Chaos
 		const FVec3 DConnectorX = ConnectorXs[1] - ConnectorXs[0];
 		const FVec3 Connector0Correction = DConnectorX - Axis * Error; 
 
-		const FVec3 Arm0 = ConnectorXs[0] + Connector0Correction - P(0);
-		const FVec3 Arm1 = ConnectorXs[1] - P(1);
+		const FVec3 Arm0 = ConnectorXs[0] + Connector0Correction - Body0().P();
+		const FVec3 Arm1 = ConnectorXs[1] - Body1().P();
 		const FVec3 CV0 = V(0) + FVec3::CrossProduct(W(0), Arm0);
 		const FVec3 CV1 = V(1) + FVec3::CrossProduct(W(1), Arm1);
 		const FVec3 CV = CV1 - CV0;
diff --git a/Engine/Source/Runtime/Experimental/Chaos/Private/Chaos/PBDJointConstraints.cpp b/Engine/Source/Runtime/Experimental/Chaos/Private/Chaos/PBDJointConstraints.cpp
--- a/Engine/Source/Runtime/Experimental/Chaos/Private/Chaos/PBDJointConstraints.cpp
+++ b/Engine/Source/Runtime/Experimental/Chaos/Private/Chaos/PBDJointConstraints.cpp
@@ -1423,6 +1423,7 @@ namespace Chaos
 			if ((SolverType == EConstraintSolverType::StandardPbd) || (It == 0))
 			{
 				// @todo(chaos): support reverse parent/child
+				Solver.Body1().UpdateRotationDependentState();
 				Solver.UpdateMasses(FReal(0), FReal(1));
 			}
 
@@ -1430,15 +1431,9 @@ namespace Chaos
 
 			if (SolverType == EConstraintSolverType::StandardPbd)
 			{
-				if (Solver.Body0().IsDynamic())
-				{
-					Solver.Body0().SolverBody().ApplyCorrections();
-					Solver.Body0().UpdateRotationDependentState();
-				}
 				if (Solver.Body1().IsDynamic())
 				{
 					Solver.Body1().SolverBody().ApplyCorrections();
-					Solver.Body1().UpdateRotationDependentState();
 				}
 			}
 		}
diff --git a/Engine/Source/Runtime/Experimental/Chaos/Public/Chaos/Joint/PBDJointCachedSolverGaussSeidel.h b/Engine/Source/Runtime/Experimental/Chaos/Public/Chaos/Joint/PBDJointCachedSolverGaussSeidel.h
--- a/Engine/Source/Runtime/Experimental/Chaos/Public/Chaos/Joint/PBDJointCachedSolverGaussSeidel.h
+++ b/Engine/Source/Runtime/Experimental/Chaos/Public/Chaos/Joint/PBDJointCachedSolverGaussSeidel.h
@@ -531,48 +531,6 @@ struct FAxisConstraintDatas
     		const int32 ConstraintIndex,
     		const FReal Dt);
 		
-		/** Linear / Angular padding */
-		
-		void CalculateLinearConstraintPadding(
-			const int32 ConstraintIndex,
-			const FReal Dt,
-			const FReal Restitution,
-			FReal& InOutPos);
-
-		void CalculateAngularConstraintPadding(
-			const int32 ConstraintIndex,
-			const FReal Restitution,
-			FReal& InOutAngle);
-
-		inline bool HasLinearConstraintPadding(const int32 AxisIndex) const
-		{
-			return LinearConstraintPadding[AxisIndex] >= 0.0f;
-		}
-
-		inline FReal GetLinearConstraintPadding(const int32 AxisIndex) const
-		{
-			return FMath::Max(LinearConstraintPadding[AxisIndex], 0.0f);
-		}
-
-		inline void SetLinearConstraintPadding(const int32 AxisIndex, FReal Padding)
-		{
-			LinearConstraintPadding[AxisIndex] = Padding;
-		}
-		
-		inline bool HasAngularConstraintPadding(const int32 ConstraintIndex) const
-		{
-			return AngularConstraintPadding[ConstraintIndex] >= 0.0f;
-		}
-
-		inline FReal GetAngularConstraintPadding(const int32 ConstraintIndex) const
-		{
-			return FMath::Max(AngularConstraintPadding[ConstraintIndex], 0.0f);
-		}
-
-		inline void SetAngularConstraintPadding(const int32 ConstraintIndex, FReal Padding)
-		{
-			AngularConstraintPadding[ConstraintIndex] = Padding;
-		}
 		
 		// The cached body state on which the joint operates
 		FConstraintSolverBody SolverBodies[MaxConstrainedBodies];
@@ -601,10 +559,6 @@ struct FAxisConstraintDatas
 		// \todo(chaos): remove Stiffness from SolverSettings (because it is not a solver constant)
 		FReal SolverStiffness;
 
-		// Constraint padding which can act something like a velocity constraint (for reslockfreetitution)
-		FVec3 LinearConstraintPadding;
-		FVec3 AngularConstraintPadding;
-
 		// Tolerances below which we stop solving
 		FReal PositionTolerance;					// Distance error below which we consider a constraint or drive solved
 		FReal AngleTolerance;						// Angle error below which we consider a constraint or drive solved
diff --git a/Engine/Source/Runtime/Experimental/Chaos/Public/Chaos/Joint/PBDJointSolverGaussSeidel.h b/Engine/Source/Runtime/Experimental/Chaos/Public/Chaos/Joint/PBDJointSolverGaussSeidel.h
--- a/Engine/Source/Runtime/Experimental/Chaos/Public/Chaos/Joint/PBDJointSolverGaussSeidel.h
+++ b/Engine/Source/Runtime/Experimental/Chaos/Public/Chaos/Joint/PBDJointSolverGaussSeidel.h
@@ -127,18 +127,6 @@ namespace Chaos
 			return NetAngularImpulse;
 		}
 
-		inline int32 GetNumActiveConstraints() const
-		{
-			// We use -1 as unitialized, but that should not be exposed outside the solver
-			return FMath::Max(NumActiveConstraints, 0);
-		}
-
-		inline bool GetIsActive() const
-		{
-			return bIsActive;
-		}
-
-
 		FPBDJointSolver();
 
 		// Called once per frame to initialze the joint solver from the joint settings
@@ -199,14 +187,6 @@ namespace Chaos
 		void UpdateMass0();
 		void UpdateMass1();
 
-		// Check to see if this constraint still needs further solved
-		// @todo(chaos): the term "active" is used inconsistently with the meaning elsewhere. Active should
-		// mean "contributed impusles". Should use "solved" rather than "active"
-		bool UpdateIsActive(
-			const FReal Dt,
-			const FPBDJointSolverSettings& SolverSettings,
-			const FPBDJointSettings& JointSettings);
-
 		void ApplyPositionConstraints(
 			const FReal Dt,
 			const FPBDJointSolverSettings& SolverSettings,
@@ -639,37 +619,6 @@ namespace Chaos
 			const FVec3 Axis,
 			FReal& InOutAngle);
 
-		inline bool HasLinearConstraintPadding(const int32 AxisIndex) const
-		{
-			return LinearConstraintPadding[AxisIndex] >= 0.0f;
-		}
-
-		inline FReal GetLinearConstraintPadding(const int32 AxisIndex) const
-		{
-			return FMath::Max(LinearConstraintPadding[AxisIndex], 0.0f);
-		}
-
-		inline void SetLinearConstraintPadding(const int32 AxisIndex, FReal Padding)
-		{
-			LinearConstraintPadding[AxisIndex] = Padding;
-		}
-
-
-		inline bool HasAngularConstraintPadding(const EJointAngularConstraintIndex ConstraintIndex) const
-		{
-			return AngularConstraintPadding[(int32)ConstraintIndex] >= 0.0f;
-		}
-
-		inline FReal GetAngularConstraintPadding(const EJointAngularConstraintIndex ConstraintIndex) const
-		{
-			return FMath::Max(AngularConstraintPadding[(int32)ConstraintIndex], 0.0f);
-		}
-
-		inline void SetAngularConstraintPadding(const EJointAngularConstraintIndex ConstraintIndex, FReal Padding)
-		{
-			AngularConstraintPadding[(int32)ConstraintIndex] = Padding;
-		}
-
 		// Calculate linear velocities along constraint axes based on the constraint type in JointSettings. CV is the relative linear velocity of joint connectors.
 		void CalculateConstraintAxisLinearVelocities(
 			const FPBDJointSettings& JointSettings,
@@ -762,10 +711,6 @@ namespace Chaos
 		// \todo(chaos): remove Stiffness from SolverSettings (because it is not a solver constant)
 		FReal SolverStiffness;
 
-		// Constraint padding which can act something like a velocity constraint (for reslockfreetitution)
-		FVec3 LinearConstraintPadding;
-		FVec3 AngularConstraintPadding;
-
 		// Tolerances below which we stop solving
 		FReal PositionTolerance;					// Distance error below which we consider a constraint or drive solved
 		FReal AngleTolerance;						// Angle error below which we consider a constraint or drive solved
@@ -775,8 +720,6 @@ namespace Chaos
 		FRotation3 LastQs[MaxConstrainedBodies];	// Rotations at the beginning of the iteration
 		FVec3 InitConstraintAxisLinearVelocities; // Linear velocities along the constraint axes at the begining of the frame, used by restitution
 		FVec3 InitConstraintAxisAngularVelocities; // Angular velocities along the constraint axes at the begining of the frame, used by restitution
-		int32 NumActiveConstraints;					// The number of active constraints and drives in the last iteration (-1 initial value)
-		bool bIsActive;								// Whether any constraints actually moved any bodies in last iteration
 	};
 
 }
\ No newline at end of file
