Chaos : fix bug where the OtherComponent returns the same as the main component in the Chaos Physics Collision event

diff --git a/Engine/Source/Runtime/Experimental/ChaosSolverEngine/Private/Chaos/ChaosGameplayEventDispatcher.cpp b/Engine/Source/Runtime/Experimental/ChaosSolverEngine/Private/Chaos/ChaosGameplayEventDispatcher.cpp
--- a/Engine/Source/Runtime/Experimental/ChaosSolverEngine/Private/Chaos/ChaosGameplayEventDispatcher.cpp
+++ b/Engine/Source/Runtime/Experimental/ChaosSolverEngine/Private/Chaos/ChaosGameplayEventDispatcher.cpp
@@ -332,7 +332,7 @@ void UChaosGameplayEventDispatcher::HandleCollisionEvents(const Chaos::FCollisio
 
 								Chaos::FCollidingData const& CollisionDataItem = CollisionData[CollisionIdx];
 
-								IPhysicsProxyBase* const PhysicsProxy1 = CollisionDataItem.Proxy2 ? PhysicsProxy0 : CollisionDataItem.Proxy2;
+								IPhysicsProxyBase* const PhysicsProxy1 = CollisionDataItem.Proxy2 ? CollisionDataItem.Proxy2: PhysicsProxy0;
 
 								// Are the proxies pending destruction? If they are no longer tracked by the PhysScene, the proxy is deleted or pending deletion.
 								if (Scene.GetOwningComponent<UPrimitiveComponent>(PhysicsProxy0) == nullptr || Scene.GetOwningComponent<UPrimitiveComponent>(PhysicsProxy1) == nullptr)
