Chaos geometry collection : add support for supporting internal cluster through blueprint hitResults
- added a ItemIndex abstration class to deal with regular transform indices and cluster index ones ( UniqueIdx )
- simplified the BP functions that are leverage the new ItemIndex abstraction
- simplified  the proxy equovalent function by consolidating the code in Clustering class
- adjusted removal crumbling  code to use the new APIs
- propagate internal cluster information from PT to GT ( no GT particle , but a way to query internal cluster index by the GT children for queries to use )

diff --git a/Engine/Config/BaseEngine.ini b/Engine/Config/BaseEngine.ini
--- a/Engine/Config/BaseEngine.ini
+++ b/Engine/Config/BaseEngine.ini
@@ -987,6 +987,8 @@ DefaultWorldPartitionHLODLayer=/Engine/Maps/Templates/HLODs/HLODLayer_Instanced.
 +StructRedirects=(OldName="AnimInstanceSubsystemData",NewName="/Script/Engine.AnimSubsystemInstance")
 +ClassRedirects=(OldName="/Script/AnimationModifiers.AnimationBlueprintLibrary",NewName="/Script/AnimationBlueprintLibrary.AnimationBlueprintLibrary")
 +PropertyRedirects=(OldName="GeometryCollectionComponent.Simulating", NewName="SimulatePhysics")
++PropertyRedirects=(OldName="GeometryCollectionComponent.ApplyExternalStrain.Index", NewName="GeometryCollectionComponent.ApplyExternalStrain.ItemIndex")
++PropertyRedirects=(OldName="GeometryCollectionComponent.CrumbleCluster.Index", NewName="GeometryCollectionComponent.CrumbleCluster.ItemIndex")
 +PropertyRedirects=(OldName="ControlRigComponent.OnPreUpdateDelegate", NewName="OnPreForwardsSolveDelegate")
 +PropertyRedirects=(OldName="ControlRigComponent.OnPostUpdateDelegate", NewName="OnPostForwardsSolveDelegate")
 +PropertyRedirects=(OldName="AnimNode_Mirror.BlendTimeOnMirrorStateChange",NewName="AnimNode_Mirror.BlendTime")
diff --git a/Engine/Source/Runtime/Engine/Private/Collision/CollisionConversions.cpp b/Engine/Source/Runtime/Engine/Private/Collision/CollisionConversions.cpp
--- a/Engine/Source/Runtime/Engine/Private/Collision/CollisionConversions.cpp
+++ b/Engine/Source/Runtime/Engine/Private/Collision/CollisionConversions.cpp
@@ -159,7 +159,7 @@ static void SetHitResultFromShapeAndFaceIndex(const FPhysicsShape& Shape,  const
 				if (ActorProxy->GetType() == EPhysicsProxyType::GeometryCollectionType)
 				{
 					const FGeometryCollectionPhysicsProxy* ConcreteProxy = static_cast<const FGeometryCollectionPhysicsProxy*>(ActorProxy);
-					OutResult.Item = ConcreteProxy->GetTransformGroupIndexFromGTParticle(&Actor);
+					OutResult.Item = ConcreteProxy->GetItemIndexFromGTParticle_External(&Actor).GetItemIndex();
 				}
 			}
 		}
diff --git a/Engine/Source/Runtime/Experimental/Chaos/Private/Chaos/PBDRigidClustering.cpp b/Engine/Source/Runtime/Experimental/Chaos/Private/Chaos/PBDRigidClustering.cpp
--- a/Engine/Source/Runtime/Experimental/Chaos/Private/Chaos/PBDRigidClustering.cpp
+++ b/Engine/Source/Runtime/Experimental/Chaos/Private/Chaos/PBDRigidClustering.cpp
@@ -943,6 +943,33 @@ namespace Chaos
 		return Child; 
 	}
 
+	FPBDRigidParticleHandle* FRigidClustering::FindClosestChild(const FPBDRigidClusteredParticleHandle* ClusteredParticle, const FVec3& WorldLocation) const
+	{
+		if (const TArray<FPBDRigidParticleHandle*>* ChildrenHandles = GetChildrenMap().Find(ClusteredParticle))
+		{
+			return FindClosestParticle(*ChildrenHandles, WorldLocation); 
+		}
+		return nullptr;
+	}
+
+	FPBDRigidParticleHandle* FRigidClustering::FindClosestParticle(const TArray<FPBDRigidParticleHandle*>& Particles, const FVec3& WorldLocation) const
+	{
+		FPBDRigidParticleHandle* ClosestChildHandle = nullptr;
+		
+		// @todo(chaos) we should offer a more precise way to query than the distance from center of mass
+		FReal ClosestSquaredDist = TNumericLimits<FReal>::Max();
+		for (FPBDRigidParticleHandle* ChildHandle: Particles)
+        {
+        	const FReal SquaredDist = (ChildHandle->X() - WorldLocation).SizeSquared();
+        	if (SquaredDist < ClosestSquaredDist)
+        	{
+        		ClosestSquaredDist = SquaredDist;
+        		ClosestChildHandle = ChildHandle;
+        	}
+        }
+        return ClosestChildHandle;
+	}
+
 	DECLARE_CYCLE_STAT(TEXT("TPBDRigidClustering<>::GenerateConnectionGraph"), STAT_GenerateConnectionGraph, STATGROUP_Chaos);
 	void 
 	FRigidClustering::GenerateConnectionGraph(
@@ -1057,18 +1084,9 @@ namespace Chaos
 				{
 					if (GraphPropagationBasedCollisionImpulseProcessing)
 					{
-						// first find the closest child
-						FPBDRigidParticleHandle* ClosestChild = nullptr;
-						FReal ClosestDistanceSquared = TNumericLimits<FReal>::Max();
-						for (FPBDRigidParticleHandle* Child:ParentToChildren)
-						{
-							const FReal DistanceSquared = (Child->X() - ContactWorldLocation).SizeSquared();
-							if (DistanceSquared < ClosestDistanceSquared)
-							{
-								ClosestDistanceSquared = DistanceSquared;
-								ClosestChild = Child;
-							}
-						}
+						// propagation based breaking model start from the closest particle and propagate through the connection graph
+						// propagation logic is dealt when evaluating the strain
+						FPBDRigidParticleHandle* ClosestChild = FindClosestParticle(ParentToChildren, ContactWorldLocation);
 						if (ClosestChild)
 						{
 							if (TPBDRigidClusteredParticleHandle<FReal, 3>* ClusteredChild = ClosestChild->CastToClustered())
@@ -1225,6 +1243,24 @@ namespace Chaos
 
 	}
 
+	bool FRigidClustering::BreakCluster(const FPBDRigidClusteredParticleHandle* ClusteredParticle)
+	{
+		// max strain will allow to unconditionally release the children when strain is evaluated
+		const FReal MaxStrain = TNumericLimits<FReal>::Max();
+		if (TArray<FPBDRigidParticleHandle*>* ChildrenHandles = GetChildrenMap().Find(ClusteredParticle))
+		{
+			for (FPBDRigidParticleHandle* ChildHandle: *ChildrenHandles)
+			{
+				if (Chaos::FPBDRigidClusteredParticleHandle* ClusteredChildHandle = ChildHandle->CastToClustered())
+				{
+					ClusteredChildHandle->SetExternalStrain(MaxStrain);
+				}
+			}
+			return true;
+		}
+		return false;
+	}
+
 	static void ConnectClusteredNodes(FPBDRigidClusteredParticleHandle* ClusteredChild1, FPBDRigidClusteredParticleHandle* ClusteredChild2)
 	{
 		check(ClusteredChild1 && ClusteredChild2);
diff --git a/Engine/Source/Runtime/Experimental/Chaos/Private/PhysicsProxy/GeometryCollectionPhysicsProxy.cpp b/Engine/Source/Runtime/Experimental/Chaos/Private/PhysicsProxy/GeometryCollectionPhysicsProxy.cpp
--- a/Engine/Source/Runtime/Experimental/Chaos/Private/PhysicsProxy/GeometryCollectionPhysicsProxy.cpp
+++ b/Engine/Source/Runtime/Experimental/Chaos/Private/PhysicsProxy/GeometryCollectionPhysicsProxy.cpp
@@ -81,12 +81,6 @@ FAutoConsoleVariableRef CVarGeometryCollectionAlwaysGenerateConnectionGraph(
 	bGeometryCollectionAlwaysGenerateConnectionGraph,
 	TEXT("When enabled, always  generate the cluster's connection graph instead of using the rest collection stored one - Note: this should only be used for troubleshooting.[def: false]"));
 
-bool bGeometryCollectionCrumbleUsingStrain = true;
-FAutoConsoleVariableRef CVarGeometryCollectionCrumbleUsingStrain(
-	TEXT("p.GeometryCollection.CrumbleUsingStrain"),
-	bGeometryCollectionCrumbleUsingStrain,
-	TEXT("When enabled, cluster crumbling is using strain instead of directly calling the clustering methods.[def: true]"));
-
 DEFINE_LOG_CATEGORY_STATIC(UGCC_LOG, Error, All);
 
 static const FSharedSimulationSizeSpecificData& GetSizeSpecificData(const TArray<FSharedSimulationSizeSpecificData>& SizeSpecificData, const FGeometryCollection& RestCollection, const int32 TransformIndex, const FBox& BoundingBox);
@@ -753,6 +747,28 @@ void FGeometryCollectionPhysicsProxy::CreateNonClusteredParticles(Chaos::FPBDRig
 	}
 }
 
+Chaos::FPBDRigidClusteredParticleHandle* FGeometryCollectionPhysicsProxy::FindClusteredParticleHandleByItemIndex(FGeometryCollectionItemIndex ItemIndex) const
+{
+	Chaos::FPBDRigidClusteredParticleHandle* ResultHandle = nullptr;;
+	if (ItemIndex.IsInternalCluster())
+	{
+		const int32 InternalClusterUniqueIdx = ItemIndex.GetInternalClusterIndex();
+		if (Chaos::FPBDRigidClusteredParticleHandle* const* InternalClusterHandle = UniqueIdxToInternalClusterHandle.Find(InternalClusterUniqueIdx))
+		{
+			ResultHandle = *InternalClusterHandle;
+		}
+	}
+	else
+	{
+		const int32 TransformIndex = ItemIndex.GetTransformIndex();
+		if (SolverParticleHandles.IsValidIndex(TransformIndex))
+		{
+			ResultHandle = SolverParticleHandles[TransformIndex];
+		}
+	}
+	return ResultHandle;
+}
+
 void FGeometryCollectionPhysicsProxy::SetSleepingState(const Chaos::FPBDRigidsSolver& RigidsSolver)
 {
 	for (FClusterHandle* Handle: SolverParticleHandles)
@@ -1518,135 +1534,47 @@ void FGeometryCollectionPhysicsProxy::DisableParticles_External(TArray<int32>&&
 			});
 	}
 }
-static void ApplyStrainToAllClusterChildren(Chaos::FRigidClustering& Clustering, Chaos::FPBDRigidClusteredParticleHandle* ClusteredParent, Chaos::FReal StrainValue)
-{
-	if (TArray<Chaos::FPBDRigidParticleHandle*>* Children = Clustering.GetChildrenMap().Find(ClusteredParent))
-	{
-		for (Chaos::FPBDRigidParticleHandle* ChildHandle: *Children)
-		{
-			if (Chaos::FPBDRigidClusteredParticleHandle* ClusteredChildHandle = ChildHandle->CastToClustered())
-			{
-				ClusteredChildHandle->SetExternalStrain(FMath::Max(ClusteredChildHandle->GetExternalStrain(), StrainValue));
-			}
-		}
-	}
-}
-
-
-void FGeometryCollectionPhysicsProxy::BreakInternalClusterParents_External(TArray<int32>&& TransformGroupIndices)
-{
-	check(IsInGameThread());
-
-	if (Chaos::FPhysicsSolver* RBDSolver = GetSolver<Chaos::FPhysicsSolver>())
-	{
-		RBDSolver->EnqueueCommandImmediate([this, RBDSolver, IndicesToBreakParent = MoveTemp(TransformGroupIndices)]()
-		{
-			// could use a set but the normal usage will always have a very small amount of unique parents and the array is more cache friendly
-			TArray<Chaos::FPBDRigidClusteredParticleHandle*> ParentToBreak;
-			for (int32 TransformIdx : IndicesToBreakParent)
-			{
-				if (Chaos::FPBDRigidClusteredParticleHandle* ChildHandle = SolverParticleHandles[TransformIdx])
-				{
-					if (Chaos::FPBDRigidClusteredParticleHandle* ClusteredParent =  ChildHandle->Parent())
-					{
-						if (ClusteredParent->InternalCluster())
-						{
-							ParentToBreak.AddUnique(ClusteredParent);
-						}
-					}
-				}
-				for (Chaos::FPBDRigidClusteredParticleHandle* ClusteredParent: ParentToBreak)
-				{
-					Chaos::FRigidClustering& Clustering = RBDSolver->GetEvolution()->GetRigidClustering();
-					if (!ClusteredParent->Disabled() && Clustering.GetChildrenMap().Contains(ClusteredParent))
-					{
-						if (bGeometryCollectionCrumbleUsingStrain)
-						{
-							ApplyStrainToAllClusterChildren(Clustering, ClusteredParent, TNumericLimits<Chaos::FReal>::Max());
-						}
-						else
-						{
-							Clustering.ReleaseClusterParticles(ClusteredParent, true /* bForceRelease */);
-						}
-					}
-				}
-			}
-		});
-	}
-}
 
-void FGeometryCollectionPhysicsProxy::BreakClusters_External(TArray<int32>&& TransformGroupIndices)
+void FGeometryCollectionPhysicsProxy::BreakClusters_External(TArray<FGeometryCollectionItemIndex>&& ItemIndices)
 {
 	check(IsInGameThread());
 
 	if (Chaos::FPhysicsSolver* RBDSolver = GetSolver<Chaos::FPhysicsSolver>())
 	{
-		RBDSolver->EnqueueCommandImmediate([this, RBDSolver, IndicesToBreakParent = MoveTemp(TransformGroupIndices)]()
+		RBDSolver->EnqueueCommandImmediate([this, RBDSolver, IndicesToBreakParent = MoveTemp(ItemIndices)]()
 		{
 			Chaos::FRigidClustering& Clustering = RBDSolver->GetEvolution()->GetRigidClustering();
-			for (int32 TransformIdx : IndicesToBreakParent)
+			for (const FGeometryCollectionItemIndex& ItemIndex : IndicesToBreakParent)
 			{
-				if (Chaos::FPBDRigidClusteredParticleHandle* ClusterHandle = SolverParticleHandles[TransformIdx])
+				if (Chaos::FPBDRigidClusteredParticleHandle* ClusteredHandle = FindClusteredParticleHandleByItemIndex(ItemIndex))
 				{
-					if (!ClusterHandle->Disabled() && Clustering.GetChildrenMap().Contains(ClusterHandle))
-					{
-						if (bGeometryCollectionCrumbleUsingStrain)
-						{
-							ApplyStrainToAllClusterChildren(Clustering, ClusterHandle, TNumericLimits<Chaos::FReal>::Max());
-						}
-						else
-						{
-							Clustering.ReleaseClusterParticles(ClusterHandle, true /* bForceRelease */);
-						}
-					}
+					Clustering.BreakCluster(ClusteredHandle);
 				}
 			}
 		});
 	}
 }
 
-void FGeometryCollectionPhysicsProxy::ApplyStrain_External(int32 TransformGroupIndex, const FVector& Location, float StrainValue)
+void FGeometryCollectionPhysicsProxy::ApplyStrain_External(FGeometryCollectionItemIndex ItemIndex, const FVector& WorldLocation, float StrainValue)
 {
 	check(IsInGameThread());
 
 	if (Chaos::FPhysicsSolver* RBDSolver = GetSolver<Chaos::FPhysicsSolver>())
 	{
-		RBDSolver->EnqueueCommandImmediate([this, StrainValue, RBDSolver, Location, TransformGroupIndex]()
+		RBDSolver->EnqueueCommandImmediate([this, StrainValue, RBDSolver, WorldLocation, ItemIndex]()
 		{
-			// @todo(chaos) move most of the logic into a method in clustering 
-			Chaos::FRigidClustering& Clustering = RBDSolver->GetEvolution()->GetRigidClustering();
-			if (Chaos::FPBDRigidClusteredParticleHandle* ClusterHandle = SolverParticleHandles[TransformGroupIndex])
+			if (Chaos::FPBDRigidClusteredParticleHandle* ClusteredHandle = FindClusteredParticleHandleByItemIndex(ItemIndex))
 			{
-				const bool bIsCluster = (ClusterHandle->ClusterIds().NumChildren > 0);
-				const bool bHasInternalClusterParent = ClusterHandle->Parent()? ClusterHandle->Parent()->InternalCluster(): false;
-				if (bHasInternalClusterParent || (!bIsCluster && !ClusterHandle->Disabled()))
+				const bool bIsCluster = (ClusteredHandle->ClusterIds().NumChildren > 0);
+				if (bIsCluster)
 				{
-					ClusterHandle->SetExternalStrain(FMath::Max(ClusterHandle->GetExternalStrain(), StrainValue));
-				}
-				else
-				{
-					if (bIsCluster)
+					Chaos::FRigidClustering& Clustering = RBDSolver->GetEvolution()->GetRigidClustering();
+					Chaos::FPBDRigidParticleHandle* ClosestChildHandle = Clustering.FindClosestChild(ClusteredHandle, WorldLocation);
+					if (ClosestChildHandle)
 					{
-						if (const TArray<Chaos::FPBDRigidParticleHandle*>* ChildrenHandles = Clustering.GetChildrenMap().Find(ClusterHandle))
+						if (Chaos::FPBDRigidClusteredParticleHandle* ClusteredChildHandle = ClosestChildHandle->CastToClustered())
 						{
-							FVector::FReal ClosestSquaredDist = TNumericLimits<FVector::FReal>::Max();
-							Chaos::FPBDRigidParticleHandle* ClosestChildHandle = nullptr; 
-							for (Chaos::FPBDRigidParticleHandle* ChildHandle: *ChildrenHandles)
-							{
-								const FVector::FReal SquaredDist = (ChildHandle->X() - Location).SizeSquared();
-								if (SquaredDist < ClosestSquaredDist)
-								{
-									ClosestSquaredDist = SquaredDist;
-									ClosestChildHandle = ChildHandle;
-								}
-							}
-							if (ClosestChildHandle)
-							{
-								if (Chaos::FPBDRigidClusteredParticleHandle* ClusteredChildHandle = ClosestChildHandle->CastToClustered())
-								{
-									ClusteredChildHandle->SetExternalStrain(FMath::Max(ClusteredChildHandle->GetExternalStrain(), StrainValue));
-								}
-							}
+							ClusteredChildHandle->SetExternalStrain(FMath::Max(ClusteredChildHandle->GetExternalStrain(), StrainValue));
 						}
 					}
 				}
@@ -1991,6 +1919,8 @@ void FGeometryCollectionPhysicsProxy::BufferPhysicsResults_External(Chaos::FDirt
 		TargetResults.Parent[TransformGroupIndex] = GameThreadCollection.Parent[TransformGroupIndex];
 		TargetResults.Transforms[TransformGroupIndex] = GameThreadCollection.Transform[TransformGroupIndex];
 
+		TargetResults.InternalClusterUniqueIdx[TransformGroupIndex] = INDEX_NONE;
+		
 #if WITH_EDITORONLY_DATA
 		TargetResults.DamageInfo[TransformGroupIndex].Damage = 0.0f;
 		TargetResults.DamageInfo[TransformGroupIndex].DamageThreshold = 0.0f;
@@ -2050,15 +1980,18 @@ void FGeometryCollectionPhysicsProxy::BufferPhysicsResults_Internal(Chaos::FPBDR
 	const bool IsActorScaled = !ActorToWorld.GetScale3D().Equals(FVector::OneVector);
 	const FTransform ActorScaleTransform(FQuat::Identity,  FVector::ZeroVector, ActorToWorld.GetScale3D());
 
+	UniqueIdxToInternalClusterHandle.Reset();
+	
 	const int32 NumTransformGroupElements = TargetResults.Transforms.Num(); 
 	if(NumTransformGroupElements > 0)
 	{ 
 		SCOPE_CYCLE_COUNTER(STAT_CalcParticleToWorld);
-		
+
 		for (int32 TransformGroupIndex = 0; TransformGroupIndex < NumTransformGroupElements; ++TransformGroupIndex)
 		{
 			TargetResults.Transforms[TransformGroupIndex] = PhysicsThreadCollection.Transform[TransformGroupIndex];
 			TargetResults.Parent[TransformGroupIndex] = PhysicsThreadCollection.Parent[TransformGroupIndex];
+			TargetResults.InternalClusterUniqueIdx[TransformGroupIndex] = INDEX_NONE;
 
 			TargetResults.States[TransformGroupIndex].DisabledState = true;
 			TargetResults.States[TransformGroupIndex].HasInternalClusterParent = false;
@@ -2161,6 +2094,13 @@ void FGeometryCollectionPhysicsProxy::BufferPhysicsResults_Internal(Chaos::FPBDR
 
 						if(ClusterParent->InternalCluster())
 						{
+							const int32 InternalClusterParentUniqueIdx = ClusterParent->UniqueIdx().Idx;
+							if (!UniqueIdxToInternalClusterHandle.Contains(InternalClusterParentUniqueIdx))
+							{
+								UniqueIdxToInternalClusterHandle.Add(InternalClusterParentUniqueIdx, ClusterParent);
+							}
+							TargetResults.InternalClusterUniqueIdx[TransformGroupIndex] = ClusterParent->UniqueIdx().Idx;
+							
 							const FTransform ParticleToWorld = Handle->ChildToParent() * FRigidTransform3(ClusterParent->X(), ClusterParent->R());    // aka ClusterChildToWorld
 							TargetResults.ParticleXs[TransformGroupIndex] = ParticleToWorld.GetTranslation();
 							TargetResults.ParticleRs[TransformGroupIndex] = ParticleToWorld.GetRotation();
@@ -2418,6 +2358,17 @@ bool FGeometryCollectionPhysicsProxy::PullFromPhysicsState(const Chaos::FDirtyGe
 			}
 		}
 
+		// internal cluster index map update
+		GTParticlesToInternalClusterUniqueIdx.Reset();
+		for (int32 TransformGroupIndex = 0; TransformGroupIndex < NumTransforms; ++TransformGroupIndex)
+		{
+			const int32 InternalClusterUniqueIdx = TargetResults.InternalClusterUniqueIdx[TransformGroupIndex]; 
+			if (InternalClusterUniqueIdx > INDEX_NONE)
+			{
+				GTParticlesToInternalClusterUniqueIdx.Add(GTParticles[TransformGroupIndex].Get(), InternalClusterUniqueIdx);	
+			}
+		}
+
 		//question: why do we need this? Sleeping objects will always have to update GPU
 		if (bIsCollectionDirty)
 		{
diff --git a/Engine/Source/Runtime/Experimental/Chaos/Public/Chaos/PBDRigidClustering.h b/Engine/Source/Runtime/Experimental/Chaos/Public/Chaos/PBDRigidClustering.h
--- a/Engine/Source/Runtime/Experimental/Chaos/Public/Chaos/PBDRigidClustering.h
+++ b/Engine/Source/Runtime/Experimental/Chaos/Public/Chaos/PBDRigidClustering.h
@@ -146,6 +146,15 @@ public:
 		FPBDRigidClusteredParticleHandle* ClusteredParticle,
 		const FClusterDestoryParameters& Parameters = FClusterDestoryParameters());
 
+	/*
+	*  BreakCluster
+	*    Breaks a cluster (internal or not) by applying max external strain to all its children 
+	* 
+	*    @param ClusteredParticle handle of the cluster to break
+	*    @return true if the cluster was successfully found and act upon 
+	*/
+	bool BreakCluster(const FPBDRigidClusteredParticleHandle* ClusteredParticle);
+	
 	//
 	// Operational 
 	//
@@ -251,6 +260,26 @@ public:
 	*/
 	const TArrayCollectionArray<TArray<TConnectivityEdge<FReal>>>& GetConnectivityEdges() const { return MParticles.ConnectivityEdgesArray(); }
 
+	/*
+	*  FindClosestChild
+	*    Find the closest child of an active cluster from a world position
+	*    current implementation will pick the closest based on the distance from the center
+	*    future implementation may expose option for more precise queries
+	*    @param ClusteredParticle active cluster handle to query the children from
+	*    @param WorldLocation world space location to find the closest child from   
+	*/
+	FPBDRigidParticleHandle* FindClosestChild(const FPBDRigidClusteredParticleHandle* ClusteredParticle, const FVec3& WorldLocation) const;
+
+	/*
+	*  FindClosest
+	*    Find the closest particle from an array of  particle
+	*    current implementation will pick the closest based on the distance from the center
+	*    future implementation may expose option for more precise queries
+	*    @param Particles array of clustered particles to select from  
+	*    @param WorldLocation world space location to find the closest particle from   
+	*/
+	FPBDRigidParticleHandle* FindClosestParticle(const TArray<FPBDRigidParticleHandle*>& Particles, const FVec3& WorldLocation) const;
+	
 	/**
 	* GenerateConnectionGraph
 	*   Creates a connection graph for the given index using the creation parameters. This will not
diff --git a/Engine/Source/Runtime/Experimental/Chaos/Public/GeometryCollectionProxyData.h b/Engine/Source/Runtime/Experimental/Chaos/Public/GeometryCollectionProxyData.h
--- a/Engine/Source/Runtime/Experimental/Chaos/Public/GeometryCollectionProxyData.h
+++ b/Engine/Source/Runtime/Experimental/Chaos/Public/GeometryCollectionProxyData.h
@@ -161,6 +161,7 @@ public:
 		
 		Transforms.SetNumUninitialized(NumTransforms);
 		Parent.SetNumUninitialized(NumTransforms);
+		InternalClusterUniqueIdx.SetNumUninitialized(NumTransforms);
 #if WITH_EDITORONLY_DATA
 		DamageInfo.SetNumUninitialized(NumTransforms);
 #endif		
@@ -185,6 +186,7 @@ public:
 
 	TArray<FTransform> Transforms;
 	TArray<int32> Parent;
+	TArray<int32> InternalClusterUniqueIdx;
 
 #if WITH_EDITORONLY_DATA
 	struct FDamageInfo
diff --git a/Engine/Source/Runtime/Experimental/Chaos/Public/PhysicsProxy/GeometryCollectionPhysicsProxy.h b/Engine/Source/Runtime/Experimental/Chaos/Public/PhysicsProxy/GeometryCollectionPhysicsProxy.h
--- a/Engine/Source/Runtime/Experimental/Chaos/Public/PhysicsProxy/GeometryCollectionPhysicsProxy.h
+++ b/Engine/Source/Runtime/Experimental/Chaos/Public/PhysicsProxy/GeometryCollectionPhysicsProxy.h
@@ -29,6 +29,88 @@ namespace Chaos
 	class FPBDRigidsEvolutionBase;
 }
 
+/**
+ * index abstraction for HitResults
+ * this allows regular bones and internal clusters to be represented by a unique int32 index that can be passed to HitResults
+ * and in return understood by the geometry collection physics proxy 
+ */ 
+struct FGeometryCollectionItemIndex
+{
+public:
+	static FGeometryCollectionItemIndex CreateInternalClusterItemIndex(int32 ClusterUniqueIdx)
+	{
+		return FGeometryCollectionItemIndex(ClusterUniqueIdx, true);
+	}
+
+	static FGeometryCollectionItemIndex CreateTransformItemIndex(int32 TransformIndex)
+	{
+		return FGeometryCollectionItemIndex(TransformIndex, false);
+	}
+
+	static FGeometryCollectionItemIndex CreateFromExistingItemIndex(int32 ItemIndex)
+	{
+		FGeometryCollectionItemIndex Result(INDEX_NONE, false);
+		Result.ItemIndex = ItemIndex;
+		return Result;
+	}
+	
+	static FGeometryCollectionItemIndex CreateInvalidItemIndex()
+	{
+		return FGeometryCollectionItemIndex(INDEX_NONE, false);
+	}
+	
+public:
+	FGeometryCollectionItemIndex(const FGeometryCollectionItemIndex& Other)
+		: ItemIndex(Other.ItemIndex)
+	{}
+	
+	bool IsInternalCluster() const
+	{
+		return ItemIndex < INDEX_NONE;
+	}
+
+	int32 GetInternalClusterIndex() const
+	{
+		check(IsInternalCluster());
+		return (ItemIndex - InternalClusterBaseIndex);
+	}
+
+	int32 GetTransformIndex() const
+	{
+		check(!IsInternalCluster());
+		return ItemIndex;
+	}
+	
+	bool IsValid() const
+	{
+		return ItemIndex != INDEX_NONE;
+	}
+
+	int32 GetItemIndex() const
+	{
+		return ItemIndex;
+	}
+
+	bool operator==(const FGeometryCollectionItemIndex& Other) const
+	{
+		return ItemIndex == Other.ItemIndex;
+	}
+	
+private:
+	static const int32 InternalClusterBaseIndex = TNumericLimits<int32>::Lowest();   
+
+	FGeometryCollectionItemIndex(int32 Index, bool bInternalCluster)
+		: ItemIndex(INDEX_NONE)
+	{
+		if (Index > INDEX_NONE)
+		{
+			ItemIndex = Index + (bInternalCluster? InternalClusterBaseIndex: 0); 	
+		}
+	}
+	
+	int32 ItemIndex;
+};
+
 class FStubGeometryCollectionData : public Chaos::FParticleData 
 {
 public:
@@ -236,22 +318,41 @@ public:
 		return INDEX_NONE;
 	}
 
-	int32 GetTransformGroupIndexFromGTParticle(const Chaos::FGeometryParticle* GTPParticle) const
+	FGeometryCollectionItemIndex GetInternalClusterParentItemIndex_External(int32 ChildTransformIndex) const
+	{
+		// first find the GTParticle matching the Transform index
+		if (ChildTransformIndex >= 0 && ChildTransformIndex < GTParticles.Num())
+		{
+			const TUniquePtr<Chaos::FGeometryParticle>& ChildGTParticle = GTParticles[ChildTransformIndex];
+			if (const int32* InternalClusterUniqueIdx = GTParticlesToInternalClusterUniqueIdx.Find(ChildGTParticle.Get()))
+			{
+				return FGeometryCollectionItemIndex::CreateInternalClusterItemIndex(*InternalClusterUniqueIdx);
+			}
+		}
+		return FGeometryCollectionItemIndex::CreateInvalidItemIndex();
+	}
+	
+	FGeometryCollectionItemIndex GetItemIndexFromGTParticle_External(const Chaos::FGeometryParticle* GTPParticle) const
 	{
+		// internal cluster have  no representation on the GT, so we use the child GT particle to find the matching internal cluster unique index 
+		if (const int32* InternalClusterUniqueIdx = GTParticlesToInternalClusterUniqueIdx.Find(GTPParticle))
+		{
+			return FGeometryCollectionItemIndex::CreateInternalClusterItemIndex(*InternalClusterUniqueIdx);
+		}
+		// regular particle that has a matchig transform index 
 		if (const int32* TransformGroupIndex = GTParticlesToTransformGroupIndex.Find(GTPParticle))
 		{
-			return *TransformGroupIndex;
+			return FGeometryCollectionItemIndex::CreateTransformItemIndex(*TransformGroupIndex);
 		}
-		return INDEX_NONE;
+		return FGeometryCollectionItemIndex::CreateInvalidItemIndex();
 	}
 
 	bool GetIsObjectDynamic() const { return IsObjectDynamic; }
 
 	void DisableParticles_External(TArray<int32>&& TransformGroupIndices);
 
-	void BreakInternalClusterParents_External(TArray<int32>&& TransformGroupIndices);
-	void BreakClusters_External(TArray<int32>&& TransformGroupIndices);
-	void ApplyStrain_External(int32 TransformGroupIndex, const FVector& Location, float StrainValue);
+	void BreakClusters_External(TArray<FGeometryCollectionItemIndex>&& ItemIndices);
+	void ApplyStrain_External(FGeometryCollectionItemIndex ItemIndex, const FVector& WorldLocation, float StrainValue);
 
 	FProxyInterpolationData& GetInterpolationData() { return InterpolationData; }
 	const FProxyInterpolationData& GetInterpolationData() const { return InterpolationData; }
@@ -303,6 +404,8 @@ protected:
 
 	void CreateNonClusteredParticles(Chaos::FPBDRigidsSolver* RigidsSolver,	const FGeometryCollection& RestCollection, const FGeometryDynamicCollection& DynamicCollection);
 
+	Chaos::FPBDRigidClusteredParticleHandle* FindClusteredParticleHandleByItemIndex(FGeometryCollectionItemIndex ItemIndex) const;
+	
 private:
 
 	FSimulationParameters Parameters;
@@ -320,7 +423,7 @@ private:
 	TArray<FClusterHandle*> SolverClusterHandles; // make a TArray of the base clase with type
 	TArray<FClusterHandle*> SolverParticleHandles;// make a TArray of base class and join with above
 	TMap<FParticleHandle*, int32> HandleToTransformGroupIndex;
-
+	TMap<int32, FClusterHandle*> UniqueIdxToInternalClusterHandle;
 
 	//
 	// Buffer Results State Information
@@ -333,6 +436,9 @@ private:
 
 	TManagedArray<TUniquePtr<Chaos::FGeometryParticle>> GTParticles;
 	TMap<Chaos::FGeometryParticle*, int32> GTParticlesToTransformGroupIndex;
+	TMap<Chaos::FGeometryParticle*, int32> GTParticlesToInternalClusterUniqueIdx;
+
+	TMap<int32, TUniquePtr<Chaos::FGeometryParticle>> GTInternalClustersByUniqueIdx;
 
 	// These are read on both threads and should not be changed
 	const FCollisionFilterData SimFilter;
diff --git a/Engine/Source/Runtime/Experimental/GeometryCollectionEngine/Private/GeometryCollection/GeometryCollectionComponent.cpp b/Engine/Source/Runtime/Experimental/GeometryCollectionEngine/Private/GeometryCollection/GeometryCollectionComponent.cpp
--- a/Engine/Source/Runtime/Experimental/GeometryCollectionEngine/Private/GeometryCollection/GeometryCollectionComponent.cpp
+++ b/Engine/Source/Runtime/Experimental/GeometryCollectionEngine/Private/GeometryCollection/GeometryCollectionComponent.cpp
@@ -3593,7 +3593,7 @@ void UGeometryCollectionComponent::IncrementSleepTimer(float DeltaTime)
 		const TManagedArray<int32>& Parents = DynamicCollection->Parent;
 		
 		TArray<int32> ToDisable;
-		TArray<int32> ToBreakParent;
+		TArray<FGeometryCollectionItemIndex> ToCrumbleParent;
 		for (int32 TransformIdx = 0; TransformIdx < SleepTimer.Num(); ++TransformIdx)
 		{
 			// update the sleeping timer accordingly
@@ -3629,7 +3629,11 @@ void UGeometryCollectionComponent::IncrementSleepTimer(float DeltaTime)
 					// if parent is internal and dynamic we need to break it before decaying 
 					if (InternalClusterParentType[TransformIdx] == (uint8)Chaos::EInternalClusterType::Dynamic)
 					{
-						ToBreakParent.AddUnique(TransformIdx);
+						FGeometryCollectionItemIndex InternalClusterItemindex = PhysicsProxy->GetInternalClusterParentItemIndex_External(TransformIdx);
+						if (InternalClusterItemindex.IsValid())
+						{
+							ToCrumbleParent.AddUnique(InternalClusterItemindex);
+						}
 					}
 
 					DynamicCollection->MakeDirty();
@@ -3645,9 +3649,9 @@ void UGeometryCollectionComponent::IncrementSleepTimer(float DeltaTime)
 			}
 		}
 
-		if (ToBreakParent.Num())
+		if (ToCrumbleParent.Num())
 		{
-			PhysicsProxy->BreakInternalClusterParents_External(MoveTemp(ToBreakParent));
+			PhysicsProxy->BreakClusters_External(MoveTemp(ToCrumbleParent));
 		}
 		if (ToDisable.Num())
 		{
@@ -3676,7 +3680,7 @@ void UGeometryCollectionComponent::IncrementBreakTimer(float DeltaTime)
 		TManagedArray<float>& Decay = DynamicCollection->ModifyAttribute<float>("Decay", FGeometryCollection::TransformGroup);
 		
 		TArray<int32> ToDisable;
-		TArray<int32> ToCrumble;
+		TArray<FGeometryCollectionItemIndex> ToCrumble;
 		for (int32 TransformIdx = 0; TransformIdx < BreakTimer.Num(); ++TransformIdx)
 		{
 			const bool bIsBroken = DynamicCollection->Active[TransformIdx];
@@ -3703,7 +3707,7 @@ void UGeometryCollectionComponent::IncrementBreakTimer(float DeltaTime)
 					const bool bUseClusterCrumbling = (BreakRemovalDuration[TransformIdx] < 0);
 					if (bUseClusterCrumbling)
 					{
-						ToCrumble.AddUnique(TransformIdx);
+						ToCrumble.AddUnique(FGeometryCollectionItemIndex::CreateTransformItemIndex(TransformIdx));
 						Decay[TransformIdx] = 0;
 					}
 					else
@@ -3730,42 +3734,20 @@ void UGeometryCollectionComponent::IncrementBreakTimer(float DeltaTime)
 	}
 }
 
-void UGeometryCollectionComponent::ApplyExternalStrain(int32 Index, const FVector& Location, float Strain)
+void UGeometryCollectionComponent::ApplyExternalStrain(int32 ItemIndex, const FVector& Location, float Strain)
 {
-	if (DynamicCollection && PhysicsProxy)
+	if (PhysicsProxy)
 	{
-		const int32 NumTransform = DynamicCollection->NumElements(FGeometryCollection::TransformGroup);
-		if (Index >= 0 && Index < NumTransform)
-		{
-			PhysicsProxy->ApplyStrain_External(Index, Location, Strain);
-		}
+		PhysicsProxy->ApplyStrain_External(FGeometryCollectionItemIndex::CreateFromExistingItemIndex(ItemIndex), Location, Strain);
 	}
 }
 
-bool UGeometryCollectionComponent::CrumbleCluster(int32 Index)
+void UGeometryCollectionComponent::CrumbleCluster(int32 ItemIndex)
 {
-	bool Result = false;
-	if (DynamicCollection && PhysicsProxy)
+	if (PhysicsProxy)
 	{
-		const int32 NumTransform = DynamicCollection->NumElements(FGeometryCollection::TransformGroup);
-		if (Index >= 0 && Index < NumTransform)
-		{
-			const TManagedArray<uint8>& InternalClusterParentType = DynamicCollection->GetAttribute<uint8>("InternalClusterParentTypeArray", FGeometryCollection::TransformGroup);
-			const bool bHasInternalClusterParent = (InternalClusterParentType[Index] != (uint8)Chaos::EInternalClusterType::None);
-			const bool bIsCluster = (DynamicCollection->Children[Index].Num() > 0);
-			if (bIsCluster)
-			{
-				PhysicsProxy->BreakClusters_External({Index});
-				Result = true;
-			}
-			else if (bHasInternalClusterParent)
-			{
-				PhysicsProxy->BreakInternalClusterParents_External({Index});
-				Result = true;
-			}
-		}
+		PhysicsProxy->BreakClusters_External({FGeometryCollectionItemIndex::CreateFromExistingItemIndex(ItemIndex)});
 	}
-	return Result;
 }
 
 bool UGeometryCollectionComponent::CalculateInnerSphere(int32 TransformIndex, FSphere& SphereOut) const
diff --git a/Engine/Source/Runtime/Experimental/GeometryCollectionEngine/Public/GeometryCollection/GeometryCollectionComponent.h b/Engine/Source/Runtime/Experimental/GeometryCollectionEngine/Public/GeometryCollection/GeometryCollectionComponent.h
--- a/Engine/Source/Runtime/Experimental/GeometryCollectionEngine/Public/GeometryCollection/GeometryCollectionComponent.h
+++ b/Engine/Source/Runtime/Experimental/GeometryCollectionEngine/Public/GeometryCollection/GeometryCollectionComponent.h
@@ -475,23 +475,21 @@ public:
 	TObjectPtr<AChaosSolverActor> ChaosSolverActor;
 
 	/**
-	 * Apply an external strain a specific piece of the geometry collection
-	 * if the piece is a cluster then the strain will be applied to the closest child   
-	 * @param Index index of the piece to apply strain on
+	 * Apply an external strain to specific piece of the geometry collection
+	 * @param ItemIndex item index ( from HitResult) of the piece to apply strain on
 	 * @param Location world location of where to apply the strain
 	 * @param Strain strain / damage to apply 
 	 */
 	UFUNCTION(BlueprintCallable, Category = "ChaosPhysics")
-	void ApplyExternalStrain(int32 Index, const FVector& Location, float Strain);
+	void ApplyExternalStrain(int32 ItemIndex, const FVector& Location, float Strain);
 
 	/**
 	 * Crumbe a cluster into all its pieces
 	 * if the piece is not a cluster but is part of an internal/partial cluster, then the owning cluster will crumble   
-	 * @param Index index of the piece to crumble
-	 * @return  true of a crumble operation was performed
+	 * @param ItemIndex item index ( from HitResult) of the cluster to crumble
 	 */
 	UFUNCTION(BlueprintCallable, Category = "ChaosPhysics")
-	bool CrumbleCluster(int32 Index);
+	void CrumbleCluster(int32 ItemIndex);
 	
 	/** RestCollection */
 	UFUNCTION(BlueprintCallable, Category = "ChaosPhysics")
