GeometryCollection : FVertexArrayInterface
- Added api to store Vertex selection lists on the collection.

diff --git a/Engine/Source/Runtime/Experimental/Chaos/Private/GeometryCollection/GeometryCollectionSelectionAlgo.cpp b/Engine/Source/Runtime/Experimental/Chaos/Private/GeometryCollection/GeometryCollectionSelectionAlgo.cpp
new file mode 100644
--- /dev/null
+++ b/Engine/Source/Runtime/Experimental/Chaos/Private/GeometryCollection/GeometryCollectionSelectionAlgo.cpp
@@ -0,0 +1,175 @@
+// Copyright Epic Games, Inc. All Rights Reserved.
+
+/*=============================================================================
+	GeometryCollection.cpp: FGeometryCollection methods.
+=============================================================================*/
+
+#include "GeometryCollection/GeometryCollectionSelectionAlgo.h"
+
+namespace GeometryCollection
+{
+
+	// Groups 
+	const FName FSelectionAlgo::UnboundGroup = "Unbound";
+	const FName FSelectionAlgo::WeightedUnboundGroup = "WeightedUnbound";
+	const FName FSelectionAlgo::BoundGroup = "Bound";
+	const FName FSelectionAlgo::WeightedBoundGroup = "WeightedBound";
+
+	// Attributes
+	const FName FSelectionAlgo::IndexAttribute = "Index";
+	const FName FSelectionAlgo::WeightAttribute = "Weights";
+	const FName FSelectionAlgo::BoneIndexAttribute = "BoneIndex";
+
+	FSelectionAlgo::FSelectionAlgo(FManagedArrayCollection* InCollection)
+		: Self(InCollection)
+	{}
+
+	//
+	//  Initialization
+	//
+
+	void FSelectionAlgo::InitUnboundedGroup(FManagedArrayCollection* Collection, FName GroupName, FName DependencyGroup)
+	{
+		if (!Collection->HasGroup(GroupName))
+		{
+			Collection->AddAttribute<TArray<int32>>(FSelectionAlgo::IndexAttribute, GroupName, { DependencyGroup });
+		}
+		ensure(Collection->FindAttributeTyped<TArray<int32>>(FSelectionAlgo::IndexAttribute, GroupName) != nullptr);
+	}
+
+	void FSelectionAlgo::InitWeightedUnboundedGroup(FManagedArrayCollection* Collection, FName GroupName, FName DependencyGroup)
+	{
+		if (!Collection->HasGroup(GroupName))
+		{
+			Collection->AddAttribute<TArray<int32>>(FSelectionAlgo::IndexAttribute, GroupName, { DependencyGroup });
+			Collection->AddAttribute<TArray<float>>(FSelectionAlgo::WeightAttribute, GroupName);
+		}
+		ensure(Collection->FindAttributeTyped<TArray<int32>>(FSelectionAlgo::IndexAttribute, GroupName) != nullptr);
+		ensure(Collection->FindAttributeTyped<TArray<float>>(FSelectionAlgo::WeightAttribute, GroupName) != nullptr);
+	}
+
+	void FSelectionAlgo::InitBoundedGroup(FManagedArrayCollection* Collection, FName GroupName, FName DependencyGroup, FName BoneDependencyGroup)
+	{
+		if (!Collection->HasGroup(GroupName))
+		{
+			Collection->AddAttribute<TArray<int32>>(FSelectionAlgo::IndexAttribute, GroupName, { DependencyGroup });
+			Collection->AddAttribute<int32>(FSelectionAlgo::BoneIndexAttribute, GroupName, { BoneDependencyGroup });
+		}
+		ensure(Collection->FindAttributeTyped<TArray<int32>>(FSelectionAlgo::IndexAttribute, GroupName) != nullptr);
+		ensure(Collection->FindAttributeTyped<int32>(FSelectionAlgo::BoneIndexAttribute, GroupName) != nullptr);
+	}
+
+	void FSelectionAlgo::InitWeightedBoundedGroup(FManagedArrayCollection* Collection, FName GroupName, FName DependencyGroup, FName BoneDependencyGroup)
+	{
+		if (!Collection->HasGroup(GroupName))
+		{
+			Collection->AddAttribute<TArray<int32>>(FSelectionAlgo::IndexAttribute, GroupName, { DependencyGroup });
+			Collection->AddAttribute<TArray<float>>(FSelectionAlgo::WeightAttribute, GroupName);
+			Collection->AddAttribute<int32>(FSelectionAlgo::BoneIndexAttribute, GroupName, { BoneDependencyGroup });
+		}
+		ensure(Collection->FindAttributeTyped<TArray<int32>>(FSelectionAlgo::IndexAttribute, GroupName) != nullptr);
+		ensure(Collection->FindAttributeTyped<TArray<float>>(FSelectionAlgo::WeightAttribute, GroupName) != nullptr);
+		ensure(Collection->FindAttributeTyped<int32>(FSelectionAlgo::BoneIndexAttribute, GroupName) != nullptr);
+	}
+
+
+	//
+	//  AddSelection
+	//
+
+	FSelectionAlgo::FSelectionKey FSelectionAlgo::AddSelection(FManagedArrayCollection* Collection, const TArray<int32>& InIndices, FName DependencyGroup)
+	{
+		FName GroupName(FSelectionAlgo::UnboundGroup.ToString() + "_" + DependencyGroup.ToString());
+		InitUnboundedGroup(Collection, GroupName, DependencyGroup);
+
+		int Idx = Collection->AddElements(1, GroupName);
+		Collection->ModifyAttribute<TArray<int32>>(FSelectionAlgo::IndexAttribute, GroupName)[Idx] = InIndices;
+		return FSelectionKey(Idx, FSelectionAlgo::UnboundGroup);
+	}
+
+	FSelectionAlgo::FSelectionKey FSelectionAlgo::AddSelection(FManagedArrayCollection* Collection, const TArray<int32>& InIndices, const TArray<float>& InWeights, FName DependencyGroup)
+	{
+		FName GroupName(FSelectionAlgo::WeightedUnboundGroup.ToString() + "_" + DependencyGroup.ToString());
+		InitWeightedUnboundedGroup(Collection, GroupName, DependencyGroup);
+
+		int Idx = Collection->AddElements(1, GroupName);
+		Collection->ModifyAttribute<TArray<int32>>(FSelectionAlgo::IndexAttribute, GroupName)[Idx] = InIndices;
+		Collection->ModifyAttribute<TArray<float>>(FSelectionAlgo::WeightAttribute, GroupName)[Idx] = InWeights;
+		return FSelectionKey(Idx, GroupName);
+	}
+
+	FSelectionAlgo::FSelectionKey FSelectionAlgo::AddSelection(FManagedArrayCollection* Collection, const int32 InBoneIndex, const TArray<int32>& InIndices, FName DependencyGroup, FName BoneDependencyGroup)
+	{
+		FName GroupName(FSelectionAlgo::BoundGroup.ToString() + "_" + DependencyGroup.ToString());
+		InitBoundedGroup(Collection, GroupName, DependencyGroup, BoneDependencyGroup);
+
+		int Idx = Collection->AddElements(1, GroupName);
+		Collection->ModifyAttribute<TArray<int32>>(FSelectionAlgo::IndexAttribute, GroupName)[Idx] = InIndices;
+		Collection->ModifyAttribute<int32>(FSelectionAlgo::BoneIndexAttribute, GroupName)[Idx] = InBoneIndex;
+		return FSelectionKey(Idx, GroupName);
+	}
+
+	FSelectionAlgo::FSelectionKey FSelectionAlgo::AddSelection(FManagedArrayCollection* Collection, const int32 InBoneIndex, const TArray<int32>& InIndices, const TArray<float>& InWeights, FName DependencyGroup, FName BoneDependencyGroup)
+	{
+		FName GroupName(FSelectionAlgo::WeightedBoundGroup.ToString() + "_" + DependencyGroup.ToString());
+		InitWeightedBoundedGroup(Collection, GroupName, DependencyGroup, BoneDependencyGroup);
+
+		int Idx = Collection->AddElements(1, GroupName);
+		Collection->ModifyAttribute<TArray<int32>>(FSelectionAlgo::IndexAttribute, GroupName)[Idx] = InIndices;
+		Collection->ModifyAttribute<TArray<float>>(FSelectionAlgo::WeightAttribute, GroupName)[Idx] = InWeights;
+		Collection->ModifyAttribute<int32>(FSelectionAlgo::BoneIndexAttribute, GroupName)[Idx] = InBoneIndex;
+		return FSelectionKey(Idx, GroupName);
+	}
+
+	//
+	//  GetSelection
+	//
+
+
+	void FSelectionAlgo::GetSelection(const FManagedArrayCollection* Collection, const FSelectionAlgo::FSelectionKey& Key, TArray<int32>& OutIndices)
+	{
+		if (Collection->HasGroup(Key.GroupName) && 0 <= Key.Index && Key.Index < Collection->NumElements(Key.GroupName))
+		{
+			if (Collection->FindAttribute<TArray<int32>>(FSelectionAlgo::IndexAttribute, Key.GroupName))
+				OutIndices = Collection->GetAttribute<TArray<int32>>(FSelectionAlgo::IndexAttribute, Key.GroupName)[Key.Index];
+		}
+	}
+
+	void FSelectionAlgo::GetSelection(const FManagedArrayCollection* Collection, const FSelectionAlgo::FSelectionKey& Key, TArray<int32>& OutIndices, TArray<float>& OutWeights)
+	{
+		if (Collection->HasGroup(Key.GroupName) && 0 <= Key.Index && Key.Index < Collection->NumElements(Key.GroupName))
+		{
+			if (Collection->FindAttribute<TArray<int32>>(FSelectionAlgo::IndexAttribute, Key.GroupName))
+				OutIndices = Collection->GetAttribute<TArray<int32>>(FSelectionAlgo::IndexAttribute, Key.GroupName)[Key.Index];
+			if (Collection->FindAttribute<TArray<float>>(FSelectionAlgo::WeightAttribute, Key.GroupName))
+				OutWeights = Collection->GetAttribute<TArray<float>>(FSelectionAlgo::WeightAttribute, Key.GroupName)[Key.Index];
+		}
+	}
+
+	void FSelectionAlgo::GetSelection(const FManagedArrayCollection* Collection, const FSelectionAlgo::FSelectionKey& Key, int32& OutBoneIndex, TArray<int32>& OutIndices)
+	{
+		if (Collection->HasGroup(Key.GroupName) && 0 <= Key.Index && Key.Index < Collection->NumElements(Key.GroupName))
+		{
+			if (Collection->FindAttribute<TArray<int32>>(FSelectionAlgo::IndexAttribute, Key.GroupName))
+				OutIndices = Collection->GetAttribute<TArray<int32>>(FSelectionAlgo::IndexAttribute, Key.GroupName)[Key.Index];
+			if (Collection->FindAttribute<int32>(FSelectionAlgo::BoneIndexAttribute, Key.GroupName))
+				OutBoneIndex = Collection->GetAttribute<int32>(FSelectionAlgo::BoneIndexAttribute, Key.GroupName)[Key.Index];
+		}
+	}
+
+	void FSelectionAlgo::GetSelection(const FManagedArrayCollection* Collection, const FSelectionAlgo::FSelectionKey& Key, int32& OutBoneIndex, TArray<int32>& OutIndices, TArray<float>& OutWeights)
+	{
+		if (Collection->HasGroup(Key.GroupName) && 0 <= Key.Index && Key.Index < Collection->NumElements(Key.GroupName))
+		{
+			if (Collection->FindAttribute<TArray<int32>>(FSelectionAlgo::IndexAttribute, Key.GroupName))
+				OutIndices = Collection->GetAttribute<TArray<int32>>(FSelectionAlgo::IndexAttribute, Key.GroupName)[Key.Index];
+			if (Collection->FindAttribute<TArray<float>>(FSelectionAlgo::WeightAttribute, Key.GroupName))
+				OutWeights = Collection->GetAttribute<TArray<float>>(FSelectionAlgo::WeightAttribute, Key.GroupName)[Key.Index];
+			if (Collection->FindAttribute<int32>(FSelectionAlgo::BoneIndexAttribute, Key.GroupName))
+				OutBoneIndex = Collection->GetAttribute<int32>(FSelectionAlgo::BoneIndexAttribute, Key.GroupName)[Key.Index];
+		}
+	}
+
+};
+
+
diff --git a/Engine/Source/Runtime/Experimental/Chaos/Public/GeometryCollection/GeometryCollectionSelectionAlgo.h b/Engine/Source/Runtime/Experimental/Chaos/Public/GeometryCollection/GeometryCollectionSelectionAlgo.h
new file mode 100644
--- /dev/null
+++ b/Engine/Source/Runtime/Experimental/Chaos/Public/GeometryCollection/GeometryCollectionSelectionAlgo.h
@@ -0,0 +1,136 @@
+// Copyright Epic Games, Inc. All Rights Reserved.
+
+#pragma once
+#include "Containers/UnrealString.h"
+#include "Containers/Array.h"
+#include "GeometryCollection/ManagedArrayCollection.h"
+
+namespace GeometryCollection
+{
+
+	/**
+	* FSelectionAlgo
+	* 
+	* Defines common API for storing a sib indexing of a separate group. For example, to store
+	* a subset of vertices defined in the group "VerticesGroup", use the function :
+	*    Key = this->AddSelection({1,2,3,4});
+	* The key can be saved and later used to query the values:
+	*	TArray<int32> SavedVertices;
+	*   this->GetSelection(Key, SavedVertices);
+	*
+	* The following groups are created on the collection based on which API is called. 
+	* 
+	*	<Group> = FSelectionAlgo::UnboundGroup_<DependencyGroup>
+	*	- FindAttribute<TArray<int32>>(FVertexSetInterface::IndexAttribute, <Group>);
+	* 
+	*	<Group> = FSelectionAlgo::WeightedUnboundGroup_<DependencyGroup>
+	*	- FindAttribute<TArray<int32>>(FVertexSetInterface::IndexAttribute, <Group>);
+	*	- FindAttribute<TArray<float>>(FVertexSetInterface::WeightAttribute, <Group>);
+	* 
+	*	<Group> = FSelectionAlgo::BoundGroup_<DependencyGroup>
+	*	- FindAttribute<TArray<int32>>(FVertexSetInterface::IndexAttribute, <Group>);
+	*	- FindAttribute<int32>(FVertexSetInterface::BoneIndexAttribute,<Group>);
+	*
+	*	<Group> = FSelectionAlgo::WeightedBoundGroup_<DependencyGroup>
+	*	- FindAttribute<TArray<int32>>(FVertexSetInterface::IndexAttribute, <Group>);
+	* 	- FindAttribute<TArray<float>>(FVertexSetInterface::WeightAttribute, <Group>);
+	*	- FindAttribute<int32>(FVertexSetInterface::BoneIndexAttribute, <Group>);
+	* 
+	*/
+	class CHAOS_API FSelectionAlgo
+ 	{
+		FManagedArrayCollection* Self;
+
+		static void InitUnboundedGroup(FManagedArrayCollection* Self, FName GroupName, FName DependencyGroup);
+		static void InitWeightedUnboundedGroup(FManagedArrayCollection* Self, FName GroupName, FName DependencyGroup);
+		static void InitBoundedGroup(FManagedArrayCollection* Self, FName GroupName, FName DependencyGroup, FName BoneDependencyGroup);
+		static void InitWeightedBoundedGroup(FManagedArrayCollection* Self, FName GroupName, FName DependencyGroup, FName BoneDependencyGroup);
+
+	public:
+
+		// groups
+		static const FName UnboundGroup;
+		static const FName WeightedUnboundGroup;
+		static const FName BoundGroup;
+		static const FName WeightedBoundGroup;
+
+		// Attributes
+		static const FName IndexAttribute;
+		static const FName WeightAttribute;
+		static const FName BoneIndexAttribute;
+
+		struct FSelectionKey {
+			FSelectionKey(int32 InIndex = INDEX_NONE, FName InGroupName = "") :
+				Index(InIndex), GroupName(InGroupName) {}
+			int32 Index;
+			FName GroupName;
+		};
+
+		/**
+		* FSelectionAlgo Constuctor
+		* @param VertixDependencyGroup : GroupName the index attribute is dependent on. 
+		*/
+		FSelectionAlgo(FManagedArrayCollection* InSelf);
+
+		/**
+		* Add the indices to the FVertexSetInterface::UnboundGroup 
+		* @param Indices : Array of indices to add. 
+		*/
+		static FSelectionKey AddSelection(FManagedArrayCollection* Self, const TArray<int32>& Indices, FName DependencyGroup);
+
+		/**
+		* Add the indices to the FVertexSetInterface::WeightedUnboundGroup
+		* @param Indices : Array of indices to add.
+		* @param Indices : Array of indices to add.
+		*/
+		static FSelectionKey AddSelection(FManagedArrayCollection* Self, const TArray<int32>& Indices, const TArray<float>& Weights, FName DependencyGroup);
+
+		/**
+		* Add the indices to the FVertexSetInterface::BoundGroup
+		* @param Indices : Array of indices to add.
+		* @param Indices : Array of indices to add.
+		*/
+		static FSelectionKey AddSelection(FManagedArrayCollection* Self, const int32 BoneIndex, const TArray<int32>& Indices, FName DependencyGroup, FName BoneDependencyGroup = FName(""));
+
+		/**
+		* Add the indices to the FVertexSetInterface::BoundGroup
+		* @param Indices : Array of indices to add.
+		* @param Indices : Array of indices to add.
+		*/
+		static FSelectionKey AddSelection(FManagedArrayCollection* Self, const int32 BoneIndex, const TArray<int32>& Indices, const TArray<float>& Weights, FName DependencyGroup, FName BoneDependencyGroup = FName(""));
+
+		/**
+		* Return the vertex list from the given key 
+		* @param Key : <GroupName and Index>
+		* @param Indices : Return indices, empty if not found. 
+		*/
+		static void GetSelection(const FManagedArrayCollection* Self, const FSelectionKey& Key, TArray<int32>& Indices);
+
+		/**
+		* Return the vertex list, bone index, and weights from the given key
+		* @param Key : <GroupName and Index>
+		* @param BoneIndex : Return BoneIndex, INDEX_NONE if not found.
+		* @param Indices : Return indices, empty if not found.
+		* @param Weights : Return vertex weights, empty if not found.
+		*/
+		static void GetSelection(const FManagedArrayCollection* Self, const FSelectionKey& Key, TArray<int32>& Indices, TArray<float>& Weights);
+
+		/**
+		* Return the vertex list and bone index from the given key
+		* @param Key : <GroupName and Index>
+		* @param BoneIndex : Return BoneIndex, INDEX_NONE if not found.
+		* @param Indices : Return indices, empty if not found.
+		*/
+		static void GetSelection(const FManagedArrayCollection* Self, const FSelectionKey& Key, int32& BoneIndex, TArray<int32>& Indices);
+
+		/**
+		* Return the vertex list, bone index, and weights from the given key
+		* @param Key : <GroupName and Index>
+		* @param BoneIndex : Return BoneIndex, INDEX_NONE if not found.
+		* @param Indices : Return indices, empty if not found.
+		* @param Weights : Return vertex weights, empty if not found.
+		*/
+		static void GetSelection(const FManagedArrayCollection* Self, const FSelectionKey& Key, int32& BoneIndex, TArray<int32>& Indices, TArray<float>& Weights);
+	};
+
+}
\ No newline at end of file
