Chaos:
    Fixed disabled constraints never being removed from the constraint graph leading to severe performance degredation after long runtimes
    Fixed disbaled constraints never being removed from their owning particle joint lists - leading to dangling joint pointers in the particles

diff --git a/Engine/Source/Runtime/Experimental/Chaos/Private/Chaos/PBDJointConstraints.cpp b/Engine/Source/Runtime/Experimental/Chaos/Private/Chaos/PBDJointConstraints.cpp
--- a/Engine/Source/Runtime/Experimental/Chaos/Private/Chaos/PBDJointConstraints.cpp
+++ b/Engine/Source/Runtime/Experimental/Chaos/Private/Chaos/PBDJointConstraints.cpp
@@ -482,11 +482,9 @@ namespace Chaos
 		FConstraintContainerHandle* ConstraintHandle = Handles[ConstraintIndex];
 		if (ConstraintHandle != nullptr)
 		{
-			if (!ConstraintStates[ConstraintIndex].bDisabled)
-			{
-				ConstraintParticles[ConstraintIndex][0]->RemoveConstraintHandle(ConstraintHandle);
-				ConstraintParticles[ConstraintIndex][1]->RemoveConstraintHandle(ConstraintHandle);
-			}
+			ConstraintParticles[ConstraintIndex][0]->RemoveConstraintHandle(ConstraintHandle);
+			ConstraintParticles[ConstraintIndex][1]->RemoveConstraintHandle(ConstraintHandle);
+
 			// Release the handle for the freed constraint
 			HandleAllocator.FreeHandle(ConstraintHandle);
 			Handles[ConstraintIndex] = nullptr;
diff --git a/Engine/Source/Runtime/Experimental/Chaos/Private/PhysicsProxy/JointConstraintProxy.cpp b/Engine/Source/Runtime/Experimental/Chaos/Private/PhysicsProxy/JointConstraintProxy.cpp
--- a/Engine/Source/Runtime/Experimental/Chaos/Private/PhysicsProxy/JointConstraintProxy.cpp
+++ b/Engine/Source/Runtime/Experimental/Chaos/Private/PhysicsProxy/JointConstraintProxy.cpp
@@ -105,10 +105,14 @@ void FJointConstraintPhysicsProxy::InitializeOnPhysicsThread(FPBDRigidsSolver* I
 
 void FJointConstraintPhysicsProxy::DestroyOnPhysicsThread(FPBDRigidsSolver* InSolver)
 {
-	if (Constraint_PT && Constraint_PT->IsValid())
+	if (Constraint_PT)
 	{
 		// @todo(chaos): clean up constraint management
-		InSolver->GetEvolution()->RemoveConstraintFromConstraintGraph(Constraint_PT);
+		if(Constraint_PT->IsInConstraintGraph())
+		{
+			InSolver->GetEvolution()->RemoveConstraintFromConstraintGraph(Constraint_PT);
+		}
+
 		auto& JointConstraints = InSolver->GetJointConstraints();
 		JointConstraints.RemoveConstraint(Constraint_PT->GetConstraintIndex());
 	}
