Chaos : Add several geometry collectiion UFunction
- ApplyLinearVelocity / ApplyAngularVelocity
- ApplyBreakingLinearVelocity / ApplyBreakingAngularVelocity
- GetInitialLevel

diff --git a/Engine/Source/Runtime/Experimental/Chaos/Private/PhysicsProxy/GeometryCollectionPhysicsProxy.cpp b/Engine/Source/Runtime/Experimental/Chaos/Private/PhysicsProxy/GeometryCollectionPhysicsProxy.cpp
--- a/Engine/Source/Runtime/Experimental/Chaos/Private/PhysicsProxy/GeometryCollectionPhysicsProxy.cpp
+++ b/Engine/Source/Runtime/Experimental/Chaos/Private/PhysicsProxy/GeometryCollectionPhysicsProxy.cpp
@@ -1583,6 +1583,101 @@ void FGeometryCollectionPhysicsProxy::ApplyStrain_External(FGeometryCollectionIt
 	}
 }
 
+template <typename TLambda>
+static void ApplyToBreakingChildren_Internal(Chaos::FRigidClustering& Clustering, Chaos::FPBDRigidClusteredParticleHandle* ClusteredHandle, TLambda Action)
+{
+	const bool bIsCluster = (ClusteredHandle->ClusterIds().NumChildren > 0);
+	if (bIsCluster && !ClusteredHandle->Disabled())
+	{
+		Chaos::FRigidClustering::FClusterMap& ChildrenMap = Clustering.GetChildrenMap();
+		if (const TArray<Chaos::FPBDRigidParticleHandle*>* ChildrenHandles = ChildrenMap.Find(ClusteredHandle))
+		{
+			for (Chaos::FPBDRigidParticleHandle* ChildHandle: *ChildrenHandles)
+			{
+				if (Chaos::FPBDRigidClusteredParticleHandle* ClusteredChildHandle = ChildHandle->CastToClustered())
+				{
+					if (ClusteredChildHandle->GetExternalStrain() > ClusteredChildHandle->Strain())
+					{
+						Action(ClusteredChildHandle);
+					}
+				}
+			}
+		}
+	}
+}
+
+void FGeometryCollectionPhysicsProxy::ApplyBreakingLinearVelocity_External(FGeometryCollectionItemIndex ItemIndex, const FVector& LinearVelocity)
+{
+	check(IsInGameThread());
+
+	if (Chaos::FPhysicsSolver* RBDSolver = GetSolver<Chaos::FPhysicsSolver>())
+	{
+		RBDSolver->EnqueueCommandImmediate([this, RBDSolver, ItemIndex, LinearVelocity]()
+		{
+			if (Chaos::FPBDRigidClusteredParticleHandle* ClusteredHandle = FindClusteredParticleHandleByItemIndex(ItemIndex))
+			{
+				ApplyToBreakingChildren_Internal(RBDSolver->GetEvolution()->GetRigidClustering(), ClusteredHandle,
+					[LinearVelocity](Chaos::FPBDRigidClusteredParticleHandle* ClusteredChildHandle)
+					{
+						ClusteredChildHandle->V() += LinearVelocity;
+					});
+			}
+		});
+	}
+}
+
+void FGeometryCollectionPhysicsProxy::ApplyBreakingAngularVelocity_External(FGeometryCollectionItemIndex ItemIndex, const FVector& AngularVelocity)
+{
+	check(IsInGameThread());
+
+	if (Chaos::FPhysicsSolver* RBDSolver = GetSolver<Chaos::FPhysicsSolver>())
+	{
+		RBDSolver->EnqueueCommandImmediate([this, RBDSolver, ItemIndex, AngularVelocity]()
+		{
+			if (Chaos::FPBDRigidClusteredParticleHandle* ClusteredHandle = FindClusteredParticleHandleByItemIndex(ItemIndex))
+			{
+				ApplyToBreakingChildren_Internal(RBDSolver->GetEvolution()->GetRigidClustering(), ClusteredHandle,
+					[AngularVelocity](Chaos::FPBDRigidClusteredParticleHandle* ClusteredChildHandle)
+					{
+						ClusteredChildHandle->W() += AngularVelocity;
+					});
+			}
+		});
+	}
+}
+
+void FGeometryCollectionPhysicsProxy::ApplyLinearVelocity_External(FGeometryCollectionItemIndex ItemIndex, const FVector& LinearVelocity)
+{
+	check(IsInGameThread());
+
+	if (Chaos::FPhysicsSolver* RBDSolver = GetSolver<Chaos::FPhysicsSolver>())
+	{
+		RBDSolver->EnqueueCommandImmediate([this, ItemIndex, LinearVelocity]()
+		{
+			if (Chaos::FPBDRigidClusteredParticleHandle* ClusteredHandle = FindClusteredParticleHandleByItemIndex(ItemIndex))
+			{
+				ClusteredHandle->V() += LinearVelocity;
+			}
+		});
+	}
+}
+
+void FGeometryCollectionPhysicsProxy::ApplyAngularVelocity_External(FGeometryCollectionItemIndex ItemIndex, const FVector& AngularVelocity)
+{
+	check(IsInGameThread());
+
+	if (Chaos::FPhysicsSolver* RBDSolver = GetSolver<Chaos::FPhysicsSolver>())
+	{
+		RBDSolver->EnqueueCommandImmediate([this, ItemIndex, AngularVelocity]()
+		{
+			if (Chaos::FPBDRigidClusteredParticleHandle* ClusteredHandle = FindClusteredParticleHandleByItemIndex(ItemIndex))
+			{
+				ClusteredHandle->W() += AngularVelocity;
+			}
+		});
+	}
+}
+
 int32 FGeometryCollectionPhysicsProxy::CalculateHierarchyLevel(const FGeometryDynamicCollection& DynamicCollection, int32 TransformIndex)
 {
 	int32 Level = 0;
@@ -2360,12 +2455,14 @@ bool FGeometryCollectionPhysicsProxy::PullFromPhysicsState(const Chaos::FDirtyGe
 
 		// internal cluster index map update
 		GTParticlesToInternalClusterUniqueIdx.Reset();
+		InternalClusterUniqueIdxToChildrenTransformIndices.Reset();
 		for (int32 TransformGroupIndex = 0; TransformGroupIndex < NumTransforms; ++TransformGroupIndex)
 		{
 			const int32 InternalClusterUniqueIdx = TargetResults.InternalClusterUniqueIdx[TransformGroupIndex]; 
 			if (InternalClusterUniqueIdx > INDEX_NONE)
 			{
-				GTParticlesToInternalClusterUniqueIdx.Add(GTParticles[TransformGroupIndex].Get(), InternalClusterUniqueIdx);	
+				GTParticlesToInternalClusterUniqueIdx.Add(GTParticles[TransformGroupIndex].Get(), InternalClusterUniqueIdx);
+				InternalClusterUniqueIdxToChildrenTransformIndices.FindOrAdd(InternalClusterUniqueIdx).Add(TransformGroupIndex); 
 			}
 		}
 
diff --git a/Engine/Source/Runtime/Experimental/Chaos/Public/PhysicsProxy/GeometryCollectionPhysicsProxy.h b/Engine/Source/Runtime/Experimental/Chaos/Public/PhysicsProxy/GeometryCollectionPhysicsProxy.h
--- a/Engine/Source/Runtime/Experimental/Chaos/Public/PhysicsProxy/GeometryCollectionPhysicsProxy.h
+++ b/Engine/Source/Runtime/Experimental/Chaos/Public/PhysicsProxy/GeometryCollectionPhysicsProxy.h
@@ -331,6 +331,15 @@ public:
 		}
 		return FGeometryCollectionItemIndex::CreateInvalidItemIndex();
 	}
+
+	const TArray<int32>* FindInternalClusterChildrenTransformIndices_External(FGeometryCollectionItemIndex ItemIndex) const
+	{
+		if (ensure(ItemIndex.IsInternalCluster()))
+		{
+			return InternalClusterUniqueIdxToChildrenTransformIndices.Find(ItemIndex.GetInternalClusterIndex());
+		}
+		return nullptr;
+	}
 	
 	FGeometryCollectionItemIndex GetItemIndexFromGTParticle_External(const Chaos::FGeometryParticle* GTPParticle) const
 	{
@@ -353,6 +362,10 @@ public:
 
 	void BreakClusters_External(TArray<FGeometryCollectionItemIndex>&& ItemIndices);
 	void ApplyStrain_External(FGeometryCollectionItemIndex ItemIndex, const FVector& WorldLocation, float StrainValue);
+	void ApplyBreakingLinearVelocity_External(FGeometryCollectionItemIndex ItemIndex, const FVector& LinearVelocity);
+	void ApplyBreakingAngularVelocity_External(FGeometryCollectionItemIndex ItemIndex, const FVector& AngularVelocity);
+	void ApplyLinearVelocity_External(FGeometryCollectionItemIndex ItemIndex, const FVector& LinearVelocity);
+	void ApplyAngularVelocity_External(FGeometryCollectionItemIndex ItemIndex, const FVector& AngularVelocity);
 
 	FProxyInterpolationData& GetInterpolationData() { return InterpolationData; }
 	const FProxyInterpolationData& GetInterpolationData() const { return InterpolationData; }
@@ -437,6 +450,7 @@ private:
 	TManagedArray<TUniquePtr<Chaos::FGeometryParticle>> GTParticles;
 	TMap<Chaos::FGeometryParticle*, int32> GTParticlesToTransformGroupIndex;
 	TMap<Chaos::FGeometryParticle*, int32> GTParticlesToInternalClusterUniqueIdx;
+	TMap<int32, TArray<int32>> InternalClusterUniqueIdxToChildrenTransformIndices;
 
 	TMap<int32, TUniquePtr<Chaos::FGeometryParticle>> GTInternalClustersByUniqueIdx;
 
diff --git a/Engine/Source/Runtime/Experimental/GeometryCollectionEngine/Private/GeometryCollection/GeometryCollectionComponent.cpp b/Engine/Source/Runtime/Experimental/GeometryCollectionEngine/Private/GeometryCollection/GeometryCollectionComponent.cpp
--- a/Engine/Source/Runtime/Experimental/GeometryCollectionEngine/Private/GeometryCollection/GeometryCollectionComponent.cpp
+++ b/Engine/Source/Runtime/Experimental/GeometryCollectionEngine/Private/GeometryCollection/GeometryCollectionComponent.cpp
@@ -3750,6 +3750,83 @@ void UGeometryCollectionComponent::CrumbleCluster(int32 ItemIndex)
 	}
 }
 
+void UGeometryCollectionComponent::ApplyBreakingLinearVelocity(int32 ItemIndex, const FVector& LinearVelocity)
+{
+	if (PhysicsProxy)
+	{
+		PhysicsProxy->ApplyBreakingLinearVelocity_External(FGeometryCollectionItemIndex::CreateFromExistingItemIndex(ItemIndex), LinearVelocity);
+	}
+
+}
+
+void UGeometryCollectionComponent::ApplyBreakingAngularVelocity(int32 ItemIndex, const FVector& AngularVelocity)
+{
+	if (PhysicsProxy)
+	{
+		PhysicsProxy->ApplyBreakingLinearVelocity_External(FGeometryCollectionItemIndex::CreateFromExistingItemIndex(ItemIndex), AngularVelocity);
+	}
+}
+
+void UGeometryCollectionComponent::ApplyLinearVelocity(int32 ItemIndex, const FVector& LinearVelocity)
+{
+	if (PhysicsProxy)
+	{
+		PhysicsProxy->ApplyLinearVelocity_External(FGeometryCollectionItemIndex::CreateFromExistingItemIndex(ItemIndex), LinearVelocity);
+	}
+}
+
+void UGeometryCollectionComponent::ApplyAngularVelocity(int32 ItemIndex, const FVector& AngularVelocity)
+{
+	if (PhysicsProxy)
+	{
+		PhysicsProxy->ApplyAngularVelocity_External(FGeometryCollectionItemIndex::CreateFromExistingItemIndex(ItemIndex), AngularVelocity);
+	}
+}
+
+int32 UGeometryCollectionComponent::GetInitialLevel(int32 ItemIndex)
+{
+	using FGeometryCollectionPtr = const TSharedPtr<FGeometryCollection, ESPMode::ThreadSafe>; 
+
+	int32 Level = INDEX_NONE;
+	if (RestCollection && RestCollection->GetGeometryCollection())
+	{
+		const TManagedArray<int32>& Parent = RestCollection->GetGeometryCollection()->Parent;
+		int32 TransformIndex = INDEX_NONE;
+
+		FGeometryCollectionItemIndex GCItemIndex = FGeometryCollectionItemIndex::CreateFromExistingItemIndex(ItemIndex);
+
+		if (GCItemIndex.IsInternalCluster())
+		{
+			if (const TArray<int32>* Children = PhysicsProxy->FindInternalClusterChildrenTransformIndices_External(GCItemIndex))
+			{
+				if (!Children->IsEmpty())
+				{
+					// find the original cluster index from first children
+					TransformIndex = Parent[(*Children)[0]];
+				}
+			}
+		}
+		else
+		{
+			TransformIndex = GCItemIndex.GetTransformIndex();
+		}
+
+		// @todo(chaos) : use "Level" attribute when it will be properly serialized
+		// for now climb back the hierarchy
+		if (TransformIndex > INDEX_NONE)
+		{
+			Level = 0;
+			int32 ParentTransformIndex =  Parent[TransformIndex];
+			while (ParentTransformIndex != INDEX_NONE)
+			{
+				++Level;
+				ParentTransformIndex =  Parent[ParentTransformIndex];
+			}
+		}
+	}
+	return Level;
+}
+
 bool UGeometryCollectionComponent::CalculateInnerSphere(int32 TransformIndex, FSphere& SphereOut) const
 {
 	// Approximates the inscribed sphere. Returns false if no such sphere exists, if for instance the index is to an embedded geometry. 
diff --git a/Engine/Source/Runtime/Experimental/GeometryCollectionEngine/Public/GeometryCollection/GeometryCollectionComponent.h b/Engine/Source/Runtime/Experimental/GeometryCollectionEngine/Public/GeometryCollection/GeometryCollectionComponent.h
--- a/Engine/Source/Runtime/Experimental/GeometryCollectionEngine/Public/GeometryCollection/GeometryCollectionComponent.h
+++ b/Engine/Source/Runtime/Experimental/GeometryCollectionEngine/Public/GeometryCollection/GeometryCollectionComponent.h
@@ -485,11 +485,53 @@ public:
 
 	/**
 	 * Crumbe a cluster into all its pieces
-	 * if the piece is not a cluster but is part of an internal/partial cluster, then the owning cluster will crumble   
 	 * @param ItemIndex item index ( from HitResult) of the cluster to crumble
 	 */
 	UFUNCTION(BlueprintCallable, Category = "ChaosPhysics")
 	void CrumbleCluster(int32 ItemIndex);
+
+	/**
+	 * Apply linear velocity on breaking pieces for a specific cluster
+	 * If ItemIndex does not represent a cluster this will do nothing  
+	 * @param ItemIndex item index ( from HitResult) of the cluster owning the breaking pieces to apply velocity on
+	 * @param LinearVelocity linear velocity to apply
+	 */
+	UFUNCTION(BlueprintCallable, Category = "ChaosPhysics")
+	void ApplyBreakingLinearVelocity(int32 ItemIndex, const FVector& LinearVelocity);
+
+	/**
+	 * Apply linear velocity on breaking pieces for a specific cluster
+	 * If ItemIndex does not represent a cluster this will do nothing  
+	 * @param ItemIndex item index ( from HitResult) of the cluster owning the breaking pieces to apply velocity on
+	 * @param AngularVelocity linear velocity to apply
+	 */
+	UFUNCTION(BlueprintCallable, Category = "ChaosPhysics")
+	void ApplyBreakingAngularVelocity(int32 ItemIndex, const FVector& AngularVelocity);
+	
+	/**
+	 * Apply linear velocity on specific piece 
+	 * @param ItemIndex item index ( from HitResult) of the piece to apply velocity on
+	* @param LinearVelocity linear velocity to apply
+	 */
+	UFUNCTION(BlueprintCallable, Category = "ChaosPhysics")
+	void ApplyLinearVelocity(int32 ItemIndex, const FVector& LinearVelocity);
+
+	/**
+	 * Apply angular velocity on specific piece 
+	 * @param ItemIndex item index ( from HitResult) of the piece to apply velocity on
+	* @param AngularVelocity linear velocity to apply
+	 */
+	UFUNCTION(BlueprintCallable, Category = "ChaosPhysics")
+	void ApplyAngularVelocity(int32 ItemIndex, const FVector& AngularVelocity);
+
+	/**
+	 * Get the initial level of a specific piece
+	 * Initial level means the level as it is in the unbroken state 
+	 * @param ItemIndex item index ( from HitResult) of the cluster to get level from
+	 * @param Level of the piece ( 0 for root level and positive for the rest ) 
+	 */
+	UFUNCTION(BlueprintCallable, Category = "ChaosPhysics")
+	int32 GetInitialLevel(int32 ItemIndex);
 	
 	/** RestCollection */
 	UFUNCTION(BlueprintCallable, Category = "ChaosPhysics")
