Return FaceIndexwhen colliding with a TriangleMesh or HeightField
Pull request #9226

diff --git a/Engine/Source/Runtime/Experimental/Chaos/Private/Chaos/Collision/ContactPointsMiscShapes.cpp b/Engine/Source/Runtime/Experimental/Chaos/Private/Chaos/Collision/ContactPointsMiscShapes.cpp
--- a/Engine/Source/Runtime/Experimental/Chaos/Private/Chaos/Collision/ContactPointsMiscShapes.cpp
+++ b/Engine/Source/Runtime/Experimental/Chaos/Private/Chaos/Collision/ContactPointsMiscShapes.cpp
@@ -194,37 +194,41 @@ namespace Chaos
 
 		FReal ContactPhi = FLT_MAX;
 		FVec3 Location, Normal;
+		int32 ContactFaceIndex = INDEX_NONE;
 		if (const TImplicitObjectScaled<GeometryA>* ScaledConvexImplicit = A.template GetObject<const TImplicitObjectScaled<GeometryA> >())
 		{
-			if (B.GJKContactPoint(*ScaledConvexImplicit, AToBTM, CullDistance, Location, Normal, ContactPhi))
+			if (B.GJKContactPoint(*ScaledConvexImplicit, AToBTM, CullDistance, Location, Normal, ContactPhi, ContactFaceIndex))
 			{
 				Contact.ShapeContactPoints[0] = AToBTM.InverseTransformPosition(Location);
 				Contact.ShapeContactPoints[1] = Location - ContactPhi * Normal;
 				Contact.ShapeContactNormal = Normal;
 				Contact.Phi = ContactPhi;
+				Contact.FaceIndex = ContactFaceIndex;
 			}
 		}
 		else if (const TImplicitObjectInstanced<GeometryA>* InstancedConvexImplicit = A.template GetObject<const TImplicitObjectInstanced<GeometryA> >())
 		{
 			if (const GeometryA* InstancedInnerObject = static_cast<const GeometryA*>(InstancedConvexImplicit->GetInstancedObject()))
 			{
-				if (B.GJKContactPoint(*InstancedInnerObject, AToBTM, CullDistance, Location, Normal, ContactPhi))
+				if (B.GJKContactPoint(*InstancedInnerObject, AToBTM, CullDistance, Location, Normal, ContactPhi, ContactFaceIndex))
 				{
 					Contact.ShapeContactPoints[0] = AToBTM.InverseTransformPosition(Location);
 					Contact.ShapeContactPoints[1] = Location - ContactPhi * Normal;
 					Contact.ShapeContactNormal = Normal;
 					Contact.Phi = ContactPhi;
+					Contact.FaceIndex = ContactFaceIndex;
 				}
 			}
 		}
 		else if (const GeometryA* ConvexImplicit = A.template GetObject<const GeometryA>())
 		{
-			if (B.GJKContactPoint(*ConvexImplicit, AToBTM, CullDistance, Location, Normal, ContactPhi))
+			if (B.GJKContactPoint(*ConvexImplicit, AToBTM, CullDistance, Location, Normal, ContactPhi, ContactFaceIndex))
 			{
 				Contact.ShapeContactPoints[0] = AToBTM.InverseTransformPosition(Location);
 				Contact.ShapeContactPoints[1] = Location - ContactPhi * Normal;
 				Contact.ShapeContactNormal = Normal;
 				Contact.Phi = ContactPhi;
+				Contact.FaceIndex = ContactFaceIndex;
 			}
 		}
 
diff --git a/Engine/Source/Runtime/Experimental/Chaos/Private/Chaos/HeightField.cpp b/Engine/Source/Runtime/Experimental/Chaos/Private/Chaos/HeightField.cpp
--- a/Engine/Source/Runtime/Experimental/Chaos/Private/Chaos/HeightField.cpp
+++ b/Engine/Source/Runtime/Experimental/Chaos/Private/Chaos/HeightField.cpp
@@ -1619,7 +1619,7 @@ namespace Chaos
 
 	template <typename GeomType>
 	bool FHeightField::GJKContactPointImp(const GeomType& QueryGeom, const FRigidTransform3& QueryTM, const FReal Thickness,
-		FVec3& ContactLocation, FVec3& ContactNormal, FReal& ContactPhi) const
+		FVec3& ContactLocation, FVec3& ContactNormal, FReal& ContactPhi, int32& ContactFaceIndex) const
 	{
 		auto OverlapTriangle = [&](const FVec3& A, const FVec3& B, const FVec3& C,
 			FVec3& LocalContactLocation, FVec3& LocalContactNormal, FReal& LocalContactPhi) -> bool
@@ -1691,6 +1691,7 @@ namespace Chaos
 					ContactPhi = LocalContactPhi;
 					ContactLocation = LocalContactLocation;
 					ContactNormal = LocalContactNormal;
+					ContactFaceIndex = CellIndex * 2; // This assumes 2 faces per cell.
 				}
 			}
 
@@ -1701,6 +1702,8 @@ namespace Chaos
 					ContactPhi = LocalContactPhi;
 					ContactLocation = LocalContactLocation;
 					ContactNormal = LocalContactNormal;
+					// This Triangle has the same faceIndex as the previous one (this is a consequence of how the materials are stored [one per cell])
+					ContactFaceIndex = CellIndex * 2; // This assumes 2 faces per cell.
 				}
 			}
 		}
@@ -1710,44 +1713,44 @@ namespace Chaos
 		return false;
 	}
 
-	bool FHeightField::GJKContactPoint(const TBox<FReal, 3>& QueryGeom, const FRigidTransform3& QueryTM, const FReal Thickness, FVec3& ContactLocation, FVec3& ContactNormal, FReal& ContactPhi) const
+	bool FHeightField::GJKContactPoint(const TBox<FReal, 3>& QueryGeom, const FRigidTransform3& QueryTM, const FReal Thickness, FVec3& ContactLocation, FVec3& ContactNormal, FReal& ContactPhi, int32& ContactFaceIndex) const
 	{
-		return GJKContactPointImp(QueryGeom, QueryTM, Thickness, ContactLocation, ContactNormal, ContactPhi);
+		return GJKContactPointImp(QueryGeom, QueryTM, Thickness, ContactLocation, ContactNormal, ContactPhi, ContactFaceIndex);
 	}
 
-	bool FHeightField::GJKContactPoint(const TSphere<FReal, 3>& QueryGeom, const FRigidTransform3& QueryTM, const FReal Thickness, FVec3& ContactLocation, FVec3& ContactNormal, FReal& ContactPhi) const
+	bool FHeightField::GJKContactPoint(const TSphere<FReal, 3>& QueryGeom, const FRigidTransform3& QueryTM, const FReal Thickness, FVec3& ContactLocation, FVec3& ContactNormal, FReal& ContactPhi, int32& ContactFaceIndex) const
 	{
-		return GJKContactPointImp(QueryGeom, QueryTM, Thickness, ContactLocation, ContactNormal, ContactPhi);
+		return GJKContactPointImp(QueryGeom, QueryTM, Thickness, ContactLocation, ContactNormal, ContactPhi, ContactFaceIndex);
 	}
 
-	bool FHeightField::GJKContactPoint(const FCapsule& QueryGeom, const FRigidTransform3& QueryTM, const FReal Thickness, FVec3& ContactLocation, FVec3& ContactNormal, FReal& ContactPhi) const
+	bool FHeightField::GJKContactPoint(const FCapsule& QueryGeom, const FRigidTransform3& QueryTM, const FReal Thickness, FVec3& ContactLocation, FVec3& ContactNormal, FReal& ContactPhi, int32& ContactFaceIndex) const
 	{
-		return GJKContactPointImp(QueryGeom, QueryTM, Thickness, ContactLocation, ContactNormal, ContactPhi);
+		return GJKContactPointImp(QueryGeom, QueryTM, Thickness, ContactLocation, ContactNormal, ContactPhi, ContactFaceIndex);
 	}
 
-	bool FHeightField::GJKContactPoint(const FConvex& QueryGeom, const FRigidTransform3& QueryTM, const FReal Thickness, FVec3& ContactLocation, FVec3& ContactNormal, FReal& ContactPhi) const
+	bool FHeightField::GJKContactPoint(const FConvex& QueryGeom, const FRigidTransform3& QueryTM, const FReal Thickness, FVec3& ContactLocation, FVec3& ContactNormal, FReal& ContactPhi, int32& ContactFaceIndex) const
 	{
-		return GJKContactPointImp(QueryGeom, QueryTM, Thickness, ContactLocation, ContactNormal, ContactPhi);
+		return GJKContactPointImp(QueryGeom, QueryTM, Thickness, ContactLocation, ContactNormal, ContactPhi, ContactFaceIndex);
 	}
 
-	bool FHeightField::GJKContactPoint(const TImplicitObjectScaled < TBox<FReal, 3>>& QueryGeom, const FRigidTransform3& QueryTM, const FReal Thickness, FVec3& ContactLocation, FVec3& ContactNormal, FReal& ContactPhi) const
+	bool FHeightField::GJKContactPoint(const TImplicitObjectScaled < TBox<FReal, 3>>& QueryGeom, const FRigidTransform3& QueryTM, const FReal Thickness, FVec3& ContactLocation, FVec3& ContactNormal, FReal& ContactPhi, int32& ContactFaceIndex) const
 	{
-		return GJKContactPointImp(QueryGeom, QueryTM, Thickness, ContactLocation, ContactNormal, ContactPhi);
+		return GJKContactPointImp(QueryGeom, QueryTM, Thickness, ContactLocation, ContactNormal, ContactPhi, ContactFaceIndex);
 	}
 
-	bool FHeightField::GJKContactPoint(const TImplicitObjectScaled < TSphere<FReal, 3>>& QueryGeom, const FRigidTransform3& QueryTM, const FReal Thickness, FVec3& ContactLocation, FVec3& ContactNormal, FReal& ContactPhi) const
+	bool FHeightField::GJKContactPoint(const TImplicitObjectScaled < TSphere<FReal, 3>>& QueryGeom, const FRigidTransform3& QueryTM, const FReal Thickness, FVec3& ContactLocation, FVec3& ContactNormal, FReal& ContactPhi, int32& ContactFaceIndex) const
 	{
-		return GJKContactPointImp(QueryGeom, QueryTM, Thickness, ContactLocation, ContactNormal, ContactPhi);
+		return GJKContactPointImp(QueryGeom, QueryTM, Thickness, ContactLocation, ContactNormal, ContactPhi, ContactFaceIndex);
 	}
 
-	bool FHeightField::GJKContactPoint(const TImplicitObjectScaled < FCapsule>& QueryGeom, const FRigidTransform3& QueryTM, const FReal Thickness, FVec3& ContactLocation, FVec3& ContactNormal, FReal& ContactPhi) const
+	bool FHeightField::GJKContactPoint(const TImplicitObjectScaled < FCapsule>& QueryGeom, const FRigidTransform3& QueryTM, const FReal Thickness, FVec3& ContactLocation, FVec3& ContactNormal, FReal& ContactPhi, int32& ContactFaceIndex) const
 	{
-		return GJKContactPointImp(QueryGeom, QueryTM, Thickness, ContactLocation, ContactNormal, ContactPhi);
+		return GJKContactPointImp(QueryGeom, QueryTM, Thickness, ContactLocation, ContactNormal, ContactPhi, ContactFaceIndex);
 	}
 
-	bool FHeightField::GJKContactPoint(const TImplicitObjectScaled<FConvex>& QueryGeom, const FRigidTransform3& QueryTM, const FReal Thickness, FVec3& ContactLocation, FVec3& ContactNormal, FReal& ContactPhi) const
+	bool FHeightField::GJKContactPoint(const TImplicitObjectScaled<FConvex>& QueryGeom, const FRigidTransform3& QueryTM, const FReal Thickness, FVec3& ContactLocation, FVec3& ContactNormal, FReal& ContactPhi, int32& ContactFaceIndex) const
 	{
-		return GJKContactPointImp(QueryGeom, QueryTM, Thickness, ContactLocation, ContactNormal, ContactPhi);
+		return GJKContactPointImp(QueryGeom, QueryTM, Thickness, ContactLocation, ContactNormal, ContactPhi, ContactFaceIndex);
 	}
 
 	template <typename QueryGeomType>
diff --git a/Engine/Source/Runtime/Experimental/Chaos/Private/Chaos/TriangleMeshImplicitObject.cpp b/Engine/Source/Runtime/Experimental/Chaos/Private/Chaos/TriangleMeshImplicitObject.cpp
--- a/Engine/Source/Runtime/Experimental/Chaos/Private/Chaos/TriangleMeshImplicitObject.cpp
+++ b/Engine/Source/Runtime/Experimental/Chaos/Private/Chaos/TriangleMeshImplicitObject.cpp
@@ -460,7 +460,8 @@ FReal FTriangleMeshImplicitObject::PhiWithNormal(const FVec3& x, FVec3& Normal)
 	FRigidTransform3 TestXf(FVec3(0.0), FRotation3::FromIdentity());
 	FVec3 TestLocation = x;
 	FReal Depth = TNumericLimits<FReal>::Max();
-	GJKContactPointImp(TestSphere, TestXf, 0.0f, TestLocation, Normal, Depth);
+	int32 FaceIndex = INDEX_NONE;
+	GJKContactPointImp(TestSphere, TestXf, 0.0f, TestLocation, Normal, Depth, FaceIndex);
 	return Depth;
 }
 
@@ -506,7 +507,7 @@ bool FTriangleMeshImplicitObject::Raycast(const FVec3& StartPoint, const FVec3&
 }
 
 template <typename QueryGeomType>
-bool FTriangleMeshImplicitObject::GJKContactPointImp(const QueryGeomType& QueryGeom, const FRigidTransform3& QueryTM, const FReal WorldThickness, FVec3& Location, FVec3& Normal, FReal& OutContactPhi, FVec3 TriMeshScale) const
+bool FTriangleMeshImplicitObject::GJKContactPointImp(const QueryGeomType& QueryGeom, const FRigidTransform3& QueryTM, const FReal WorldThickness, FVec3& Location, FVec3& Normal, FReal& OutContactPhi, int32& OutFaceIndex, FVec3 TriMeshScale) const
 {
 	ensure(TriMeshScale != FVec3(0.0f));
 	bool bResult = false;
@@ -565,6 +566,7 @@ bool FTriangleMeshImplicitObject::GJKContactPointImp(const QueryGeomType& QueryG
 					OutContactPhi = LocalContactPhi;
 					Location = LocalContactLocation;
 					Normal = LocalContactNormal;
+					OutFaceIndex = TriIdx;
 				}
 			}
 
@@ -579,44 +581,44 @@ bool FTriangleMeshImplicitObject::GJKContactPointImp(const QueryGeomType& QueryG
 	return LambdaHelper(MElements.GetSmallIndexBuffer());
 }
 
-bool FTriangleMeshImplicitObject::GJKContactPoint(const TSphere<FReal, 3>& QueryGeom, const FRigidTransform3& QueryTM, const FReal WorldThickness, FVec3& Location, FVec3& Normal, FReal& ContactPhi) const
+bool FTriangleMeshImplicitObject::GJKContactPoint(const TSphere<FReal, 3>& QueryGeom, const FRigidTransform3& QueryTM, const FReal WorldThickness, FVec3& Location, FVec3& Normal, FReal& ContactPhi, int32& ContactFaceIndex) const
 {
-	return GJKContactPointImp(QueryGeom, QueryTM, WorldThickness, Location, Normal, ContactPhi);
+	return GJKContactPointImp(QueryGeom, QueryTM, WorldThickness, Location, Normal, ContactPhi, ContactFaceIndex);
 }
 
-bool FTriangleMeshImplicitObject::GJKContactPoint(const TBox<FReal, 3>& QueryGeom, const FRigidTransform3& QueryTM, const FReal WorldThickness, FVec3& Location, FVec3& Normal, FReal& ContactPhi) const
+bool FTriangleMeshImplicitObject::GJKContactPoint(const TBox<FReal, 3>& QueryGeom, const FRigidTransform3& QueryTM, const FReal WorldThickness, FVec3& Location, FVec3& Normal, FReal& ContactPhi, int32& ContactFaceIndex) const
 {
-	return GJKContactPointImp(QueryGeom, QueryTM, WorldThickness, Location, Normal, ContactPhi);
+	return GJKContactPointImp(QueryGeom, QueryTM, WorldThickness, Location, Normal, ContactPhi, ContactFaceIndex);
 }
 
-bool FTriangleMeshImplicitObject::GJKContactPoint(const FCapsule& QueryGeom, const FRigidTransform3& QueryTM, const FReal WorldThickness, FVec3& Location, FVec3& Normal, FReal& ContactPhi) const
+bool FTriangleMeshImplicitObject::GJKContactPoint(const FCapsule& QueryGeom, const FRigidTransform3& QueryTM, const FReal WorldThickness, FVec3& Location, FVec3& Normal, FReal& ContactPhi, int32& ContactFaceIndex) const
 {
-	return GJKContactPointImp(QueryGeom, QueryTM, WorldThickness, Location, Normal, ContactPhi);
+	return GJKContactPointImp(QueryGeom, QueryTM, WorldThickness, Location, Normal, ContactPhi, ContactFaceIndex);
 }
 
-bool FTriangleMeshImplicitObject::GJKContactPoint(const FConvex& QueryGeom, const FRigidTransform3& QueryTM, const FReal WorldThickness, FVec3& Location, FVec3& Normal, FReal& ContactPhi) const
+bool FTriangleMeshImplicitObject::GJKContactPoint(const FConvex& QueryGeom, const FRigidTransform3& QueryTM, const FReal WorldThickness, FVec3& Location, FVec3& Normal, FReal& ContactPhi, int32& ContactFaceIndex) const
 {
-	return GJKContactPointImp(QueryGeom, QueryTM, WorldThickness, Location, Normal, ContactPhi);
+	return GJKContactPointImp(QueryGeom, QueryTM, WorldThickness, Location, Normal, ContactPhi, ContactFaceIndex);
 }
 
-bool FTriangleMeshImplicitObject::GJKContactPoint(const TImplicitObjectScaled< TSphere<FReal, 3> >& QueryGeom, const FRigidTransform3& QueryTM, const FReal WorldThickness, FVec3& Location, FVec3& Normal, FReal& ContactPhi, FVec3 TriMeshScale) const
+bool FTriangleMeshImplicitObject::GJKContactPoint(const TImplicitObjectScaled< TSphere<FReal, 3> >& QueryGeom, const FRigidTransform3& QueryTM, const FReal WorldThickness, FVec3& Location, FVec3& Normal, FReal& ContactPhi, int32& ContactFaceIndex, FVec3 TriMeshScale) const
 {
-	return GJKContactPointImp(QueryGeom, QueryTM, WorldThickness, Location, Normal, ContactPhi, TriMeshScale);
+	return GJKContactPointImp(QueryGeom, QueryTM, WorldThickness, Location, Normal, ContactPhi, ContactFaceIndex, TriMeshScale);
 }
 
-bool FTriangleMeshImplicitObject::GJKContactPoint(const TImplicitObjectScaled< TBox<FReal, 3> >& QueryGeom, const FRigidTransform3& QueryTM, const FReal WorldThickness, FVec3& Location, FVec3& Normal, FReal& ContactPhi, FVec3 TriMeshScale) const
+bool FTriangleMeshImplicitObject::GJKContactPoint(const TImplicitObjectScaled< TBox<FReal, 3> >& QueryGeom, const FRigidTransform3& QueryTM, const FReal WorldThickness, FVec3& Location, FVec3& Normal, FReal& ContactPhi, int32& ContactFaceIndex, FVec3 TriMeshScale) const
 {
-	return GJKContactPointImp(QueryGeom, QueryTM, WorldThickness, Location, Normal, ContactPhi, TriMeshScale);
+	return GJKContactPointImp(QueryGeom, QueryTM, WorldThickness, Location, Normal, ContactPhi, ContactFaceIndex, TriMeshScale);
 }
 
-bool FTriangleMeshImplicitObject::GJKContactPoint(const TImplicitObjectScaled< FCapsule >& QueryGeom, const FRigidTransform3& QueryTM, const FReal WorldThickness, FVec3& Location, FVec3& Normal, FReal& ContactPhi, FVec3 TriMeshScale) const
+bool FTriangleMeshImplicitObject::GJKContactPoint(const TImplicitObjectScaled< FCapsule >& QueryGeom, const FRigidTransform3& QueryTM, const FReal WorldThickness, FVec3& Location, FVec3& Normal, FReal& ContactPhi, int32& ContactFaceIndex, FVec3 TriMeshScale) const
 {
-	return GJKContactPointImp(QueryGeom, QueryTM, WorldThickness, Location, Normal, ContactPhi, TriMeshScale);
+	return GJKContactPointImp(QueryGeom, QueryTM, WorldThickness, Location, Normal, ContactPhi, ContactFaceIndex, TriMeshScale);
 }
 
-bool FTriangleMeshImplicitObject::GJKContactPoint(const TImplicitObjectScaled< FConvex >& QueryGeom, const FRigidTransform3& QueryTM, const FReal WorldThickness, FVec3& Location, FVec3& Normal, FReal& ContactPhi, FVec3 TriMeshScale) const
+bool FTriangleMeshImplicitObject::GJKContactPoint(const TImplicitObjectScaled< FConvex >& QueryGeom, const FRigidTransform3& QueryTM, const FReal WorldThickness, FVec3& Location, FVec3& Normal, FReal& ContactPhi, int32& ContactFaceIndex, FVec3 TriMeshScale) const
 {
-	return GJKContactPointImp(QueryGeom, QueryTM, WorldThickness, Location, Normal, ContactPhi, TriMeshScale);
+	return GJKContactPointImp(QueryGeom, QueryTM, WorldThickness, Location, Normal, ContactPhi, ContactFaceIndex, TriMeshScale);
 }
 
 int32 FTriangleMeshImplicitObject::GetExternalFaceIndexFromInternal(int32 InternalFaceIndex) const
diff --git a/Engine/Source/Runtime/Experimental/Chaos/Public/Chaos/HeightField.h b/Engine/Source/Runtime/Experimental/Chaos/Public/Chaos/HeightField.h
--- a/Engine/Source/Runtime/Experimental/Chaos/Public/Chaos/HeightField.h
+++ b/Engine/Source/Runtime/Experimental/Chaos/Public/Chaos/HeightField.h
@@ -84,14 +84,14 @@ namespace Chaos
 		bool SweepGeomCCD(const TImplicitObjectScaled<FCapsule>& QueryGeom, const FRigidTransform3& StartTM, const FVec3& Dir, const FReal Length, const FReal IgnorePenetration, const FReal TargetPenetration, FReal& OutTOI, FReal& OutPhi, FVec3& OutPosition, FVec3& OutNormal, int32& OutFaceIndex, FVec3& OutFaceNormal) const;
 		bool SweepGeomCCD(const TImplicitObjectScaled<FConvex>& QueryGeom, const FRigidTransform3& StartTM, const FVec3& Dir, const FReal Length, const FReal IgnorePenetration, const FReal TargetPenetration, FReal& OutTOI, FReal& OutPhi, FVec3& OutPosition, FVec3& OutNormal, int32& OutFaceIndex, FVec3& OutFaceNormal) const;
 
-		bool GJKContactPoint(const TBox<FReal, 3>& QueryGeom, const FRigidTransform3& QueryTM, const FReal Thickness, FVec3& ContactLocation, FVec3& ContactNormal, FReal& ContactPhi) const;
-		bool GJKContactPoint(const TSphere<FReal, 3>& QueryGeom, const FRigidTransform3& QueryTM, const FReal Thickness, FVec3& ContactLocation, FVec3& ContactNormal, FReal& ContactPhi) const;
-		bool GJKContactPoint(const FCapsule& QueryGeom, const FRigidTransform3& QueryTM, const FReal Thickness, FVec3& ContactLocation, FVec3& ContactNormal, FReal& ContactPhi) const;
-		bool GJKContactPoint(const FConvex& QueryGeom, const FRigidTransform3& QueryTM, const FReal Thickness, FVec3& ContactLocation, FVec3& ContactNormal, FReal& ContactPhi) const;
-		bool GJKContactPoint(const TImplicitObjectScaled<TBox<FReal, 3>>& QueryGeom, const FRigidTransform3& QueryTM, const FReal Thickness, FVec3& ContactLocation, FVec3& ContactNormal, FReal& ContactPhi) const;
-		bool GJKContactPoint(const TImplicitObjectScaled<TSphere<FReal, 3>>& QueryGeom, const FRigidTransform3& QueryTM, const FReal Thickness, FVec3& ContactLocation, FVec3& ContactNormal, FReal& ContactPhi) const;
-		bool GJKContactPoint(const TImplicitObjectScaled<FCapsule>& QueryGeom, const FRigidTransform3& QueryTM, const FReal Thickness, FVec3& ContactLocation, FVec3& ContactNormal, FReal& ContactPhi) const;
-		bool GJKContactPoint(const TImplicitObjectScaled<FConvex>& QueryGeom, const FRigidTransform3& QueryTM, const FReal Thickness, FVec3& ContactLocation, FVec3& ContactNormal, FReal& ContactPhi) const;
+		bool GJKContactPoint(const TBox<FReal, 3>& QueryGeom, const FRigidTransform3& QueryTM, const FReal Thickness, FVec3& ContactLocation, FVec3& ContactNormal, FReal& ContactPhi, int32& ContactFaceIndex) const;
+		bool GJKContactPoint(const TSphere<FReal, 3>& QueryGeom, const FRigidTransform3& QueryTM, const FReal Thickness, FVec3& ContactLocation, FVec3& ContactNormal, FReal& ContactPhi, int32& ContactFaceIndex) const;
+		bool GJKContactPoint(const FCapsule& QueryGeom, const FRigidTransform3& QueryTM, const FReal Thickness, FVec3& ContactLocation, FVec3& ContactNormal, FReal& ContactPhi, int32& ContactFaceIndex) const;
+		bool GJKContactPoint(const FConvex& QueryGeom, const FRigidTransform3& QueryTM, const FReal Thickness, FVec3& ContactLocation, FVec3& ContactNormal, FReal& ContactPhi, int32& ContactFaceIndex) const;
+		bool GJKContactPoint(const TImplicitObjectScaled<TBox<FReal, 3>>& QueryGeom, const FRigidTransform3& QueryTM, const FReal Thickness, FVec3& ContactLocation, FVec3& ContactNormal, FReal& ContactPhi, int32& ContactFaceIndex) const;
+		bool GJKContactPoint(const TImplicitObjectScaled<TSphere<FReal, 3>>& QueryGeom, const FRigidTransform3& QueryTM, const FReal Thickness, FVec3& ContactLocation, FVec3& ContactNormal, FReal& ContactPhi, int32& ContactFaceIndex) const;
+		bool GJKContactPoint(const TImplicitObjectScaled<FCapsule>& QueryGeom, const FRigidTransform3& QueryTM, const FReal Thickness, FVec3& ContactLocation, FVec3& ContactNormal, FReal& ContactPhi, int32& ContactFaceIndex) const;
+		bool GJKContactPoint(const TImplicitObjectScaled<FConvex>& QueryGeom, const FRigidTransform3& QueryTM, const FReal Thickness, FVec3& ContactLocation, FVec3& ContactNormal, FReal& ContactPhi, int32& ContactFaceIndex) const;
 
 
 		virtual int32 FindMostOpposingFace(const FVec3& Position, const FVec3& UnitDir, int32 HintFaceIndex, FReal SearchDist) const override;
@@ -821,7 +821,7 @@ namespace Chaos
 		bool SweepGeomCCDImp(const QueryGeomType& QueryGeom, const FRigidTransform3& StartTM, const FVec3& Dir, const FReal Length, const FReal IgnorePenetration, const FReal TargetPenetration, FReal& OutTOI, FReal& OutPhi, FVec3& OutPosition, FVec3& OutNormal, int32& OutFaceIndex) const;
 
 		template <typename GeomType>
-		bool GJKContactPointImp(const GeomType& QueryGeom, const FRigidTransform3& QueryTM, const FReal Thickness, FVec3& ContactLocation, FVec3& ContactNormal, FReal& ContactPhi) const;
+		bool GJKContactPointImp(const GeomType& QueryGeom, const FRigidTransform3& QueryTM, const FReal Thickness, FVec3& ContactLocation, FVec3& ContactNormal, FReal& ContactPhi, int32& ContactFaceIndex) const;
 	};
 
 }
diff --git a/Engine/Source/Runtime/Experimental/Chaos/Public/Chaos/ImplicitObjectScaled.h b/Engine/Source/Runtime/Experimental/Chaos/Public/Chaos/ImplicitObjectScaled.h
--- a/Engine/Source/Runtime/Experimental/Chaos/Public/Chaos/ImplicitObjectScaled.h
+++ b/Engine/Source/Runtime/Experimental/Chaos/Public/Chaos/ImplicitObjectScaled.h
@@ -236,9 +236,9 @@ public:
 	}
 
 	template <typename QueryGeomType>
-	bool GJKContactPoint(const QueryGeomType& A, const FRigidTransform3& AToBTM, const FReal Thickness, FVec3& Location, FVec3& Normal, FReal& Penetration) const
+	bool GJKContactPoint(const QueryGeomType& A, const FRigidTransform3& AToBTM, const FReal Thickness, FVec3& Location, FVec3& Normal, FReal& Penetration, int32& FaceIndex) const
 	{
-		return MObject->GJKContactPoint(A, AToBTM, Thickness, Location, Normal, Penetration);
+		return MObject->GJKContactPoint(A, AToBTM, Thickness, Location, Normal, Penetration, FaceIndex);
 	}
 
 	virtual uint16 GetMaterialIndex(uint32 HintIndex) const
@@ -709,7 +709,7 @@ public:
 	}
 
 	template <typename QueryGeomType>
-	bool GJKContactPoint(const QueryGeomType& A, const FRigidTransform3& AToBTM, const FReal Thickness, FVec3& Location, FVec3& Normal, FReal& Penetration) const
+	bool GJKContactPoint(const QueryGeomType& A, const FRigidTransform3& AToBTM, const FReal Thickness, FVec3& Location, FVec3& Normal, FReal& Penetration, int32& FaceIndex) const
 	{
 		TRigidTransform<T, d> AToBTMNoScale(AToBTM.GetLocation() * MInvScale, AToBTM.GetRotation());
 
@@ -719,7 +719,7 @@ public:
 		const FReal UnscaledThickness = Thickness;
 
 		auto ScaledA = MakeScaledHelper(A, MInvScale);
-		return MObject->GJKContactPoint(ScaledA, AToBTMNoScale, UnscaledThickness, Location, Normal, Penetration, MScale);
+		return MObject->GJKContactPoint(ScaledA, AToBTMNoScale, UnscaledThickness, Location, Normal, Penetration, FaceIndex, MScale);
 	}
 
 	/** This is a low level function and assumes the internal object has a OverlapGeom function. Should not be called directly. See GeometryQueries.h : OverlapQuery */
diff --git a/Engine/Source/Runtime/Experimental/Chaos/Public/Chaos/TriangleMeshImplicitObject.h b/Engine/Source/Runtime/Experimental/Chaos/Public/Chaos/TriangleMeshImplicitObject.h
--- a/Engine/Source/Runtime/Experimental/Chaos/Public/Chaos/TriangleMeshImplicitObject.h
+++ b/Engine/Source/Runtime/Experimental/Chaos/Public/Chaos/TriangleMeshImplicitObject.h
@@ -544,15 +544,15 @@ namespace Chaos
 		bool SweepGeomCCD(const TImplicitObjectScaled<FCapsule>& QueryGeom, const FRigidTransform3& StartTM, const FVec3& Dir, const FReal Length, const FReal IgnorePenetration, const FReal TargetPenetration, FReal& OutTOI, FReal& OutPhi, FVec3& OutPosition, FVec3& OutNormal, int32& OutFaceIndex, FVec3& OutFaceNormal, const FVec3& TriMeshScale = FVec3(1.0f)) const;
 		bool SweepGeomCCD(const TImplicitObjectScaled<FConvex>& QueryGeom, const FRigidTransform3& StartTM, const FVec3& Dir, const FReal Length, const FReal IgnorePenetration, const FReal TargetPenetration, FReal& OutTOI, FReal& OutPhi, FVec3& OutPosition, FVec3& OutNormal, int32& OutFaceIndex, FVec3& OutFaceNormal, const FVec3& TriMeshScale = FVec3(1.0f)) const;
 
-		bool GJKContactPoint(const TSphere<FReal, 3>& QueryGeom, const FRigidTransform3& QueryTM, const FReal Thickness, FVec3& Location, FVec3& Normal, FReal& Penetration) const;
-		bool GJKContactPoint(const TBox<FReal, 3>& QueryGeom, const FRigidTransform3& QueryTM, const FReal Thickness, FVec3& Location, FVec3& Normal, FReal& Penetration) const;
-		bool GJKContactPoint(const FCapsule& QueryGeom, const FRigidTransform3& QueryTM, const FReal Thickness, FVec3& Location, FVec3& Normal, FReal& Penetration) const;
-		bool GJKContactPoint(const FConvex& QueryGeom, const FRigidTransform3& QueryTM, const FReal Thickness, FVec3& Location, FVec3& Normal, FReal& Penetration) const;
+		bool GJKContactPoint(const TSphere<FReal, 3>& QueryGeom, const FRigidTransform3& QueryTM, const FReal Thickness, FVec3& Location, FVec3& Normal, FReal& Penetration, int32& FaceIndex) const;
+		bool GJKContactPoint(const TBox<FReal, 3>& QueryGeom, const FRigidTransform3& QueryTM, const FReal Thickness, FVec3& Location, FVec3& Normal, FReal& Penetration, int32& FaceIndex) const;
+		bool GJKContactPoint(const FCapsule& QueryGeom, const FRigidTransform3& QueryTM, const FReal Thickness, FVec3& Location, FVec3& Normal, FReal& Penetration, int32& FaceIndex) const;
+		bool GJKContactPoint(const FConvex& QueryGeom, const FRigidTransform3& QueryTM, const FReal Thickness, FVec3& Location, FVec3& Normal, FReal& Penetration, int32& FaceIndex) const;
 
-		bool GJKContactPoint(const TImplicitObjectScaled < TSphere<FReal, 3> >& QueryGeom, const FRigidTransform3& QueryTM, const FReal Thickness, FVec3& Location, FVec3& Normal, FReal& Penetration, FVec3 TriMeshScale = FVec3(1.0f)) const;
-		bool GJKContactPoint(const TImplicitObjectScaled < TBox<FReal, 3> >& QueryGeom, const FRigidTransform3& QueryTM, const FReal Thickness, FVec3& Location, FVec3& Normal, FReal& Penetration, FVec3 TriMeshScale = FVec3(1.0f)) const;
-		bool GJKContactPoint(const TImplicitObjectScaled < FCapsule >& QueryGeom, const FRigidTransform3& QueryTM, const FReal Thickness, FVec3& Location, FVec3& Normal, FReal& Penetration, FVec3 TriMeshScale = FVec3(1.0f)) const;
-		bool GJKContactPoint(const TImplicitObjectScaled < FConvex >& QueryGeom, const FRigidTransform3& QueryTM, const FReal Thickness, FVec3& Location, FVec3& Normal, FReal& Penetration, FVec3 TriMeshScale = FVec3(1.0f)) const;
+		bool GJKContactPoint(const TImplicitObjectScaled < TSphere<FReal, 3> >& QueryGeom, const FRigidTransform3& QueryTM, const FReal Thickness, FVec3& Location, FVec3& Normal, FReal& Penetration, int32& FaceIndex, FVec3 TriMeshScale = FVec3(1.0f)) const;
+		bool GJKContactPoint(const TImplicitObjectScaled < TBox<FReal, 3> >& QueryGeom, const FRigidTransform3& QueryTM, const FReal Thickness, FVec3& Location, FVec3& Normal, FReal& Penetration, int32& FaceIndex, FVec3 TriMeshScale = FVec3(1.0f)) const;
+		bool GJKContactPoint(const TImplicitObjectScaled < FCapsule >& QueryGeom, const FRigidTransform3& QueryTM, const FReal Thickness, FVec3& Location, FVec3& Normal, FReal& Penetration, int32& FaceIndex, FVec3 TriMeshScale = FVec3(1.0f)) const;
+		bool GJKContactPoint(const TImplicitObjectScaled < FConvex >& QueryGeom, const FRigidTransform3& QueryTM, const FReal Thickness, FVec3& Location, FVec3& Normal, FReal& Penetration, int32& FaceIndex, FVec3 TriMeshScale = FVec3(1.0f)) const;
 
 		// Returns -1 if InternalFaceIndex is not in map, or map is invalid.
 		int32 GetExternalFaceIndexFromInternal(int32 InternalFaceIndex) const;
@@ -854,7 +854,7 @@ namespace Chaos
 		    : FImplicitObject(EImplicitObject::HasBoundingBox, ImplicitObjectType::TriangleMesh){};
 
 		template <typename QueryGeomType>
-		bool GJKContactPointImp(const QueryGeomType& QueryGeom, const FRigidTransform3& QueryTM, const FReal Thickness, FVec3& Location, FVec3& Normal, FReal& Penetration, FVec3 TriMeshScale = FVec3(1.0)) const;
+		bool GJKContactPointImp(const QueryGeomType& QueryGeom, const FRigidTransform3& QueryTM, const FReal Thickness, FVec3& Location, FVec3& Normal, FReal& Penetration, int32& FaceIndex, FVec3 TriMeshScale = FVec3(1.0)) const;
 
 		template<typename QueryGeomType>
 		bool OverlapGeomImp(const QueryGeomType& QueryGeom, const FRigidTransform3& QueryTM, const FReal Thickness, FMTDInfo* OutMTD = nullptr, FVec3 TriMeshScale = FVec3(1.0f)) const;
