Dataflow - Inactive nodes can now setup to be passthrough nodes

diff --git a/Engine/Plugins/Experimental/GeometryCollectionPlugin/Source/GeometryCollectionNodes/Public/Dataflow/GeometryCollectionNodes.h b/Engine/Plugins/Experimental/GeometryCollectionPlugin/Source/GeometryCollectionNodes/Public/Dataflow/GeometryCollectionNodes.h
--- a/Engine/Plugins/Experimental/GeometryCollectionPlugin/Source/GeometryCollectionNodes/Public/Dataflow/GeometryCollectionNodes.h
+++ b/Engine/Plugins/Experimental/GeometryCollectionPlugin/Source/GeometryCollectionNodes/Public/Dataflow/GeometryCollectionNodes.h
@@ -49,14 +49,14 @@ public:
 	UPROPERTY(EditAnywhere, Category = "Dataflow", meta = (EditCondition = "false", EditConditionHides));
 	float Scale = 1.0;
 
-	UPROPERTY(meta = (DataflowInput, DataflowOutput, DisplayName = "Collection"))
+	UPROPERTY(meta = (DataflowInput, DataflowOutput, DisplayName = "Collection", Passthrough = "Collection"))
 	FManagedArrayCollection Collection;
 
 	FExampleCollectionEditDataflowNode(const Dataflow::FNodeParameters& InParam, FGuid InGuid = FGuid::NewGuid())
 		: FDataflowNode(InParam, InGuid)
 	{
+		RegisterOutputConnection(&Collection, &Collection);
 		RegisterInputConnection(&Collection);
-		RegisterOutputConnection(&Collection);
 	}
 
 	virtual void Evaluate(Dataflow::FContext& Context, const FDataflowOutput* Out) const override;
diff --git a/Engine/Source/Runtime/Experimental/Dataflow/Core/Private/Dataflow/DataflowConnection.cpp b/Engine/Source/Runtime/Experimental/Dataflow/Core/Private/Dataflow/DataflowConnection.cpp
--- a/Engine/Source/Runtime/Experimental/Dataflow/Core/Private/Dataflow/DataflowConnection.cpp
+++ b/Engine/Source/Runtime/Experimental/Dataflow/Core/Private/Dataflow/DataflowConnection.cpp
@@ -14,7 +14,7 @@ FDataflowConnection::FDataflowConnection(Dataflow::FPin::EDirection InDirection,
 	, Guid(InGuid)
 {}
 
-
+// @todo(dataflow) check that the return types are all consistent
 uint32 FDataflowConnection::GetOffset() const
 {
 	if (ensure(Property != nullptr))
diff --git a/Engine/Source/Runtime/Experimental/Dataflow/Core/Private/Dataflow/DataflowGraph.cpp b/Engine/Source/Runtime/Experimental/Dataflow/Core/Private/Dataflow/DataflowGraph.cpp
--- a/Engine/Source/Runtime/Experimental/Dataflow/Core/Private/Dataflow/DataflowGraph.cpp
+++ b/Engine/Source/Runtime/Experimental/Dataflow/Core/Private/Dataflow/DataflowGraph.cpp
@@ -21,29 +21,29 @@ namespace Dataflow
 
 	void FGraph::RemoveNode(TSharedPtr<FDataflowNode> Node)
 	{
-		for (FDataflowConnection* Output : Node->GetOutputs())
+		for (FDataflowOutput* Output : Node->GetOutputs())
 		{
 			if (Output)
 			{
-				for (FDataflowConnection* Input : Output->GetConnectedInputs())
+				for (FDataflowInput* Input : Output->GetConnectedInputs())
 				{
 					if (Input)
 					{
-						Disconnect(Output, Input);
+						Disconnect(Input, Output);
 					}
 				}
 			}
 		}
-		for (FDataflowConnection* Input : Node->GetInputs())
+		for (FDataflowInput* Input : Node->GetInputs())
 		{
 			if (Input)
 			{
-				TArray<FDataflowConnection*> Outputs = Input->GetConnectedOutputs();
-				for (FDataflowConnection* Out : Outputs)
+				TArray<FDataflowOutput*> Outputs = Input->GetConnectedOutputs();
+				for (FDataflowOutput* Output : Outputs)
 				{
-					if (Out)
+					if (Output)
 					{
-						Disconnect(Out, Input);
+						Disconnect(Input, Output);
 					}
 				}
 			}
@@ -57,42 +57,60 @@ namespace Dataflow
 		//            or implement a better sync for the EdGraph and DataflowGraph
 		if (Connection->GetDirection() == FPin::EDirection::INPUT)
 		{
-			TArray<FDataflowConnection*> BaseOutputs = Connection->GetConnectedOutputs();
-			for (FDataflowConnection* Output : BaseOutputs)
+			FDataflowInput* ConnectionIn = reinterpret_cast<FDataflowInput*>(Connection);
+			TArray<FDataflowOutput*> BaseOutputs = ConnectionIn->GetConnectedOutputs();
+			for (FDataflowOutput* Output : BaseOutputs)
 			{
-				Disconnect(Connection, Output);
+				Disconnect(ConnectionIn, Output);
 			}
 		}
 		else if (Connection->GetDirection() == FPin::EDirection::OUTPUT)
 		{
-			TArray<FDataflowConnection*> BaseInputs = Connection->GetConnectedInputs();
-			for (FDataflowConnection* Input : BaseInputs)
+			FDataflowOutput* ConnectionOut = reinterpret_cast<FDataflowOutput*>(Connection);
+			TArray<FDataflowInput*> BaseInputs = ConnectionOut->GetConnectedInputs();
+			for (FDataflowInput* Input : BaseInputs)
 			{
-				Disconnect(Input, Connection);
+				Disconnect(Input, ConnectionOut);
 			}
 		}
 	}
 
+	void FGraph::ClearConnections(FDataflowInput* InConnection)
+	{
+		for (FDataflowOutput* Output : InConnection->GetConnectedOutputs())
+		{
+			Disconnect(InConnection, Output);
+		}
+	}
 
-	void FGraph::Connect(FDataflowConnection* Input, FDataflowConnection* Output)
+	void FGraph::ClearConnections(FDataflowOutput* OutConnection)
 	{
-		if (ensure(Input && Output))
+		for (FDataflowInput* Input : OutConnection->GetConnectedInputs())
 		{
-			Input->AddConnection(Output);
-			Output->AddConnection(Input);
+			Disconnect(Input, OutConnection);
+		}
+	}
+
+
+	void FGraph::Connect(FDataflowInput* InputConnection, FDataflowOutput* OutputConnection)
+	{
+		if (ensure(InputConnection && OutputConnection))
+		{
+			InputConnection->AddConnection(OutputConnection);
+			OutputConnection->AddConnection(InputConnection);
 			Connections.Add(FLink(
-				Input->GetOwningNode()->GetGuid(), Input->GetGuid(), 
-				Output->GetOwningNode()->GetGuid(), Output->GetGuid()));
+				InputConnection->GetOwningNode()->GetGuid(), InputConnection->GetGuid(), 
+				OutputConnection->GetOwningNode()->GetGuid(), OutputConnection->GetGuid()));
 		}
 	}
 
-	void FGraph::Disconnect(FDataflowConnection* Input, FDataflowConnection* Output)
+	void FGraph::Disconnect(FDataflowInput* InputConnection, FDataflowOutput* OutputConnection)
 	{
-		Input->RemoveConnection(Output);
-		Output->RemoveConnection(Input);
+		InputConnection->RemoveConnection(OutputConnection);
+		OutputConnection->RemoveConnection(InputConnection);
 		Connections.RemoveSwap(FLink(
-			Input->GetOwningNode()->GetGuid(), Input->GetGuid(),
-			Output->GetOwningNode()->GetGuid(), Output->GetGuid()));
+			InputConnection->GetOwningNode()->GetGuid(), InputConnection->GetGuid(),
+			OutputConnection->GetOwningNode()->GetGuid(), OutputConnection->GetGuid()));
 
 	}
 
@@ -115,7 +133,8 @@ namespace Dataflow
 
 				DATAFLOW_OPTIONAL_BLOCK_WRITE_BEGIN()
 				{
-					TArray< FDataflowConnection* > IO = Node->GetOutputs();
+					TArray< FDataflowConnection* > IO;
+					IO.Append(Node->GetOutputs());
 					IO.Append(Node->GetInputs());
 					ArNum = IO.Num();
 					Ar << ArNum;
@@ -157,7 +176,9 @@ namespace Dataflow
 
 					int ArNumInner;
 					Ar << ArNumInner;
-					TArray< FDataflowConnection* > IO = Node->GetOutputs();  IO.Append(Node->GetInputs());
+					TArray< FDataflowConnection* > IO;
+					IO.Append(Node->GetOutputs());
+					IO.Append(Node->GetInputs());
 					for (int Cdx = 0; Cdx < ArNumInner; Cdx++)
 					{
 						Ar << ArGuid << ArType << ArName;
@@ -191,7 +212,7 @@ namespace Dataflow
 					{
 						if (ConnectionGuidMap[Con.Input]->GetType() == ConnectionGuidMap[Con.Output]->GetType())
 						{
-							Connect(ConnectionGuidMap[Con.Input], ConnectionGuidMap[Con.Output]);
+							Connect(reinterpret_cast<FDataflowInput*>(ConnectionGuidMap[Con.Input]), reinterpret_cast<FDataflowOutput*>(ConnectionGuidMap[Con.Output]));
 						}
 					}
 				}
diff --git a/Engine/Source/Runtime/Experimental/Dataflow/Core/Private/Dataflow/DataflowInputOutput.cpp b/Engine/Source/Runtime/Experimental/Dataflow/Core/Private/Dataflow/DataflowInputOutput.cpp
--- a/Engine/Source/Runtime/Experimental/Dataflow/Core/Private/Dataflow/DataflowInputOutput.cpp
+++ b/Engine/Source/Runtime/Experimental/Dataflow/Core/Private/Dataflow/DataflowInputOutput.cpp
@@ -40,20 +40,20 @@ bool FDataflowInput::RemoveConnection(FDataflowConnection* InOutput)
 	return false;
 }
 
-TArray< FDataflowConnection* > FDataflowInput::GetConnectedOutputs()
+TArray< FDataflowOutput* > FDataflowInput::GetConnectedOutputs()
 {
-	TArray<FDataflowConnection* > RetList;
-	if (FDataflowConnection* Conn = GetConnection())
+	TArray<FDataflowOutput* > RetList;
+	if (FDataflowOutput* Conn = GetConnection())
 	{
 		RetList.Add(Conn);
 	}
 	return RetList;
 }
 
-const TArray< const FDataflowConnection* > FDataflowInput::GetConnectedOutputs() const
+const TArray< const FDataflowOutput* > FDataflowInput::GetConnectedOutputs() const
 {
-	TArray<const FDataflowConnection* > RetList;
-	if (const FDataflowConnection* Conn = GetConnection())
+	TArray<const FDataflowOutput* > RetList;
+	if (const FDataflowOutput* Conn = GetConnection())
 	{
 		RetList.Add(Conn);
 	}
@@ -81,9 +81,9 @@ FDataflowOutput::FDataflowOutput(const Dataflow::FOutputParameters& Param, FGuid
 const TArray<FDataflowInput*>& FDataflowOutput::GetConnections() const { return Connections; }
 TArray<FDataflowInput*>& FDataflowOutput::GetConnections() { return Connections; }
 
-const TArray< const FDataflowConnection*> FDataflowOutput::GetConnectedInputs() const
+const TArray< const FDataflowInput*> FDataflowOutput::GetConnectedInputs() const
 {
-	TArray<const FDataflowConnection*> RetList;
+	TArray<const FDataflowInput*> RetList;
 	RetList.Reserve(Connections.Num());
 	for (FDataflowInput* Ptr : Connections) 
 	{ 
@@ -92,9 +92,9 @@ const TArray< const FDataflowConnection*> FDataflowOutput::GetConnectedInputs()
 	return RetList;
 }
 
-TArray< FDataflowConnection*> FDataflowOutput::GetConnectedInputs()
+TArray< FDataflowInput*> FDataflowOutput::GetConnectedInputs()
 {
-	TArray<FDataflowConnection*> RetList;
+	TArray<FDataflowInput*> RetList;
 	RetList.Reserve(Connections.Num());
 	for (FDataflowInput* Ptr : Connections) 
 	{ 
@@ -118,21 +118,6 @@ bool FDataflowOutput::RemoveConnection(FDataflowConnection* InInput)
 	Connections.RemoveSwap((FDataflowInput*)InInput); return true;
 }
 
-bool FDataflowOutput::Evaluate(Dataflow::FContext& Context) const
-{
-	check(OwningNode);
-
-	if (OwningNode->bActive)
-	{
-		OwningNode->Evaluate(Context, this);
-		if (!Context.HasData(CacheKey()))
-		{
-			ensureMsgf(false, TEXT("Failed to evaluate output (%s:%s)"), *OwningNode->GetName().ToString(), *GetName().ToString());
-		}
-	}
-	return OwningNode->bActive;
-}
-
 void FDataflowOutput::Invalidate()
 {
 	if (CacheKeyValue != UINT_MAX)
diff --git a/Engine/Source/Runtime/Experimental/Dataflow/Core/Private/Dataflow/DataflowNode.cpp b/Engine/Source/Runtime/Experimental/Dataflow/Core/Private/Dataflow/DataflowNode.cpp
--- a/Engine/Source/Runtime/Experimental/Dataflow/Core/Private/Dataflow/DataflowNode.cpp
+++ b/Engine/Source/Runtime/Experimental/Dataflow/Core/Private/Dataflow/DataflowNode.cpp
@@ -12,13 +12,13 @@
 // Inputs
 //
 
-void FDataflowNode::AddInput(FDataflowConnection* InPtr)
+void FDataflowNode::AddInput(FDataflowInput* InPtr)
 {
 	if (InPtr)
 	{
-		for (TPair<uint32, FDataflowConnection*> Elem : Inputs)
+		for (TPair<uint32, FDataflowInput*> Elem : Inputs)
 		{
-			FDataflowConnection* In = Elem.Value;
+			FDataflowInput* In = Elem.Value;
 			ensureMsgf(!In->GetName().IsEqual(InPtr->GetName()), TEXT("Add Input Failed: Existing Node input already defined with name (%s)"), *InPtr->GetName().ToString());
 		}
 
@@ -35,12 +35,12 @@ void FDataflowNode::AddInput(FDataflowConnection* InPtr)
 
 FDataflowInput* FDataflowNode::FindInput(FName InName)
 {
-	for (TPair<uint32, FDataflowConnection*> Elem : Inputs)
+	for (TPair<uint32, FDataflowInput*> Elem : Inputs)
 	{
-		FDataflowConnection* Con = Elem.Value;
+		FDataflowInput* Con = Elem.Value;
 		if (Con->GetName().IsEqual(InName))
 		{
-			return (FDataflowInput*)Con;
+			return Con;
 		}
 	}
 	return nullptr;
@@ -49,12 +49,12 @@ FDataflowInput* FDataflowNode::FindInput(FName InName)
 
 const FDataflowInput* FDataflowNode::FindInput(const void* Reference) const
 {
-	for (TPair<uint32, FDataflowConnection*> Elem : Inputs)
+	for (TPair<uint32, FDataflowInput*> Elem : Inputs)
 	{
-		FDataflowConnection* Con = Elem.Value;
+		FDataflowInput* Con = Elem.Value;
 		if (Con->RealAddress() == (size_t)Reference)
 		{
-			return (FDataflowInput*)Con;
+			return Con;
 		}
 	}
 	return nullptr;
@@ -62,29 +62,29 @@ const FDataflowInput* FDataflowNode::FindInput(const void* Reference) const
 
 FDataflowInput* FDataflowNode::FindInput(void* Reference)
 {
-	for (TPair<uint32, FDataflowConnection*> Elem : Inputs)
+	for (TPair<uint32, FDataflowInput*> Elem : Inputs)
 	{
-		FDataflowConnection* Con = Elem.Value;
+		FDataflowInput* Con = Elem.Value;
 		if (Con->RealAddress() == (size_t)Reference)
 		{
-			return (FDataflowInput*)Con;
+			return Con;
 		}
 	}
 	return nullptr;
 }
 
-TArray< FDataflowConnection* > FDataflowNode::GetInputs() const
+TArray< FDataflowInput* > FDataflowNode::GetInputs() const
 {
-	TArray< FDataflowConnection* > Result;
+	TArray< FDataflowInput* > Result;
 	Inputs.GenerateValueArray(Result);
 	return Result;
 }
 
 void FDataflowNode::ClearInputs()
 {
-	for (TPair<uint32, FDataflowConnection*> Elem : Inputs)
+	for (TPair<uint32, FDataflowInput*> Elem : Inputs)
 	{
-		FDataflowConnection* Con = Elem.Value;
+		FDataflowInput* Con = Elem.Value;
 		delete Con;
 	}
 	Inputs.Reset();
@@ -96,13 +96,13 @@ void FDataflowNode::ClearInputs()
 //
 
 
-void FDataflowNode::AddOutput(FDataflowConnection* InPtr)
+void FDataflowNode::AddOutput(FDataflowOutput* InPtr)
 {
 	if (InPtr)
 	{
-		for (TPair<uint32, FDataflowConnection*> Elem : Outputs)
+		for (TPair<uint32, FDataflowOutput*> Elem : Outputs)
 		{
-			FDataflowConnection* Out = Elem.Value;
+			FDataflowOutput* Out = Elem.Value;
 			ensureMsgf(!Out->GetName().IsEqual(InPtr->GetName()), TEXT("Add Output Failed: Existing Node output already defined with name (%s)"), *InPtr->GetName().ToString());
 		}
 
@@ -121,9 +121,9 @@ void FDataflowNode::AddOutput(FDataflowConnection* InPtr)
 
 FDataflowOutput* FDataflowNode::FindOutput(FName InName)
 {
-	for (TPair<uint32, FDataflowConnection*> Elem : Outputs)
+	for (TPair<uint32, FDataflowOutput*> Elem : Outputs)
 	{
-		FDataflowConnection* Con = Elem.Value;
+		FDataflowOutput* Con = Elem.Value;
 		if (Con->GetName().IsEqual(InName))
 		{
 			return (FDataflowOutput*)Con;
@@ -134,9 +134,9 @@ FDataflowOutput* FDataflowNode::FindOutput(FName InName)
 
 const FDataflowOutput* FDataflowNode::FindOutput(const void* Reference) const
 {
-	for (TPair<uint32, FDataflowConnection*> Elem : Outputs)
+	for (TPair<uint32, FDataflowOutput*> Elem : Outputs)
 	{
-		FDataflowConnection* Con = Elem.Value;
+		FDataflowOutput* Con = Elem.Value;
 		if (Con->RealAddress() == (size_t)Reference)
 		{
 			return (FDataflowOutput*)Con;
@@ -147,9 +147,9 @@ const FDataflowOutput* FDataflowNode::FindOutput(const void* Reference) const
 
 FDataflowOutput* FDataflowNode::FindOutput(void* Reference)
 {
-	for (TPair<uint32, FDataflowConnection*> Elem : Outputs)
+	for (TPair<uint32, FDataflowOutput*> Elem : Outputs)
 	{
-		FDataflowConnection* Con = Elem.Value;
+		FDataflowOutput* Con = Elem.Value;
 		if (Con->RealAddress() == (size_t)Reference)
 		{
 			return (FDataflowOutput*)Con;
@@ -158,9 +158,9 @@ FDataflowOutput* FDataflowNode::FindOutput(void* Reference)
 	return nullptr;
 }
 
-TArray< FDataflowConnection* > FDataflowNode::GetOutputs() const
+TArray< FDataflowOutput* > FDataflowNode::GetOutputs() const
 {
-	TArray< FDataflowConnection* > Result;
+	TArray< FDataflowOutput* > Result;
 	Outputs.GenerateValueArray(Result);
 	return Result;
 }
@@ -168,9 +168,9 @@ TArray< FDataflowConnection* > FDataflowNode::GetOutputs() const
 
 void FDataflowNode::ClearOutputs()
 {
-	for (TPair<uint32, FDataflowConnection*> Elem : Outputs)
+	for (TPair<uint32, FDataflowOutput*> Elem : Outputs)
 	{
-		FDataflowConnection* Con = Elem.Value;
+		FDataflowOutput* Con = Elem.Value;
 		delete Con;
 	}
 	Outputs.Reset();
@@ -180,14 +180,14 @@ void FDataflowNode::ClearOutputs()
 TArray<Dataflow::FPin> FDataflowNode::GetPins() const
 {
 	TArray<Dataflow::FPin> RetVal;
-	for (TPair<uint32, FDataflowConnection*> Elem : Inputs)
+	for (TPair<uint32, FDataflowInput*> Elem : Inputs)
 	{
-		FDataflowConnection* Con = Elem.Value;
+		FDataflowInput* Con = Elem.Value;
 		RetVal.Add({ Dataflow::FPin::EDirection::INPUT,Con->GetType(), Con->GetName() });
 	}
-	for (TPair<uint32, FDataflowConnection*> Elem : Outputs)
+	for (TPair<uint32, FDataflowOutput*> Elem : Outputs)
 	{
-		FDataflowConnection* Con = Elem.Value;
+		FDataflowOutput* Con = Elem.Value;
 		RetVal.Add({ Dataflow::FPin::EDirection::OUTPUT,Con->GetType(), Con->GetName() });
 	}
 	return RetVal;
@@ -195,14 +195,14 @@ TArray<Dataflow::FPin> FDataflowNode::GetPins() const
 
 void FDataflowNode::InvalidateOutputs()
 {
-	for (TPair<uint32, FDataflowConnection*> Elem : Outputs)
+	for (TPair<uint32, FDataflowOutput*> Elem : Outputs)
 	{
-		FDataflowConnection* Con = Elem.Value;
+		FDataflowOutput* Con = Elem.Value;
 		Con->Invalidate();
 	}
 }
 
-void FDataflowNode::RegisterInputConnection(const void* Data)
+void FDataflowNode::RegisterInputConnection(const void* InProperty)
 {
 	if (TUniquePtr<FStructOnScope> ScriptOnStruct = TUniquePtr<FStructOnScope>(NewStructOnScope()))
 	{
@@ -212,7 +212,7 @@ void FDataflowNode::RegisterInputConnection(const void* Data)
 			{
 				FProperty* Property = *PropertyIt;
 				size_t RealAddress = (size_t)this + Property->GetOffset_ForInternal();
-				if (RealAddress == (size_t)Data)
+				if (RealAddress == (size_t)InProperty)
 				{
 					FName PropName(Property->GetName());
 					FName PropType(Property->GetCPPType());
@@ -223,23 +223,38 @@ void FDataflowNode::RegisterInputConnection(const void* Data)
 	}
 }
 
-void FDataflowNode::RegisterOutputConnection(const void* Data)
+void FDataflowNode::RegisterOutputConnection(const void* InProperty, const void* Passthrough)
 {
 	if (TUniquePtr<FStructOnScope> ScriptOnStruct = TUniquePtr<FStructOnScope>(NewStructOnScope()))
 	{
 		if (const UStruct* Struct = ScriptOnStruct->GetStruct())
 		{
+			FDataflowOutput* OutputConnection = nullptr;
+			size_t PassthroughOffset = INDEX_NONE;
 			for (TFieldIterator<FProperty> PropertyIt(Struct); PropertyIt; ++PropertyIt)
 			{
 				FProperty* Property = *PropertyIt;
 				size_t RealAddress = (size_t)this + Property->GetOffset_ForInternal();
-				if (RealAddress == (size_t)Data)
+				if (RealAddress == (size_t)InProperty)
 				{
 					FName PropName(Property->GetName());
 					FName PropType(Property->GetCPPType());
-					AddOutput(new FDataflowOutput({ PropType, PropName, this, Property }));
+					OutputConnection = new FDataflowOutput({ PropType, PropName, this, Property });
 				}
+				if (RealAddress == (size_t)Passthrough)
+				{
+					PassthroughOffset = (uint32)Property->GetOffset_ForInternal();
+				}
+			}
+			if(OutputConnection != nullptr)
+			{
+				if(PassthroughOffset != INDEX_NONE)
+				{
+					OutputConnection->SetPassthroughOffsetAddress(PassthroughOffset);
+				}
+				AddOutput(OutputConnection);
 			}
+			
 		}
 	}
 }
@@ -267,11 +282,49 @@ bool FDataflowNode::ValidateConnections()
 				}
 				if (Property->HasMetaData(TEXT("DataflowOutput")))
 				{
-					if(!FindOutput(PropName))
+					const FDataflowOutput* OutputConnection = FindOutput(PropName);
+					if(!OutputConnection)
 					{
 						UE_LOG(LogChaos, Warning, TEXT("Missing dataflow RegisterOutputConnection in constructor for (%s:%s)"), *GetName().ToString(),*PropName.ToString());
 						bValid = false;
 					}
+
+					// Validate passthrough connections if they exist
+					if (const FString* PassthroughName = Property->FindMetaData(TEXT("Passthrough")))
+					{
+						void* PassthroughConnectionAddress = OutputConnection->GetPassthroughRealAddress();
+						if(PassthroughConnectionAddress == nullptr)
+						{
+							UE_LOG(LogChaos, Warning, TEXT("Missing Passthrough registration for (%s:%s)"), *GetName().ToString(),*PropName.ToString());
+							bValid = false;
+						}
+
+						const FDataflowInput* PassthroughConnectionInput = FindInput(FName(*PassthroughName));
+						const FDataflowInput* PassthroughConnectionInputFromArg = FindInput(PassthroughConnectionAddress);
+
+						if(PassthroughConnectionInputFromArg != PassthroughConnectionInput)
+						{
+							UE_LOG(LogChaos, Warning, TEXT("Mismatch in declared and registered Passthrough connection; (%s:%s vs %s)"), *GetName().ToString(), *PropName.ToString(), *PassthroughConnectionInputFromArg->GetName().ToString());
+							bValid = false;
+						}
+
+						if(!PassthroughConnectionInput)
+						{
+							UE_LOG(LogChaos, Warning, TEXT("Incorrect Passthrough Connection set for (%s:%s)"), *GetName().ToString(),*PropName.ToString());
+							bValid = false;
+						}
+
+						else if(OutputConnection->GetType() != PassthroughConnectionInput->GetType())
+						{
+							UE_LOG(LogChaos, Warning, TEXT("Passthrough connection types mismatch for (%s:%s)"), *GetName().ToString(),*PropName.ToString());
+							bValid = false;
+						}
+					}
+					else if(OutputConnection->GetPassthroughRealAddress()) 
+					{
+						UE_LOG(LogChaos, Warning, TEXT("Missing Passthrough decleration for (%s:%s)"), *GetName().ToString(),*PropName.ToString());
+						bValid = false;
+					}
 				}
 #endif
 			}
diff --git a/Engine/Source/Runtime/Experimental/Dataflow/Core/Public/Dataflow/DataflowConnection.h b/Engine/Source/Runtime/Experimental/Dataflow/Core/Public/Dataflow/DataflowConnection.h
--- a/Engine/Source/Runtime/Experimental/Dataflow/Core/Public/Dataflow/DataflowConnection.h
+++ b/Engine/Source/Runtime/Experimental/Dataflow/Core/Public/Dataflow/DataflowConnection.h
@@ -72,40 +72,12 @@ public:
 	virtual bool AddConnection(FDataflowConnection* In) { return false; };
 	virtual bool RemoveConnection(FDataflowConnection* In) { return false; }
 
-	virtual TArray< FDataflowConnection* > GetConnectedInputs() { return TArray<FDataflowConnection* >(); }
-	virtual const TArray< const FDataflowConnection* > GetConnectedInputs() const { return TArray<const FDataflowConnection* >(); }
-
-	virtual TArray< FDataflowConnection* > GetConnectedOutputs() { return TArray<FDataflowConnection* >(); }
-	virtual const TArray< const FDataflowConnection* > GetConnectedOutputs() const { return TArray<const FDataflowConnection* >(); }
-
 	template<class T>
 	bool IsA(const T* InVar) const
 	{
 		return (size_t)OwningNode + (size_t)GetOffset() == (size_t)InVar;
-
-	}
-
-	template<class T> const T& GetValueAsInput(Dataflow::FContext& Context, const T& Default) const
-	{
-		if (GetConnectedOutputs().Num())
-		{
-			ensure(GetConnectedOutputs().Num() == 1);
-			if (const FDataflowConnection* ConnectionBase = (GetConnectedOutputs()[0]))
-			{
-				if (!ConnectionBase->Evaluate(Context))
-				{
-					Context.SetData(ConnectionBase->CacheKey(), new Dataflow::ContextCache<T>(Property, new T(Default)));
-				}
-				if (Context.HasData(ConnectionBase->CacheKey()))
-				{
-					return Context.GetDataReference<T>(ConnectionBase->CacheKey(), Default);
-				}
-			}
-		}
-		return Default;
 	}
 
 	virtual void Invalidate() {};
-	virtual bool Evaluate(Dataflow::FContext& Context) const { return false; };
 
 };
diff --git a/Engine/Source/Runtime/Experimental/Dataflow/Core/Public/Dataflow/DataflowGraph.h b/Engine/Source/Runtime/Experimental/Dataflow/Core/Public/Dataflow/DataflowGraph.h
--- a/Engine/Source/Runtime/Experimental/Dataflow/Core/Public/Dataflow/DataflowGraph.h
+++ b/Engine/Source/Runtime/Experimental/Dataflow/Core/Public/Dataflow/DataflowGraph.h
@@ -85,12 +85,14 @@ namespace Dataflow
 			return TSharedPtr<FDataflowNode>(nullptr);
 		}
 
-
 		void RemoveNode(TSharedPtr<FDataflowNode> Node);
 
-		void ClearConnections(FDataflowConnection*);
-		void Connect(FDataflowConnection* Input, FDataflowConnection* Output);
-		void Disconnect(FDataflowConnection* Input, FDataflowConnection* Output);
+		void ClearConnections(FDataflowConnection* ConnectionBase);
+		void ClearConnections(FDataflowInput* Input);
+		void ClearConnections(FDataflowOutput* Output);
+
+		void Connect(FDataflowInput* InputConnection, FDataflowOutput* OutputConnection);
+		void Disconnect(FDataflowInput* InputConnection, FDataflowOutput* OutputConnection);
 
 		virtual void Serialize(FArchive& Ar);
 		const TSet<FName>& GetDisabledNodes() const { return DisabledNodes; }
diff --git a/Engine/Source/Runtime/Experimental/Dataflow/Core/Public/Dataflow/DataflowInputOutput.h b/Engine/Source/Runtime/Experimental/Dataflow/Core/Public/Dataflow/DataflowInputOutput.h
--- a/Engine/Source/Runtime/Experimental/Dataflow/Core/Public/Dataflow/DataflowInputOutput.h
+++ b/Engine/Source/Runtime/Experimental/Dataflow/Core/Public/Dataflow/DataflowInputOutput.h
@@ -50,19 +50,14 @@ public:
 	FDataflowOutput* GetConnection() { return Connection; }
 	const FDataflowOutput* GetConnection() const { return Connection; }
 
-	virtual TArray< FDataflowConnection* > GetConnectedOutputs() override;
-	virtual const TArray< const FDataflowConnection* > GetConnectedOutputs() const override;
-
-	template<class T> const T& GetValue(Dataflow::FContext& Context, const T& Default) const
-	{
-		return GetValueAsInput<T>(Context, Default);
-	}
+	virtual TArray< FDataflowOutput* > GetConnectedOutputs();
+	virtual const TArray< const FDataflowOutput* > GetConnectedOutputs() const;
 
+	template<class T> const T& GetValue(Dataflow::FContext& Context, const T& Default) const;
 
 	virtual void Invalidate() override;
 };
 
-
 //
 // Output
 //
@@ -94,6 +89,8 @@ struct DATAFLOWCORE_API FDataflowOutput : public FDataflowConnection
 
 	TArray< FDataflowInput* > Connections;
 
+	size_t PassthroughOffsetAddress = INDEX_NONE;
+
 public:
 	static FDataflowOutput NoOpOutput;
 
@@ -102,13 +99,27 @@ public:
 	TArray<FDataflowInput*>& GetConnections();
 	const TArray<FDataflowInput*>& GetConnections() const;
 
-	virtual TArray<FDataflowConnection*> GetConnectedInputs() override;
-	virtual const TArray<const FDataflowConnection*> GetConnectedInputs() const override;
+	virtual TArray<FDataflowInput*> GetConnectedInputs();
+	virtual const TArray<const FDataflowInput*> GetConnectedInputs() const;
 
 	virtual bool AddConnection(FDataflowConnection* InOutput) override;
 
 	virtual bool RemoveConnection(FDataflowConnection* InInput) override;
 
+	virtual FORCEINLINE void SetPassthroughOffsetAddress(const size_t InPassthroughOffsetAddress)
+	{
+		PassthroughOffsetAddress = InPassthroughOffsetAddress;
+	}
+
+	virtual FORCEINLINE void* GetPassthroughRealAddress() const
+	{
+		if(PassthroughOffsetAddress != INDEX_NONE)
+		{
+			return (void*)((size_t)OwningNode + PassthroughOffsetAddress);
+		}
+		return nullptr;
+	}
+ 
 	template<class T>
 	void SetValue(const T& InVal, Dataflow::FContext& Context) const
 	{
@@ -120,7 +131,7 @@ public:
 
 	template<class T> const T& GetValue(Dataflow::FContext& Context, const T& Default) const
 	{
-		if (!this->Evaluate(Context))
+		if (!this->Evaluate<T>(Context))
 		{
 			Context.SetData(CacheKey(), new Dataflow::ContextCache<T>(Property, new T(Default)));
 		}
@@ -134,8 +145,52 @@ public:
 		return Default;
 	}
 
-		
+	template<class T>
 	bool Evaluate(Dataflow::FContext& Context) const;
 	virtual void Invalidate() override;
 
-};
\ No newline at end of file
+};
+ 
+template<class T>
+const T& FDataflowInput::GetValue(Dataflow::FContext& Context, const T& Default) const
+{
+	if (GetConnectedOutputs().Num())
+	{
+		ensure(GetConnectedOutputs().Num() == 1);
+		if (const FDataflowOutput* ConnectionBase = GetConnection())
+		{
+			if (!ConnectionBase->Evaluate<T>(Context))
+			{
+				Context.SetData(ConnectionBase->CacheKey(), new Dataflow::ContextCache<T>(Property, new T(Default)));
+			}
+			if (Context.HasData(ConnectionBase->CacheKey()))
+			{
+				const T& data = Context.GetDataReference<T>(ConnectionBase->CacheKey(), Default);
+				return data;
+			}
+		}
+	}
+	return Default;
+}
+ 
+template<class T>
+bool FDataflowOutput::Evaluate(Dataflow::FContext& Context) const
+{
+	check(OwningNode);
+ 
+	if (OwningNode->bActive)
+	{
+		OwningNode->Evaluate(Context, this);
+	}
+	else if(const FDataflowInput* PassthroughInput = OwningNode->FindInput(GetPassthroughRealAddress()))
+	{
+		SetValue<T>(PassthroughInput->GetValue<T>(Context, *reinterpret_cast<const T*>(PassthroughInput->RealAddress())), Context);
+	}
+ 
+	// Validation
+	if (!Context.HasData(CacheKey()))
+	{
+		ensureMsgf(false, TEXT("Failed to evaluate output (%s:%s)"), *OwningNode->GetName().ToString(), *GetName().ToString());
+	}
+	return true;
+}
\ No newline at end of file
diff --git a/Engine/Source/Runtime/Experimental/Dataflow/Core/Public/Dataflow/DataflowNode.h b/Engine/Source/Runtime/Experimental/Dataflow/Core/Public/Dataflow/DataflowNode.h
--- a/Engine/Source/Runtime/Experimental/Dataflow/Core/Public/Dataflow/DataflowNode.h
+++ b/Engine/Source/Runtime/Experimental/Dataflow/Core/Public/Dataflow/DataflowNode.h
@@ -40,8 +40,8 @@ struct DATAFLOWCORE_API FDataflowNode
 	FGuid Guid;
 	FName Name;
 
-	TMap< int, FDataflowConnection* > Inputs;
-	TMap< int, FDataflowConnection* > Outputs;
+	TMap< int, FDataflowInput* > Inputs;
+	TMap< int, FDataflowOutput* > Outputs;
 
 	UPROPERTY(EditAnywhere, Category = "Dataflow")
 	bool bActive = true;
@@ -77,8 +77,8 @@ struct DATAFLOWCORE_API FDataflowNode
 	TArray<Dataflow::FPin> GetPins() const;
 
 
-	void AddInput(FDataflowConnection* InPtr);
-	TArray< FDataflowConnection* > GetInputs() const;
+	void AddInput(FDataflowInput* InPtr);
+	TArray< FDataflowInput* > GetInputs() const;
 	void ClearInputs();
 
 	FDataflowInput* FindInput(FName Name);
@@ -86,8 +86,8 @@ struct DATAFLOWCORE_API FDataflowNode
 	const FDataflowInput* FindInput(const void* Reference) const;
 
 
-	void AddOutput(FDataflowConnection* InPtr);
-	TArray< FDataflowConnection* > GetOutputs() const;
+	void AddOutput(FDataflowOutput* InPtr);
+	TArray< FDataflowOutput* > GetOutputs() const;
 	void ClearOutputs();
 
 	FDataflowOutput* FindOutput(FName Name);
@@ -104,8 +104,8 @@ struct DATAFLOWCORE_API FDataflowNode
 	virtual FStructOnScope* NewStructOnScope() { return nullptr; }
 
 	/** Register the Input and Outputs after the creation in the factory */
-	void RegisterInputConnection(const void*);
-	void RegisterOutputConnection(const void*);
+	void RegisterInputConnection(const void* Property);
+	void RegisterOutputConnection(const void* Property, const void* Passthrough = nullptr);
 
 	//
 	// Evaluation
@@ -125,7 +125,7 @@ struct DATAFLOWCORE_API FDataflowNode
 	template<class T> const T& GetValue(Dataflow::FContext& Context, const T* Reference) const
 	{
 		checkSlow(FindInput(Reference));
-		return FindInput(Reference)->template GetValueAsInput<T>(Context, *Reference);
+		return FindInput(Reference)->template GetValue<T>(Context, *Reference);
 	}
 
 	/**
@@ -141,7 +141,7 @@ struct DATAFLOWCORE_API FDataflowNode
 	template<class T> const T& GetValue(Dataflow::FContext& Context, const T* Reference, const T& Default) const
 	{
 		checkSlow(FindInput(Reference));
-		return FindInput(Reference)->template GetValueAsInput<T>(Context, Default);
+		return FindInput(Reference)->template GetValue<T>(Context, Default);
 	}
 
 	/**
diff --git a/Engine/Source/Runtime/Experimental/Dataflow/Engine/Private/Dataflow/DataflowEdNode.cpp b/Engine/Source/Runtime/Experimental/Dataflow/Engine/Private/Dataflow/DataflowEdNode.cpp
--- a/Engine/Source/Runtime/Experimental/Dataflow/Engine/Private/Dataflow/DataflowEdNode.cpp
+++ b/Engine/Source/Runtime/Experimental/Dataflow/Engine/Private/Dataflow/DataflowEdNode.cpp
@@ -58,7 +58,7 @@ void UDataflowEdNode::PinConnectionListChanged(UEdGraphPin* Pin)
 		{
 			if (Pin->Direction == EEdGraphPinDirection::EGPD_Input)
 			{
-				if (FDataflowConnection* ConnectionInput = DataflowNode->FindInput(FName(Pin->GetName())))
+				if (FDataflowInput* ConnectionInput = DataflowNode->FindInput(FName(Pin->GetName())))
 				{
 					DataflowGraph->ClearConnections(ConnectionInput);
 					for (UEdGraphPin* LinkedCon : Pin->LinkedTo)
@@ -69,7 +69,7 @@ void UDataflowEdNode::PinConnectionListChanged(UEdGraphPin* Pin)
 							{
 								if (TSharedPtr<FDataflowNode> LinkedDataflowNode = DataflowGraph->FindBaseNode(LinkedNode->GetDataflowNodeGuid()))
 								{
-									if (FDataflowConnection* LinkedConBase = LinkedDataflowNode->FindOutput(FName(LinkedCon->GetName())))
+									if (FDataflowOutput* LinkedConBase = LinkedDataflowNode->FindOutput(FName(LinkedCon->GetName())))
 									{
 										DataflowGraph->Connect(ConnectionInput, LinkedConBase);
 									}
@@ -81,23 +81,20 @@ void UDataflowEdNode::PinConnectionListChanged(UEdGraphPin* Pin)
 			}
 			else if (Pin->Direction == EEdGraphPinDirection::EGPD_Output)
 			{
-				if (FDataflowConnection* Con = DataflowNode->FindOutput(FName(Pin->GetName())))
+				if (FDataflowOutput* ConnectionOutput = DataflowNode->FindOutput(FName(Pin->GetName())))
 				{
-					if (FDataflowConnection* ConnectionOutput = DataflowNode->FindOutput(FName(Pin->GetName())))
+					DataflowGraph->ClearConnections(ConnectionOutput);
+					for (UEdGraphPin* LinkedCon : Pin->LinkedTo)
 					{
-						DataflowGraph->ClearConnections(ConnectionOutput);
-						for (UEdGraphPin* LinkedCon : Pin->LinkedTo)
+						if (UDataflowEdNode* LinkedNode = Cast<UDataflowEdNode>(LinkedCon->GetOwningNode()))
 						{
-							if (UDataflowEdNode* LinkedNode = Cast<UDataflowEdNode>(LinkedCon->GetOwningNode()))
+							if (ensure(LinkedNode->IsBound()))
 							{
-								if (ensure(LinkedNode->IsBound()))
+								if (TSharedPtr<FDataflowNode> LinkedDataflowNode = DataflowGraph->FindBaseNode(LinkedNode->GetDataflowNodeGuid()))
 								{
-									if (TSharedPtr<FDataflowNode> LinkedDataflowNode = DataflowGraph->FindBaseNode(LinkedNode->GetDataflowNodeGuid()))
+									if (FDataflowInput* LinkedConBase = LinkedDataflowNode->FindInput(FName(LinkedCon->GetName())))
 									{
-										if (FDataflowConnection* LinkedConBase = LinkedDataflowNode->FindInput(FName(LinkedCon->GetName())))
-										{
-											DataflowGraph->Connect(LinkedConBase, ConnectionOutput);
-										}
+										DataflowGraph->Connect(LinkedConBase, ConnectionOutput);
 									}
 								}
 							}
