Constraints: Refactor constraint channels out of MovieScene and ControlRig and into Constraint. Note still need to create a ConstraintEditor plugin to move the ConstraintChannelEditor code out of MovieSceneTools but wanted to get this in asap.

diff --git a/Engine/Plugins/Animation/ControlRig/Source/ControlRig/Public/Sequencer/MovieSceneControlRigParameterSection.h b/Engine/Plugins/Animation/ControlRig/Source/ControlRig/Public/Sequencer/MovieSceneControlRigParameterSection.h
--- a/Engine/Plugins/Animation/ControlRig/Source/ControlRig/Public/Sequencer/MovieSceneControlRigParameterSection.h
+++ b/Engine/Plugins/Animation/ControlRig/Source/ControlRig/Public/Sequencer/MovieSceneControlRigParameterSection.h
@@ -16,9 +16,8 @@
 #include "Channels/MovieSceneIntegerChannel.h"
 #include "Channels/MovieSceneByteChannel.h"
 #include "Channels/MovieSceneBoolChannel.h"
-#include "Channels/MovieSceneConstraintChannel.h"
 #include "Sequencer/MovieSceneControlRigSpaceChannel.h"
-
+#include "ConstraintChannel.h"
 
 #include "MovieSceneControlRigParameterSection.generated.h"
 
@@ -106,22 +105,6 @@ struct CONTROLRIG_API FSpaceControlNameAndChannel
 	FMovieSceneControlRigSpaceChannel SpaceCurve;
 };
 
-USTRUCT()
-struct CONTROLRIG_API FConstraintAndActiveChannel
-{
-	GENERATED_USTRUCT_BODY()
-
-	FConstraintAndActiveChannel(){}
-	FConstraintAndActiveChannel(const TObjectPtr<UTickableConstraint>& InConstraint)
-		: Constraint(InConstraint)
-	{};
-
-	UPROPERTY()
-	TSoftObjectPtr<UTickableConstraint> Constraint;
-
-	UPROPERTY()
-	FMovieSceneConstraintChannel ActiveChannel;
-};
 
 /**
 *  Data that's queried during an interrogation
diff --git a/Engine/Plugins/Animation/ControlRig/Source/ControlRigEditor/Private/Sequencer/ConstraintChannelHelper.cpp b/Engine/Plugins/Animation/ControlRig/Source/ControlRigEditor/Private/Sequencer/ConstraintChannelHelper.cpp
--- a/Engine/Plugins/Animation/ControlRig/Source/ControlRigEditor/Private/Sequencer/ConstraintChannelHelper.cpp
+++ b/Engine/Plugins/Animation/ControlRig/Source/ControlRigEditor/Private/Sequencer/ConstraintChannelHelper.cpp
@@ -14,7 +14,7 @@
 #include "TransformConstraint.h"
 #include "Algo/Copy.h"
 #include "Algo/Unique.h"
-#include "Channels/MovieSceneConstraintChannel.h"
+#include "ConstraintChannel.h"
 
 #include "Tools/BakingHelper.h"
 #include "Tools/ConstraintBaker.h"
diff --git a/Engine/Source/Editor/MovieSceneTools/MovieSceneTools.Build.cs b/Engine/Source/Editor/MovieSceneTools/MovieSceneTools.Build.cs
--- a/Engine/Source/Editor/MovieSceneTools/MovieSceneTools.Build.cs
+++ b/Engine/Source/Editor/MovieSceneTools/MovieSceneTools.Build.cs
@@ -31,6 +31,7 @@
 				"Sequencer",
                 "EditorWidgets",
 				"SequencerCore",
+				"Constraints"
             }
         );
 
diff --git a/Engine/Source/Editor/MovieSceneTools/Private/Channels/ConstraintChannelEditor.cpp b/Engine/Source/Editor/MovieSceneTools/Private/Channels/ConstraintChannelEditor.cpp
--- a/Engine/Source/Editor/MovieSceneTools/Private/Channels/ConstraintChannelEditor.cpp
+++ b/Engine/Source/Editor/MovieSceneTools/Private/Channels/ConstraintChannelEditor.cpp
@@ -2,7 +2,7 @@
 
 #include "ConstraintChannelEditor.h"
 
-#include "Channels/MovieSceneConstraintChannel.h"
+#include "ConstraintChannel.h"
 
 #include "KeyBarCurveModel.h"
 #include "SequencerSectionPainter.h"
diff --git a/Engine/Source/Editor/MovieSceneTools/Private/MovieSceneToolsModule.cpp b/Engine/Source/Editor/MovieSceneTools/Private/MovieSceneToolsModule.cpp
--- a/Engine/Source/Editor/MovieSceneTools/Private/MovieSceneToolsModule.cpp
+++ b/Engine/Source/Editor/MovieSceneTools/Private/MovieSceneToolsModule.cpp
@@ -77,7 +77,7 @@
 #include "Channels/EventChannelCurveModel.h"
 #include "Channels/SCurveEditorEventChannelView.h"
 #include "Channels/MovieSceneAudioTriggerChannel.h"
-#include "Channels/MovieSceneConstraintChannel.h"
+#include "ConstraintChannel.h"
 #include "Channels/ConstraintChannelEditor.h"
 #include "Sections/MovieSceneEventSection.h"
 
diff --git a/Engine/Source/Runtime/Experimental/Animation/Constraints/Constraints.Build.cs b/Engine/Source/Runtime/Experimental/Animation/Constraints/Constraints.Build.cs
--- a/Engine/Source/Runtime/Experimental/Animation/Constraints/Constraints.Build.cs
+++ b/Engine/Source/Runtime/Experimental/Animation/Constraints/Constraints.Build.cs
@@ -16,7 +16,8 @@ public class Constraints : ModuleRules
 				"Core",
 				"CoreUObject",
 				"Engine",
-				"AnimationCore"
+				"AnimationCore",
+				"MovieScene"
 			}
 		);
 	}
diff --git a/Engine/Source/Runtime/Experimental/Animation/Constraints/Private/ConstraintChannel.cpp b/Engine/Source/Runtime/Experimental/Animation/Constraints/Private/ConstraintChannel.cpp
new file mode 100644
--- /dev/null
+++ b/Engine/Source/Runtime/Experimental/Animation/Constraints/Private/ConstraintChannel.cpp
@@ -0,0 +1,19 @@
+// Copyright Epic Games, Inc. All Rights Reserved.
+
+#include "ConstraintChannel.h"
+
+bool FMovieSceneConstraintChannel::Evaluate(FFrameTime InTime, bool& OutValue) const
+{
+	const TArrayView<const FFrameNumber> TimeArray = GetTimes();
+	if (TimeArray.IsEmpty())
+	{
+		return false;
+	}
+
+	if (InTime.FrameNumber < TimeArray[0])
+	{
+		return false;
+	}
+
+	return FMovieSceneBoolChannel::Evaluate(InTime, OutValue);
+}
diff --git a/Engine/Source/Runtime/Experimental/Animation/Constraints/Public/ConstraintChannel.h b/Engine/Source/Runtime/Experimental/Animation/Constraints/Public/ConstraintChannel.h
new file mode 100644
--- /dev/null
+++ b/Engine/Source/Runtime/Experimental/Animation/Constraints/Public/ConstraintChannel.h
@@ -0,0 +1,68 @@
+// Copyright Epic Games, Inc. All Rights Reserved.
+
+#pragma once
+
+#include "Channels/MovieSceneBoolChannel.h"
+#include "MovieSceneClipboard.h"
+#include "ConstraintsManager.h"
+#include "ConstraintChannel.generated.h"
+
+USTRUCT()
+struct CONSTRAINTS_API FMovieSceneConstraintChannel : public FMovieSceneBoolChannel
+{
+	GENERATED_BODY()
+
+	FMovieSceneConstraintChannel() {};
+
+	// @todo 
+	bool Evaluate(FFrameTime InTime, bool& OutValue) const;
+};
+
+template<>
+struct TStructOpsTypeTraits<FMovieSceneConstraintChannel> : public TStructOpsTypeTraitsBase2<FMovieSceneConstraintChannel>
+{
+	enum { WithStructuredSerializeFromMismatchedTag = true };
+};
+
+
+template<>
+struct TMovieSceneChannelTraits<FMovieSceneConstraintChannel> : TMovieSceneChannelTraitsBase<FMovieSceneConstraintChannel>
+{
+	enum { SupportsDefaults = false };
+
+#if WITH_EDITOR
+
+	/** Byte channels can have external values (ie, they can get their values from external objects for UI purposes) */
+	typedef TMovieSceneExternalValue<bool> ExtendedEditorDataType;
+
+#endif
+};
+
+// #if WITH_EDITOR
+// namespace MovieSceneClipboard
+// {
+// 	template<> inline FName GetKeyTypeName<bool>()
+// 	{
+// 		static FName Name("Bool");
+// 		return Name;
+// 	}
+// }
+// #endif
+
+
+USTRUCT()
+struct CONSTRAINTS_API FConstraintAndActiveChannel
+{
+	GENERATED_USTRUCT_BODY()
+
+	FConstraintAndActiveChannel() {}
+	FConstraintAndActiveChannel(const TObjectPtr<UTickableConstraint>& InConstraint)
+		: Constraint(InConstraint)
+	{};
+
+	UPROPERTY()
+	TSoftObjectPtr<UTickableConstraint> Constraint;
+
+	UPROPERTY()
+	FMovieSceneConstraintChannel ActiveChannel;
+};
\ No newline at end of file
diff --git a/Engine/Source/Runtime/MovieScene/Private/Channels/MovieSceneConstraintChannel.cpp b/Engine/Source/Runtime/MovieScene/Private/Channels/MovieSceneConstraintChannel.cpp
deleted file mode 100644
diff --git a/Engine/Source/Runtime/MovieScene/Public/Channels/MovieSceneConstraintChannel.h b/Engine/Source/Runtime/MovieScene/Public/Channels/MovieSceneConstraintChannel.h
deleted file mode 100644
