Chaos destruction : simplify removal runtime system code paths
- both removal system are now using common methods
- both of them are supporting crumbling of clusters ( both regular and internal )
- Abstraction of the logic into facade classes to avoid direct access to the attributes
- Fix bug with slow moving detection always return true

diff --git a/Engine/Plugins/Experimental/ChaosEditor/Source/FractureEditor/Private/SGeometryCollectionOutliner.h b/Engine/Plugins/Experimental/ChaosEditor/Source/FractureEditor/Private/SGeometryCollectionOutliner.h
--- a/Engine/Plugins/Experimental/ChaosEditor/Source/FractureEditor/Private/SGeometryCollectionOutliner.h
+++ b/Engine/Plugins/Experimental/ChaosEditor/Source/FractureEditor/Private/SGeometryCollectionOutliner.h
@@ -16,6 +16,7 @@
 
 #include "GeometryCollection/GeometryCollection.h"
 #include "GeometryCollection/GeometryCollectionComponent.h"
+#include "GeometryCollection/GeometryCollectionEngineRemoval.h"
 
 #include "SGeometryCollectionOutliner.generated.h"
 
diff --git a/Engine/Source/Runtime/Experimental/Chaos/Private/GeometryCollectionProxyData.cpp b/Engine/Source/Runtime/Experimental/Chaos/Private/GeometryCollectionProxyData.cpp
--- a/Engine/Source/Runtime/Experimental/Chaos/Private/GeometryCollectionProxyData.cpp
+++ b/Engine/Source/Runtime/Experimental/Chaos/Private/GeometryCollectionProxyData.cpp
@@ -65,3 +65,54 @@ FGeometryDynamicCollection::FGeometryDynamicCollection()
 	AddExternalAttribute(SimulatableParticlesAttribute, FGeometryCollection::TransformGroup, SimulatableParticles);
 
 }
+
+///////////////////////////////////////////////////////////////////////////////////////////////////////////////////////////
+
+FGeometryCollectionDynamicStateFacade::FGeometryCollectionDynamicStateFacade(FManagedArrayCollection& InCollection)
+	: ActiveAttribute(InCollection, FGeometryDynamicCollection::ActiveAttribute,  FTransformCollection::TransformGroup)
+	, DynamicStateAttribute(InCollection, FGeometryDynamicCollection::DynamicStateAttribute,  FTransformCollection::TransformGroup)
+	, ChildrenAttribute(InCollection, "Children",  FTransformCollection::TransformGroup)
+	, ParentAttribute(InCollection, "Parent",  FTransformCollection::TransformGroup)
+	, InternalClusterParentTypeAttribute(InCollection, "InternalClusterParentTypeArray", FGeometryCollection::TransformGroup)
+{
+}
+
+bool FGeometryCollectionDynamicStateFacade::IsValid() const
+{
+	return ActiveAttribute.IsValid()
+		&& DynamicStateAttribute.IsValid()
+		&& ChildrenAttribute.IsValid()
+		&& ParentAttribute.IsValid()
+		&& InternalClusterParentTypeAttribute.IsValid()
+		;
+}
+
+bool FGeometryCollectionDynamicStateFacade::IsDynamicOrSleeping(int32 TransformIndex) const
+{
+	const int32 State = DynamicStateAttribute.Get()[TransformIndex];
+	return (State == (int)EObjectStateTypeEnum::Chaos_Object_Sleeping) || (State == (int)EObjectStateTypeEnum::Chaos_Object_Dynamic);
+}
+
+bool FGeometryCollectionDynamicStateFacade::IsSleeping(int32 TransformIndex) const
+{
+	const int32 State = DynamicStateAttribute.Get()[TransformIndex];
+	return (State == (int)EObjectStateTypeEnum::Chaos_Object_Sleeping);
+}
+
+bool FGeometryCollectionDynamicStateFacade::HasChildren(int32 TransformIndex) const
+{
+	return (ChildrenAttribute.Get()[TransformIndex].Num() > 0);
+}
+
+bool FGeometryCollectionDynamicStateFacade::HasBrokenOff(int32 TransformIndex) const
+{
+	const bool IsActive = ActiveAttribute.Get()[TransformIndex];
+	const bool HasParent = (ParentAttribute.Get()[TransformIndex] != INDEX_NONE);
+	return IsActive && (!HasParent) && IsDynamicOrSleeping(TransformIndex);
+}
+
+bool FGeometryCollectionDynamicStateFacade::HasDynamicInternalClusterParent(int32 TransformIndex) const
+{
+	const uint8 InternalParenttype = InternalClusterParentTypeAttribute.Get()[TransformIndex];
+	return InternalParenttype == (uint8)Chaos::EInternalClusterType::Dynamic;
+}
\ No newline at end of file
diff --git a/Engine/Source/Runtime/Experimental/Chaos/Public/GeometryCollection/GeometryCollection.h b/Engine/Source/Runtime/Experimental/Chaos/Public/GeometryCollection/GeometryCollection.h
--- a/Engine/Source/Runtime/Experimental/Chaos/Public/GeometryCollection/GeometryCollection.h
+++ b/Engine/Source/Runtime/Experimental/Chaos/Public/GeometryCollection/GeometryCollection.h
@@ -7,6 +7,7 @@
 #include "Misc/Crc.h"
 
 #include "GeometryCollection/GeometryCollectionConvexPropertiesInterface.h"
+#include "GeometryCollection/ManagedArrayAccessor.h"
 
 namespace Chaos
 {
@@ -344,51 +345,13 @@ FORCEINLINE Chaos::FChaosArchive& operator<<(Chaos::FChaosArchive& Ar, FGeometry
 	return Ar;
 }
 
-
-template <typename T>
-struct AttributeAccessor
-{
-public:
-	AttributeAccessor(FManagedArrayCollection& InCollection, const FName& AttributeName, const FName& AttributeGroup)
-		: Collection(InCollection)
-		, Name(AttributeName)
-		, Group(AttributeGroup)
-	{
-		AttributeArray = InCollection.FindAttributeTyped<T>(AttributeName, AttributeGroup);
-
-	}
-	bool IsValid() const { return AttributeArray != nullptr; }
-	const TManagedArray<T>& Get() const
-	{
-		check(IsValid());
-		return *AttributeArray;
-	}
-	TManagedArray<T>& Modify() const
-	{
-		check(IsValid());
-		AttributeArray->MarkDirty();
-		return *AttributeArray;
-	}
-	void Add()
-	{
-		AttributeArray = &Collection.AddAttribute<T>(Name, Group);
-	}
-
-private:
-	FManagedArrayCollection& Collection;
-	FName Name;
-	FName Group;
-	TManagedArray<T>* AttributeArray;
-};
-
-
 class CHAOS_API FGeometryCollectionMeshFacade
 {
 public:
 	FGeometryCollectionMeshFacade(FManagedArrayCollection& InCollection);
 
 	/** 
-	 * returns true iof all the necessary attributes are present
+	 * returns true if all the necessary attributes are present
 	 * if not then the API can be used to create  
 	 */
 	bool IsValid() const;	
@@ -398,20 +361,20 @@ public:
 	 */
 	void AddAttributes();
 
-	AttributeAccessor<FVector3f> Vertex;
-	AttributeAccessor<FVector3f> TangentU;
-	AttributeAccessor<FVector3f> TangentV;
-	AttributeAccessor<FVector3f> Normal;
-	AttributeAccessor<TArray<FVector2f>> UVs;
-	AttributeAccessor<FLinearColor> Color;
-	AttributeAccessor<int32> BoneMap;
-	AttributeAccessor<int32> VertexStart;
-	AttributeAccessor<int32> VertexCount;
-
-	AttributeAccessor<FIntVector> Indices;
-	AttributeAccessor<bool> Visible;
-	AttributeAccessor<int32> MaterialIndex;
-	AttributeAccessor<int32> MaterialID;
-	AttributeAccessor<int32> FaceStart;
-	AttributeAccessor<int32> FaceCount;
+	TManagedArrayAccessor<FVector3f> Vertex;
+	TManagedArrayAccessor<FVector3f> TangentU;
+	TManagedArrayAccessor<FVector3f> TangentV;
+	TManagedArrayAccessor<FVector3f> Normal;
+	TManagedArrayAccessor<TArray<FVector2f>> UVs;
+	TManagedArrayAccessor<FLinearColor> Color;
+	TManagedArrayAccessor<int32> BoneMap;
+	TManagedArrayAccessor<int32> VertexStart;
+	TManagedArrayAccessor<int32> VertexCount;
+
+	TManagedArrayAccessor<FIntVector> Indices;
+	TManagedArrayAccessor<bool> Visible;
+	TManagedArrayAccessor<int32> MaterialIndex;
+	TManagedArrayAccessor<int32> MaterialID;
+	TManagedArrayAccessor<int32> FaceStart;
+	TManagedArrayAccessor<int32> FaceCount;
 };
\ No newline at end of file
diff --git a/Engine/Source/Runtime/Experimental/Chaos/Public/GeometryCollection/ManagedArrayAccessor.h b/Engine/Source/Runtime/Experimental/Chaos/Public/GeometryCollection/ManagedArrayAccessor.h
new file mode 100644
--- /dev/null
+++ b/Engine/Source/Runtime/Experimental/Chaos/Public/GeometryCollection/ManagedArrayAccessor.h
@@ -0,0 +1,60 @@
+// Copyright Epic Games, Inc. All Rights Reserved.
+#pragma once
+#include "GeometryCollection/ManagedArrayCollection.h"
+
+/**
+ * this class wraps a managed array
+ * this provides a convenient API for optional attributes in a collection facade
+ */
+template <typename T>
+struct TManagedArrayAccessor
+{
+public:
+	TManagedArrayAccessor(FManagedArrayCollection& InCollection, const FName& AttributeName, const FName& AttributeGroup)
+		: Collection(InCollection)
+		, Name(AttributeName)
+		, Group(AttributeGroup)
+	{
+		AttributeArray = InCollection.FindAttributeTyped<T>(AttributeName, AttributeGroup);
+
+	}
+	bool IsValid() const { return AttributeArray != nullptr; }
+
+	/** get the attribute for read only */
+	const TManagedArray<T>& Get() const
+	{
+		check(IsValid());
+		return *AttributeArray;
+	}
+
+	/** get the attribute for modification */
+	TManagedArray<T>& Modify() const
+	{
+		check(IsValid());
+		AttributeArray->MarkDirty();
+		return *AttributeArray;
+	}
+
+	/** add the attribute if it does not exists yet */
+	TManagedArray<T>& Add()
+	{
+		AttributeArray = &Collection.AddAttribute<T>(Name, Group);
+		return *AttributeArray;
+	}
+
+	/** add and fill the attribute if it does not exist yet */
+	void AddAndFill(const T& Value)
+	{
+		if (!Collection.HasAttribute(Name, Group))
+		{
+			AttributeArray = &Collection.AddAttribute<T>(Name, Group);
+			AttributeArray->Fill(Value);
+		}
+	}
+
+private:
+	FManagedArrayCollection& Collection;
+	FName Name;
+	FName Group;
+	TManagedArray<T>* AttributeArray;
+};
diff --git a/Engine/Source/Runtime/Experimental/Chaos/Public/GeometryCollectionProxyData.h b/Engine/Source/Runtime/Experimental/Chaos/Public/GeometryCollectionProxyData.h
--- a/Engine/Source/Runtime/Experimental/Chaos/Public/GeometryCollectionProxyData.h
+++ b/Engine/Source/Runtime/Experimental/Chaos/Public/GeometryCollectionProxyData.h
@@ -88,6 +88,50 @@ public:
 	TManagedArray<bool> SimulatableParticles;
 };
 
+/**
+ * Provides an API for dynamic state related attributes
+ * physics state , broken state, current parent (normal or internal clusters )
+ * To be used with the dynamic collection
+ */
+class CHAOS_API FGeometryCollectionDynamicStateFacade
+{
+public:
+	FGeometryCollectionDynamicStateFacade(FManagedArrayCollection& InCollection);
+
+	/** returns true if all the necessary attributes are present */
+	bool IsValid() const;
+
+	/** return true if the transform is in a dynamic or sleeping state */
+	bool IsDynamicOrSleeping(int32 TransformIndex) const;
+
+	/** return true if the transform is in a sleeping state */
+	bool IsSleeping(int32 TransformIndex) const;
+
+	/** whether there's children attached to this transfom (Cluster) */
+	bool HasChildren(int32 TransformIndex) const;
+	
+	/** return true if the transform has broken off its parent */
+	bool HasBrokenOff(int32 TransformIndex) const;
+
+	/** return true if the transform has an internal cluster parent */
+	bool HasDynamicInternalClusterParent(int32 TransformIndex) const;
+	
+private:
+	/** Active state, true means that the transform is active or broken off from its parent */
+	TManagedArrayAccessor<bool> ActiveAttribute;
+
+	/** physics state of the transform (Dynamic, kinematic, static, sleeping) */
+	TManagedArrayAccessor<int32> DynamicStateAttribute;
+
+	/** currently attached children (potentially different from the initial children setup) */
+	TManagedArrayAccessor<TSet<int32>> ChildrenAttribute;
+	
+	/** Current parent (potentially different from the initial parent) */
+	TManagedArrayAccessor<int32> ParentAttribute;
+
+	/** type of internal state parent */
+	TManagedArrayAccessor<uint8> InternalClusterParentTypeAttribute;
+};
 
 class FGeometryCollectioPerFrameData
 {
diff --git a/Engine/Source/Runtime/Experimental/GeometryCollectionEngine/Private/GeometryCollection/GeometryCollectionComponent.cpp b/Engine/Source/Runtime/Experimental/GeometryCollectionEngine/Private/GeometryCollection/GeometryCollectionComponent.cpp
--- a/Engine/Source/Runtime/Experimental/GeometryCollectionEngine/Private/GeometryCollection/GeometryCollectionComponent.cpp
+++ b/Engine/Source/Runtime/Experimental/GeometryCollectionEngine/Private/GeometryCollection/GeometryCollectionComponent.cpp
@@ -44,6 +44,7 @@
 
 #include "Rendering/NaniteResources.h"
 #include "PrimitiveSceneInfo.h"
+#include "GeometryCollection/GeometryCollectionEngineRemoval.h"
 
 #if !(UE_BUILD_SHIPPING || UE_BUILD_TEST)
 #include "Logging/MessageLog.h"
@@ -2087,106 +2088,24 @@ void UGeometryCollectionComponent::ResetDynamicCollection()
 		GetSimulationTypeArrayCopyOnWrite();
 		GetStatusFlagsArrayCopyOnWrite();
 
+		FGeometryCollectionDecayDynamicFacade DecayDynamicFacade(*DynamicCollection);
+		
 		// we are not testing for bAllowRemovalOnSleep, so that we can enable it at runtime if necessary
 		if (RestCollection->bRemoveOnMaxSleep)
 		{
-			if (!DynamicCollection->HasAttribute("SleepTimer", FGeometryCollection::TransformGroup))
-			{
-				TManagedArray<float>& SleepTimer = DynamicCollection->AddAttribute<float>("SleepTimer", FGeometryCollection::TransformGroup);
-				SleepTimer.Fill(0.f);
-			}
-
-			if (!DynamicCollection->HasAttribute("Decay", FGeometryCollection::TransformGroup))
-			{
-				TManagedArray<float>& Decay = DynamicCollection->AddAttribute<float>("Decay", FGeometryCollection::TransformGroup);
-				Decay.Fill(0);
-			}
-
-			if (!DynamicCollection->HasAttribute("LastPosition", FGeometryCollection::TransformGroup))
-			{
-				TManagedArray<FVector>& LastPosition = DynamicCollection->AddAttribute<FVector>("LastPosition", FGeometryCollection::TransformGroup);
-				LastPosition.Fill(FVector::Zero());
-			}
-			
-			if (!DynamicCollection->HasAttribute("UniformScale", FGeometryCollection::TransformGroup))
-			{
-				TManagedArray<FTransform>& UniformScale = DynamicCollection->AddAttribute<FTransform>("UniformScale", FGeometryCollection::TransformGroup);
-				UniformScale.Fill(FTransform::Identity);
-			}
+			DecayDynamicFacade.AddAttributes();
 
-			if (!DynamicCollection->HasAttribute("MaxSleepTime", FGeometryCollection::TransformGroup))
-			{
-				float MinTime = FMath::Max(0.0f, RestCollection->MaximumSleepTime.X);
-				float MaxTime = FMath::Max(MinTime, RestCollection->MaximumSleepTime.Y);
-				TManagedArray<float>& MaxSleepTime = DynamicCollection->AddAttribute<float>("MaxSleepTime", FGeometryCollection::TransformGroup);
-				for (int32 Idx = 0; Idx < MaxSleepTime.Num(); ++Idx)
-				{
-					MaxSleepTime[Idx] = FMath::RandRange(MinTime, MaxTime);
-				}
-			}
-
-			if (!DynamicCollection->HasAttribute("SleepRemovalDuration", FGeometryCollection::TransformGroup))
-			{
-				float MinTime = FMath::Max(0.0f, RestCollection->RemovalDuration.X);
-				float MaxTime = FMath::Max(MinTime, RestCollection->RemovalDuration.Y);
-				TManagedArray<float>& SleepRemovalDuration = DynamicCollection->AddAttribute<float>("SleepRemovalDuration", FGeometryCollection::TransformGroup);
-				for (int32 Idx = 0; Idx < SleepRemovalDuration.Num(); ++Idx)
-				{
-					SleepRemovalDuration[Idx] = FMath::RandRange(MinTime, MaxTime);
-				}
-			}
+			FGeometryCollectionRemoveOnSleepDynamicFacade RemoveOnSleepDynamicFacade(*DynamicCollection);
+			RemoveOnSleepDynamicFacade.AddAttributes(RestCollection->MaximumSleepTime, RestCollection->RemovalDuration);
 		}
 		
 		// Remove on break feature related dynamic attribute arrays
 		if (const TManagedArray<FVector4f>* RemoveOnBreak = RestCollection->GetGeometryCollection()->FindAttribute<FVector4f>("RemoveOnBreak", FGeometryCollection::TransformGroup))
 		{
-			// Needs to be checked here too because remove on sleep and remove on break features are orthogonal
-			if (!DynamicCollection->HasAttribute("Decay", FGeometryCollection::TransformGroup))
-			{
-				TManagedArray<float>& Decay = DynamicCollection->AddAttribute<float>("Decay", FGeometryCollection::TransformGroup);
-				Decay.Fill(0);
-			}
-
-			if (!DynamicCollection->HasAttribute("UniformScale", FGeometryCollection::TransformGroup))
-			{
-				TManagedArray<FTransform>& UniformScale = DynamicCollection->AddAttribute<FTransform>("UniformScale", FGeometryCollection::TransformGroup);
-				UniformScale.Fill(FTransform::Identity);
-			}
-		
-			if (!DynamicCollection->HasAttribute("BreakTimer", FGeometryCollection::TransformGroup))
-			{
-				TManagedArray<float>& BreakTimer = DynamicCollection->AddAttribute<float>("BreakTimer", FGeometryCollection::TransformGroup);
-				BreakTimer.Fill(-1); // disabled timer , waiting for the break event
-			}
-			
-			if (!DynamicCollection->HasAttribute("PostBreakDuration", FGeometryCollection::TransformGroup))
-			{
-				TManagedArray<float>& PostBreakDuration = DynamicCollection->AddAttribute<float>("PostBreakDuration", FGeometryCollection::TransformGroup);
-				for (int32 Idx = 0; Idx < PostBreakDuration.Num(); ++Idx)
-				{
-					const FRemoveOnBreakData RemoveOnBreakData{ (*RemoveOnBreak)[Idx] };
-					const FVector2f PostBreakTimer = RemoveOnBreakData.GetBreakTimer();
-					const float MinTime = FMath::Max(0.0f, PostBreakTimer.X);
-					const float MaxTime = FMath::Max(MinTime, PostBreakTimer.Y);
-					PostBreakDuration[Idx] = RemoveOnBreakData.IsEnabled()? FMath::RandRange(MinTime, MaxTime): -1;
-				}
-			}
+			DecayDynamicFacade.AddAttributes();
 
-			if (!DynamicCollection->HasAttribute("BreakRemovalDuration", FGeometryCollection::TransformGroup))
-			{
-				const TManagedArray<TSet<int32>>& Children = DynamicCollection->Children;
-				TManagedArray<float>& BreakRemovalDuration = DynamicCollection->AddAttribute<float>("BreakRemovalDuration", FGeometryCollection::TransformGroup);
-				for (int32 Idx = 0; Idx < BreakRemovalDuration.Num(); ++Idx)
-				{
-					const FRemoveOnBreakData RemoveOnBreakData{ (*RemoveOnBreak)[Idx] };
-					const FVector2f RemovalTimer = RemoveOnBreakData.GetRemovalTimer();
-					const float MinTime = FMath::Max(0.0f, RemovalTimer.X);
-					const float MaxTime = FMath::Max(MinTime, RemovalTimer.Y);
-					const bool bIsCluster = (Children[Idx].Num() > 0);
-					const bool bUseClusterCrumbling = (bIsCluster && RemoveOnBreakData.GetClusterCrumbling());
-					BreakRemovalDuration[Idx] = bUseClusterCrumbling? -1: FMath::RandRange(MinTime, MaxTime);
-				}
-			}
+			FGeometryCollectionRemoveOnBreakDynamicFacade RemoveOnBreakDynamicFacade(*DynamicCollection);
+			RemoveOnBreakDynamicFacade.AddAttributes(*RemoveOnBreak, DynamicCollection->Children);
 		}
 
 		SetRenderStateDirty();
@@ -3584,6 +3503,72 @@ void UGeometryCollectionComponent::InitializeEmbeddedGeometry()
 	}
 }
 
+struct FGeometryCollectionDecayContext
+{
+	FGeometryCollectionDecayContext(FGeometryCollectionPhysicsProxy& PhysicsProxyIn, FGeometryCollectionDecayDynamicFacade& DecayFacadeIn)
+		: PhysicsProxy(PhysicsProxyIn)
+		, DecayFacade(DecayFacadeIn)
+		, DirtyDynamicCollection(false)
+	{}
+
+	FGeometryCollectionPhysicsProxy& PhysicsProxy;
+	FGeometryCollectionDecayDynamicFacade& DecayFacade;
+	
+	bool DirtyDynamicCollection;
+	TArray<int32> ToDisable;
+	TArray<FGeometryCollectionItemIndex> ToCrumble;
+
+	void Process(FGeometryDynamicCollection& DynamicCollection)
+	{
+		if (DirtyDynamicCollection)
+		{
+			DynamicCollection.MakeDirty();
+		}
+		if (ToCrumble.Num())
+		{
+			PhysicsProxy.BreakClusters_External(MoveTemp(ToCrumble));
+		}
+		if (ToDisable.Num())
+		{
+			PhysicsProxy.DisableParticles_External(MoveTemp(ToDisable));
+		}
+	}
+};
+
+void UGeometryCollectionComponent::UpdateDecay(int32 TransformIdx, float UpdatedDecay, bool bUseClusterCrumbling, bool bHasDynamicInternalClusterParent, FGeometryCollectionDecayContext& ContextInOut)
+{
+	TManagedArray<float>& Decay = ContextInOut.DecayFacade.DecayAttribute.Modify();
+	if (UpdatedDecay > Decay[TransformIdx])
+	{
+		ContextInOut.DirtyDynamicCollection = true;
+		Decay[TransformIdx] = UpdatedDecay;
+
+		if (bUseClusterCrumbling)
+		{
+			if (bHasDynamicInternalClusterParent)
+			{
+				FGeometryCollectionItemIndex InternalClusterItemindex = ContextInOut.PhysicsProxy.GetInternalClusterParentItemIndex_External(TransformIdx);
+				if (InternalClusterItemindex.IsValid())
+				{
+					ContextInOut.ToCrumble.AddUnique(InternalClusterItemindex);
+					Decay[TransformIdx] = 0.0f;
+				}
+			}
+			else
+			{
+				ContextInOut.ToCrumble.AddUnique(FGeometryCollectionItemIndex::CreateTransformItemIndex(TransformIdx));
+				Decay[TransformIdx] = 0.0f;
+			}
+		}
+		else if (Decay[TransformIdx] >= 1.0f)
+		{
+			// Disable the particle if it has decayed the requisite time
+			Decay[TransformIdx] = 1.0f;
+			ContextInOut.ToDisable.Add(TransformIdx);
+		}
+	}
+}
+
 void UGeometryCollectionComponent::IncrementSleepTimer(float DeltaTime)
 {
 	if (DeltaTime <= 0 || !RestCollection->bRemoveOnMaxSleep || !bAllowRemovalOnSleep)
@@ -3592,86 +3577,46 @@ void UGeometryCollectionComponent::IncrementSleepTimer(float DeltaTime)
 	}
 	
 	// If a particle is sleeping, increment its sleep timer, otherwise reset it.
-	if (DynamicCollection && PhysicsProxy 
-		&& DynamicCollection->HasAttribute("SleepTimer", FGeometryCollection::TransformGroup) 
-		&& DynamicCollection->HasAttribute("MaxSleepTime", FGeometryCollection::TransformGroup)
-		&& DynamicCollection->HasAttribute("SleepRemovalDuration", FGeometryCollection::TransformGroup)
-		&& DynamicCollection->HasAttribute("Decay", FGeometryCollection::TransformGroup)
-		&& DynamicCollection->HasAttribute("LastPosition", FGeometryCollection::TransformGroup)
-		&& DynamicCollection->HasAttribute("InternalClusterParentTypeArray", FGeometryCollection::TransformGroup))
-	{
-		TManagedArray<float>& SleepTimer = DynamicCollection->ModifyAttribute<float>("SleepTimer", FGeometryCollection::TransformGroup);
-		const TManagedArray<float>& SleepRemovalDuration = DynamicCollection->GetAttribute<float>("SleepRemovalDuration", FGeometryCollection::TransformGroup);
-		const TManagedArray<float>& MaxSleepTime = DynamicCollection->GetAttribute<float>("MaxSleepTime", FGeometryCollection::TransformGroup);
-		TManagedArray<float>& Decay = DynamicCollection->ModifyAttribute<float>("Decay", FGeometryCollection::TransformGroup);
-		TManagedArray<FVector>& LastPosition = DynamicCollection->ModifyAttribute<FVector>("LastPosition", FGeometryCollection::TransformGroup);
-		const TManagedArray<uint8>& InternalClusterParentType = DynamicCollection->GetAttribute<uint8>("InternalClusterParentTypeArray", FGeometryCollection::TransformGroup);
-		const TManagedArray<int32>& Parents = DynamicCollection->Parent;
-		
-		TArray<int32> ToDisable;
-		TArray<FGeometryCollectionItemIndex> ToCrumbleParent;
-		for (int32 TransformIdx = 0; TransformIdx < SleepTimer.Num(); ++TransformIdx)
-		{
-			// update the sleeping timer accordingly
-			const int32 State = DynamicCollection->DynamicState[TransformIdx];
-			const bool bIsDynamicOrSleeping = (State == (int)EObjectStateTypeEnum::Chaos_Object_Sleeping) || (State == (int)EObjectStateTypeEnum::Chaos_Object_Dynamic);
-			const bool bNoParent = Parents[TransformIdx] == INDEX_NONE;
-			const bool bHasKinematicOrStaticInternalClusterParent = (InternalClusterParentType[TransformIdx] == (uint8)Chaos::EInternalClusterType::KinematicOrStatic);
-			const bool IsRemovalActive = (MaxSleepTime[TransformIdx] >= 0) && bIsDynamicOrSleeping && bNoParent && !bHasKinematicOrStaticInternalClusterParent;
-			if (IsRemovalActive)
-			{
-				bool bIsNotMoving = false;
-				if (RestCollection->bSlowMovingAsSleeping)
-				{
-					const FVector CurrentPosition = DynamicCollection->Transform[TransformIdx].GetTranslation();
-					LastPosition[TransformIdx] = CurrentPosition;
-					
-					const FVector::FReal InstantVelocity = (CurrentPosition-LastPosition[TransformIdx]).Size() / (FVector::FReal)DeltaTime; 
-					bIsNotMoving = (InstantVelocity < RestCollection->SlowMovingVelocityThreshold);
-				}
-				
-				const bool IsSleeping = (DynamicCollection->DynamicState[TransformIdx] == (int)EObjectStateTypeEnum::Chaos_Object_Sleeping);
-				const bool HasStartedDecaying = (Decay[TransformIdx] > 0); 
-				if (IsSleeping || bIsNotMoving || HasStartedDecaying)
-				{
-					SleepTimer[TransformIdx] += DeltaTime;
-				}
+	if (DynamicCollection && PhysicsProxy)
+	{
+		FGeometryCollectionRemoveOnSleepDynamicFacade RemoveOnSleepFacade(*DynamicCollection);
+		FGeometryCollectionDecayDynamicFacade DecayFacade(*DynamicCollection);
+		FGeometryCollectionDynamicStateFacade DynamicStateFacade(*DynamicCollection);
+
+		if (RemoveOnSleepFacade.IsValid()
+			&& DecayFacade.IsValid()
+			&& DynamicStateFacade.IsValid())
+		{
+			FGeometryCollectionDecayContext DecayContext(*PhysicsProxy, DecayFacade);
 			
-				// update the decay and disable the particle when decay has completed 
-				const bool bZeroRemovalDuration = (SleepRemovalDuration[TransformIdx] < UE_SMALL_NUMBER);
-				const float UpdatedDecay = bZeroRemovalDuration? 1.f: FMath::Clamp<float>((SleepTimer[TransformIdx] - MaxSleepTime[TransformIdx]) / SleepRemovalDuration[TransformIdx], 0.f, 1.f);
-				if (UpdatedDecay > Decay[TransformIdx])
+			TManagedArray<float>& Decay = DecayFacade.DecayAttribute.Modify();
+			for (int32 TransformIdx = 0; TransformIdx < Decay.Num(); ++TransformIdx)
+			{
+				// update the sleeping timer accordingly
+				const bool IsRemovalActive = RemoveOnSleepFacade.IsRemovalActive(TransformIdx);
+				const bool HasBrokenOff = DynamicStateFacade.HasBrokenOff(TransformIdx);
+				const bool HasDynamicInternalClusterParent = DynamicStateFacade.HasDynamicInternalClusterParent(TransformIdx);
+				if (IsRemovalActive && (HasBrokenOff || HasDynamicInternalClusterParent))
 				{
-					// if parent is internal and dynamic we need to break it before decaying 
-					if (InternalClusterParentType[TransformIdx] == (uint8)Chaos::EInternalClusterType::Dynamic)
+					// if decay has started we do not need to check slow moving or sleeping state anymore  
+					bool ShouldUpdateTimer = (Decay[TransformIdx] > 0);
+					if (!ShouldUpdateTimer && RestCollection->bSlowMovingAsSleeping)
 					{
-						FGeometryCollectionItemIndex InternalClusterItemindex = PhysicsProxy->GetInternalClusterParentItemIndex_External(TransformIdx);
-						if (InternalClusterItemindex.IsValid())
-						{
-							ToCrumbleParent.AddUnique(InternalClusterItemindex);
-						}
+						const FVector CurrentPosition = DynamicCollection->Transform[TransformIdx].GetTranslation();
+						ShouldUpdateTimer |= RemoveOnSleepFacade.ComputeSlowMovingState(TransformIdx, CurrentPosition, DeltaTime, RestCollection->SlowMovingVelocityThreshold);
 					}
-
-					DynamicCollection->MakeDirty();
-					Decay[TransformIdx] = UpdatedDecay;
-
-					if (Decay[TransformIdx] >= 1.0f)
+					if (ShouldUpdateTimer || DynamicStateFacade.IsSleeping(TransformIdx))
 					{
-						// Disable the particle if it has been asleep for the requisite time
-						Decay[TransformIdx] = 1.0f;
-						ToDisable.Add(TransformIdx);
+						RemoveOnSleepFacade.UpdateSleepTimer(TransformIdx, DeltaTime);
 					}
+			 
+					// update the decay and disable the particle when decay has completed 
+					const float UpdatedDecay = RemoveOnSleepFacade.ComputeDecay(TransformIdx);
+					UpdateDecay(TransformIdx, UpdatedDecay, DynamicStateFacade.HasChildren(TransformIdx), HasDynamicInternalClusterParent, DecayContext);
 				}
 			}
-		}
 
-		if (ToCrumbleParent.Num())
-		{
-			PhysicsProxy->BreakClusters_External(MoveTemp(ToCrumbleParent));
-		}
-		if (ToDisable.Num())
-		{
-			PhysicsProxy->DisableParticles_External(MoveTemp(ToDisable));
+			DecayContext.Process(*DynamicCollection);
 		}
 	}
 }
@@ -3683,69 +3628,38 @@ void UGeometryCollectionComponent::IncrementBreakTimer(float DeltaTime)
 		return;
 	}
 	
-	if (DynamicCollection && PhysicsProxy 
-		&& DynamicCollection->HasAttribute("BreakTimer", FGeometryCollection::TransformGroup) 
-		&& DynamicCollection->HasAttribute("PostBreakDuration", FGeometryCollection::TransformGroup)
-		&& DynamicCollection->HasAttribute("BreakRemovalDuration", FGeometryCollection::TransformGroup)
-		&& DynamicCollection->HasAttribute("Decay", FGeometryCollection::TransformGroup)
-		)
-	{
-		TManagedArray<float>& BreakTimer = DynamicCollection->ModifyAttribute<float>("BreakTimer", FGeometryCollection::TransformGroup);
-		const TManagedArray<float>& PostBreakDuration = DynamicCollection->GetAttribute<float>("PostBreakDuration", FGeometryCollection::TransformGroup);
-		const TManagedArray<float>& BreakRemovalDuration = DynamicCollection->GetAttribute<float>("BreakRemovalDuration", FGeometryCollection::TransformGroup);
-		TManagedArray<float>& Decay = DynamicCollection->ModifyAttribute<float>("Decay", FGeometryCollection::TransformGroup);
-		
-		TArray<int32> ToDisable;
-		TArray<FGeometryCollectionItemIndex> ToCrumble;
-		for (int32 TransformIdx = 0; TransformIdx < BreakTimer.Num(); ++TransformIdx)
+	if (DynamicCollection && PhysicsProxy)
+	{
+		FGeometryCollectionRemoveOnBreakDynamicFacade RemoveOnBreakFacade(*DynamicCollection);
+		FGeometryCollectionDecayDynamicFacade DecayFacade(*DynamicCollection);
+		FGeometryCollectionDynamicStateFacade DynamicStateFacade(*DynamicCollection);
+
+		if (RemoveOnBreakFacade.IsValid()
+			&& DecayFacade.IsValid()
+			&& DynamicStateFacade.IsValid())
 		{
-			const bool bIsBroken = DynamicCollection->Active[TransformIdx];
-			const bool bIsRemovedFromBreakEnabled = (PostBreakDuration[TransformIdx] >= 0);
+			FGeometryCollectionDecayContext DecayContext(*PhysicsProxy, DecayFacade);
 			
-			const int32 State = DynamicCollection->DynamicState[TransformIdx];
-			const bool bIsDynamicOrSleeping = (State == (int)EObjectStateTypeEnum::Chaos_Object_Sleeping) || (State == (int)EObjectStateTypeEnum::Chaos_Object_Dynamic);
-			
-			if (bIsBroken && bIsRemovedFromBreakEnabled && bIsDynamicOrSleeping)
+			TManagedArray<float>& Decay = DecayFacade.DecayAttribute.Modify();
+			for (int32 TransformIdx = 0; TransformIdx < Decay.Num(); ++TransformIdx)
 			{
-				BreakTimer[TransformIdx] += DeltaTime;
-				const bool bPostBreakTimeExpired = (BreakTimer[TransformIdx] >= PostBreakDuration[TransformIdx]);
-				const bool bZeroRemovalDuration = (BreakRemovalDuration[TransformIdx] < UE_SMALL_NUMBER);
-				float UpdatedBreakDecay = 0;
-				if (bPostBreakTimeExpired)
-				{
-					UpdatedBreakDecay = bZeroRemovalDuration? 1.f: FMath::Clamp<float>((BreakTimer[TransformIdx] - PostBreakDuration[TransformIdx]) / BreakRemovalDuration[TransformIdx], 0.f, 1.f);
-				}
-				if (UpdatedBreakDecay > Decay[TransformIdx])
+				const bool HasDynamicInternalClusterParent = DynamicStateFacade.HasDynamicInternalClusterParent(TransformIdx);
+				if (RemoveOnBreakFacade.IsRemovalActive(TransformIdx)
+					&& (DynamicStateFacade.HasBrokenOff(TransformIdx) || HasDynamicInternalClusterParent))
 				{
-					DynamicCollection->MakeDirty();
-					Decay[TransformIdx] = UpdatedBreakDecay;
-
-					const bool bUseClusterCrumbling = (BreakRemovalDuration[TransformIdx] < 0);
-					if (bUseClusterCrumbling)
+					bool UseClusterCrumbling = RemoveOnBreakFacade.UseClusterCrumbling(TransformIdx);
+					if (HasDynamicInternalClusterParent && DynamicCollection->Parent[TransformIdx] > INDEX_NONE)
 					{
-						ToCrumble.AddUnique(FGeometryCollectionItemIndex::CreateTransformItemIndex(TransformIdx));
-						Decay[TransformIdx] = 0;
-					}
-					else
-					{
-						if (Decay[TransformIdx] >= 1.0f)
-						{
-							// Disable the particle if it has decayed  the requisite time
-							Decay[TransformIdx] = 1.0f;
-							ToDisable.Add(TransformIdx);
-						}
+						// look up the state of the original parent to decide if the internal cluster parent should cumble 
+						UseClusterCrumbling = RemoveOnBreakFacade.UseClusterCrumbling(DynamicCollection->Parent[TransformIdx]);
 					}
+					
+					const float UpdatedBreakDecay = RemoveOnBreakFacade.UpdateBreakTimerAndComputeDecay(TransformIdx, DeltaTime);
+					UpdateDecay(TransformIdx, UpdatedBreakDecay, UseClusterCrumbling, HasDynamicInternalClusterParent, DecayContext);
 				}
 			}
-		}
 
-		if (ToCrumble.Num())
-		{
-			PhysicsProxy->BreakClusters_External(MoveTemp(ToCrumble));
-		}
-		if (ToDisable.Num())
-		{
-			PhysicsProxy->DisableParticles_External(MoveTemp(ToDisable));
+			DecayContext.Process(*DynamicCollection);
 		}
 	}
 }
diff --git a/Engine/Source/Runtime/Experimental/GeometryCollectionEngine/Private/GeometryCollection/GeometryCollectionEngineRemoval.cpp b/Engine/Source/Runtime/Experimental/GeometryCollectionEngine/Private/GeometryCollection/GeometryCollectionEngineRemoval.cpp
new file mode 100644
--- /dev/null
+++ b/Engine/Source/Runtime/Experimental/GeometryCollectionEngine/Private/GeometryCollection/GeometryCollectionEngineRemoval.cpp
@@ -0,0 +1,188 @@
+// Copyright Epic Games, Inc. All Rights Reserved.
+
+#include "GeometryCollection/GeometryCollectionEngineRemoval.h"
+#include "GeometryCollection/GeometryCollection.h"
+
+///////////////////////////////////////////////////////////////////////////////////////////////////////////////////////////
+
+FGeometryCollectionRemoveOnBreakDynamicFacade::FGeometryCollectionRemoveOnBreakDynamicFacade(FManagedArrayCollection& InCollection)
+	: BreakTimerAttribute(InCollection, "BreakTimer", FGeometryCollection::TransformGroup)
+	, PostBreakDurationAttribute(InCollection, "PostBreakDuration", FGeometryCollection::TransformGroup)
+	, BreakRemovalDurationAttribute(InCollection, "BreakRemovalDuration", FGeometryCollection::TransformGroup)
+{
+}
+
+bool FGeometryCollectionRemoveOnBreakDynamicFacade::IsValid() const
+{
+	return BreakTimerAttribute.IsValid()
+		&& PostBreakDurationAttribute.IsValid()
+		&& BreakRemovalDurationAttribute.IsValid()
+	;
+}
+
+void FGeometryCollectionRemoveOnBreakDynamicFacade::AddAttributes(const TManagedArray<FVector4f>& RemoveOnBreakAttribute, const TManagedArray<TSet<int32>>& ChildrenAttribute)
+{
+	BreakTimerAttribute.AddAndFill(DisabledBreakTimer);
+
+	if (!PostBreakDurationAttribute.IsValid())
+	{
+		TManagedArray<float>& PostBreakDuration = PostBreakDurationAttribute.Add();
+		for (int32 Idx = 0; Idx < PostBreakDuration.Num(); ++Idx)
+		{
+			const FRemoveOnBreakData RemoveOnBreakData{ RemoveOnBreakAttribute[Idx] };
+			const FVector2f PostBreakTimer = RemoveOnBreakData.GetBreakTimer();
+			const float MinTime = FMath::Max(0.0f, PostBreakTimer.X);
+			const float MaxTime = FMath::Max(MinTime, PostBreakTimer.Y);
+			PostBreakDuration[Idx] = RemoveOnBreakData.IsEnabled()? FMath::RandRange(MinTime, MaxTime): DisabledPostBreakDuration;
+		}
+	}
+
+	if (!BreakRemovalDurationAttribute.IsValid())
+	{
+		TManagedArray<float>& BreakRemovalDuration = BreakRemovalDurationAttribute.Add();
+		for (int32 Idx = 0; Idx < BreakRemovalDuration.Num(); ++Idx)
+		{
+			const FRemoveOnBreakData RemoveOnBreakData{ RemoveOnBreakAttribute[Idx] };
+			const FVector2f RemovalTimer = RemoveOnBreakData.GetRemovalTimer();
+			const float MinTime = FMath::Max(0.0f, RemovalTimer.X);
+			const float MaxTime = FMath::Max(MinTime, RemovalTimer.Y);
+			const bool bIsCluster = (ChildrenAttribute[Idx].Num() > 0);
+			const bool bUseClusterCrumbling = (bIsCluster && RemoveOnBreakData.GetClusterCrumbling());
+			BreakRemovalDuration[Idx] = bUseClusterCrumbling? CrumblingRemovalTimer: FMath::RandRange(MinTime, MaxTime);
+		}
+	}
+}
+
+bool FGeometryCollectionRemoveOnBreakDynamicFacade::IsRemovalActive(int32 TransformIndex) const
+{
+	return (PostBreakDurationAttribute.Get()[TransformIndex] > DisabledPostBreakDuration);
+}
+
+bool FGeometryCollectionRemoveOnBreakDynamicFacade::UseClusterCrumbling(int32 TransformIndex) const
+{
+	return (BreakRemovalDurationAttribute.Get()[TransformIndex] <= CrumblingRemovalTimer);
+}
+
+float FGeometryCollectionRemoveOnBreakDynamicFacade::UpdateBreakTimerAndComputeDecay(int32 TransformIndex, float DeltaTime)
+{
+	float& BreakTimerRef = BreakTimerAttribute.Modify()[TransformIndex];
+	const float PostBreakDuration = PostBreakDurationAttribute.Get()[TransformIndex];
+	const float BreakRemovalDuration = BreakRemovalDurationAttribute.Get()[TransformIndex];
+
+	float UpdatedBreakDecay = 0;
+	if (PostBreakDuration >= 0)
+	{
+		BreakTimerRef += DeltaTime;
+	
+		const bool bPostBreakTimeExpired = (BreakTimerRef >= PostBreakDuration) && (PostBreakDuration >= 0);
+		const bool bZeroRemovalDuration = (BreakRemovalDuration < UE_SMALL_NUMBER);
+		if (bPostBreakTimeExpired)
+		{
+			UpdatedBreakDecay = bZeroRemovalDuration? 1.f: FMath::Clamp<float>((BreakTimerRef - PostBreakDuration) / BreakRemovalDuration, 0.f, 1.f);
+		}
+	}
+	return UpdatedBreakDecay;
+}
+
+///////////////////////////////////////////////////////////////////////////////////////////////////////////////////////////
+
+FGeometryCollectionRemoveOnSleepDynamicFacade::FGeometryCollectionRemoveOnSleepDynamicFacade(FManagedArrayCollection& InCollection)
+	: SleepTimerAttribute(InCollection, "SleepTimer", FGeometryCollection::TransformGroup)
+	, MaxSleepTimeAttribute(InCollection, "MaxSleepTime", FGeometryCollection::TransformGroup)
+	, SleepRemovalDurationAttribute(InCollection, "SleepRemovalDuration", FGeometryCollection::TransformGroup)
+	, LastPositionAttribute(InCollection, "LastPosition", FGeometryCollection::TransformGroup)
+{
+}
+
+bool FGeometryCollectionRemoveOnSleepDynamicFacade::IsValid() const
+{
+	return SleepTimerAttribute.IsValid()
+		&& MaxSleepTimeAttribute.IsValid()
+		&& SleepRemovalDurationAttribute.IsValid()
+		&& LastPositionAttribute.IsValid()
+	;
+}
+
+void FGeometryCollectionRemoveOnSleepDynamicFacade::AddAttributes(const FVector2D& MaximumSleepTime, const FVector2D& RemovalDuration)
+{
+	SleepTimerAttribute.AddAndFill(0.0f);
+
+	LastPositionAttribute.AddAndFill(FVector::ZeroVector);
+	
+	if (!MaxSleepTimeAttribute.IsValid())
+	{
+		const float MinTime = FMath::Max(0.0f, MaximumSleepTime.X);
+		const float MaxTime = FMath::Max(MinTime, MaximumSleepTime.Y);
+		TManagedArray<float>& MaxSleepTime = MaxSleepTimeAttribute.Add();
+		for (int32 Idx = 0; Idx < MaxSleepTime.Num(); ++Idx)
+		{
+			MaxSleepTime[Idx] = FMath::RandRange(MinTime, MaxTime);
+		}
+	}
+
+	if (!SleepRemovalDurationAttribute.IsValid())
+	{
+		float MinTime = FMath::Max(0.0f, RemovalDuration.X);
+		float MaxTime = FMath::Max(MinTime, RemovalDuration.Y);
+		TManagedArray<float>& SleepRemovalDuration = SleepRemovalDurationAttribute.Add();
+		for (int32 Idx = 0; Idx < SleepRemovalDuration.Num(); ++Idx)
+		{
+			SleepRemovalDuration[Idx] = FMath::RandRange(MinTime, MaxTime);
+		}
+	}
+}
+
+bool FGeometryCollectionRemoveOnSleepDynamicFacade::IsRemovalActive(int32 TransformIndex) const
+{
+	return (MaxSleepTimeAttribute.Get()[TransformIndex] >= 0);
+}
+
+bool FGeometryCollectionRemoveOnSleepDynamicFacade::ComputeSlowMovingState(int32 TransformIndex, const FVector& Position, float DeltaTime, FVector::FReal VelocityThreshold)
+{
+	bool IsSlowMoving = false;
+	FVector& LastPositionRef = LastPositionAttribute.Modify()[TransformIndex];
+	if (DeltaTime > 0)
+	{
+		const FVector::FReal InstantVelocity = (Position-LastPositionRef).Size() / (FVector::FReal)DeltaTime; 
+		IsSlowMoving = (InstantVelocity < VelocityThreshold);
+	}
+	LastPositionRef = Position;
+	return IsSlowMoving;
+}
+
+void FGeometryCollectionRemoveOnSleepDynamicFacade::UpdateSleepTimer(int32 TransformIndex, float DeltaTime)
+{
+	check(SleepTimerAttribute.IsValid());
+	SleepTimerAttribute.Modify()[TransformIndex] += DeltaTime;
+}
+
+float FGeometryCollectionRemoveOnSleepDynamicFacade::ComputeDecay(int32 TransformIndex) const
+{
+	const float SleepRemovalDuration = SleepRemovalDurationAttribute.Get()[TransformIndex];
+	const float MaxSleepTime = MaxSleepTimeAttribute.Get()[TransformIndex];
+	const float SleepTimer = SleepTimerAttribute.Get()[TransformIndex];
+	const bool bZeroRemovalDuration = (SleepRemovalDuration < UE_SMALL_NUMBER);
+	return bZeroRemovalDuration? 1.f: FMath::Clamp<float>((SleepTimer -MaxSleepTime) / SleepRemovalDuration, 0.f, 1.f);
+}
+
+
+///////////////////////////////////////////////////////////////////////////////////////////////////////////////////////////
+
+FGeometryCollectionDecayDynamicFacade::FGeometryCollectionDecayDynamicFacade(FManagedArrayCollection& InCollection)
+	: DecayAttribute(InCollection, "Decay", FGeometryCollection::TransformGroup)
+	, UniformScaleAttribute(InCollection, "UniformScale", FGeometryCollection::TransformGroup)
+{
+}
+
+bool FGeometryCollectionDecayDynamicFacade::IsValid() const
+{
+	return DecayAttribute.IsValid()
+		&& UniformScaleAttribute.IsValid()
+	;
+}
+
+void FGeometryCollectionDecayDynamicFacade::AddAttributes()
+{
+	DecayAttribute.AddAndFill(0.0f);
+	UniformScaleAttribute.AddAndFill(FTransform::Identity);
+}
\ No newline at end of file
diff --git a/Engine/Source/Runtime/Experimental/GeometryCollectionEngine/Public/GeometryCollection/GeometryCollectionComponent.h b/Engine/Source/Runtime/Experimental/GeometryCollectionEngine/Public/GeometryCollection/GeometryCollectionComponent.h
--- a/Engine/Source/Runtime/Experimental/GeometryCollectionEngine/Public/GeometryCollection/GeometryCollectionComponent.h
+++ b/Engine/Source/Runtime/Experimental/GeometryCollectionEngine/Public/GeometryCollection/GeometryCollectionComponent.h
@@ -35,6 +35,8 @@ class UChaosPhysicalMaterial;
 class AChaosSolverActor;
 struct FGeometryCollectionEmbeddedExemplar;
 class UInstancedStaticMeshComponent;
+class FGeometryCollectionDecayDynamicFacade;
+struct FGeometryCollectionDecayContext;
 struct FDamageCollector;
 
 DECLARE_DYNAMIC_MULTICAST_DELEGATE_OneParam(FOnChaosBreakEvent, const FChaosBreakEvent&, BreakEvent);
@@ -167,42 +169,6 @@ namespace GeometryCollection
 	ENUM_CLASS_FLAGS(EEditUpdate);
 }
 
-// this class helps reading and writing the packed data used in the managed array
-struct FRemoveOnBreakData
-{
-public:
-	inline static const FVector4f DisabledPackedData{ -1, 0, 0, 0 };   
-	
-	FRemoveOnBreakData()
-		: PackedData(FRemoveOnBreakData::DisabledPackedData)
-	{}
-
-	FRemoveOnBreakData(const FVector4f& InPackedData)
-		: PackedData(InPackedData)
-	{}
-
-	FRemoveOnBreakData(bool bEnable, const FVector2f& BreakTimer, bool bClusterCrumbling, const FVector2f& RemovalTimer)
-	{
-		PackedData.X = FMath::Min(FMath::Abs(BreakTimer.X), FMath::Abs(BreakTimer.Y)); // Min break timer
-		PackedData.Y = FMath::Max(FMath::Abs(BreakTimer.X), FMath::Abs(BreakTimer.Y)); // Min break timer
-		PackedData.Z = FMath::Min(FMath::Abs(RemovalTimer.X), FMath::Abs(RemovalTimer.Y)); // Min removal timer
-		PackedData.W = FMath::Max(FMath::Abs(RemovalTimer.X), FMath::Abs(RemovalTimer.Y)); // Max removal timer
-
-		PackedData.X *= bEnable? +1: -1;
-		PackedData.Z *= bClusterCrumbling? -1: +1;
-	}
-	
-	bool IsEnabled() const { return PackedData.X >= 0; }
-	bool GetClusterCrumbling() const { return IsEnabled() && (PackedData.Z < 0); } 
-	FVector2f GetBreakTimer() const	{ return FVector2f(FMath::Abs(PackedData.X), FMath::Abs(PackedData.Y)); }
-	FVector2f GetRemovalTimer() const { return FVector2f(FMath::Abs(PackedData.Z), FMath::Abs(PackedData.W)); }
-
-	const FVector4f& GetPackedData() const { return PackedData; }
-
-private:
-	FVector4f PackedData;
-};
-
 /**
 *	FGeometryCollectionEdit
 *     Structured RestCollection access where the scope
@@ -1035,6 +1001,7 @@ private:
 
 	void IncrementSleepTimer(float DeltaTime);
 	void IncrementBreakTimer(float DeltaTime);
+	void UpdateDecay(int32 TransformIdx, float UpdatedDecay, bool UseClusterCrumbling, bool HasDynamicInternalClusterParent, FGeometryCollectionDecayContext& ContextInOut);
 	bool CalculateInnerSphere(int32 TransformIndex, FSphere& SphereOut) const;
 	void ProcessRepData();
 
diff --git a/Engine/Source/Runtime/Experimental/GeometryCollectionEngine/Public/GeometryCollection/GeometryCollectionEngineRemoval.h b/Engine/Source/Runtime/Experimental/GeometryCollectionEngine/Public/GeometryCollection/GeometryCollectionEngineRemoval.h
new file mode 100644
--- /dev/null
+++ b/Engine/Source/Runtime/Experimental/GeometryCollectionEngine/Public/GeometryCollection/GeometryCollectionEngineRemoval.h
@@ -0,0 +1,181 @@
+// Copyright Epic Games, Inc. All Rights Reserved.
+
+#pragma once
+
+#include "CoreMinimal.h"
+#include "GeometryCollection/ManagedArrayAccessor.h"
+
+struct FManagedArrayCollection;
+
+/**
+ * this class helps reading and writing the packed data used in the managed array
+ */
+struct FRemoveOnBreakData
+{
+public:
+	inline static const FVector4f DisabledPackedData{ -1, 0, 0, 0 };   
+	
+	FRemoveOnBreakData()
+		: PackedData(FRemoveOnBreakData::DisabledPackedData)
+	{}
+
+	FRemoveOnBreakData(const FVector4f& InPackedData)
+		: PackedData(InPackedData)
+	{}
+
+	FRemoveOnBreakData(bool bEnable, const FVector2f& BreakTimer, bool bClusterCrumbling, const FVector2f& RemovalTimer)
+	{
+		PackedData.X = FMath::Min(FMath::Abs(BreakTimer.X), FMath::Abs(BreakTimer.Y)); // Min break timer
+		PackedData.Y = FMath::Max(FMath::Abs(BreakTimer.X), FMath::Abs(BreakTimer.Y)); // Max break timer
+		PackedData.Z = FMath::Min(FMath::Abs(RemovalTimer.X), FMath::Abs(RemovalTimer.Y)); // Min removal timer
+		PackedData.W = FMath::Max(FMath::Abs(RemovalTimer.X), FMath::Abs(RemovalTimer.Y)); // Max removal timer
+
+		PackedData.X *= bEnable? +1: -1;
+		PackedData.Z *= bClusterCrumbling? -1: +1;
+	}
+	
+	bool IsEnabled() const { return PackedData.X >= 0; }
+	bool GetClusterCrumbling() const { return IsEnabled() && (PackedData.Z < 0); } 
+	FVector2f GetBreakTimer() const	{ return FVector2f(FMath::Abs(PackedData.X), FMath::Abs(PackedData.Y)); }
+	FVector2f GetRemovalTimer() const { return FVector2f(FMath::Abs(PackedData.Z), FMath::Abs(PackedData.W)); }
+
+	const FVector4f& GetPackedData() const { return PackedData; }
+
+private:
+	FVector4f PackedData;
+};
+
+/**
+ * Provides an API for the run time aspect of the remove on break feature
+ * this is to be used with the dynamic collection
+ */
+class FGeometryCollectionRemoveOnBreakDynamicFacade
+{
+public:
+	static constexpr float DisabledBreakTimer = -1;
+	static constexpr float DisabledPostBreakDuration = -1;
+	static constexpr float CrumblingRemovalTimer = -1;
+	
+	FGeometryCollectionRemoveOnBreakDynamicFacade(FManagedArrayCollection& InCollection);
+
+	/** 
+	 * returns true if all the necessary attributes are present
+	 * if not then the API can be used to create  
+	 */
+	bool IsValid() const;	
+
+	/**
+	 * Add the necessary attributes if they are missing and initialize them if necessary
+	 * @param RemoveOnBreakAttribute remove on break attribute from the rest collection
+	 * @param ChildrenAttribute remove on break attribute from the rest collection
+	*/
+	void AddAttributes(const TManagedArray<FVector4f>& RemoveOnBreakAttribute, const TManagedArray<TSet<int32>>& ChildrenAttribute);
+
+	/** true if the removal is active for a specific piece */
+	bool IsRemovalActive(int32 TransformIndex) const;
+
+	/** true if a specific transfom uses cluster crumbling */
+	bool UseClusterCrumbling(int32 TransformIndex) const;
+	
+	/**
+	 * Update break timer and return the matching decay  
+	 * @param TransformIndex index of the transform to update
+	 * @param DeltaTime elapsed time since the last update in second
+	 * @return decay value computed from the timer and duration ( [0,1] range )
+	 */
+	float UpdateBreakTimerAndComputeDecay(int32 TransformIndex, float DeltaTime);
+	
+private:
+	/** Time elapsed since the break in seconds */
+	TManagedArrayAccessor<float> BreakTimerAttribute;
+
+	/** duration after the break before the removal process starts */
+	TManagedArrayAccessor<float> PostBreakDurationAttribute;
+	
+	/** removal duration */
+	TManagedArrayAccessor<float> BreakRemovalDurationAttribute;
+};
+
+/**
+ * Provides an API for the run time aspect of the remove on sleep feature
+ * this is to be used with the dynamic collection
+ */
+class FGeometryCollectionRemoveOnSleepDynamicFacade
+{
+public:
+	FGeometryCollectionRemoveOnSleepDynamicFacade(FManagedArrayCollection& InCollection);
+
+	/** 
+	 * returns true if all the necessary attributes are present
+	 * if not then the API can be used to create  
+	 */
+	bool IsValid() const;	
+
+	/**
+	 * Add the necessary attributes if they are missing and initialize them if necessary
+	 * @param MaximumSleepTime range of time to initialize the sleep duration
+	 * @param RemovalDuration range of time to initialize the removal duration
+	*/
+	void AddAttributes(const FVector2D& MaximumSleepTime, const FVector2D& RemovalDuration);
+
+	/** true if the removal is active for a specific piece */
+	bool IsRemovalActive(int32 TransformIndex) const;
+
+	/**
+	 * Compute the slow moving state and update from the last position
+	 * After calling this method, LastPosition will be updated with Position
+	 * @param Position Current world position
+	 * @param DeltaTime elapsed time since last update
+	 * @return true if the piece if the piece is considered slow moving 
+	 **/
+	bool ComputeSlowMovingState(int32 TransformIndex, const FVector& Position, float DeltaTime, FVector::FReal VelocityThreshold);
+	
+	/**
+	 * Update the sleep timer
+	 * @param TransformIndex index of the transform to update
+	 * @param DeltaTime elapsed time since last update
+	 */
+	void UpdateSleepTimer(int32 TransformIndex, float DeltaTime);
+
+	/** Compute decay from elapsed timer and durantion attributes */
+	float ComputeDecay(int32 TransformIndex) const;
+
+	
+private:
+	/** Time elapsed since the sleep detection */
+	TManagedArrayAccessor<float> SleepTimerAttribute;
+	
+	/** duration after the sleep detection before the removal process starts (read only from outside) */
+	TManagedArrayAccessor<float> MaxSleepTimeAttribute;
+	
+	/** removal duration (read only from outside) */
+	TManagedArrayAccessor<float> SleepRemovalDurationAttribute;
+
+	/** Last position used to detect slow moving pieces */
+	TManagedArrayAccessor<FVector> LastPositionAttribute;
+};
+
+/**
+ * Provides an API for decay related attributes ( use for remove on break and remove on sleep )
+ */
+class FGeometryCollectionDecayDynamicFacade
+{
+public:
+	FGeometryCollectionDecayDynamicFacade(FManagedArrayCollection& InCollection);
+
+	/** 
+	 * returns true if all the necessary attributes are present
+	 * if not then the API can be used to create  
+	 */
+	bool IsValid() const;	
+
+	/** Add the necessary attributes if they are missing and initialize them if necessary */
+	void AddAttributes();
+
+	/** state of decay ([0-1] range) */
+	// @todo(chaos) this should eventually move to a common removal facade when break and sleep removal are consolidated
+	TManagedArrayAccessor<float> DecayAttribute;
+
+	/** scale transform used to shrink the geometry collection piece */
+	TManagedArrayAccessor<FTransform> UniformScaleAttribute;
+};
