Dataflow - Fixed Connect and Disconnect methods. Change reinterpret_cast to static_cast.

diff --git a/Engine/Source/Runtime/Experimental/Dataflow/Core/Private/Dataflow/DataflowGraph.cpp b/Engine/Source/Runtime/Experimental/Dataflow/Core/Private/Dataflow/DataflowGraph.cpp
--- a/Engine/Source/Runtime/Experimental/Dataflow/Core/Private/Dataflow/DataflowGraph.cpp
+++ b/Engine/Source/Runtime/Experimental/Dataflow/Core/Private/Dataflow/DataflowGraph.cpp
@@ -29,7 +29,7 @@ namespace Dataflow
 				{
 					if (Input)
 					{
-						Disconnect(Input, Output);
+						Disconnect(Output, Input);
 					}
 				}
 			}
@@ -43,7 +43,7 @@ namespace Dataflow
 				{
 					if (Output)
 					{
-						Disconnect(Input, Output);
+						Disconnect(Output, Input);
 					}
 				}
 			}
@@ -57,20 +57,20 @@ namespace Dataflow
 		//            or implement a better sync for the EdGraph and DataflowGraph
 		if (Connection->GetDirection() == FPin::EDirection::INPUT)
 		{
-			FDataflowInput* ConnectionIn = reinterpret_cast<FDataflowInput*>(Connection);
+			FDataflowInput* ConnectionIn = static_cast<FDataflowInput*>(Connection);
 			TArray<FDataflowOutput*> BaseOutputs = ConnectionIn->GetConnectedOutputs();
 			for (FDataflowOutput* Output : BaseOutputs)
 			{
-				Disconnect(ConnectionIn, Output);
+				Disconnect(Output, ConnectionIn);
 			}
 		}
 		else if (Connection->GetDirection() == FPin::EDirection::OUTPUT)
 		{
-			FDataflowOutput* ConnectionOut = reinterpret_cast<FDataflowOutput*>(Connection);
+			FDataflowOutput* ConnectionOut = static_cast<FDataflowOutput*>(Connection);
 			TArray<FDataflowInput*> BaseInputs = ConnectionOut->GetConnectedInputs();
 			for (FDataflowInput* Input : BaseInputs)
 			{
-				Disconnect(Input, ConnectionOut);
+				Disconnect(ConnectionOut, Input);
 			}
 		}
 	}
@@ -79,7 +79,7 @@ namespace Dataflow
 	{
 		for (FDataflowOutput* Output : InConnection->GetConnectedOutputs())
 		{
-			Disconnect(InConnection, Output);
+			Disconnect(Output, InConnection);
 		}
 	}
 
@@ -87,30 +87,30 @@ namespace Dataflow
 	{
 		for (FDataflowInput* Input : OutConnection->GetConnectedInputs())
 		{
-			Disconnect(Input, OutConnection);
+			Disconnect(OutConnection, Input);
 		}
 	}
 
 
-	void FGraph::Connect(FDataflowInput* InputConnection, FDataflowOutput* OutputConnection)
+	void FGraph::Connect(FDataflowOutput* OutputConnection, FDataflowInput* InputConnection)
 	{
-		if (ensure(InputConnection && OutputConnection))
+		if (ensure(OutputConnection && InputConnection))
 		{
-			InputConnection->AddConnection(OutputConnection);
 			OutputConnection->AddConnection(InputConnection);
+			InputConnection->AddConnection(OutputConnection);
 			Connections.Add(FLink(
-				InputConnection->GetOwningNode()->GetGuid(), InputConnection->GetGuid(), 
-				OutputConnection->GetOwningNode()->GetGuid(), OutputConnection->GetGuid()));
+				OutputConnection->GetOwningNode()->GetGuid(), OutputConnection->GetGuid(),
+				InputConnection->GetOwningNode()->GetGuid(), InputConnection->GetGuid()));
 		}
 	}
 
-	void FGraph::Disconnect(FDataflowInput* InputConnection, FDataflowOutput* OutputConnection)
+	void FGraph::Disconnect(FDataflowOutput* OutputConnection, FDataflowInput* InputConnection)
 	{
-		InputConnection->RemoveConnection(OutputConnection);
 		OutputConnection->RemoveConnection(InputConnection);
+		InputConnection->RemoveConnection(OutputConnection);
 		Connections.RemoveSwap(FLink(
-			InputConnection->GetOwningNode()->GetGuid(), InputConnection->GetGuid(),
-			OutputConnection->GetOwningNode()->GetGuid(), OutputConnection->GetGuid()));
+			OutputConnection->GetOwningNode()->GetGuid(), OutputConnection->GetGuid(),
+			InputConnection->GetOwningNode()->GetGuid(), InputConnection->GetGuid()));
 
 	}
 
@@ -212,7 +212,7 @@ namespace Dataflow
 					{
 						if (ConnectionGuidMap[Con.Input]->GetType() == ConnectionGuidMap[Con.Output]->GetType())
 						{
-							Connect(reinterpret_cast<FDataflowInput*>(ConnectionGuidMap[Con.Input]), reinterpret_cast<FDataflowOutput*>(ConnectionGuidMap[Con.Output]));
+							Connect(static_cast<FDataflowOutput*>(ConnectionGuidMap[Con.Output]), static_cast<FDataflowInput*>(ConnectionGuidMap[Con.Input]));
 						}
 					}
 				}
diff --git a/Engine/Source/Runtime/Experimental/Dataflow/Core/Public/Dataflow/DataflowGraph.h b/Engine/Source/Runtime/Experimental/Dataflow/Core/Public/Dataflow/DataflowGraph.h
--- a/Engine/Source/Runtime/Experimental/Dataflow/Core/Public/Dataflow/DataflowGraph.h
+++ b/Engine/Source/Runtime/Experimental/Dataflow/Core/Public/Dataflow/DataflowGraph.h
@@ -91,8 +91,8 @@ namespace Dataflow
 		void ClearConnections(FDataflowInput* Input);
 		void ClearConnections(FDataflowOutput* Output);
 
-		void Connect(FDataflowInput* InputConnection, FDataflowOutput* OutputConnection);
-		void Disconnect(FDataflowInput* InputConnection, FDataflowOutput* OutputConnection);
+		void Connect(FDataflowOutput* OutputConnection, FDataflowInput* InputConnection);
+		void Disconnect(FDataflowOutput* OutputConnection, FDataflowInput* InputConnection);
 
 		virtual void Serialize(FArchive& Ar);
 		const TSet<FName>& GetDisabledNodes() const { return DisabledNodes; }
diff --git a/Engine/Source/Runtime/Experimental/Dataflow/Engine/Private/Dataflow/DataflowEdNode.cpp b/Engine/Source/Runtime/Experimental/Dataflow/Engine/Private/Dataflow/DataflowEdNode.cpp
--- a/Engine/Source/Runtime/Experimental/Dataflow/Engine/Private/Dataflow/DataflowEdNode.cpp
+++ b/Engine/Source/Runtime/Experimental/Dataflow/Engine/Private/Dataflow/DataflowEdNode.cpp
@@ -69,9 +69,9 @@ void UDataflowEdNode::PinConnectionListChanged(UEdGraphPin* Pin)
 							{
 								if (TSharedPtr<FDataflowNode> LinkedDataflowNode = DataflowGraph->FindBaseNode(LinkedNode->GetDataflowNodeGuid()))
 								{
-									if (FDataflowOutput* LinkedConBase = LinkedDataflowNode->FindOutput(FName(LinkedCon->GetName())))
+									if (FDataflowOutput* LinkedConOutput = LinkedDataflowNode->FindOutput(FName(LinkedCon->GetName())))
 									{
-										DataflowGraph->Connect(ConnectionInput, LinkedConBase);
+										DataflowGraph->Connect(LinkedConOutput, ConnectionInput);
 									}
 								}
 							}
@@ -92,9 +92,9 @@ void UDataflowEdNode::PinConnectionListChanged(UEdGraphPin* Pin)
 							{
 								if (TSharedPtr<FDataflowNode> LinkedDataflowNode = DataflowGraph->FindBaseNode(LinkedNode->GetDataflowNodeGuid()))
 								{
-									if (FDataflowInput* LinkedConBase = LinkedDataflowNode->FindInput(FName(LinkedCon->GetName())))
+									if (FDataflowInput* LinkedConInput = LinkedDataflowNode->FindInput(FName(LinkedCon->GetName())))
 									{
-										DataflowGraph->Connect(LinkedConBase, ConnectionOutput);
+										DataflowGraph->Connect(ConnectionOutput, LinkedConInput);
 									}
 								}
 							}
