Dataflow - Changed GetOffset() and RealAddress()'s return types to make it more consistent.

diff --git a/Engine/Source/Runtime/Experimental/Dataflow/Core/Private/Dataflow/DataflowConnection.cpp b/Engine/Source/Runtime/Experimental/Dataflow/Core/Private/Dataflow/DataflowConnection.cpp
--- a/Engine/Source/Runtime/Experimental/Dataflow/Core/Private/Dataflow/DataflowConnection.cpp
+++ b/Engine/Source/Runtime/Experimental/Dataflow/Core/Private/Dataflow/DataflowConnection.cpp
@@ -14,14 +14,13 @@ FDataflowConnection::FDataflowConnection(Dataflow::FPin::EDirection InDirection,
 	, Guid(InGuid)
 {}
 
-// @todo(dataflow) check that the return types are all consistent
-uint32 FDataflowConnection::GetOffset() const
+int32 FDataflowConnection::GetOffset() const
 {
 	if (ensure(Property != nullptr))
 	{
 		return Property->GetOffset_ForInternal();
 	}
-	return (uint32)INDEX_NONE;
+	return INDEX_NONE;
 }
 
 
diff --git a/Engine/Source/Runtime/Experimental/Dataflow/Core/Private/Dataflow/DataflowNode.cpp b/Engine/Source/Runtime/Experimental/Dataflow/Core/Private/Dataflow/DataflowNode.cpp
--- a/Engine/Source/Runtime/Experimental/Dataflow/Core/Private/Dataflow/DataflowNode.cpp
+++ b/Engine/Source/Runtime/Experimental/Dataflow/Core/Private/Dataflow/DataflowNode.cpp
@@ -52,7 +52,7 @@ const FDataflowInput* FDataflowNode::FindInput(const void* Reference) const
 	for (TPair<uint32, FDataflowInput*> Elem : Inputs)
 	{
 		FDataflowInput* Con = Elem.Value;
-		if (Con->RealAddress() == (size_t)Reference)
+		if (Con->RealAddress() == Reference)
 		{
 			return Con;
 		}
@@ -65,7 +65,7 @@ FDataflowInput* FDataflowNode::FindInput(void* Reference)
 	for (TPair<uint32, FDataflowInput*> Elem : Inputs)
 	{
 		FDataflowInput* Con = Elem.Value;
-		if (Con->RealAddress() == (size_t)Reference)
+		if (Con->RealAddress() == Reference)
 		{
 			return Con;
 		}
@@ -137,7 +137,7 @@ const FDataflowOutput* FDataflowNode::FindOutput(const void* Reference) const
 	for (TPair<uint32, FDataflowOutput*> Elem : Outputs)
 	{
 		FDataflowOutput* Con = Elem.Value;
-		if (Con->RealAddress() == (size_t)Reference)
+		if (Con->RealAddress() == Reference)
 		{
 			return (FDataflowOutput*)Con;
 		}
@@ -150,7 +150,7 @@ FDataflowOutput* FDataflowNode::FindOutput(void* Reference)
 	for (TPair<uint32, FDataflowOutput*> Elem : Outputs)
 	{
 		FDataflowOutput* Con = Elem.Value;
-		if (Con->RealAddress() == (size_t)Reference)
+		if (Con->RealAddress() == Reference)
 		{
 			return (FDataflowOutput*)Con;
 		}
diff --git a/Engine/Source/Runtime/Experimental/Dataflow/Core/Public/Dataflow/DataflowConnection.h b/Engine/Source/Runtime/Experimental/Dataflow/Core/Public/Dataflow/DataflowConnection.h
--- a/Engine/Source/Runtime/Experimental/Dataflow/Core/Public/Dataflow/DataflowConnection.h
+++ b/Engine/Source/Runtime/Experimental/Dataflow/Core/Public/Dataflow/DataflowConnection.h
@@ -56,7 +56,7 @@ public:
 	const FDataflowNode* GetOwningNode() const { return OwningNode; }
 
 	Dataflow::FPin::EDirection GetDirection() const { return Direction; }
-	uint32 GetOffset( ) const;
+	int32 GetOffset( ) const;
 
 	FName GetType() const { return Type; }
 
@@ -66,8 +66,8 @@ public:
 	FName GetName() const { return Name; }
 	void SetName(FName InName) { Name = InName; }
 
-	size_t RealAddress() const { ensure(OwningNode);  return (size_t)OwningNode + (size_t)GetOffset(); };
-	size_t CacheKey() const { return RealAddress(); };
+	void* RealAddress() const { ensure(OwningNode);  return (void*)((size_t)OwningNode + (size_t)GetOffset()); };
+	size_t CacheKey() const { return (size_t)RealAddress(); };
 
 	virtual bool AddConnection(FDataflowConnection* In) { return false; };
 	virtual bool RemoveConnection(FDataflowConnection* In) { return false; }
