Chaos geometry collection : BP functions
- extended apply external strain node to have a radius and proapgation  parameters
- added an internal strain node similar to external strain

diff --git a/Engine/Source/Runtime/Experimental/Chaos/Private/Chaos/PBDRigidClustering.cpp b/Engine/Source/Runtime/Experimental/Chaos/Private/Chaos/PBDRigidClustering.cpp
--- a/Engine/Source/Runtime/Experimental/Chaos/Private/Chaos/PBDRigidClustering.cpp
+++ b/Engine/Source/Runtime/Experimental/Chaos/Private/Chaos/PBDRigidClustering.cpp
@@ -952,7 +952,7 @@ namespace Chaos
 		return nullptr;
 	}
 
-	FPBDRigidParticleHandle* FRigidClustering::FindClosestParticle(const TArray<FPBDRigidParticleHandle*>& Particles, const FVec3& WorldLocation) const
+	FPBDRigidParticleHandle* FRigidClustering::FindClosestParticle(const TArray<FPBDRigidParticleHandle*>& Particles, const FVec3& WorldLocation)
 	{
 		FPBDRigidParticleHandle* ClosestChildHandle = nullptr;
 		
@@ -970,6 +970,49 @@ namespace Chaos
         return ClosestChildHandle;
 	}
 
+	TArray<FPBDRigidParticleHandle*> FRigidClustering::FindChildrenWithinRadius(const FPBDRigidClusteredParticleHandle* ClusteredParticle, const FVec3& WorldLocation, FReal Radius, bool bAlwaysReturnClosest) const
+	{
+		if (const TArray<FPBDRigidParticleHandle*>* ChildrenHandles = GetChildrenMap().Find(ClusteredParticle))
+		{
+			return FindParticlesWithinRadius(*ChildrenHandles, WorldLocation, Radius, bAlwaysReturnClosest); 
+		}
+		TArray<FPBDRigidParticleHandle*> EmptyArray;
+		return EmptyArray;
+	}
+
+	TArray<FPBDRigidParticleHandle*> FRigidClustering::FindParticlesWithinRadius(const TArray<FPBDRigidParticleHandle*>& Particles, const FVec3& WorldLocation, FReal Radius, bool bAlwaysReturnClosest)
+	{
+		TArray<FPBDRigidParticleHandle*> Result;
+		
+		FPBDRigidParticleHandle* ClosestChildHandle = nullptr;
+		
+		// @todo(chaos) we should offer a more precise way to query than the distance from center of mass
+		FReal ClosestSquaredDist = TNumericLimits<FReal>::Max();
+		
+		const FReal RadiusSquared = Radius * Radius;
+		for (FPBDRigidParticleHandle* ChildHandle: Particles)
+		{
+			const FReal SquaredDist = (ChildHandle->X() - WorldLocation).SizeSquared();
+			if (SquaredDist <= RadiusSquared)
+			{
+				Result.Add(ChildHandle);
+			}
+			if (bAlwaysReturnClosest)
+			{
+				if (SquaredDist < ClosestSquaredDist)
+				{
+					ClosestSquaredDist = SquaredDist;
+					ClosestChildHandle = ChildHandle;
+				}
+			}
+		}
+		if (bAlwaysReturnClosest && ClosestChildHandle && Result.Num() == 0)
+		{
+			Result.Add(ClosestChildHandle);
+		}
+		return Result;
+	}
+	
 	DECLARE_CYCLE_STAT(TEXT("TPBDRigidClustering<>::GenerateConnectionGraph"), STAT_GenerateConnectionGraph, STATGROUP_Chaos);
 	void 
 	FRigidClustering::GenerateConnectionGraph(
diff --git a/Engine/Source/Runtime/Experimental/Chaos/Private/PhysicsProxy/GeometryCollectionPhysicsProxy.cpp b/Engine/Source/Runtime/Experimental/Chaos/Private/PhysicsProxy/GeometryCollectionPhysicsProxy.cpp
--- a/Engine/Source/Runtime/Experimental/Chaos/Private/PhysicsProxy/GeometryCollectionPhysicsProxy.cpp
+++ b/Engine/Source/Runtime/Experimental/Chaos/Private/PhysicsProxy/GeometryCollectionPhysicsProxy.cpp
@@ -747,7 +747,7 @@ void FGeometryCollectionPhysicsProxy::CreateNonClusteredParticles(Chaos::FPBDRig
 	}
 }
 
-Chaos::FPBDRigidClusteredParticleHandle* FGeometryCollectionPhysicsProxy::FindClusteredParticleHandleByItemIndex(FGeometryCollectionItemIndex ItemIndex) const
+Chaos::FPBDRigidClusteredParticleHandle* FGeometryCollectionPhysicsProxy::FindClusteredParticleHandleByItemIndex_Internal(FGeometryCollectionItemIndex ItemIndex) const
 {
 	Chaos::FPBDRigidClusteredParticleHandle* ResultHandle = nullptr;;
 	if (ItemIndex.IsInternalCluster())
@@ -1546,7 +1546,7 @@ void FGeometryCollectionPhysicsProxy::BreakClusters_External(TArray<FGeometryCol
 			Chaos::FRigidClustering& Clustering = RBDSolver->GetEvolution()->GetRigidClustering();
 			for (const FGeometryCollectionItemIndex& ItemIndex : IndicesToBreakParent)
 			{
-				if (Chaos::FPBDRigidClusteredParticleHandle* ClusteredHandle = FindClusteredParticleHandleByItemIndex(ItemIndex))
+				if (Chaos::FPBDRigidClusteredParticleHandle* ClusteredHandle = FindClusteredParticleHandleByItemIndex_Internal(ItemIndex))
 				{
 					Clustering.BreakCluster(ClusteredHandle);
 				}
@@ -1555,42 +1555,122 @@ void FGeometryCollectionPhysicsProxy::BreakClusters_External(TArray<FGeometryCol
 	}
 }
 
-void FGeometryCollectionPhysicsProxy::ApplyStrain_External(FGeometryCollectionItemIndex ItemIndex, const FVector& WorldLocation, float StrainValue)
+template <typename TAction>
+static void ApplyToChildrenAtPointWithRadiusAndPropagation_Internal(
+	Chaos::FRigidClustering& Clustering,
+	Chaos::FPBDRigidClusteredParticleHandle& ClusteredHandle,
+	const FVector& WorldLocation,
+	float Radius,
+	int32 PropagationDepth,
+	float PropagationFactor,
+	TAction Action)
 {
-	check(IsInGameThread());
-
-	if (Chaos::FPhysicsSolver* RBDSolver = GetSolver<Chaos::FPhysicsSolver>())
+	const bool bIsCluster = (ClusteredHandle.ClusterIds().NumChildren > 0);
+	if (bIsCluster)
 	{
-		RBDSolver->EnqueueCommandImmediate([this, StrainValue, RBDSolver, WorldLocation, ItemIndex]()
+		TArray<Chaos::FPBDRigidParticleHandle*> ChildrenHandles = Clustering.FindChildrenWithinRadius(&ClusteredHandle, WorldLocation, Radius, true /* always return closest */ );
+		if (ChildrenHandles.Num())
 		{
-			if (Chaos::FPBDRigidClusteredParticleHandle* ClusteredHandle = FindClusteredParticleHandleByItemIndex(ItemIndex))
+			TArray<Chaos::FPBDRigidParticleHandle*> ProcessedHandles;
+			Chaos::FReal PropagationMultiplier = 1.0f;
+
+			int32 CurrentPropagationDepth = PropagationDepth;
+			while (CurrentPropagationDepth >= 0)
 			{
-				const bool bIsCluster = (ClusteredHandle->ClusterIds().NumChildren > 0);
-				if (bIsCluster)
+				// apply to current
+				for (Chaos::FPBDRigidParticleHandle* ChildHandle: ChildrenHandles)
 				{
-					Chaos::FRigidClustering& Clustering = RBDSolver->GetEvolution()->GetRigidClustering();
-					Chaos::FPBDRigidParticleHandle* ClosestChildHandle = Clustering.FindClosestChild(ClusteredHandle, WorldLocation);
-					if (ClosestChildHandle)
+					if (Chaos::FPBDRigidClusteredParticleHandle* ClusteredChildHandle = ChildHandle->CastToClustered())
 					{
-						if (Chaos::FPBDRigidClusteredParticleHandle* ClusteredChildHandle = ClosestChildHandle->CastToClustered())
+						Action(ClusteredChildHandle, PropagationMultiplier);
+//						ClusteredChildHandle->SetExternalStrain(FMath::Max(ClusteredChildHandle->GetExternalStrain(), StrainValue * PropagationMultiplier));
+					}
+				}
+
+				// find handles to propagate to
+				if (CurrentPropagationDepth > 0)
+				{
+					// only propagate to newly added ( avoids allocating a new temporary array for next level handles )
+					// move the entire ChildrenHandle array to the processed ones so that it also resets it for us
+					const int32 StartIndex = ProcessedHandles.Num(); 
+					ProcessedHandles.Append(MoveTemp(ChildrenHandles));
+					
+					for (int32 ChildIndex = StartIndex; ChildIndex < ProcessedHandles.Num(); ++ChildIndex)
+					{
+						if (Chaos::FPBDRigidClusteredParticleHandle* ClusteredChildHandle = ProcessedHandles[ChildIndex]->CastToClustered())
 						{
-							ClusteredChildHandle->SetExternalStrain(FMath::Max(ClusteredChildHandle->GetExternalStrain(), StrainValue));
+							for (const Chaos::TConnectivityEdge<Chaos::FReal>& Edge:ClusteredChildHandle->ConnectivityEdges())
+							{
+								if (Edge.Sibling && !ChildrenHandles.Contains(Edge.Sibling))
+								{
+									ChildrenHandles.Add(Edge.Sibling);
+								}
+							}
 						}
 					}
 				}
+				
+				PropagationMultiplier *= static_cast<Chaos::FReal>(PropagationFactor);
+				--CurrentPropagationDepth;
+			}
+		}
+	}
+}
+
+
+void FGeometryCollectionPhysicsProxy::ApplyExternalStrain_External(FGeometryCollectionItemIndex ItemIndex, const FVector& WorldLocation, float Radius, int32 PropagationDepth, float PropagationFactor, float StrainValue)
+{
+	check(IsInGameThread());
+
+	if (Chaos::FPhysicsSolver* RBDSolver = GetSolver<Chaos::FPhysicsSolver>())
+	{
+		RBDSolver->EnqueueCommandImmediate([this, Radius, StrainValue, PropagationDepth, PropagationFactor, RBDSolver, WorldLocation, ItemIndex]()
+		{
+			if (Chaos::FPBDRigidClusteredParticleHandle* ClusteredHandle = FindClusteredParticleHandleByItemIndex_Internal(ItemIndex))
+			{
+				ApplyToChildrenAtPointWithRadiusAndPropagation_Internal(
+					RBDSolver->GetEvolution()->GetRigidClustering(), *ClusteredHandle,
+					WorldLocation, Radius, PropagationDepth, PropagationFactor,
+					[StrainValue](Chaos::FPBDRigidClusteredParticleHandle* ClusteredChildHandle, Chaos::FReal PropagationMultiplier)
+					{
+						ClusteredChildHandle->SetExternalStrain(FMath::Max(ClusteredChildHandle->GetExternalStrain(), StrainValue * PropagationMultiplier));
+					});
+			}
+		});
+	}
+}
+
+void FGeometryCollectionPhysicsProxy::ApplyInternalStrain_External(FGeometryCollectionItemIndex ItemIndex, const FVector& WorldLocation, float Radius, int32 PropagationDepth, float PropagationFactor, float StrainValue)
+{
+	check(IsInGameThread());
+
+	if (Chaos::FPhysicsSolver* RBDSolver = GetSolver<Chaos::FPhysicsSolver>())
+	{
+		RBDSolver->EnqueueCommandImmediate([this, Radius, StrainValue, PropagationDepth, PropagationFactor, RBDSolver, WorldLocation, ItemIndex]()
+		{
+			if (Chaos::FPBDRigidClusteredParticleHandle* ClusteredHandle = FindClusteredParticleHandleByItemIndex_Internal(ItemIndex))
+			{
+				ApplyToChildrenAtPointWithRadiusAndPropagation_Internal(
+					RBDSolver->GetEvolution()->GetRigidClustering(), *ClusteredHandle,
+					WorldLocation, Radius, PropagationDepth, PropagationFactor,
+					[StrainValue](Chaos::FPBDRigidClusteredParticleHandle* ClusteredChildHandle, Chaos::FReal PropagationMultiplier)
+					{
+						const Chaos::FReal NewInternalStrain = ClusteredChildHandle->Strain() - ((Chaos::FReal)StrainValue * PropagationMultiplier);
+						ClusteredChildHandle->SetStrain(FMath::Max(0,NewInternalStrain));
+					});
 			}
 		});
 	}
 }
 
-template <typename TLambda>
-static void ApplyToBreakingChildren_Internal(Chaos::FRigidClustering& Clustering, Chaos::FPBDRigidClusteredParticleHandle* ClusteredHandle, TLambda Action)
+template <typename TAction>
+static void ApplyToBreakingChildren_Internal(Chaos::FRigidClustering& Clustering, Chaos::FPBDRigidClusteredParticleHandle& ClusteredHandle, TAction Action)
 {
-	const bool bIsCluster = (ClusteredHandle->ClusterIds().NumChildren > 0);
-	if (bIsCluster && !ClusteredHandle->Disabled())
+	const bool bIsCluster = (ClusteredHandle.ClusterIds().NumChildren > 0);
+	if (bIsCluster && !ClusteredHandle.Disabled())
 	{
 		Chaos::FRigidClustering::FClusterMap& ChildrenMap = Clustering.GetChildrenMap();
-		if (const TArray<Chaos::FPBDRigidParticleHandle*>* ChildrenHandles = ChildrenMap.Find(ClusteredHandle))
+		if (const TArray<Chaos::FPBDRigidParticleHandle*>* ChildrenHandles = ChildrenMap.Find(&ClusteredHandle))
 		{
 			for (Chaos::FPBDRigidParticleHandle* ChildHandle: *ChildrenHandles)
 			{
@@ -1614,9 +1694,9 @@ void FGeometryCollectionPhysicsProxy::ApplyBreakingLinearVelocity_External(FGeom
 	{
 		RBDSolver->EnqueueCommandImmediate([this, RBDSolver, ItemIndex, LinearVelocity]()
 		{
-			if (Chaos::FPBDRigidClusteredParticleHandle* ClusteredHandle = FindClusteredParticleHandleByItemIndex(ItemIndex))
+			if (Chaos::FPBDRigidClusteredParticleHandle* ClusteredHandle = FindClusteredParticleHandleByItemIndex_Internal(ItemIndex))
 			{
-				ApplyToBreakingChildren_Internal(RBDSolver->GetEvolution()->GetRigidClustering(), ClusteredHandle,
+				ApplyToBreakingChildren_Internal(RBDSolver->GetEvolution()->GetRigidClustering(), *ClusteredHandle,
 					[LinearVelocity](Chaos::FPBDRigidClusteredParticleHandle* ClusteredChildHandle)
 					{
 						ClusteredChildHandle->V() += LinearVelocity;
@@ -1634,9 +1714,9 @@ void FGeometryCollectionPhysicsProxy::ApplyBreakingAngularVelocity_External(FGeo
 	{
 		RBDSolver->EnqueueCommandImmediate([this, RBDSolver, ItemIndex, AngularVelocity]()
 		{
-			if (Chaos::FPBDRigidClusteredParticleHandle* ClusteredHandle = FindClusteredParticleHandleByItemIndex(ItemIndex))
+			if (Chaos::FPBDRigidClusteredParticleHandle* ClusteredHandle = FindClusteredParticleHandleByItemIndex_Internal(ItemIndex))
 			{
-				ApplyToBreakingChildren_Internal(RBDSolver->GetEvolution()->GetRigidClustering(), ClusteredHandle,
+				ApplyToBreakingChildren_Internal(RBDSolver->GetEvolution()->GetRigidClustering(), *ClusteredHandle,
 					[AngularVelocity](Chaos::FPBDRigidClusteredParticleHandle* ClusteredChildHandle)
 					{
 						ClusteredChildHandle->W() += AngularVelocity;
@@ -1654,7 +1734,7 @@ void FGeometryCollectionPhysicsProxy::ApplyLinearVelocity_External(FGeometryColl
 	{
 		RBDSolver->EnqueueCommandImmediate([this, ItemIndex, LinearVelocity]()
 		{
-			if (Chaos::FPBDRigidClusteredParticleHandle* ClusteredHandle = FindClusteredParticleHandleByItemIndex(ItemIndex))
+			if (Chaos::FPBDRigidClusteredParticleHandle* ClusteredHandle = FindClusteredParticleHandleByItemIndex_Internal(ItemIndex))
 			{
 				ClusteredHandle->V() += LinearVelocity;
 			}
@@ -1670,7 +1750,7 @@ void FGeometryCollectionPhysicsProxy::ApplyAngularVelocity_External(FGeometryCol
 	{
 		RBDSolver->EnqueueCommandImmediate([this, ItemIndex, AngularVelocity]()
 		{
-			if (Chaos::FPBDRigidClusteredParticleHandle* ClusteredHandle = FindClusteredParticleHandleByItemIndex(ItemIndex))
+			if (Chaos::FPBDRigidClusteredParticleHandle* ClusteredHandle = FindClusteredParticleHandleByItemIndex_Internal(ItemIndex))
 			{
 				ClusteredHandle->W() += AngularVelocity;
 			}
diff --git a/Engine/Source/Runtime/Experimental/Chaos/Public/Chaos/PBDRigidClustering.h b/Engine/Source/Runtime/Experimental/Chaos/Public/Chaos/PBDRigidClustering.h
--- a/Engine/Source/Runtime/Experimental/Chaos/Public/Chaos/PBDRigidClustering.h
+++ b/Engine/Source/Runtime/Experimental/Chaos/Public/Chaos/PBDRigidClustering.h
@@ -278,7 +278,32 @@ public:
 	*    @param Particles array of clustered particles to select from  
 	*    @param WorldLocation world space location to find the closest particle from   
 	*/
-	FPBDRigidParticleHandle* FindClosestParticle(const TArray<FPBDRigidParticleHandle*>& Particles, const FVec3& WorldLocation) const;
+	static FPBDRigidParticleHandle* FindClosestParticle(const TArray<FPBDRigidParticleHandle*>& Particles, const FVec3& WorldLocation);
+
+	/*
+	*  FindChildrenWithinRadius
+	*    Find the children of an active cluster from a world position and a radius
+	*    current implementation will pick the closest based on the distance from the center
+	*    future implementation may expose option for more precise queries
+	*    if bAlwaysReturnClosest is checked this will always return  the closest from the location even if the radius does not encompass any center of mass  
+	*    @param ClusteredParticle active cluster handle to query the children from
+	*    @param WorldLocation world space location to find the closest child from
+	*    @param Radius Radius to use from the WorldLocation
+	*    @param bAlwaysReturnClosest if radius query does not return anything still return the closest from the point
+	*/
+	TArray<FPBDRigidParticleHandle*> FindChildrenWithinRadius(const FPBDRigidClusteredParticleHandle* ClusteredParticle, const FVec3& WorldLocation, FReal Radius, bool bAlwaysReturnClosest) const;
+
+	/*
+	*  FindParticlesWithinRadius
+	*    Find the closest particle from an array of  particle from a world position and a radius 
+	*    current implementation will pick the closest based on the distance from the center
+	*    future implementation may expose option for more precise queries
+	*    if bAlwaysReturnClosest is checked this will always return  the closest from the location even if the radius does not encompass any center of mass
+	*    @param Particles array of clustered particles to select from  
+	*    @param WorldLocation world space location to find the closest particle from
+	*    @param bAlwaysReturnClosest if radius query does not return anything still return the closest from the point
+	*/
+	static TArray<FPBDRigidParticleHandle*> FindParticlesWithinRadius(const TArray<FPBDRigidParticleHandle*>& Particles, const FVec3& WorldLocation, FReal Radius, bool bAlwaysReturnClosest);
 	
 	/**
 	* GenerateConnectionGraph
diff --git a/Engine/Source/Runtime/Experimental/Chaos/Public/PhysicsProxy/GeometryCollectionPhysicsProxy.h b/Engine/Source/Runtime/Experimental/Chaos/Public/PhysicsProxy/GeometryCollectionPhysicsProxy.h
--- a/Engine/Source/Runtime/Experimental/Chaos/Public/PhysicsProxy/GeometryCollectionPhysicsProxy.h
+++ b/Engine/Source/Runtime/Experimental/Chaos/Public/PhysicsProxy/GeometryCollectionPhysicsProxy.h
@@ -361,7 +361,8 @@ public:
 	void DisableParticles_External(TArray<int32>&& TransformGroupIndices);
 
 	void BreakClusters_External(TArray<FGeometryCollectionItemIndex>&& ItemIndices);
-	void ApplyStrain_External(FGeometryCollectionItemIndex ItemIndex, const FVector& WorldLocation, float StrainValue);
+	void ApplyExternalStrain_External(FGeometryCollectionItemIndex ItemIndex, const FVector& WorldLocation, float Radius, int32 PropagationDepth, float PropagationFactor, float StrainValue);
+	void ApplyInternalStrain_External(FGeometryCollectionItemIndex ItemIndex, const FVector& WorldLocation, float Radius, int32 PropagationDepth, float PropagationFactor, float StrainValue);
 	void ApplyBreakingLinearVelocity_External(FGeometryCollectionItemIndex ItemIndex, const FVector& LinearVelocity);
 	void ApplyBreakingAngularVelocity_External(FGeometryCollectionItemIndex ItemIndex, const FVector& AngularVelocity);
 	void ApplyLinearVelocity_External(FGeometryCollectionItemIndex ItemIndex, const FVector& LinearVelocity);
@@ -417,7 +418,7 @@ protected:
 
 	void CreateNonClusteredParticles(Chaos::FPBDRigidsSolver* RigidsSolver,	const FGeometryCollection& RestCollection, const FGeometryDynamicCollection& DynamicCollection);
 
-	Chaos::FPBDRigidClusteredParticleHandle* FindClusteredParticleHandleByItemIndex(FGeometryCollectionItemIndex ItemIndex) const;
+	Chaos::FPBDRigidClusteredParticleHandle* FindClusteredParticleHandleByItemIndex_Internal(FGeometryCollectionItemIndex ItemIndex) const;
 	
 private:
 
diff --git a/Engine/Source/Runtime/Experimental/GeometryCollectionEngine/Private/GeometryCollection/GeometryCollectionComponent.cpp b/Engine/Source/Runtime/Experimental/GeometryCollectionEngine/Private/GeometryCollection/GeometryCollectionComponent.cpp
--- a/Engine/Source/Runtime/Experimental/GeometryCollectionEngine/Private/GeometryCollection/GeometryCollectionComponent.cpp
+++ b/Engine/Source/Runtime/Experimental/GeometryCollectionEngine/Private/GeometryCollection/GeometryCollectionComponent.cpp
@@ -3664,11 +3664,19 @@ void UGeometryCollectionComponent::IncrementBreakTimer(float DeltaTime)
 	}
 }
 
-void UGeometryCollectionComponent::ApplyExternalStrain(int32 ItemIndex, const FVector& Location, float Strain)
+void UGeometryCollectionComponent::ApplyExternalStrain(int32 ItemIndex, const FVector& Location, float Radius, int32 PropagationDepth, float PropagationFactor, float Strain)
 {
 	if (PhysicsProxy)
 	{
-		PhysicsProxy->ApplyStrain_External(FGeometryCollectionItemIndex::CreateFromExistingItemIndex(ItemIndex), Location, Strain);
+		PhysicsProxy->ApplyExternalStrain_External(FGeometryCollectionItemIndex::CreateFromExistingItemIndex(ItemIndex), Location, Radius, PropagationDepth, PropagationFactor, Strain);
+	}
+}
+
+void UGeometryCollectionComponent::ApplyInternalStrain(int32 ItemIndex, const FVector& Location, float Radius, int32 PropagationDepth, float PropagationFactor, float Strain)
+{
+	if (PhysicsProxy)
+	{
+		PhysicsProxy->ApplyInternalStrain_External(FGeometryCollectionItemIndex::CreateFromExistingItemIndex(ItemIndex), Location, Radius, PropagationDepth, PropagationFactor, Strain);
 	}
 }
 
diff --git a/Engine/Source/Runtime/Experimental/GeometryCollectionEngine/Public/GeometryCollection/GeometryCollectionComponent.h b/Engine/Source/Runtime/Experimental/GeometryCollectionEngine/Public/GeometryCollection/GeometryCollectionComponent.h
--- a/Engine/Source/Runtime/Experimental/GeometryCollectionEngine/Public/GeometryCollection/GeometryCollectionComponent.h
+++ b/Engine/Source/Runtime/Experimental/GeometryCollectionEngine/Public/GeometryCollection/GeometryCollectionComponent.h
@@ -453,11 +453,26 @@ public:
 	 * Apply an external strain to specific piece of the geometry collection
 	 * @param ItemIndex item index ( from HitResult) of the piece to apply strain on
 	 * @param Location world location of where to apply the strain
+	 * @param Radius radius from the location point to apply the strain to ( using the center of mass of the pieces )
+	 * @param PropagationDepth How many level of connection to follow to propagate the strain through
+	 * @param PropagationFactor when using propagation, the factor to multiply the strain from one level to the other, allowing falloff effect
 	 * @param Strain strain / damage to apply 
 	 */
 	UFUNCTION(BlueprintCallable, Category = "ChaosPhysics")
-	void ApplyExternalStrain(int32 ItemIndex, const FVector& Location, float Strain);
+	void ApplyExternalStrain(int32 ItemIndex, const FVector& Location, float Radius = 0, int32 PropagationDepth = 0, float PropagationFactor = 1, float Strain = 0);
 
+	/**
+	 * Apply an internal strain to specific piece of the geometry collection
+	 * @param ItemIndex item index ( from HitResult) of the piece to apply strain on
+	 * @param Location world location of where to apply the strain
+	 * @param Radius radius from the location point to apply the strain to ( using the center of mass of the pieces )
+	 * @param PropagationDepth How many level of connection to follow to propagate the strain through
+	 * @param PropagationFactor when using propagation, the factor to multiply the strain from one level to the other, allowing falloff effect
+	 * @param Strain strain / damage to apply 
+	 */
+	UFUNCTION(BlueprintCallable, Category = "ChaosPhysics")
+	void ApplyInternalStrain(int32 ItemIndex, const FVector& Location, float Radius = 0, int32 PropagationDepth = 0, float PropagationFactor = 1, float Strain = 0);
+	
 	/**
 	 * Crumbe a cluster into all its pieces
 	 * @param ItemIndex item index ( from HitResult) of the cluster to crumble
