Implement Geometry Collection replication abandon cluster above level.
Add counters for replication fracture/cluster counts.
Compute GC level attribute during cook and serialize.

diff --git a/Engine/Source/Runtime/Core/Public/UObject/FortniteNCBranchObjectVersion.h b/Engine/Source/Runtime/Core/Public/UObject/FortniteNCBranchObjectVersion.h
--- a/Engine/Source/Runtime/Core/Public/UObject/FortniteNCBranchObjectVersion.h
+++ b/Engine/Source/Runtime/Core/Public/UObject/FortniteNCBranchObjectVersion.h
@@ -38,6 +38,10 @@ struct CORE_API FFortniteNCBranchObjectVersion
 
 		// Transfer PhysicsAsset SolverSettings (iteration counts etc) to new structure
 		PhysicsAssetNewSolverSettings,
+		
+		// Chaos GeometryCollection now saves levels attribute values
+		ChaosGeometryCollectionSaveLevelsAttribute,
+
 
 		// -----<new versions can be added above this line>-------------------------------------------------
 		VersionPlusOne,
diff --git a/Engine/Source/Runtime/Experimental/Chaos/Private/ChaosStats.cpp b/Engine/Source/Runtime/Experimental/Chaos/Private/ChaosStats.cpp
--- a/Engine/Source/Runtime/Experimental/Chaos/Private/ChaosStats.cpp
+++ b/Engine/Source/Runtime/Experimental/Chaos/Private/ChaosStats.cpp
@@ -27,6 +27,8 @@ DEFINE_STAT(STAT_GCCUGlobalMatrices);
 DEFINE_STAT(STAT_GCInitDynamicData);
 DEFINE_STAT(STAT_GCTotalTransforms);
 DEFINE_STAT(STAT_GCChangedTransforms);
+DEFINE_STAT(STAT_GCReplicatedClusters);
+DEFINE_STAT(STAT_GCReplicatedFractures);
 DEFINE_STAT(STAT_LockWaits);
 DEFINE_STAT(STAT_GeomBeginFrame);
 DEFINE_STAT(STAT_SkelMeshUpdateAnim);
diff --git a/Engine/Source/Runtime/Experimental/Chaos/Private/GeometryCollection/GeometryCollection.cpp b/Engine/Source/Runtime/Experimental/Chaos/Private/GeometryCollection/GeometryCollection.cpp
--- a/Engine/Source/Runtime/Experimental/Chaos/Private/GeometryCollection/GeometryCollection.cpp
+++ b/Engine/Source/Runtime/Experimental/Chaos/Private/GeometryCollection/GeometryCollection.cpp
@@ -11,6 +11,7 @@
 #include "GeometryCollection/GeometryCollectionProximityUtility.h"
 #include "GeometryCollection/GeometryCollectionClusteringUtility.h"
 #include "GeometryCollection/GeometryCollectionConvexUtility.h"
+#include "UObject/FortniteNCBranchObjectVersion.h"
 
 #include <iostream>
 #include <fstream>
@@ -953,6 +954,8 @@ void FGeometryCollection::Serialize(Chaos::FChaosArchive& Ar)
 		FGeometryCollectionConvexPropertiesInterface::CleanInterfaceForCook();
 	}
 
+	Ar.UsingCustomVersion(FFortniteNCBranchObjectVersion::GUID);
+
 	Super::Serialize(Ar);
 
 	if (Ar.IsLoading())
@@ -1210,6 +1213,14 @@ void FGeometryCollection::Serialize(Chaos::FChaosArchive& Ar)
 			
 			Version = 9;
 		}
+
+		if (Ar.CustomVer(FFortniteNCBranchObjectVersion::GUID) < FFortniteNCBranchObjectVersion::ChaosGeometryCollectionSaveLevelsAttribute)
+		{
+			// Level attribute previously serialized with bSave = false, so was not serializing level data.
+			// We now compute this during cook and need to serialize, so convert attribute to bSave = true.
+			RemoveAttribute("Level", FTransformCollection::TransformGroup);
+			AddAttribute<int32>("Level", FTransformCollection::TransformGroup, FConstructionParameters(FName(), /*bSave=*/true));
+		}
 	}
 }
 
diff --git a/Engine/Source/Runtime/Experimental/Chaos/Private/GeometryCollection/ManagedArrayCollection.cpp b/Engine/Source/Runtime/Experimental/Chaos/Private/GeometryCollection/ManagedArrayCollection.cpp
--- a/Engine/Source/Runtime/Experimental/Chaos/Private/GeometryCollection/ManagedArrayCollection.cpp
+++ b/Engine/Source/Runtime/Experimental/Chaos/Private/GeometryCollection/ManagedArrayCollection.cpp
@@ -370,24 +370,30 @@ void FManagedArrayCollection::CopyMatchingAttributesFrom(
 }
 
 void FManagedArrayCollection::CopyAttribute(const FManagedArrayCollection& InCollection, FName Name, FName Group)
+{
+	CopyAttribute(InCollection, /*SrcName=*/Name, /*DestName=*/Name, Group);
+}
+
+void FManagedArrayCollection::CopyAttribute(const FManagedArrayCollection& InCollection, FName SrcName, FName DestName, FName Group)
 {
 	SyncGroupSizeFrom(InCollection, Group);
-	FKeyType Key = FManagedArrayCollection::MakeMapKey(Name, Group);
+	FKeyType SrcKey = FManagedArrayCollection::MakeMapKey(SrcName, Group);
+	FKeyType DestKey = FManagedArrayCollection::MakeMapKey(DestName, Group);
 
-	if (!HasAttribute(Name, Group))
+	if (!HasAttribute(DestName, Group))
 	{
-		const FValueType& V = InCollection.Map[Key];
+		const FValueType& V = InCollection.Map[SrcKey];
 		EArrayType Type = V.ArrayType;
 		FValueType Value(Type, *NewManagedTypedArray(Type));
 		Value.Value->Resize(NumElements(Group));
 		Value.GroupIndexDependency = V.GroupIndexDependency;
 		Value.Saved = V.Saved;
 		Value.bExternalValue = V.bExternalValue;
-		Map.Add(Key, MoveTemp(Value));
+		Map.Add(DestKey, MoveTemp(Value));
 	}
 
-	const FValueType& OriginalValue = InCollection.Map[Key];
-	const FValueType& DestValue = Map[Key];	
+	const FValueType& OriginalValue = InCollection.Map[SrcKey];
+	const FValueType& DestValue = Map[DestKey];
 	check(OriginalValue.ArrayType == DestValue.ArrayType);
 	DestValue.Value->Init(*OriginalValue.Value);
 }
diff --git a/Engine/Source/Runtime/Experimental/Chaos/Private/PhysicsProxy/GeometryCollectionPhysicsProxy.cpp b/Engine/Source/Runtime/Experimental/Chaos/Private/PhysicsProxy/GeometryCollectionPhysicsProxy.cpp
--- a/Engine/Source/Runtime/Experimental/Chaos/Private/PhysicsProxy/GeometryCollectionPhysicsProxy.cpp
+++ b/Engine/Source/Runtime/Experimental/Chaos/Private/PhysicsProxy/GeometryCollectionPhysicsProxy.cpp
@@ -609,6 +609,12 @@ void FGeometryCollectionPhysicsProxy::Initialize(Chaos::FPBDRigidsEvolutionBase
 			}
 		}
 	}
+
+	// Add levels to physics thread collection, rename to make it clear these are initial levels from rest collection and not updated during sim.
+	PhysicsThreadCollection.CopyAttribute(*Parameters.RestCollection, /*SrcName=*/"Level", /*DestName=*/"InitialLevel", FTransformCollection::TransformGroup);
+	GameThreadCollection.CopyAttribute(*Parameters.RestCollection, /*SrcName=*/"Level", /*DestName=*/"InitialLevel", FTransformCollection::TransformGroup);
+
+
 }
 
 
@@ -1769,6 +1775,57 @@ int32 FGeometryCollectionPhysicsProxy::CalculateHierarchyLevel(const FGeometryDy
 	return Level;
 }
 
+
+
+int32 FGeometryCollectionPhysicsProxy::CalculateAndSetLevel(int32 TransformGroupIndex, const TManagedArray<int32>& Parent, TManagedArray<int32>& Levels)
+{
+	// Count levels up to root or first parent that has level initialized
+	int32 LevelsFromRoot = 0;
+
+	// Parents up tree from current index that need level initialization, ordered depth first
+	TArray<int32, TInlineAllocator<8>> ParentStack;
+
+	int32 ParentTransformGroupIndex = Parent[TransformGroupIndex];
+	while (ParentTransformGroupIndex != INDEX_NONE)
+	{
+		LevelsFromRoot++;
+
+		int32 ParentLevel = Levels[ParentTransformGroupIndex];
+		if (ParentLevel > 0)
+		{
+			// If parent level is not 0 it has already been computed, 
+			// can early out and not search for true root.
+			Levels[TransformGroupIndex] = ParentLevel + LevelsFromRoot;
+			break;
+		}
+
+		// save this parent so we can set it's level once we determine depth
+		ParentStack.Push(ParentTransformGroupIndex);
+
+		ParentTransformGroupIndex = Parent[ParentTransformGroupIndex];
+	}
+
+	if (LevelsFromRoot > 0) // Not at root, update levels of current particle and uninitialized parents
+	{
+		if (Levels[TransformGroupIndex] == 0)
+		{
+			// If we are uninitialized, we did not early out after finding an initialized parent up the tree, and LevelsFromRoot is new level.
+			Levels[TransformGroupIndex] = LevelsFromRoot;
+		}
+
+		// Initialize newly discovered parent levels
+		for (int32 Idx = 0; Idx < ParentStack.Num(); ++Idx)
+		{
+			int32 ParentTransformGroupIdx = ParentStack[Idx];
+
+			// Level of parent is leaf level minus parent's distance from leaf minus 1.
+			Levels[ParentTransformGroupIdx] = Levels[TransformGroupIndex] - Idx - 1;
+		}
+	}
+
+	return Levels[TransformGroupIndex];
+}
+
 void FGeometryCollectionPhysicsProxy::InitializeRemoveOnFracture(FParticlesType& Particles, const TManagedArray<int32>& DynamicState)
 {
 	/*
@@ -2972,6 +3029,8 @@ void FGeometryCollectionPhysicsProxy::InitializeSharedCollisionStructures(
 	TManagedArray<FRealSingle>& CollectionMass = RestCollection.AddAttribute<FRealSingle>(TEXT("Mass"), FTransformCollection::TransformGroup);
 	TManagedArray<TUniquePtr<FSimplicial>>& CollectionSimplicials =	RestCollection.AddAttribute<TUniquePtr<FSimplicial>>(FGeometryDynamicCollection::SimplicialsAttribute, FTransformCollection::TransformGroup);
 
+	TManagedArray<int32>& Levels = RestCollection.AddAttribute<int32>(TEXT("Level"), FTransformCollection::TransformGroup);
+
 	RestCollection.RemoveAttribute(FGeometryDynamicCollection::ImplicitsAttribute, FTransformCollection::TransformGroup);
 	TManagedArray<FGeometryDynamicCollection::FSharedImplicit>& CollectionImplicits = RestCollection.AddAttribute<FGeometryDynamicCollection::FSharedImplicit>(FGeometryDynamicCollection::ImplicitsAttribute, FTransformCollection::TransformGroup);
 
@@ -3452,6 +3511,9 @@ void FGeometryCollectionPhysicsProxy::InitializeSharedCollisionStructures(
 
 				TriangleMeshesArray[ClusterTransformIdx] = MoveTemp(UnionMesh);
 			}
+
+			// Set level of TransformGroupIndex in Levels attribute.
+			CalculateAndSetLevel(TransformGroupIndex, Parent, Levels);
 		}
 
 		InitRemoveOnFracture(RestCollection, SharedParams);
diff --git a/Engine/Source/Runtime/Experimental/Chaos/Public/ChaosStats.h b/Engine/Source/Runtime/Experimental/Chaos/Public/ChaosStats.h
--- a/Engine/Source/Runtime/Experimental/Chaos/Public/ChaosStats.h
+++ b/Engine/Source/Runtime/Experimental/Chaos/Public/ChaosStats.h
@@ -42,6 +42,8 @@ DECLARE_CYCLE_STAT_EXTERN(TEXT("Geometry Collection Component CalculateGlobalMat
 DECLARE_CYCLE_STAT_EXTERN(TEXT("Geometry Collection InitDynamicData"), STAT_GCInitDynamicData, STATGROUP_Chaos, CHAOS_API);
 DECLARE_DWORD_COUNTER_STAT_EXTERN(TEXT("Geometry Collection Total Transforms"), STAT_GCTotalTransforms, STATGROUP_Chaos, CHAOS_API);
 DECLARE_DWORD_COUNTER_STAT_EXTERN(TEXT("Geometry Collection Changed Transforms"), STAT_GCChangedTransforms, STATGROUP_Chaos, CHAOS_API);
+DECLARE_DWORD_COUNTER_STAT_EXTERN(TEXT("Geometry Collection Replicated Clusters"), STAT_GCReplicatedClusters, STATGROUP_ChaosCounters, CHAOS_API);
+DECLARE_DWORD_COUNTER_STAT_EXTERN(TEXT("Geometry Collection Replicated Fractures"), STAT_GCReplicatedFractures, STATGROUP_ChaosCounters, CHAOS_API);
 DECLARE_CYCLE_STAT_EXTERN(TEXT("Physics Lock Waits"), STAT_LockWaits, STATGROUP_Chaos, CHAOS_API);
 DECLARE_CYCLE_STAT_EXTERN(TEXT("Geometry Collection Begin Frame"), STAT_GeomBeginFrame, STATGROUP_Chaos, CHAOS_API);
 DECLARE_CYCLE_STAT_EXTERN(TEXT("Skeletal Mesh Update Anim"), STAT_SkelMeshUpdateAnim, STATGROUP_Chaos, CHAOS_API);
diff --git a/Engine/Source/Runtime/Experimental/Chaos/Public/GeometryCollection/ManagedArrayCollection.h b/Engine/Source/Runtime/Experimental/Chaos/Public/GeometryCollection/ManagedArrayCollection.h
--- a/Engine/Source/Runtime/Experimental/Chaos/Public/GeometryCollection/ManagedArrayCollection.h
+++ b/Engine/Source/Runtime/Experimental/Chaos/Public/GeometryCollection/ManagedArrayCollection.h
@@ -415,6 +415,14 @@ public:
 	*/
 	void CopyAttribute(const FManagedArrayCollection& InCollection, FName Name, FName Group);
 
+	/**
+	* Copy an attribute. Will perform an implicit group sync.
+	* @param SrcName - The name of the attribute being copied from
+	* @param DestName - The name of the attribute being copied to
+	* @param Group - The group that manages the attribute
+	*/
+	void CopyAttribute(const FManagedArrayCollection& InCollection, FName SrcName, FName DestName, FName Group);
+
 	/**
 	* Copy attributes that match the input collection. This is a utility to easily sync collections
 	* @param InCollection - All groups from this collection found in the input will be sized accordingly
diff --git a/Engine/Source/Runtime/Experimental/Chaos/Public/PhysicsProxy/GeometryCollectionPhysicsProxy.h b/Engine/Source/Runtime/Experimental/Chaos/Public/PhysicsProxy/GeometryCollectionPhysicsProxy.h
--- a/Engine/Source/Runtime/Experimental/Chaos/Public/PhysicsProxy/GeometryCollectionPhysicsProxy.h
+++ b/Engine/Source/Runtime/Experimental/Chaos/Public/PhysicsProxy/GeometryCollectionPhysicsProxy.h
@@ -404,6 +404,14 @@ protected:
 	void SetSleepingState(const Chaos::FPBDRigidsSolver& RigidsSolver);
 	void DirtyAllParticles(const Chaos::FPBDRigidsSolver& RigidsSolver);
 	
+	/** 
+	 * Traverses the parents of TransformGroupIdx counting number of levels,
+	 * and sets levels array value for TransformGroupIdx and its parents if not yet initialized.
+	 * If level is already set, retrieve stored level.
+	 * Uninitialized level array should be correct size and defaulted to zeros.
+	 */
+	static int32 CalculateAndSetLevel(int32 TransformGroupIdx, const TManagedArray<int32>& Parent, TManagedArray<int32>& Levels);
+
 	/** 
 	 * Traverses the parents of \p TransformIndex in \p GeometryCollection, counting
 	 * the number of levels until the next parent is \c INDEX_NONE.
diff --git a/Engine/Source/Runtime/Experimental/GeometryCollectionEngine/Private/GeometryCollection/DerivedDataGeometryCollectionCooker.cpp b/Engine/Source/Runtime/Experimental/GeometryCollectionEngine/Private/GeometryCollection/DerivedDataGeometryCollectionCooker.cpp
--- a/Engine/Source/Runtime/Experimental/GeometryCollectionEngine/Private/GeometryCollection/DerivedDataGeometryCollectionCooker.cpp
+++ b/Engine/Source/Runtime/Experimental/GeometryCollectionEngine/Private/GeometryCollection/DerivedDataGeometryCollectionCooker.cpp
@@ -71,7 +71,7 @@ bool FDerivedDataGeometryCollectionCooker::Build(TArray<uint8>& OutData)
 
 const TCHAR* FDerivedDataGeometryCollectionCooker::GetVersionString() const
 {
-	const TCHAR* VersionString = TEXT("3448151914544729A210FBB9B0524868");
+	const TCHAR* VersionString = TEXT("007038F125ED47A4AF928FAC54FD9493");
 
 	static FString CachedNaniteVersionString;
 	if (CachedNaniteVersionString.IsEmpty())
diff --git a/Engine/Source/Runtime/Experimental/GeometryCollectionEngine/Private/GeometryCollection/GeometryCollectionComponent.cpp b/Engine/Source/Runtime/Experimental/GeometryCollectionEngine/Private/GeometryCollection/GeometryCollectionComponent.cpp
--- a/Engine/Source/Runtime/Experimental/GeometryCollectionEngine/Private/GeometryCollection/GeometryCollectionComponent.cpp
+++ b/Engine/Source/Runtime/Experimental/GeometryCollectionEngine/Private/GeometryCollection/GeometryCollectionComponent.cpp
@@ -1411,6 +1411,8 @@ void UGeometryCollectionComponent::UpdateRepData()
 		const FPBDRigidsSolver* Solver = PhysicsProxy->GetSolver<Chaos::FPBDRigidsSolver>();
 		const FRigidClustering& RigidClustering = Solver->GetEvolution()->GetRigidClustering();
 
+		const TManagedArray<int32>& Levels = PhysicsProxy->GetPhysicsCollection().GetAttribute<int32>("InitialLevel", FGeometryCollection::TransformGroup);
+
 		//see if we have any new clusters that are enabled
 		TSet<FPBDRigidClusteredParticleHandle*> Processed;
 		for (FPBDRigidClusteredParticleHandle* Particle : PhysicsProxy->GetParticles())
@@ -1436,20 +1438,46 @@ void UGeometryCollectionComponent::UpdateRepData()
 
 			if (bProcess && Root->Disabled() == false && ClustersToRep->Find(Root) == nullptr)
 			{
-				//first time in here so needs a new count
-				//TODO: check root needs to replicate if abandon by level is enabled
-				ClustersToRep->Add(Root);
+				const int32 TransformGroupIdx = PhysicsProxy->GetTransformGroupIndexFromHandle(Root);
+				int32 Level = INDEX_NONE;
+				if (Root->InternalCluster() == false)
+				{
+					Level = Levels[TransformGroupIdx];
+				}
+				else
+				{
+					// Use internal cluster child's index to compute level.
+					const TArray<FPBDRigidParticleHandle*>& Children = RigidClustering.GetChildrenMap()[Root];
+					const int32 ChildTransformGroupIdx = PhysicsProxy->GetTransformGroupIndexFromHandle(Children[0]);
+					Level = Levels[ChildTransformGroupIdx] - 1;
+				}
+
+				if (!bEnableAbandonAfterLevel|| Level < ReplicationAbandonClusterLevel)
+				{
+					// not already replicated and not abandoned level, start replicating cluster
+					ClustersToRep->Add(Root);
+					bClustersChanged = true;
+				}
+
 				if(Root->InternalCluster() == false && bFirstUpdate == false)	//if bFirstUpdate it must be that these are the initial roots of the GC. These did not break off so no need to replicate
 				{
 					//a one off so record it
-					const int32 TransformGroupIdx = PhysicsProxy->GetTransformGroupIndexFromHandle(Root);
 					ensureMsgf(TransformGroupIdx >= 0, TEXT("Non-internal cluster should always have a group index"));
 					ensureMsgf(TransformGroupIdx < TNumericLimits<uint16>::Max(), TEXT("Trying to replicate GC with more than 65k pieces. We assumed uint16 would suffice"));
-					RepData.OneOffActivated.Add(FGeometryCollectionActivatedCluster(TransformGroupIdx, Root->V(), Root->W()));
-					bClustersChanged = true;
+
+					// Because we cull ClustersToRep with abandoned level, we must make sure we don't add duplicates to one off activated.
+					// TODO: avoid search for entry for perf
+					FGeometryCollectionActivatedCluster OneOffActivated(TransformGroupIdx, Root->V(), Root->W());
+					if(!RepData.OneOffActivated.Contains(OneOffActivated))
+					{
+						bClustersChanged = true;
+						RepData.OneOffActivated.Add(OneOffActivated);
+					}
 				}
 			}
 		}
+		
+		INC_DWORD_STAT_BY(STAT_GCReplicatedFractures, RepData.OneOffActivated.Num());
 
 		//build up clusters to replicate and compare with previous frame
 		TArray<FGeometryCollectionClusterRep> Clusters;
@@ -1512,6 +1540,8 @@ void UGeometryCollectionComponent::UpdateRepData()
 		{
 			RepData.Clusters = MoveTemp(Clusters);
 
+			INC_DWORD_STAT_BY(STAT_GCReplicatedClusters, RepData.Clusters.Num());
+
 			MARK_PROPERTY_DIRTY_FROM_NAME(UGeometryCollectionComponent, RepData, this);
 			++RepData.Version;
 
diff --git a/Engine/Source/Runtime/Experimental/GeometryCollectionEngine/Public/GeometryCollection/GeometryCollectionComponent.h b/Engine/Source/Runtime/Experimental/GeometryCollectionEngine/Public/GeometryCollection/GeometryCollectionComponent.h
--- a/Engine/Source/Runtime/Experimental/GeometryCollectionEngine/Public/GeometryCollection/GeometryCollectionComponent.h
+++ b/Engine/Source/Runtime/Experimental/GeometryCollectionEngine/Public/GeometryCollection/GeometryCollectionComponent.h
@@ -323,6 +323,11 @@ struct FGeometryCollectionActivatedCluster
 	uint16 ActivatedIndex;
 	FVector_NetQuantize100 InitialLinearVelocity;
 	FVector_NetQuantize100 InitialAngularVelocity;
+
+	bool operator==(const FGeometryCollectionActivatedCluster& Other) const
+	{
+		return ActivatedIndex == Other.ActivatedIndex;
+	}
 };
 
 FORCEINLINE FArchive& operator<<(FArchive& Ar, FGeometryCollectionActivatedCluster& ActivatedCluster)
