Chaos : remove unecessary forward declaration causing uissues with integration

diff --git a/Engine/Source/Runtime/Experimental/Chaos/Public/GeometryCollection/Facades/CollectionAnchoringFacade.h b/Engine/Source/Runtime/Experimental/Chaos/Public/GeometryCollection/Facades/CollectionAnchoringFacade.h
--- a/Engine/Source/Runtime/Experimental/Chaos/Public/GeometryCollection/Facades/CollectionAnchoringFacade.h
+++ b/Engine/Source/Runtime/Experimental/Chaos/Public/GeometryCollection/Facades/CollectionAnchoringFacade.h
@@ -6,8 +6,6 @@
 #include "GeometryCollection/ManagedArrayAccessor.h"
 #include "GeometryCollection/ManagedArrayCollection.h"
 
-struct FManagedArrayCollection;
-
 namespace Chaos::Facades
 {
 	/**
