Changes to make transient / un-Saved properties work more logically for Geometry Collection data:

- make the ManagedArrayCollection 'Saved' flag control whether the whole property is saved, rather than just its data

- when transacting, save everything (ignore the Saved flag)

- also when loading, remove properties that were not in the map we loaded from

- stop manually removing some unsaved attributes, as the Saved flag should take care of that instead

- stop loading from the DDC (after already loading from the transaction history) for every undo/redo transaction

diff --git a/Engine/Source/Runtime/Core/Public/UObject/UE5MainStreamObjectVersion.h b/Engine/Source/Runtime/Core/Public/UObject/UE5MainStreamObjectVersion.h
--- a/Engine/Source/Runtime/Core/Public/UObject/UE5MainStreamObjectVersion.h
+++ b/Engine/Source/Runtime/Core/Public/UObject/UE5MainStreamObjectVersion.h
@@ -253,6 +253,9 @@ struct CORE_API FUE5MainStreamObjectVersion
 		// Made stationary local and skylights behave similar to SM5
 		MobileStationaryLocalLights,
 
+		// Made ManagedArrayCollection::FValueType::Value always serialize when FValueType is
+		ManagedArrayCollectionAlwaysSerializeValue,
+
 		// -----<new versions can be added above this line>-------------------------------------------------
 		VersionPlusOne,
 		LatestVersion = VersionPlusOne - 1
diff --git a/Engine/Source/Runtime/Experimental/Chaos/Private/GeometryCollection/GeometryCollection.cpp b/Engine/Source/Runtime/Experimental/Chaos/Private/GeometryCollection/GeometryCollection.cpp
--- a/Engine/Source/Runtime/Experimental/Chaos/Private/GeometryCollection/GeometryCollection.cpp
+++ b/Engine/Source/Runtime/Experimental/Chaos/Private/GeometryCollection/GeometryCollection.cpp
@@ -81,7 +81,7 @@ void FGeometryCollection::Construct()
 	// Material Group
 	AddExternalAttribute<FGeometryCollectionSection>("Sections", FGeometryCollection::MaterialGroup, Sections, FacesDependency);
 
-	FGeometryCollectionConvexPropertiesInterface::InitializeInterface();
+	InitializeInterfaces();
 }
 
 
@@ -661,7 +661,12 @@ void FGeometryCollection::Empty()
 	{
 		EmptyGroup(GroupName);
 	}
-	// re-initialize the interface to its default initial state
+	// re-initialize interfaces
+	InitializeInterfaces();
+}
+
+void FGeometryCollection::InitializeInterfaces()
+{
 	FGeometryCollectionConvexPropertiesInterface::InitializeInterface();
 }
 
@@ -1027,8 +1032,7 @@ void FGeometryCollection::Serialize(Chaos::FChaosArchive& Ar)
 			}
 		}
 
-		RemoveAttribute("ExplodedTransform", FTransformCollection::TransformGroup);
-		RemoveAttribute("ExplodedVector", FTransformCollection::TransformGroup);
+
 
 		// Version 5 introduced accurate SimulationType tagging
 		if (Version < 5)
@@ -1114,8 +1118,9 @@ void FGeometryCollection::Serialize(Chaos::FChaosArchive& Ar)
 		{
 			if (!HasAttribute("UVs", FGeometryCollection::VerticesGroup))
 			{
+				// Note: As UVs is an external attribute that is always added by Construct, this should never be encountered
 				UE_LOG(FGeometryCollectionLogging, Log, TEXT("GeometryCollection updated to multiple UV sets."));
-				AddAttribute<TArray<FVector2f>>("UVs", FGeometryCollection::VerticesGroup);				
+				AddExternalAttribute<TArray<FVector2f>>("UVs", FGeometryCollection::VerticesGroup, UVs);
 			}
 
 			TManagedArray<TArray<FVector2f>>& MultipleUVs = ModifyAttribute<TArray<FVector2f>>("UVs", FGeometryCollection::VerticesGroup);
@@ -1221,6 +1226,9 @@ void FGeometryCollection::Serialize(Chaos::FChaosArchive& Ar)
 			RemoveAttribute("Level", FTransformCollection::TransformGroup);
 			AddAttribute<int32>("Level", FTransformCollection::TransformGroup, FConstructionParameters(FName(), /*bSave=*/true));
 		}
+
+		// Finally, make sure expected interfaces are initialized
+		InitializeInterfaces();
 	}
 }
 
diff --git a/Engine/Source/Runtime/Experimental/Chaos/Private/GeometryCollection/ManagedArrayCollection.cpp b/Engine/Source/Runtime/Experimental/Chaos/Private/GeometryCollection/ManagedArrayCollection.cpp
--- a/Engine/Source/Runtime/Experimental/Chaos/Private/GeometryCollection/ManagedArrayCollection.cpp
+++ b/Engine/Source/Runtime/Experimental/Chaos/Private/GeometryCollection/ManagedArrayCollection.cpp
@@ -4,6 +4,7 @@
 #include "GeometryCollection/GeometryCollectionAlgo.h"
 #include "Templates/UnrealTemplate.h"
 #include "Chaos/ChaosArchive.h"
+#include "UObject/UE5MainStreamObjectVersion.h"
 
 DEFINE_LOG_CATEGORY_STATIC(FManagedArrayCollectionLogging, NoLogging, All);
 
@@ -462,6 +463,8 @@ static const FName GuidName("GUID");
 
 void FManagedArrayCollection::Serialize(Chaos::FChaosArchive& Ar)
 {
+	Ar.UsingCustomVersion(FUE5MainStreamObjectVersion::GUID);
+
 	if (Ar.IsSaving()) Version = 9;
 	Ar << Version;
 
@@ -499,6 +502,19 @@ void FManagedArrayCollection::Serialize(Chaos::FChaosArchive& Ar)
 			}
 		}
 
+		TArray<FKeyType> ToRemoveKeys;
+		for (const TTuple<FKeyType, FValueType>& Pair : Map)
+		{
+			if (!Pair.Value.bExternalValue && !TmpMap.Find(Pair.Key))
+			{
+				ToRemoveKeys.Add(Pair.Key);
+			}
+		}
+		for (const FKeyType& Key : ToRemoveKeys)
+		{
+			Map.Remove(Key);
+		}
+
 #if WITH_EDITOR
 		//it's possible new entries have been added but are not in old content. Resize these.
 		for (TTuple<FKeyType, FValueType>& Pair : Map)
@@ -522,10 +538,31 @@ void FManagedArrayCollection::Serialize(Chaos::FChaosArchive& Ar)
 		}
 
 	}
-	else
+	else // Ar.IsSaving()
 	{
 		Ar << GroupInfo;
-		Ar << Map;
+		// Unless it's an undo/redo transaction, strip out the keys that we don't want to save
+		if (!Ar.IsTransacting())
+		{
+			TMap<FKeyType, FValueType> ToSaveMap = Map;
+			TArray<FKeyType> ToRemoveKeys;
+			for (TTuple<FKeyType, FValueType>& Pair : ToSaveMap)
+			{
+				if (!Pair.Value.Saved)
+				{
+					ToRemoveKeys.Add(Pair.Key);
+				}
+			}
+			for (const FKeyType& Key : ToRemoveKeys)
+			{
+				ToSaveMap.Remove(Key);
+			}
+			Ar << ToSaveMap;
+		}
+		else
+		{
+			Ar << Map;
+		}
 	}
 }
 
@@ -555,17 +592,19 @@ FArchive& operator<<(FArchive& Ar, FManagedArrayCollection::FValueType& ValueIn)
 	if (Version >= 2)
 	{
 		Ar << ValueIn.GroupIndexDependency;
-		Ar << ValueIn.Saved;	//question: should we be saving if Saved is false?
+		Ar << ValueIn.Saved;
 	}
 
 	if (ValueIn.Value == nullptr)
 	{
 		ValueIn.Value = NewManagedTypedArray(ValueIn.ArrayType);
 	}
-
-	if (ValueIn.Saved)
+	
+	// Note: We switched to always saving the value here, and use the Saved flag
+	// to remove the property from the overall Map (see FManagedArrayCollection::Serialize above)
+	bool bNewSavedBehavior = Ar.CustomVer(FUE5MainStreamObjectVersion::GUID) >= FUE5MainStreamObjectVersion::ManagedArrayCollectionAlwaysSerializeValue;
+	if (bNewSavedBehavior || ValueIn.Saved)
 	{
-		//todo(ocohen): need a better way to enforce this
 		ValueIn.Value->Serialize(static_cast<Chaos::FChaosArchive&>(Ar));
 	}
 
diff --git a/Engine/Source/Runtime/Experimental/Chaos/Public/GeometryCollection/GeometryCollection.h b/Engine/Source/Runtime/Experimental/Chaos/Public/GeometryCollection/GeometryCollection.h
--- a/Engine/Source/Runtime/Experimental/Chaos/Public/GeometryCollection/GeometryCollection.h
+++ b/Engine/Source/Runtime/Experimental/Chaos/Public/GeometryCollection/GeometryCollection.h
@@ -181,6 +181,9 @@ public:
 	//
 	//
 
+	// Initialize any interfaces on the geometry collection (i.e., the FGeometryCollectionConvexPropertiesInterface)
+	virtual void InitializeInterfaces();
+
 
 	/**
 	*  Update bounding box entries for the geometry
diff --git a/Engine/Source/Runtime/Experimental/GeometryCollectionEngine/Private/GeometryCollection/GeometryCollectionObject.cpp b/Engine/Source/Runtime/Experimental/GeometryCollectionEngine/Private/GeometryCollection/GeometryCollectionObject.cpp
--- a/Engine/Source/Runtime/Experimental/GeometryCollectionEngine/Private/GeometryCollection/GeometryCollectionObject.cpp
+++ b/Engine/Source/Runtime/Experimental/GeometryCollectionEngine/Private/GeometryCollection/GeometryCollectionObject.cpp
@@ -835,7 +835,7 @@ void UGeometryCollection::Serialize(FArchive& Ar)
 	//for all versions loaded, make sure sim data is up to date
  	if (Ar.IsLoading())
 	{
-		EnsureDataIsCooked();	//make sure loaded content is built
+		EnsureDataIsCooked(true, Ar.IsTransacting());	//make sure loaded content is built
 	}
 #endif
 }
@@ -1252,11 +1252,11 @@ bool UGeometryCollection::Modify(bool bAlwaysMarkDirty /*= true*/)
 	return bSuperResult;
 }
 
-void UGeometryCollection::EnsureDataIsCooked(bool bInitResources)
+void UGeometryCollection::EnsureDataIsCooked(bool bInitResources, bool bIsTransacting)
 {
 	if (StateGuid != LastBuiltGuid)
 	{
-		CreateSimulationDataImp(/*bCopyFromDDC=*/ true);
+		CreateSimulationDataImp(/*bCopyFromDDC=*/ !bIsTransacting);
 
 		if (FApp::CanEverRender() && bInitResources)
 		{
diff --git a/Engine/Source/Runtime/Experimental/GeometryCollectionEngine/Public/GeometryCollection/GeometryCollectionObject.h b/Engine/Source/Runtime/Experimental/GeometryCollectionEngine/Public/GeometryCollection/GeometryCollectionObject.h
--- a/Engine/Source/Runtime/Experimental/GeometryCollectionEngine/Public/GeometryCollection/GeometryCollectionObject.h
+++ b/Engine/Source/Runtime/Experimental/GeometryCollectionEngine/Public/GeometryCollection/GeometryCollectionObject.h
@@ -322,7 +322,7 @@ public:
 #endif
 
 #if WITH_EDITOR
-	void EnsureDataIsCooked(bool bInitResources = true);
+	void EnsureDataIsCooked(bool bInitResources = true, bool bIsTransacting = false);
 #endif
 
 	/** Accessors for internal geometry collection */
