Fix static analysis warning (value used before being checked).

diff --git a/Engine/Source/Runtime/Experimental/GeometryCollectionEngine/Private/GeometryCollection/GeometryCollectionComponent.cpp b/Engine/Source/Runtime/Experimental/GeometryCollectionEngine/Private/GeometryCollection/GeometryCollectionComponent.cpp
--- a/Engine/Source/Runtime/Experimental/GeometryCollectionEngine/Private/GeometryCollection/GeometryCollectionComponent.cpp
+++ b/Engine/Source/Runtime/Experimental/GeometryCollectionEngine/Private/GeometryCollection/GeometryCollectionComponent.cpp
@@ -1443,6 +1443,9 @@ void UGeometryCollectionComponent::UpdateRepData()
 				int32 Level = INDEX_NONE;
 				if (Root->InternalCluster() == false)
 				{
+					ensureMsgf(TransformGroupIdx >= 0, TEXT("Non-internal cluster should always have a group index"));
+					ensureMsgf(TransformGroupIdx < TNumericLimits<uint16>::Max(), TEXT("Trying to replicate GC with more than 65k pieces. We assumed uint16 would suffice"));
+
 					Level = Levels[TransformGroupIdx];
 				}
 				else
