ChaosFlesh : Multiple Components

diff --git a/Engine/Source/Runtime/Experimental/Chaos/Public/Chaos/XPBDCorotatedConstraints.h b/Engine/Source/Runtime/Experimental/Chaos/Public/Chaos/XPBDCorotatedConstraints.h
--- a/Engine/Source/Runtime/Experimental/Chaos/Public/Chaos/XPBDCorotatedConstraints.h
+++ b/Engine/Source/Runtime/Experimental/Chaos/Public/Chaos/XPBDCorotatedConstraints.h
@@ -24,8 +24,6 @@ namespace Chaos::Softs
 		//this one only accepts tetmesh input and mesh
 		FXPBDCorotatedConstraints(
 			const Chaos::Softs::FSolverParticles& InParticles,
-			int32 ParticleOffset,
-			int32 ParticleCount,
 			const TArray<TVector<int32, 4>>& InMesh,
 			const bool bRecordMetricIn = true,
 			const Chaos::Softs::FSolverReal& EMesh = (FSolverReal)10.0,
diff --git a/Engine/Source/Runtime/Experimental/Chaos/Public/Chaos/XPBDCorotatedFiberConstraints.h b/Engine/Source/Runtime/Experimental/Chaos/Public/Chaos/XPBDCorotatedFiberConstraints.h
--- a/Engine/Source/Runtime/Experimental/Chaos/Public/Chaos/XPBDCorotatedFiberConstraints.h
+++ b/Engine/Source/Runtime/Experimental/Chaos/Public/Chaos/XPBDCorotatedFiberConstraints.h
@@ -25,8 +25,6 @@ namespace Chaos::Softs
 		//this one only accepts tetmesh input and mesh
 		FXPBDCorotatedFiberConstraints(
 			const FSolverParticles& InParticles,
-			int32 ParticleOffset,
-			int32 ParticleCount,
 			const TArray<TVector<int32, 4>>& InMesh,
 			const bool bRecordMetricIn = true,
 			const FSolverReal& EMesh = (FSolverReal)10.0,
@@ -34,7 +32,7 @@ namespace Chaos::Softs
 			const FSolverVec3 InFiberDir = FSolverVec3((FSolverReal)1., (FSolverReal)0., (FSolverReal)0.),
 			const FSolverReal InSigmaMax = (FSolverReal)3e5
 		)
-			: Base(InParticles, ParticleOffset, ParticleCount, InMesh, bRecordMetricIn, EMesh, NuMesh), SigmaMax(InSigmaMax), FiberDir(InFiberDir)
+			: Base(InParticles, InMesh, bRecordMetricIn, EMesh, NuMesh), SigmaMax(InSigmaMax), FiberDir(InFiberDir)
 		{
 			LambdaArray.Init((FSolverReal)0., 3 * MeshConstraints.Num());
 		}
