Fix warning in InterchangeProjectSetting by using full class name instead of the short name
#rnx

diff --git a/Engine/Source/Runtime/Experimental/Interchange/Engine/Public/InterchangeProjectSettings.h b/Engine/Source/Runtime/Experimental/Interchange/Engine/Public/InterchangeProjectSettings.h
--- a/Engine/Source/Runtime/Experimental/Interchange/Engine/Public/InterchangeProjectSettings.h
+++ b/Engine/Source/Runtime/Experimental/Interchange/Engine/Public/InterchangeProjectSettings.h
@@ -18,7 +18,7 @@ struct FInterchangePipelineStack
 	GENERATED_BODY()
 	
 	/** The starting mesh for the blueprint **/
-	UPROPERTY(EditAnywhere, Category = Interchange, meta=(AllowedClasses = "InterchangePipelineBase, InterchangeBlueprintPipelineBase, InterchangePythonPipelineAsset"))
+	UPROPERTY(EditAnywhere, Category = Interchange, meta=(AllowedClasses = "/Script/InterchangeCore.InterchangePipelineBase, /Script/InterchangeEngine.InterchangeBlueprintPipelineBase, /Script/InterchangeEngine.InterchangePythonPipelineAsset"))
 	TArray<FSoftObjectPath> Pipelines;
 };
 
