GeometryCollection : Selection Facade

diff --git a/Engine/Source/Runtime/Experimental/Chaos/Private/GeometryCollection/Facades/CollectionSelectionFacade.cpp b/Engine/Source/Runtime/Experimental/Chaos/Private/GeometryCollection/Facades/CollectionSelectionFacade.cpp
new file mode 100644
--- /dev/null
+++ b/Engine/Source/Runtime/Experimental/Chaos/Private/GeometryCollection/Facades/CollectionSelectionFacade.cpp
@@ -0,0 +1,175 @@
+// Copyright Epic Games, Inc. All Rights Reserved.
+
+/*=============================================================================
+	GeometryCollection.cpp: FGeometryCollection methods.
+=============================================================================*/
+
+#include "GeometryCollection/Facades/CollectionSelectionFacade.h"
+
+namespace GeometryCollection::Facades
+{
+
+	// Groups 
+	const FName FSelectionFacade::UnboundGroup = "Unbound";
+	const FName FSelectionFacade::WeightedUnboundGroup = "WeightedUnbound";
+	const FName FSelectionFacade::BoundGroup = "Bound";
+	const FName FSelectionFacade::WeightedBoundGroup = "WeightedBound";
+
+	// Attributes
+	const FName FSelectionFacade::IndexAttribute = "Index";
+	const FName FSelectionFacade::WeightAttribute = "Weights";
+	const FName FSelectionFacade::BoneIndexAttribute = "BoneIndex";
+
+	FSelectionFacade::FSelectionFacade(FManagedArrayCollection* InCollection)
+		: Self(InCollection)
+	{}
+
+	//
+	//  Initialization
+	//
+
+	void FSelectionFacade::InitUnboundedGroup(FManagedArrayCollection* Collection, FName GroupName, FName DependencyGroup)
+	{
+		if (!Collection->HasGroup(GroupName))
+		{
+			Collection->AddAttribute<TArray<int32>>(FSelectionFacade::IndexAttribute, GroupName, { DependencyGroup });
+		}
+		ensure(Collection->FindAttributeTyped<TArray<int32>>(FSelectionFacade::IndexAttribute, GroupName) != nullptr);
+	}
+
+	void FSelectionFacade::InitWeightedUnboundedGroup(FManagedArrayCollection* Collection, FName GroupName, FName DependencyGroup)
+	{
+		if (!Collection->HasGroup(GroupName))
+		{
+			Collection->AddAttribute<TArray<int32>>(FSelectionFacade::IndexAttribute, GroupName, { DependencyGroup });
+			Collection->AddAttribute<TArray<float>>(FSelectionFacade::WeightAttribute, GroupName);
+		}
+		ensure(Collection->FindAttributeTyped<TArray<int32>>(FSelectionFacade::IndexAttribute, GroupName) != nullptr);
+		ensure(Collection->FindAttributeTyped<TArray<float>>(FSelectionFacade::WeightAttribute, GroupName) != nullptr);
+	}
+
+	void FSelectionFacade::InitBoundedGroup(FManagedArrayCollection* Collection, FName GroupName, FName DependencyGroup, FName BoneDependencyGroup)
+	{
+		if (!Collection->HasGroup(GroupName))
+		{
+			Collection->AddAttribute<TArray<int32>>(FSelectionFacade::IndexAttribute, GroupName, { DependencyGroup });
+			Collection->AddAttribute<int32>(FSelectionFacade::BoneIndexAttribute, GroupName, { BoneDependencyGroup });
+		}
+		ensure(Collection->FindAttributeTyped<TArray<int32>>(FSelectionFacade::IndexAttribute, GroupName) != nullptr);
+		ensure(Collection->FindAttributeTyped<int32>(FSelectionFacade::BoneIndexAttribute, GroupName) != nullptr);
+	}
+
+	void FSelectionFacade::InitWeightedBoundedGroup(FManagedArrayCollection* Collection, FName GroupName, FName DependencyGroup, FName BoneDependencyGroup)
+	{
+		if (!Collection->HasGroup(GroupName))
+		{
+			Collection->AddAttribute<TArray<int32>>(FSelectionFacade::IndexAttribute, GroupName, { DependencyGroup });
+			Collection->AddAttribute<TArray<float>>(FSelectionFacade::WeightAttribute, GroupName);
+			Collection->AddAttribute<int32>(FSelectionFacade::BoneIndexAttribute, GroupName, { BoneDependencyGroup });
+		}
+		ensure(Collection->FindAttributeTyped<TArray<int32>>(FSelectionFacade::IndexAttribute, GroupName) != nullptr);
+		ensure(Collection->FindAttributeTyped<TArray<float>>(FSelectionFacade::WeightAttribute, GroupName) != nullptr);
+		ensure(Collection->FindAttributeTyped<int32>(FSelectionFacade::BoneIndexAttribute, GroupName) != nullptr);
+	}
+
+
+	//
+	//  AddSelection
+	//
+
+	FSelectionFacade::FSelectionKey FSelectionFacade::AddSelection(FManagedArrayCollection* Collection, const TArray<int32>& InIndices, FName DependencyGroup)
+	{
+		FName GroupName(FSelectionFacade::UnboundGroup.ToString() + "_" + DependencyGroup.ToString());
+		InitUnboundedGroup(Collection, GroupName, DependencyGroup);
+
+		int Idx = Collection->AddElements(1, GroupName);
+		Collection->ModifyAttribute<TArray<int32>>(FSelectionFacade::IndexAttribute, GroupName)[Idx] = InIndices;
+		return FSelectionKey(Idx, FSelectionFacade::UnboundGroup);
+	}
+
+	FSelectionFacade::FSelectionKey FSelectionFacade::AddSelection(FManagedArrayCollection* Collection, const TArray<int32>& InIndices, const TArray<float>& InWeights, FName DependencyGroup)
+	{
+		FName GroupName(FSelectionFacade::WeightedUnboundGroup.ToString() + "_" + DependencyGroup.ToString());
+		InitWeightedUnboundedGroup(Collection, GroupName, DependencyGroup);
+
+		int Idx = Collection->AddElements(1, GroupName);
+		Collection->ModifyAttribute<TArray<int32>>(FSelectionFacade::IndexAttribute, GroupName)[Idx] = InIndices;
+		Collection->ModifyAttribute<TArray<float>>(FSelectionFacade::WeightAttribute, GroupName)[Idx] = InWeights;
+		return FSelectionKey(Idx, GroupName);
+	}
+
+	FSelectionFacade::FSelectionKey FSelectionFacade::AddSelection(FManagedArrayCollection* Collection, const int32 InBoneIndex, const TArray<int32>& InIndices, FName DependencyGroup, FName BoneDependencyGroup)
+	{
+		FName GroupName(FSelectionFacade::BoundGroup.ToString() + "_" + DependencyGroup.ToString());
+		InitBoundedGroup(Collection, GroupName, DependencyGroup, BoneDependencyGroup);
+
+		int Idx = Collection->AddElements(1, GroupName);
+		Collection->ModifyAttribute<TArray<int32>>(FSelectionFacade::IndexAttribute, GroupName)[Idx] = InIndices;
+		Collection->ModifyAttribute<int32>(FSelectionFacade::BoneIndexAttribute, GroupName)[Idx] = InBoneIndex;
+		return FSelectionKey(Idx, GroupName);
+	}
+
+	FSelectionFacade::FSelectionKey FSelectionFacade::AddSelection(FManagedArrayCollection* Collection, const int32 InBoneIndex, const TArray<int32>& InIndices, const TArray<float>& InWeights, FName DependencyGroup, FName BoneDependencyGroup)
+	{
+		FName GroupName(FSelectionFacade::WeightedBoundGroup.ToString() + "_" + DependencyGroup.ToString());
+		InitWeightedBoundedGroup(Collection, GroupName, DependencyGroup, BoneDependencyGroup);
+
+		int Idx = Collection->AddElements(1, GroupName);
+		Collection->ModifyAttribute<TArray<int32>>(FSelectionFacade::IndexAttribute, GroupName)[Idx] = InIndices;
+		Collection->ModifyAttribute<TArray<float>>(FSelectionFacade::WeightAttribute, GroupName)[Idx] = InWeights;
+		Collection->ModifyAttribute<int32>(FSelectionFacade::BoneIndexAttribute, GroupName)[Idx] = InBoneIndex;
+		return FSelectionKey(Idx, GroupName);
+	}
+
+	//
+	//  GetSelection
+	//
+
+
+	void FSelectionFacade::GetSelection(const FManagedArrayCollection* Collection, const FSelectionFacade::FSelectionKey& Key, TArray<int32>& OutIndices)
+	{
+		if (Collection->HasGroup(Key.GroupName) && 0 <= Key.Index && Key.Index < Collection->NumElements(Key.GroupName))
+		{
+			if (Collection->FindAttribute<TArray<int32>>(FSelectionFacade::IndexAttribute, Key.GroupName))
+				OutIndices = Collection->GetAttribute<TArray<int32>>(FSelectionFacade::IndexAttribute, Key.GroupName)[Key.Index];
+		}
+	}
+
+	void FSelectionFacade::GetSelection(const FManagedArrayCollection* Collection, const FSelectionFacade::FSelectionKey& Key, TArray<int32>& OutIndices, TArray<float>& OutWeights)
+	{
+		if (Collection->HasGroup(Key.GroupName) && 0 <= Key.Index && Key.Index < Collection->NumElements(Key.GroupName))
+		{
+			if (Collection->FindAttribute<TArray<int32>>(FSelectionFacade::IndexAttribute, Key.GroupName))
+				OutIndices = Collection->GetAttribute<TArray<int32>>(FSelectionFacade::IndexAttribute, Key.GroupName)[Key.Index];
+			if (Collection->FindAttribute<TArray<float>>(FSelectionFacade::WeightAttribute, Key.GroupName))
+				OutWeights = Collection->GetAttribute<TArray<float>>(FSelectionFacade::WeightAttribute, Key.GroupName)[Key.Index];
+		}
+	}
+
+	void FSelectionFacade::GetSelection(const FManagedArrayCollection* Collection, const FSelectionFacade::FSelectionKey& Key, int32& OutBoneIndex, TArray<int32>& OutIndices)
+	{
+		if (Collection->HasGroup(Key.GroupName) && 0 <= Key.Index && Key.Index < Collection->NumElements(Key.GroupName))
+		{
+			if (Collection->FindAttribute<TArray<int32>>(FSelectionFacade::IndexAttribute, Key.GroupName))
+				OutIndices = Collection->GetAttribute<TArray<int32>>(FSelectionFacade::IndexAttribute, Key.GroupName)[Key.Index];
+			if (Collection->FindAttribute<int32>(FSelectionFacade::BoneIndexAttribute, Key.GroupName))
+				OutBoneIndex = Collection->GetAttribute<int32>(FSelectionFacade::BoneIndexAttribute, Key.GroupName)[Key.Index];
+		}
+	}
+
+	void FSelectionFacade::GetSelection(const FManagedArrayCollection* Collection, const FSelectionFacade::FSelectionKey& Key, int32& OutBoneIndex, TArray<int32>& OutIndices, TArray<float>& OutWeights)
+	{
+		if (Collection->HasGroup(Key.GroupName) && 0 <= Key.Index && Key.Index < Collection->NumElements(Key.GroupName))
+		{
+			if (Collection->FindAttribute<TArray<int32>>(FSelectionFacade::IndexAttribute, Key.GroupName))
+				OutIndices = Collection->GetAttribute<TArray<int32>>(FSelectionFacade::IndexAttribute, Key.GroupName)[Key.Index];
+			if (Collection->FindAttribute<TArray<float>>(FSelectionFacade::WeightAttribute, Key.GroupName))
+				OutWeights = Collection->GetAttribute<TArray<float>>(FSelectionFacade::WeightAttribute, Key.GroupName)[Key.Index];
+			if (Collection->FindAttribute<int32>(FSelectionFacade::BoneIndexAttribute, Key.GroupName))
+				OutBoneIndex = Collection->GetAttribute<int32>(FSelectionFacade::BoneIndexAttribute, Key.GroupName)[Key.Index];
+		}
+	}
+
+};
+
+
diff --git a/Engine/Source/Runtime/Experimental/Chaos/Private/GeometryCollection/GeometryCollectionSelectionAlgo.cpp b/Engine/Source/Runtime/Experimental/Chaos/Private/GeometryCollection/GeometryCollectionSelectionAlgo.cpp
deleted file mode 100644
diff --git a/Engine/Source/Runtime/Experimental/Chaos/Public/GeometryCollection/Facades/CollectionSelectionFacade.h b/Engine/Source/Runtime/Experimental/Chaos/Public/GeometryCollection/Facades/CollectionSelectionFacade.h
new file mode 100644
--- /dev/null
+++ b/Engine/Source/Runtime/Experimental/Chaos/Public/GeometryCollection/Facades/CollectionSelectionFacade.h
@@ -0,0 +1,136 @@
+// Copyright Epic Games, Inc. All Rights Reserved.
+
+#pragma once
+#include "Containers/UnrealString.h"
+#include "Containers/Array.h"
+#include "GeometryCollection/ManagedArrayCollection.h"
+
+namespace GeometryCollection::Facades
+{
+
+	/**
+	* FSelectionFacade
+	* 
+	* Defines common API for storing a sib indexing of a separate group. For example, to store
+	* a subset of vertices defined in the group "VerticesGroup", use the function :
+	*    Key = this->AddSelection({1,2,3,4});
+	* The key can be saved and later used to query the values:
+	*	TArray<int32> SavedVertices;
+	*   this->GetSelection(Key, SavedVertices);
+	*
+	* The following groups are created on the collection based on which API is called. 
+	* 
+	*	<Group> = FSelectionFacade::UnboundGroup_<DependencyGroup>
+	*	- FindAttribute<TArray<int32>>(FVertexSetInterface::IndexAttribute, <Group>);
+	* 
+	*	<Group> = FSelectionFacade::WeightedUnboundGroup_<DependencyGroup>
+	*	- FindAttribute<TArray<int32>>(FVertexSetInterface::IndexAttribute, <Group>);
+	*	- FindAttribute<TArray<float>>(FVertexSetInterface::WeightAttribute, <Group>);
+	* 
+	*	<Group> = FSelectionFacade::BoundGroup_<DependencyGroup>
+	*	- FindAttribute<TArray<int32>>(FVertexSetInterface::IndexAttribute, <Group>);
+	*	- FindAttribute<int32>(FVertexSetInterface::BoneIndexAttribute,<Group>);
+	*
+	*	<Group> = FSelectionFacade::WeightedBoundGroup_<DependencyGroup>
+	*	- FindAttribute<TArray<int32>>(FVertexSetInterface::IndexAttribute, <Group>);
+	* 	- FindAttribute<TArray<float>>(FVertexSetInterface::WeightAttribute, <Group>);
+	*	- FindAttribute<int32>(FVertexSetInterface::BoneIndexAttribute, <Group>);
+	* 
+	*/
+	class CHAOS_API FSelectionFacade
+ 	{
+		FManagedArrayCollection* Self;
+
+		static void InitUnboundedGroup(FManagedArrayCollection* Self, FName GroupName, FName DependencyGroup);
+		static void InitWeightedUnboundedGroup(FManagedArrayCollection* Self, FName GroupName, FName DependencyGroup);
+		static void InitBoundedGroup(FManagedArrayCollection* Self, FName GroupName, FName DependencyGroup, FName BoneDependencyGroup);
+		static void InitWeightedBoundedGroup(FManagedArrayCollection* Self, FName GroupName, FName DependencyGroup, FName BoneDependencyGroup);
+
+	public:
+
+		// groups
+		static const FName UnboundGroup;
+		static const FName WeightedUnboundGroup;
+		static const FName BoundGroup;
+		static const FName WeightedBoundGroup;
+
+		// Attributes
+		static const FName IndexAttribute;
+		static const FName WeightAttribute;
+		static const FName BoneIndexAttribute;
+
+		struct FSelectionKey {
+			FSelectionKey(int32 InIndex = INDEX_NONE, FName InGroupName = "") :
+				Index(InIndex), GroupName(InGroupName) {}
+			int32 Index;
+			FName GroupName;
+		};
+
+		/**
+		* FSelectionFacade Constuctor
+		* @param VertixDependencyGroup : GroupName the index attribute is dependent on. 
+		*/
+		FSelectionFacade(FManagedArrayCollection* InSelf);
+
+		/**
+		* Add the indices to the FVertexSetInterface::UnboundGroup 
+		* @param Indices : Array of indices to add. 
+		*/
+		static FSelectionKey AddSelection(FManagedArrayCollection* Self, const TArray<int32>& Indices, FName DependencyGroup);
+
+		/**
+		* Add the indices to the FVertexSetInterface::WeightedUnboundGroup
+		* @param Indices : Array of indices to add.
+		* @param Indices : Array of indices to add.
+		*/
+		static FSelectionKey AddSelection(FManagedArrayCollection* Self, const TArray<int32>& Indices, const TArray<float>& Weights, FName DependencyGroup);
+
+		/**
+		* Add the indices to the FVertexSetInterface::BoundGroup
+		* @param Indices : Array of indices to add.
+		* @param Indices : Array of indices to add.
+		*/
+		static FSelectionKey AddSelection(FManagedArrayCollection* Self, const int32 BoneIndex, const TArray<int32>& Indices, FName DependencyGroup, FName BoneDependencyGroup = FName(""));
+
+		/**
+		* Add the indices to the FVertexSetInterface::BoundGroup
+		* @param Indices : Array of indices to add.
+		* @param Indices : Array of indices to add.
+		*/
+		static FSelectionKey AddSelection(FManagedArrayCollection* Self, const int32 BoneIndex, const TArray<int32>& Indices, const TArray<float>& Weights, FName DependencyGroup, FName BoneDependencyGroup = FName(""));
+
+		/**
+		* Return the vertex list from the given key 
+		* @param Key : <GroupName and Index>
+		* @param Indices : Return indices, empty if not found. 
+		*/
+		static void GetSelection(const FManagedArrayCollection* Self, const FSelectionKey& Key, TArray<int32>& Indices);
+
+		/**
+		* Return the vertex list, bone index, and weights from the given key
+		* @param Key : <GroupName and Index>
+		* @param BoneIndex : Return BoneIndex, INDEX_NONE if not found.
+		* @param Indices : Return indices, empty if not found.
+		* @param Weights : Return vertex weights, empty if not found.
+		*/
+		static void GetSelection(const FManagedArrayCollection* Self, const FSelectionKey& Key, TArray<int32>& Indices, TArray<float>& Weights);
+
+		/**
+		* Return the vertex list and bone index from the given key
+		* @param Key : <GroupName and Index>
+		* @param BoneIndex : Return BoneIndex, INDEX_NONE if not found.
+		* @param Indices : Return indices, empty if not found.
+		*/
+		static void GetSelection(const FManagedArrayCollection* Self, const FSelectionKey& Key, int32& BoneIndex, TArray<int32>& Indices);
+
+		/**
+		* Return the vertex list, bone index, and weights from the given key
+		* @param Key : <GroupName and Index>
+		* @param BoneIndex : Return BoneIndex, INDEX_NONE if not found.
+		* @param Indices : Return indices, empty if not found.
+		* @param Weights : Return vertex weights, empty if not found.
+		*/
+		static void GetSelection(const FManagedArrayCollection* Self, const FSelectionKey& Key, int32& BoneIndex, TArray<int32>& Indices, TArray<float>& Weights);
+	};
+
+}
\ No newline at end of file
diff --git a/Engine/Source/Runtime/Experimental/Chaos/Public/GeometryCollection/GeometryCollectionSelectionAlgo.h b/Engine/Source/Runtime/Experimental/Chaos/Public/GeometryCollection/GeometryCollectionSelectionAlgo.h
deleted file mode 100644
