Resolve error C2872: 'FSphere': ambiguous symbol

I encountered this state locally, specifying the type explicitly should resolve this

#rnx

diff --git a/Engine/Source/Runtime/Experimental/GeometryCollectionEngine/Private/GeometryCollection/GeometryCollectionComponent.cpp b/Engine/Source/Runtime/Experimental/GeometryCollectionEngine/Private/GeometryCollection/GeometryCollectionComponent.cpp
--- a/Engine/Source/Runtime/Experimental/GeometryCollectionEngine/Private/GeometryCollection/GeometryCollectionComponent.cpp
+++ b/Engine/Source/Runtime/Experimental/GeometryCollectionEngine/Private/GeometryCollection/GeometryCollectionComponent.cpp
@@ -2,34 +2,35 @@
 
 #include "GeometryCollection/GeometryCollectionComponent.h"
 
+#include "AI/NavigationSystemHelpers.h"
 #include "Async/ParallelFor.h"
-#include "Components/BoxComponent.h"
+#include "Chaos/ChaosPhysicalMaterial.h"
+#include "ChaosSolversModule.h"
+#include "ChaosStats.h"
 #include "ComponentRecreateRenderStateContext.h"
-#include "GeometryCollection/GeometryCollectionObject.h"
+#include "Components/BoxComponent.h"
+#include "Engine/InstancedStaticMesh.h"
+#include "GeometryCollection/GeometryCollectionActor.h"
 #include "GeometryCollection/GeometryCollectionAlgo.h"
+#include "GeometryCollection/GeometryCollectionCache.h"
+#include "GeometryCollection/GeometryCollectionClusteringUtility.h"
 #include "GeometryCollection/GeometryCollectionComponentPluginPrivate.h"
-#include "GeometryCollection/GeometryCollectionSceneProxy.h"
+#include "GeometryCollection/GeometryCollectionDebugDrawComponent.h"
+#include "GeometryCollection/GeometryCollectionObject.h"
+#include "GeometryCollection/GeometryCollectionProximityUtility.h"
 #include "GeometryCollection/GeometryCollectionSQAccelerator.h"
+#include "GeometryCollection/GeometryCollectionSceneProxy.h"
 #include "GeometryCollection/GeometryCollectionUtility.h"
-#include "GeometryCollection/GeometryCollectionClusteringUtility.h"
-#include "GeometryCollection/GeometryCollectionProximityUtility.h"
-#include "GeometryCollection/GeometryCollectionCache.h"
-#include "GeometryCollection/GeometryCollectionActor.h"
-#include "GeometryCollection/GeometryCollectionDebugDrawComponent.h"
-#include "Physics/Experimental/PhysScene_Chaos.h"
+#include "Math/Sphere.h"
 #include "Modules/ModuleManager.h"
-#include "ChaosSolversModule.h"
-#include "ChaosStats.h"
-#include "PhysicsProxy/GeometryCollectionPhysicsProxy.h"
-#include "PhysicsSolver.h"
-#include "Physics/PhysicsFiltering.h"
-#include "Chaos/ChaosPhysicalMaterial.h"
-#include "AI/NavigationSystemHelpers.h"
-#include "Net/UnrealNetwork.h"
 #include "Net/Core/PushModel/PushModel.h"
+#include "Net/UnrealNetwork.h"
 #include "PhysicalMaterials/PhysicalMaterial.h"
+#include "Physics/Experimental/PhysScene_Chaos.h"
+#include "Physics/PhysicsFiltering.h"
 #include "PhysicsField/PhysicsFieldComponent.h"
-#include "Engine/InstancedStaticMesh.h"
+#include "PhysicsProxy/GeometryCollectionPhysicsProxy.h"
+#include "PhysicsSolver.h"
 
 #include "Algo/RemoveIf.h"
 
@@ -3322,7 +3323,7 @@ void UGeometryCollectionComponent::CalculateGlobalMatrices()
 						else
 						{
 							float ShrinkRadius = 0.0f;
-							FSphere AccumulatedSphere;
+							UE::Math::TSphere<double> AccumulatedSphere;
 							// todo(chaos) : find a faster way to do that ( precompute the data ? )
 							if (CalculateInnerSphere(Idx, AccumulatedSphere))
 							{
@@ -3804,7 +3805,7 @@ int32 UGeometryCollectionComponent::GetInitialLevel(int32 ItemIndex)
 	return Level;
 }
 
-bool UGeometryCollectionComponent::CalculateInnerSphere(int32 TransformIndex, FSphere& SphereOut) const
+bool UGeometryCollectionComponent::CalculateInnerSphere(int32 TransformIndex, UE::Math::TSphere<double>& SphereOut) const
 {
 	// Approximates the inscribed sphere. Returns false if no such sphere exists, if for instance the index is to an embedded geometry. 
 
@@ -3817,7 +3818,7 @@ bool UGeometryCollectionComponent::CalculateInnerSphere(int32 TransformIndex, FS
 	{
 		// Sphere in component space, centered on body's COM.
 		FVector COM = MassToLocal[TransformIndex].GetLocation();
-		SphereOut = FSphere(COM, InnerRadius[TransformToGeometryIndex[TransformIndex]]);
+		SphereOut = UE::Math::TSphere<double>(COM, InnerRadius[TransformToGeometryIndex[TransformIndex]]);
 		return true;
 	}
 	else if (RestCollection->GetGeometryCollection()->IsClustered(TransformIndex))
@@ -3826,7 +3827,7 @@ bool UGeometryCollectionComponent::CalculateInnerSphere(int32 TransformIndex, FS
 		bool bSphereFound = false;
 		for (int32 ChildIndex: Children[TransformIndex])
 		{
-			FSphere LocalSphere;
+			UE::Math::TSphere<double> LocalSphere;
 			if (CalculateInnerSphere(ChildIndex, LocalSphere))
 			{
 				if (!bSphereFound)
diff --git a/Engine/Source/Runtime/Experimental/GeometryCollectionEngine/Public/GeometryCollection/GeometryCollectionComponent.h b/Engine/Source/Runtime/Experimental/GeometryCollectionEngine/Public/GeometryCollection/GeometryCollectionComponent.h
--- a/Engine/Source/Runtime/Experimental/GeometryCollectionEngine/Public/GeometryCollection/GeometryCollectionComponent.h
+++ b/Engine/Source/Runtime/Experimental/GeometryCollectionEngine/Public/GeometryCollection/GeometryCollectionComponent.h
@@ -23,6 +23,7 @@
 #include "Chaos/ChaosSolverComponentTypes.h"
 #include "Chaos/PBDRigidsEvolutionFwd.h"
 #include "EngineDefines.h"
+#include "Math/MathFwd.h"
 
 #include "GeometryCollectionComponent.generated.h"
 
@@ -1021,8 +1022,8 @@ private:
 
 	void IncrementSleepTimer(float DeltaTime);
 	void IncrementBreakTimer(float DeltaTime);
+	bool CalculateInnerSphere(int32 TransformIndex, UE::Math::TSphere<double>& SphereOut) const;
 	void UpdateDecay(int32 TransformIdx, float UpdatedDecay, bool UseClusterCrumbling, bool HasDynamicInternalClusterParent, FGeometryCollectionDecayContext& ContextInOut);
-	bool CalculateInnerSphere(int32 TransformIndex, FSphere& SphereOut) const;
 	void ProcessRepData();
 
 	void BuildInitialFilterData();
