Chaos : geometry collection : fix crash in replication code when using embedded geometry

diff --git a/Engine/Source/Runtime/Experimental/GeometryCollectionEngine/Private/GeometryCollection/GeometryCollectionComponent.cpp b/Engine/Source/Runtime/Experimental/GeometryCollectionEngine/Private/GeometryCollection/GeometryCollectionComponent.cpp
--- a/Engine/Source/Runtime/Experimental/GeometryCollectionEngine/Private/GeometryCollection/GeometryCollectionComponent.cpp
+++ b/Engine/Source/Runtime/Experimental/GeometryCollectionEngine/Private/GeometryCollection/GeometryCollectionComponent.cpp
@@ -1419,6 +1419,9 @@ void UGeometryCollectionComponent::UpdateRepData()
 		TSet<FPBDRigidClusteredParticleHandle*> Processed;
 		for (FPBDRigidClusteredParticleHandle* Particle : PhysicsProxy->GetParticles())
 		{
+			// Particle can be null if we have embedded geometry 
+			if (Particle)
+			{
 			bool bProcess = true;
 			Processed.Add(Particle);
 			FPBDRigidClusteredParticleHandle* Root = Particle;
@@ -1481,6 +1484,7 @@ void UGeometryCollectionComponent::UpdateRepData()
 				}
 			}
 		}
+		}
 		
 		INC_DWORD_STAT_BY(STAT_GCReplicatedFractures, RepData.OneOffActivated.Num());
 
