Heightfield midphase inaccuracies fixed
- Fixes heightfield paint tool jitter

diff --git a/Engine/Source/Runtime/Experimental/Chaos/Private/Chaos/HeightField.cpp b/Engine/Source/Runtime/Experimental/Chaos/Private/Chaos/HeightField.cpp
--- a/Engine/Source/Runtime/Experimental/Chaos/Private/Chaos/HeightField.cpp
+++ b/Engine/Source/Runtime/Experimental/Chaos/Private/Chaos/HeightField.cpp
@@ -1087,25 +1087,24 @@ namespace Chaos
 				const FReal BoundsMinZ = CachedBounds.Min().Z;
 				const FReal BoundsMaxZ = CachedBounds.Max().Z;
 				const FReal Length2D = FMath::Sqrt(DirScaled[0] * DirScaled[0] + DirScaled[1] * DirScaled[1]);
-				const FReal FastInc = 5.0;
-				const FReal FastIncScaled = FastInc / Length2D;
-				
-				const FVec3 CellsToInspectScaled = DirScaled * FastIncScaled;
-				FReal NextStartZ = NextStart.Z;
-				const FReal Length3D = DirScaled.Length();
-				const FReal FastIncScaled3D = FastInc / Length3D;
-				const FReal CellsToInspectZ = DirScaled.Z * FastIncScaled3D;
+				const FReal FastInc2D = 5.0;  // Length in cell dimension
+				const FReal FastIncScaled = FastInc2D / Length2D;
+
+				const FVec3 InspectionStepVector = DirScaled * FastIncScaled * ScaledDx;
+				FReal NextStartZ = FMath::Clamp(NextStart.Z, BoundsMinZ, BoundsMaxZ); // Numerical errors can violate this condition, so force it here.
+				const FReal FastIncScaled3D = InspectionStepVector.Length();
+				const FReal CellsToInspectZ = FastIncScaled3D * Dir.Z;
 				FReal DistanceProcessed = 0.0;		
-				FReal FastIncScaledAccum = 0.0;
 				while (true)
 				{
 					if (!FlatGrid.IsValid(CellIdx))
 					{
 						return false;
 					}
-					FVec3 NewNextStart = CellsToInspectScaled + NextStart + UE_SMALL_NUMBER;
-					TVec2<int32> NewNextStartInt = TVec2<int32>(static_cast<int32>(NewNextStart[0]), static_cast<int32>(NewNextStart[1]));
-					TVec2<int32> NextStartInt = TVec2<int32>(static_cast<int32>(NextStart[0]), static_cast<int32>(NextStart[1]));
+
+					FVec3 NewNextStart = (InspectionStepVector + NextStart);
+					TVec2<int32> NewNextStartInt = TVec2<int32>(static_cast<int32>(NewNextStart[0] / Scale2D[0]), static_cast<int32>(NewNextStart[1] / Scale2D[1]));
+					TVec2<int32> NextStartInt = TVec2<int32>(static_cast<int32>(NextStart[0] / Scale2D[0]), static_cast<int32>(NextStart[1] / Scale2D[1]));
 
 					TVec2<int32> DiffInt = NewNextStartInt - NextStartInt;
 					TVec2<int32> AddedCellIdx = TVec2<int32>(static_cast<int32>(DiffInt[0]), static_cast<int32>(DiffInt[1]));
@@ -1115,19 +1114,18 @@ namespace Chaos
 					Bounds.Thicken(Visitor.ThicknessSimd);
 					if (Bounds.RaycastFast(Visitor.StartPointSimd, Visitor.InvDirSimd, Visitor.Parallel, CurrentLengthSimd))
 					{
-						NextStart = NextStartOri + Dir * FastIncScaledAccum;
+						NextStart = NextStartOri + Dir * DistanceProcessed;
 						break;
 					}
 
 					NextStart = NewNextStart;
-					CellIdx += AddedCellIdx;
+					CellIdx = FlatGrid.Cell(NewNextStartInt);
 					if (DistanceProcessed > CurrentLength || NextStartZ < BoundsMinZ || NextStartZ > BoundsMaxZ)
 					{
 						return false;
 					}
 					DistanceProcessed += FastIncScaled3D;
-					NextStartZ += CellsToInspectZ;
-					FastIncScaledAccum += FastIncScaled;
+					NextStartZ += CellsToInspectZ;					
 				}
 			}
 
