Interchange: Fix crash when duplicating a blueprint pipeline
#rnx

diff --git a/Engine/Source/Runtime/Experimental/Interchange/Engine/Private/InterchangeManager.cpp b/Engine/Source/Runtime/Experimental/Interchange/Engine/Private/InterchangeManager.cpp
--- a/Engine/Source/Runtime/Experimental/Interchange/Engine/Private/InterchangeManager.cpp
+++ b/Engine/Source/Runtime/Experimental/Interchange/Engine/Private/InterchangeManager.cpp
@@ -912,11 +912,13 @@ UInterchangeManager::ImportInternal(const FString& ContentPath, const UInterchan
 								AsyncHelper->Pipelines.Add(GeneratedPipeline);
 
 								//We need to save the python pipeline asset because we cannot save an asset create with a python class
-								UObject* OriginalPipeline = DuplicateObject<UObject>(Pipelines[GraphPipelineIndex].TryLoad(), GetTransientPackage());
-								if (UInterchangePythonPipelineAsset* PythonPipelineAsset = Cast<UInterchangePythonPipelineAsset>(OriginalPipeline))
+								if (UInterchangePythonPipelineAsset* OriginalPipeline = Cast<UInterchangePythonPipelineAsset>(Pipelines[GraphPipelineIndex].TryLoad()))
 								{
-									PythonPipelineAsset->SetupFromPipeline(Cast<UInterchangePythonPipelineBase>(GeneratedPipeline));
-									AsyncHelper->OriginalPipelines.Add(OriginalPipeline);
+									if (UInterchangePythonPipelineAsset* DuplicatedPythonPipelineAsset = DuplicateObject<UInterchangePythonPipelineAsset>(OriginalPipeline, GetTransientPackage()))
+									{
+										DuplicatedPythonPipelineAsset->SetupFromPipeline(Cast<UInterchangePythonPipelineBase>(GeneratedPipeline));
+										AsyncHelper->OriginalPipelines.Add(DuplicatedPythonPipelineAsset);
+									}
 								}
 								else
 								{
