Chaos : geometry collection : Allow clusters to be anchored

diff --git a/Engine/Source/Runtime/Experimental/Chaos/Private/Chaos/PBDRigidClustering.cpp b/Engine/Source/Runtime/Experimental/Chaos/Private/Chaos/PBDRigidClustering.cpp
--- a/Engine/Source/Runtime/Experimental/Chaos/Private/Chaos/PBDRigidClustering.cpp
+++ b/Engine/Source/Runtime/Experimental/Chaos/Private/Chaos/PBDRigidClustering.cpp
@@ -103,7 +103,7 @@ namespace Chaos
 		NewParticle->SetClusterId(ClusterId(nullptr, Children.Num()));
 		NewParticle->SetClusterGroupIndex(ClusterGroupIndex);
 		NewParticle->SetStrains(0.0);
-		NewParticle->SetIsAnchored(false);
+		NewParticle->SetIsAnchored(Parameters.bIsAnchored);
 
 		// Update clustering data structures.
 		if (MChildren.Contains(NewParticle))
diff --git a/Engine/Source/Runtime/Experimental/Chaos/Private/Chaos/PBDRigidClusteringAlgo.cpp b/Engine/Source/Runtime/Experimental/Chaos/Private/Chaos/PBDRigidClusteringAlgo.cpp
--- a/Engine/Source/Runtime/Experimental/Chaos/Private/Chaos/PBDRigidClusteringAlgo.cpp
+++ b/Engine/Source/Runtime/Experimental/Chaos/Private/Chaos/PBDRigidClusteringAlgo.cpp
@@ -228,46 +228,53 @@ namespace Chaos
 		{
 			if (MChildren.Contains(ClusteredCurrentNode) && MChildren[ClusteredCurrentNode].Num())
 			{
-				// TQueue is a linked list, which has no preallocator.
-				TQueue<Chaos::FPBDRigidParticleHandle*> Queue;
-				for (Chaos::FPBDRigidParticleHandle* Child : MChildren[ClusteredCurrentNode])
+				if (ClusteredCurrentNode->IsAnchored())
 				{
-					Queue.Enqueue(Child);
+					ObjectState = EObjectStateType::Kinematic;
 				}
-
-				Chaos::FPBDRigidParticleHandle* CurrentHandle;
-				while (Queue.Dequeue(CurrentHandle) && ObjectState == EObjectStateType::Dynamic)
+				else
 				{
-					bool bIsAnchored = false;
-					if (FClusterHandle CurrentClusterHandle = CurrentHandle->CastToClustered())
+					// TQueue is a linked list, which has no preallocator.
+					TQueue<Chaos::FPBDRigidParticleHandle*> Queue;
+					for (Chaos::FPBDRigidParticleHandle* Child : MChildren[ClusteredCurrentNode])
+					{
+						Queue.Enqueue(Child);
+					}
+
+					Chaos::FPBDRigidParticleHandle* CurrentHandle;
+					while (Queue.Dequeue(CurrentHandle) && ObjectState == EObjectStateType::Dynamic)
 					{
-						// @question : Maybe we should just store the leaf node bodies in a
-						// map, that will require Memory(n*log(n))
-						if (MChildren.Contains(CurrentClusterHandle))
+						bool bIsAnchored = false;
+						if (FClusterHandle CurrentClusterHandle = CurrentHandle->CastToClustered())
 						{
-							for (Chaos::FPBDRigidParticleHandle* Child : MChildren[CurrentClusterHandle])
+							// @question : Maybe we should just store the leaf node bodies in a
+							// map, that will require Memory(n*log(n))
+							if (MChildren.Contains(CurrentClusterHandle))
 							{
-								Queue.Enqueue(Child);
+								for (Chaos::FPBDRigidParticleHandle* Child : MChildren[CurrentClusterHandle])
+								{
+									Queue.Enqueue(Child);
+								}
 							}
+							
+							bIsAnchored = CurrentClusterHandle->IsAnchored();
 						}
-						
-						bIsAnchored = CurrentClusterHandle->IsAnchored();
-					}
 
-					if (bIsAnchored)
-					{
-						ObjectState = EObjectStateType::Kinematic;
-					}
-					else
-					{
-						const EObjectStateType CurrState = CurrentHandle->ObjectState();
-						if (CurrState == EObjectStateType::Kinematic)
+						if (bIsAnchored)
 						{
 							ObjectState = EObjectStateType::Kinematic;
 						}
-						else if (CurrState == EObjectStateType::Static)
+						else
 						{
-							ObjectState = EObjectStateType::Static;
+							const EObjectStateType CurrState = CurrentHandle->ObjectState();
+							if (CurrState == EObjectStateType::Kinematic)
+							{
+								ObjectState = EObjectStateType::Kinematic;
+							}
+							else if (CurrState == EObjectStateType::Static)
+							{
+								ObjectState = EObjectStateType::Static;
+							}
 						}
 					}
 				}
diff --git a/Engine/Source/Runtime/Experimental/Chaos/Private/GeometryCollection/Facades/CollectionAnchoringFacade.cpp b/Engine/Source/Runtime/Experimental/Chaos/Private/GeometryCollection/Facades/CollectionAnchoringFacade.cpp
--- a/Engine/Source/Runtime/Experimental/Chaos/Private/GeometryCollection/Facades/CollectionAnchoringFacade.cpp
+++ b/Engine/Source/Runtime/Experimental/Chaos/Private/GeometryCollection/Facades/CollectionAnchoringFacade.cpp
@@ -63,7 +63,7 @@ namespace Chaos::Facades
 		
 	bool FCollectionAnchoringFacade::IsAnchored(int32 TransformIndex) const
 	{
-		return AnchoredAttribute.Get()[TransformIndex];
+		return AnchoredAttribute.IsValid()? AnchoredAttribute.Get()[TransformIndex]: false;
 	}
 
 	void FCollectionAnchoringFacade::SetAnchored(int32 TransformIndex, bool bValue)
diff --git a/Engine/Source/Runtime/Experimental/Chaos/Private/PhysicsProxy/GeometryCollectionPhysicsProxy.cpp b/Engine/Source/Runtime/Experimental/Chaos/Private/PhysicsProxy/GeometryCollectionPhysicsProxy.cpp
--- a/Engine/Source/Runtime/Experimental/Chaos/Private/PhysicsProxy/GeometryCollectionPhysicsProxy.cpp
+++ b/Engine/Source/Runtime/Experimental/Chaos/Private/PhysicsProxy/GeometryCollectionPhysicsProxy.cpp
@@ -879,7 +879,7 @@ void FGeometryCollectionPhysicsProxy::InitializeBodiesPT(Chaos::FPBDRigidsSolver
 					CollisionParticlesPerObjectFraction);
 
 				// initialize anchoring information if available 
-				Handle->SetIsAnchored(AnchoringFacade.HasAnchoredAttribute()? AnchoringFacade.IsAnchored(TransformGroupIndex): false);
+				Handle->SetIsAnchored(AnchoringFacade.IsAnchored(TransformGroupIndex));
 				
 				if (Parameters.EnableClustering)
 				{
@@ -1247,6 +1247,9 @@ FGeometryCollectionPhysicsProxy::BuildClusters_Internal(
 			FMath::TruncToInt32(Chaos::FReal(FMath::Max(0, FMath::Min(NumCollisionParticles * CollisionParticlesPerObjectFraction, NumCollisionParticles))));
 		ClusterCreationParameters.CollisionParticles->Resize(ClampedCollisionParticlesSize);
 	}
+	Chaos::Facades::FCollectionAnchoringFacade AnchoringFacade(DynamicCollection);
+	ClusterCreationParameters.bIsAnchored = AnchoringFacade.IsAnchored(CollectionClusterIndex);
+	
 	TArray<Chaos::TPBDRigidParticleHandle<Chaos::FReal, 3>*> ChildHandlesCopy(ChildHandles);
 
 	// Construct an active cluster particle, disable children, derive M and I from children:
diff --git a/Engine/Source/Runtime/Experimental/Chaos/Public/Chaos/ClusterCreationParameters.h b/Engine/Source/Runtime/Experimental/Chaos/Public/Chaos/ClusterCreationParameters.h
--- a/Engine/Source/Runtime/Experimental/Chaos/Public/Chaos/ClusterCreationParameters.h
+++ b/Engine/Source/Runtime/Experimental/Chaos/Public/Chaos/ClusterCreationParameters.h
@@ -33,6 +33,7 @@ namespace Chaos
 			, FBVHParticles* CollisionParticlesIn = nullptr
 			, Chaos::TPBDRigidClusteredParticleHandle<Chaos::FReal,3>* ClusterParticleHandleIn = nullptr
 			, const FVec3& ScaleIn = FVec3::OneVector
+			, bool bIsAnchoredIn = false
 		)
 			: CoillisionThicknessPercent(CoillisionThicknessPercentIn)
 			, MaxNumConnections(MaxNumConnectionsIn)
@@ -44,6 +45,7 @@ namespace Chaos
 			, CollisionParticles(CollisionParticlesIn)
 			, ClusterParticleHandle(ClusterParticleHandleIn)
 			, Scale(ScaleIn)
+			, bIsAnchored(bIsAnchoredIn)
 		{}
 
 		FReal CoillisionThicknessPercent;
@@ -56,5 +58,6 @@ namespace Chaos
 		FBVHParticles* CollisionParticles;
 		Chaos::FPBDRigidClusteredParticleHandle* ClusterParticleHandle;
 		FVec3 Scale;
+		bool bIsAnchored;
 	};
 }
