Chaos geometry collection :  Add crumbling event

diff --git a/Engine/Source/Runtime/Experimental/Chaos/Private/Chaos/PBDRigidClustering.cpp b/Engine/Source/Runtime/Experimental/Chaos/Private/Chaos/PBDRigidClustering.cpp
--- a/Engine/Source/Runtime/Experimental/Chaos/Private/Chaos/PBDRigidClustering.cpp
+++ b/Engine/Source/Runtime/Experimental/Chaos/Private/Chaos/PBDRigidClustering.cpp
@@ -403,6 +403,42 @@ namespace Chaos
 		}
 	}
 
+	
+	void FRigidClustering::SendCrumblingEvent(FPBDRigidClusteredParticleHandle* ClusteredParticle)
+	{
+		check(ClusteredParticle);
+		if (IPhysicsProxyBase* Proxy = ClusteredParticle->PhysicsProxy())
+		{
+			if (Proxy->GetType() == EPhysicsProxyType::GeometryCollectionType)
+			{
+				FGeometryCollectionPhysicsProxy* ConcreteProxy = static_cast<FGeometryCollectionPhysicsProxy*>(Proxy);
+				FSimulationParameters& SimParams = ConcreteProxy->GetSimParameters(); 
+				if (SimParams.bGenerateCrumblingData)
+				{
+					FCrumblingData& ClusterCrumbling = MAllClusterCrumblings.AddDefaulted_GetRef();
+					ClusterCrumbling.Proxy = ClusteredParticle->PhysicsProxy();
+					ClusterCrumbling.Location = ClusteredParticle->X();
+					ClusterCrumbling.Orientation = ClusteredParticle->R();
+					ClusterCrumbling.LinearVelocity = ClusteredParticle->V();
+					ClusterCrumbling.AngularVelocity = ClusteredParticle->W();
+					ClusterCrumbling.Mass = ClusteredParticle->M();
+					if (ClusteredParticle->Geometry() && ClusteredParticle->Geometry()->HasBoundingBox())
+					{
+						ClusterCrumbling.LocalBounds = ClusteredParticle->Geometry()->BoundingBox();
+					}
+					if (SimParams.bGenerateCrumblingChildrenData)
+					{
+						// when sending this event, children are still attached
+						if (const FRigidHandleArray* Children = MChildren.Find(ClusteredParticle))
+						{
+							ConcreteProxy->GetTransformGroupIndicesFromHandles(*Children, ClusterCrumbling.Children);
+						}
+					}
+				}
+			}
+		}
+	}
+
 	TArray<FRigidClustering::FParticleIsland> FRigidClustering::FindIslandsInChildren(const FPBDRigidClusteredParticleHandle* ClusteredParticle)
 	{
 		const TArray<FPBDRigidParticleHandle*>& Children = MChildren[ClusteredParticle];
@@ -1288,10 +1324,10 @@ namespace Chaos
 
 	}
 
-	bool FRigidClustering::BreakCluster(const FPBDRigidClusteredParticleHandle* ClusteredParticle)
+	bool FRigidClustering::BreakCluster(FPBDRigidClusteredParticleHandle* ClusteredParticle)
 	{
 		// max strain will allow to unconditionally release the children when strain is evaluated
-		const FReal MaxStrain = TNumericLimits<FReal>::Max();
+		constexpr FReal MaxStrain = TNumericLimits<FReal>::Max();
 		if (TArray<FPBDRigidParticleHandle*>* ChildrenHandles = GetChildrenMap().Find(ClusteredParticle))
 		{
 			for (FPBDRigidParticleHandle* ChildHandle: *ChildrenHandles)
@@ -1301,6 +1337,10 @@ namespace Chaos
 					ClusteredChildHandle->SetExternalStrain(MaxStrain);
 				}
 			}
+			if (ChildrenHandles->Num() > 0)
+			{
+				SendCrumblingEvent(ClusteredParticle);
+			}
 			return true;
 		}
 		return false;
diff --git a/Engine/Source/Runtime/Experimental/Chaos/Private/ChaosStats.cpp b/Engine/Source/Runtime/Experimental/Chaos/Private/ChaosStats.cpp
--- a/Engine/Source/Runtime/Experimental/Chaos/Private/ChaosStats.cpp
+++ b/Engine/Source/Runtime/Experimental/Chaos/Private/ChaosStats.cpp
@@ -35,6 +35,7 @@ DEFINE_STAT(STAT_SkelMeshUpdateAnim);
 DEFINE_STAT(STAT_DispatchEventNotifies);
 DEFINE_STAT(STAT_DispatchCollisionEvents);
 DEFINE_STAT(STAT_DispatchBreakEvents);
+DEFINE_STAT(STAT_DispatchCrumblingEvents);
 DEFINE_STAT(STAT_BufferPhysicsResults);
 DEFINE_STAT(STAT_FlipResults);
 DEFINE_STAT(STAT_ProcessDeferredCreatePhysicsState);
@@ -61,6 +62,7 @@ DEFINE_STAT(STAT_GatherCollisionEvent)
 DEFINE_STAT(STAT_GatherBreakingEvent)
 DEFINE_STAT(STAT_GatherTrailingEvent)
 DEFINE_STAT(STAT_GatherSleepingEvent)
+DEFINE_STAT(STAT_GatherCrumblingEvent)
 DEFINE_STAT(STAT_AccelerationStructureReset);
 
 DEFINE_STAT(STAT_ParamUpdateObject);
diff --git a/Engine/Source/Runtime/Experimental/Chaos/Private/EventDefaults.cpp b/Engine/Source/Runtime/Experimental/Chaos/Private/EventDefaults.cpp
--- a/Engine/Source/Runtime/Experimental/Chaos/Private/EventDefaults.cpp
+++ b/Engine/Source/Runtime/Experimental/Chaos/Private/EventDefaults.cpp
@@ -67,6 +67,7 @@ namespace Chaos
 		RegisterTrailingEvent(EventManager);
 		RegisterSleepingEvent(EventManager);
 		RegisterRemovalEvent(EventManager);
+		RegisterCrumblingEvent(EventManager);
 	}
 
 	void FEventDefaults::RegisterCollisionEvent(FEventManager& EventManager)
@@ -552,4 +553,34 @@ namespace Chaos
 			
 			});
 	}
+
+	void FEventDefaults::RegisterCrumblingEvent(FEventManager& EventManager)
+	{
+		EventManager.template RegisterEvent<FCrumblingEventData>(EEventType::Crumbling, []
+		(const Chaos::FPBDRigidsSolver* Solver, FCrumblingEventData& CrumblingEventData, bool bResetData)
+		{
+			check(Solver);
+
+			EnsureIsInPhysicsThreadContext();
+
+			SCOPE_CYCLE_COUNTER(STAT_GatherBreakingEvent);
+
+			// @todo(chaos) : should we check for a way to globally stop this from being processed ? 
+
+			if (bResetData)
+			{
+				CrumblingEventData.Reset();
+			}
+			CrumblingEventData.SetTimeCreated(Solver->MTime);
+
+			const TArray<FCrumblingData>& AllCrumblingsArray = Solver->GetEvolution()->GetRigidClustering().GetAllClusterCrumblings();
+
+			CrumblingEventData.Reserve(AllCrumblingsArray.Num());
+			for (const FCrumblingData& Crumbling: AllCrumblingsArray)
+			{
+				// todo(chaos) : implement filtering only if necessary as crumbling event should be sparser than breaking ones
+				CrumblingEventData.AddCrumbling(Crumbling);
+			}
+		});
+	}
 }
diff --git a/Engine/Source/Runtime/Experimental/Chaos/Private/PBDRigidsSolver.cpp b/Engine/Source/Runtime/Experimental/Chaos/Private/PBDRigidsSolver.cpp
--- a/Engine/Source/Runtime/Experimental/Chaos/Private/PBDRigidsSolver.cpp
+++ b/Engine/Source/Runtime/Experimental/Chaos/Private/PBDRigidsSolver.cpp
@@ -398,6 +398,7 @@ namespace Chaos
 					bool ResetData = (MSubStepInfo.Step == 0);
 					MSolver->GetEventManager()->FillProducerData(MSolver, ResetData);
 					MSolver->GetEvolution()->ResetAllRemovals();
+					MSolver->GetEvolution()->GetRigidClustering().ResetAllClusterCrumblings();
 				}
 
 				// flip on last sub-step of frame
diff --git a/Engine/Source/Runtime/Experimental/Chaos/Public/Chaos/ExternalCollisionData.h b/Engine/Source/Runtime/Experimental/Chaos/Public/Chaos/ExternalCollisionData.h
--- a/Engine/Source/Runtime/Experimental/Chaos/Public/Chaos/ExternalCollisionData.h
+++ b/Engine/Source/Runtime/Experimental/Chaos/Public/Chaos/ExternalCollisionData.h
@@ -179,6 +179,49 @@ namespace Chaos
 		int32 TransformGroupIndex;
 	};
 
+	/*
+	CrumblingData passed from the physics solver to subsystems
+	*/
+	struct FCrumblingData
+	{
+		FCrumblingData()
+			: Proxy(nullptr)
+			, Location(FVec3::ZeroVector)
+			, Orientation(FRotation3::Identity)
+			, LinearVelocity(FVec3::ZeroVector)
+			, AngularVelocity(FVec3::ZeroVector)
+			, Mass((FReal)0.0)
+			, LocalBounds(FAABB3(FVec3((FReal)0.0), FVec3((FReal)0.0)))
+		{}
+
+		FCrumblingData(const FCrumblingData& Other)
+			: Proxy(Other.Proxy)
+			, Location(Other.Location)
+			, Orientation(Other.Orientation)
+			, LinearVelocity(Other.LinearVelocity)
+			, AngularVelocity(Other.AngularVelocity)
+			, Mass(Other.Mass)
+			, LocalBounds(Other.LocalBounds)
+			, Children(Other.Children)
+		{}
+		
+		// The pointer to the proxy should be used with caution on the Game Thread.
+		// Ideally we only ever use this as a table key when acquiring related structures.
+		// If we genuinely need to dereference the pointer for any reason, test if it is deleted (nullptr) or
+		// pending deletion: if a call to FPhysScene_Chaos::GetOwningComponent<UPrimitiveComponent>() returns
+		// nullptr, the proxy should not be used.
+		IPhysicsProxyBase* Proxy;
+		
+		FVec3 Location;
+		FRotation3 Orientation;
+		FVec3 LinearVelocity;
+		FVec3 AngularVelocity;
+		FReal Mass;
+		FAABB3 LocalBounds;
+		// optional ( see proxy options )
+		TArray<int32> Children;
+	};
+	
 	/*
 	BreakingData used in the subsystems
 	*/
diff --git a/Engine/Source/Runtime/Experimental/Chaos/Public/Chaos/PBDRigidClustering.h b/Engine/Source/Runtime/Experimental/Chaos/Public/Chaos/PBDRigidClustering.h
--- a/Engine/Source/Runtime/Experimental/Chaos/Public/Chaos/PBDRigidClustering.h
+++ b/Engine/Source/Runtime/Experimental/Chaos/Public/Chaos/PBDRigidClustering.h
@@ -153,7 +153,7 @@ public:
 	*    @param ClusteredParticle handle of the cluster to break
 	*    @return true if the cluster was successfully found and act upon 
 	*/
-	bool BreakCluster(const FPBDRigidClusteredParticleHandle* ClusteredParticle);
+	bool BreakCluster(FPBDRigidClusteredParticleHandle* ClusteredParticle);
 	
 	//
 	// Operational 
@@ -254,6 +254,14 @@ public:
 	void SetGenerateClusterBreaking(bool DoGenerate) { DoGenerateBreakingData = DoGenerate; }
 	void ResetAllClusterBreakings() { MAllClusterBreakings.Reset(); }
 
+	/*
+	*  Cluster crumbling Data
+	*     triggered when all the children of a cluster are released all at once
+	*     event is generated only if the owning proxy allows it
+	*/
+	const TArray<FCrumblingData>& GetAllClusterCrumblings() const { return MAllClusterCrumblings; }
+	void ResetAllClusterCrumblings() { MAllClusterCrumblings.Reset(); }
+
 	/*
 	* GetConnectivityEdges
 	*    Provides a list of each rigid body's current siblings and associated strain within the cluster.
@@ -351,7 +359,9 @@ public:
 	void RemoveNodeConnections(FPBDRigidClusteredParticleHandle* Child);
 
 	void RemoveChildFromParent(FPBDRigidParticleHandle* Child, const FPBDRigidClusteredParticleHandle* ClusteredParent);
+
 	void SendBreakingEvent(FPBDRigidClusteredParticleHandle* ClusteredParticle);
+	void SendCrumblingEvent(FPBDRigidClusteredParticleHandle* ClusteredParticle);
 
 	TSet<FPBDRigidParticleHandle*> ReleaseClusterParticlesImpl(
 		FPBDRigidClusteredParticleHandle* ClusteredParticle, 
@@ -381,6 +391,8 @@ private:
 	bool DoGenerateBreakingData;
 	TArray<FBreakingData> MAllClusterBreakings;
 
+	TArray<FCrumblingData> MAllClusterCrumblings;
+	
 	FReal MClusterConnectionFactor;
 	FClusterCreationParameters::EConnectionMethod MClusterUnionConnectionType;
 };
diff --git a/Engine/Source/Runtime/Experimental/Chaos/Public/ChaosStats.h b/Engine/Source/Runtime/Experimental/Chaos/Public/ChaosStats.h
--- a/Engine/Source/Runtime/Experimental/Chaos/Public/ChaosStats.h
+++ b/Engine/Source/Runtime/Experimental/Chaos/Public/ChaosStats.h
@@ -50,6 +50,7 @@ DECLARE_CYCLE_STAT_EXTERN(TEXT("Skeletal Mesh Update Anim"), STAT_SkelMeshUpdate
 DECLARE_CYCLE_STAT_EXTERN(TEXT("Dispatch Event Notifies"), STAT_DispatchEventNotifies, STATGROUP_Chaos, CHAOS_API);
 DECLARE_CYCLE_STAT_EXTERN(TEXT("Dispatch Collision Events"), STAT_DispatchCollisionEvents, STATGROUP_Chaos, CHAOS_API);
 DECLARE_CYCLE_STAT_EXTERN(TEXT("Dispatch Break Events"), STAT_DispatchBreakEvents, STATGROUP_Chaos, CHAOS_API);
+DECLARE_CYCLE_STAT_EXTERN(TEXT("Dispatch Crumbling Events"), STAT_DispatchCrumblingEvents, STATGROUP_Chaos, CHAOS_API);
 DECLARE_CYCLE_STAT_EXTERN(TEXT("BufferPhysicsResults"), STAT_BufferPhysicsResults, STATGROUP_Chaos, CHAOS_API);
 DECLARE_CYCLE_STAT_EXTERN(TEXT("Flip Results"), STAT_FlipResults, STATGROUP_Chaos, CHAOS_API);
 DECLARE_CYCLE_STAT_EXTERN(TEXT("ProcessDeferredCreatePhysicsState"), STAT_ProcessDeferredCreatePhysicsState, STATGROUP_Chaos, CHAOS_API);
@@ -77,6 +78,7 @@ DECLARE_CYCLE_STAT_EXTERN(TEXT("Gathering Breaking Event Data"), STAT_GatherBrea
 DECLARE_CYCLE_STAT_EXTERN(TEXT("Gathering Trailing Event Data"), STAT_GatherTrailingEvent, STATGROUP_Chaos, CHAOS_API);
 DECLARE_CYCLE_STAT_EXTERN(TEXT("Gathering Sleeping Event Data"), STAT_GatherSleepingEvent, STATGROUP_Chaos, CHAOS_API);
 DECLARE_CYCLE_STAT_EXTERN(TEXT("Acceleration Structure Reset"), STAT_AccelerationStructureReset, STATGROUP_Chaos, CHAOS_API);
+DECLARE_CYCLE_STAT_EXTERN(TEXT("Gathering Crumbling Event Data"), STAT_GatherCrumblingEvent, STATGROUP_Chaos, CHAOS_API);
 
 DECLARE_CYCLE_STAT_EXTERN(TEXT("Reset Collision Rule"), STAT_ResetCollisionRule, STATGROUP_Chaos, CHAOS_API);
 
diff --git a/Engine/Source/Runtime/Experimental/Chaos/Public/EventDefaults.h b/Engine/Source/Runtime/Experimental/Chaos/Public/EventDefaults.h
--- a/Engine/Source/Runtime/Experimental/Chaos/Public/EventDefaults.h
+++ b/Engine/Source/Runtime/Experimental/Chaos/Public/EventDefaults.h
@@ -42,5 +42,10 @@ namespace Chaos
 		 */
 		static void RegisterRemovalEvent(FEventManager& EventManager);
 
+		/**
+		* Register crumbling event gathering function & data type
+		*/
+		static void RegisterCrumblingEvent(FEventManager& EventManager);
+
 	};
 }
\ No newline at end of file
diff --git a/Engine/Source/Runtime/Experimental/Chaos/Public/EventManager.h b/Engine/Source/Runtime/Experimental/Chaos/Public/EventManager.h
--- a/Engine/Source/Runtime/Experimental/Chaos/Public/EventManager.h
+++ b/Engine/Source/Runtime/Experimental/Chaos/Public/EventManager.h
@@ -26,7 +26,8 @@ namespace Chaos
 		Breaking = 1,
 		Trailing = 2,
 		Sleeping = 3,
-		Removal = 4
+		Removal = 4,
+		Crumbling = 5,
 	};
 
 	typedef int32 FEventID;
diff --git a/Engine/Source/Runtime/Experimental/Chaos/Public/EventsData.h b/Engine/Source/Runtime/Experimental/Chaos/Public/EventsData.h
--- a/Engine/Source/Runtime/Experimental/Chaos/Public/EventsData.h
+++ b/Engine/Source/Runtime/Experimental/Chaos/Public/EventsData.h
@@ -22,6 +22,7 @@ namespace Chaos
 	typedef TArray<FTrailingData> FTrailingDataArray;
 	typedef TArray<FRemovalData> FRemovalDataArray;
 	typedef TArray<FSleepingData> FSleepingDataArray;
+	typedef TArray<FCrumblingData> FCrumblingDataArray;
 
 	/* Common */
 
@@ -151,4 +152,49 @@ namespace Chaos
 		FSleepingDataArray SleepingData;
 	};
 
+	/*
+	* All the crumbling events for one frame time stamped with the time for that frame
+	*/
+	struct FAllCrumblingData : public FTimeResource
+	{
+		FAllCrumblingData() : AllCrumblingsArray(FCrumblingDataArray()) {}
+
+		void Reset()
+		{
+			AllCrumblingsArray.Reset();
+		}
+
+		FCrumblingDataArray AllCrumblingsArray;
+	};
+
+	struct FCrumblingEventData
+	{
+		FCrumblingEventData() {}
+
+		void Reset()
+		{
+			CrumblingData.Reset();
+			PhysicsProxyToCrumblingIndices.Reset();
+		}
+
+		FORCEINLINE_DEBUGGABLE void Reserve(int32 Num)
+		{
+			CrumblingData.AllCrumblingsArray.Reserve(Num);
+		}
+		
+		FORCEINLINE_DEBUGGABLE void SetTimeCreated(FReal TimeCreatedIn)
+		{
+			CrumblingData.TimeCreated = TimeCreatedIn;
+		}
+
+		FORCEINLINE_DEBUGGABLE void AddCrumbling(const FCrumblingData& CrumblingToAdd)
+		{
+			const int32 NewIndex = CrumblingData.AllCrumblingsArray.Emplace(CrumblingToAdd);
+			TArray<int32>& Indices = PhysicsProxyToCrumblingIndices.PhysicsProxyToIndicesMap.FindOrAdd(CrumblingToAdd.Proxy);
+			Indices.Add(NewIndex);
+		}
+		
+		FAllCrumblingData CrumblingData;
+		FIndicesByPhysicsProxy PhysicsProxyToCrumblingIndices;
+	};
 }
diff --git a/Engine/Source/Runtime/Experimental/Chaos/Public/GeometryCollection/GeometryCollectionSimulationCoreTypes.h b/Engine/Source/Runtime/Experimental/Chaos/Public/GeometryCollection/GeometryCollectionSimulationCoreTypes.h
--- a/Engine/Source/Runtime/Experimental/Chaos/Public/GeometryCollection/GeometryCollectionSimulationCoreTypes.h
+++ b/Engine/Source/Runtime/Experimental/Chaos/Public/GeometryCollection/GeometryCollectionSimulationCoreTypes.h
@@ -198,6 +198,8 @@ struct FSimulationParameters
 		, bGenerateCollisionData(false)
 		, bGenerateTrailingData(false)
 		, bGenerateRemovalsData(false)
+		, bGenerateCrumblingData(false)
+		, bGenerateCrumblingChildrenData(false)
 		, RemoveOnFractureEnabled(false)
 		, EnableGravity(true)
 		, UseInertiaConditioning(true)
@@ -241,6 +243,8 @@ struct FSimulationParameters
 		, bGenerateCollisionData(Other.bGenerateCollisionData)
 		, bGenerateTrailingData(Other.bGenerateTrailingData)
 		, bGenerateRemovalsData(Other.bGenerateRemovalsData)
+		, bGenerateCrumblingData(Other.bGenerateCrumblingData)
+		, bGenerateCrumblingChildrenData(Other.bGenerateCrumblingChildrenData)
 		, Shared(Other.Shared)
 		, RemoveOnFractureEnabled(Other.RemoveOnFractureEnabled)
 		, EnableGravity(Other.EnableGravity)
@@ -305,6 +309,8 @@ struct FSimulationParameters
 	bool bGenerateCollisionData;
 	bool bGenerateTrailingData;
 	bool bGenerateRemovalsData;
+	bool bGenerateCrumblingData;
+	bool bGenerateCrumblingChildrenData;
 
 	FSharedSimulationParameters Shared;
 
diff --git a/Engine/Source/Runtime/Experimental/Chaos/Public/PhysicsProxy/GeometryCollectionPhysicsProxy.h b/Engine/Source/Runtime/Experimental/Chaos/Public/PhysicsProxy/GeometryCollectionPhysicsProxy.h
--- a/Engine/Source/Runtime/Experimental/Chaos/Public/PhysicsProxy/GeometryCollectionPhysicsProxy.h
+++ b/Engine/Source/Runtime/Experimental/Chaos/Public/PhysicsProxy/GeometryCollectionPhysicsProxy.h
@@ -318,6 +318,19 @@ public:
 		return INDEX_NONE;
 	}
 
+	void GetTransformGroupIndicesFromHandles(const TArray<FParticleHandle*> Handles, TArray<int32>& IndicesOut) const
+	{
+		IndicesOut.SetNumUninitialized(Handles.Num());
+		for (int32 HandleIndex = 0; HandleIndex < Handles.Num(); ++HandleIndex)
+		{
+			IndicesOut[HandleIndex] = INDEX_NONE;
+			if (const int32* TransformGroupIndex = HandleToTransformGroupIndex.Find(Handles[HandleIndex]))
+			{
+				IndicesOut[HandleIndex] = (*TransformGroupIndex);
+			}
+		}
+	}
+
 	FGeometryCollectionItemIndex GetInternalClusterParentItemIndex_External(int32 ChildTransformIndex) const
 	{
 		// first find the GTParticle matching the Transform index
diff --git a/Engine/Source/Runtime/Experimental/ChaosSolverEngine/Private/Chaos/ChaosGameplayEventDispatcher.cpp b/Engine/Source/Runtime/Experimental/ChaosSolverEngine/Private/Chaos/ChaosGameplayEventDispatcher.cpp
--- a/Engine/Source/Runtime/Experimental/ChaosSolverEngine/Private/Chaos/ChaosGameplayEventDispatcher.cpp
+++ b/Engine/Source/Runtime/Experimental/ChaosSolverEngine/Private/Chaos/ChaosGameplayEventDispatcher.cpp
@@ -79,6 +79,20 @@ static void DispatchPendingRemovalEvents(TArray<FChaosRemovalEvent> const& Event
 	}
 }
 
+static void DispatchPendingCrumblingEvents(TArray<FChaosCrumblingEvent> const& Events, TMap<TObjectPtr<UPrimitiveComponent>, FCrumblingEventCallbackWrapper> const& Registrations)
+{
+	for (FChaosCrumblingEvent const& E : Events)
+	{
+		if (E.Component)
+		{
+			if (const FCrumblingEventCallbackWrapper* const Callback = Registrations.Find(E.Component))
+			{
+				Callback->CrumblingEventCallback(E);
+			}
+		}
+	}
+}
+
 static void SetCollisionInfoFromComp(FRigidBodyCollisionInfo& Info, UPrimitiveComponent* Comp)
 {
 	if (Comp)
@@ -245,6 +259,23 @@ void UChaosGameplayEventDispatcher::UnRegisterForRemovalEvents(UPrimitiveCompone
 	}
 }
 
+void UChaosGameplayEventDispatcher::RegisterForCrumblingEvents(UPrimitiveComponent* Component, FOnCrumblingEventCallback InFunc)
+{
+	if (Component)
+	{
+		FCrumblingEventCallbackWrapper F = { InFunc };
+		CrumblingEventRegistrations.Add(Component, F);
+	}
+}
+
+void UChaosGameplayEventDispatcher::UnRegisterForCrumblingEvents(UPrimitiveComponent* Component)
+{
+	if (Component)
+	{
+		CrumblingEventRegistrations.Remove(Component);
+	}
+}
+
 void UChaosGameplayEventDispatcher::DispatchPendingWakeNotifies()
 {
 	for (auto MapItr = PendingSleepNotifies.CreateIterator(); MapItr; ++MapItr)
@@ -270,6 +301,7 @@ void UChaosGameplayEventDispatcher::RegisterChaosEvents()
 			EventManager->RegisterHandler<Chaos::FBreakingEventData>(Chaos::EEventType::Breaking, this, &UChaosGameplayEventDispatcher::HandleBreakingEvents);
 			EventManager->RegisterHandler<Chaos::FSleepingEventData>(Chaos::EEventType::Sleeping, this, &UChaosGameplayEventDispatcher::HandleSleepingEvents);
 			EventManager->RegisterHandler<Chaos::FRemovalEventData>(Chaos::EEventType::Removal, this, &UChaosGameplayEventDispatcher::HandleRemovalEvents);
+			EventManager->RegisterHandler<Chaos::FCrumblingEventData>(Chaos::EEventType::Crumbling, this, &UChaosGameplayEventDispatcher::HandleCrumblingEvents);
 		}
 	}
 }
@@ -287,6 +319,7 @@ void UChaosGameplayEventDispatcher::UnregisterChaosEvents()
 				EventManager->UnregisterHandler(Chaos::EEventType::Breaking, this);
 				EventManager->UnregisterHandler(Chaos::EEventType::Sleeping, this);
 				EventManager->UnregisterHandler(Chaos::EEventType::Removal, this);
+				EventManager->UnregisterHandler(Chaos::EEventType::Crumbling, this);
 			}
 		}
 	}
@@ -543,3 +576,47 @@ void UChaosGameplayEventDispatcher::HandleRemovalEvents(const Chaos::FRemovalEve
 
 }
 
+void UChaosGameplayEventDispatcher::HandleCrumblingEvents(const Chaos::FCrumblingEventData& Event)
+{
+	SCOPE_CYCLE_COUNTER(STAT_DispatchCrumblingEvents);
+
+	// CRUMBLING EVENTS
+
+	const float CrumblingDataTimestamp = Event.CrumblingData.TimeCreated;
+	if (CrumblingDataTimestamp > LastCrumblingDataTime)
+	{
+		LastCrumblingDataTime = CrumblingDataTimestamp;
+
+		Chaos::FCrumblingDataArray const& BreakingData = Event.CrumblingData.AllCrumblingsArray;
+
+		// let's assume crumbles are rare, so we will iterate breaks instead of registered components for now
+		TArray<FChaosCrumblingEvent> PendingCrumblingEvent;
+		for (const Chaos::FCrumblingData& CrumblingDataItem : Event.CrumblingData.AllCrumblingsArray)
+		{	
+			if (CrumblingDataItem.Proxy)
+			{
+				if (UPrimitiveComponent* const PrimComp = Cast<UPrimitiveComponent>(CrumblingDataItem.Proxy->GetOwner()))
+				{
+					if (CrumblingEventRegistrations.Contains(PrimComp))
+					{
+						// queue them up so we can release the physics data before triggering BP events
+						FChaosCrumblingEvent& CrumblingEvent = PendingCrumblingEvent.AddZeroed_GetRef();
+						CrumblingEvent.Component = PrimComp;
+						CrumblingEvent.Location = CrumblingDataItem.Location;
+						CrumblingEvent.Orientation = CrumblingDataItem.Orientation;
+						CrumblingEvent.LinearVelocity = CrumblingDataItem.LinearVelocity;
+						CrumblingEvent.AngularVelocity = CrumblingDataItem.AngularVelocity;
+						CrumblingEvent.Mass = static_cast<float>(CrumblingDataItem.Mass);
+						CrumblingEvent.LocalBounds = FBox(CrumblingDataItem.LocalBounds.Min(), CrumblingDataItem.LocalBounds.Max());
+						CrumblingEvent.Children = CrumblingDataItem.Children;
+					}
+				}
+			}
+		}
+		if (PendingCrumblingEvent.Num() > 0)
+		{
+			DispatchPendingCrumblingEvents(PendingCrumblingEvent, CrumblingEventRegistrations);
+		}
+	}
+
+}
diff --git a/Engine/Source/Runtime/Experimental/ChaosSolverEngine/Public/Chaos/ChaosGameplayEventDispatcher.h b/Engine/Source/Runtime/Experimental/ChaosSolverEngine/Public/Chaos/ChaosGameplayEventDispatcher.h
--- a/Engine/Source/Runtime/Experimental/ChaosSolverEngine/Public/Chaos/ChaosGameplayEventDispatcher.h
+++ b/Engine/Source/Runtime/Experimental/ChaosSolverEngine/Public/Chaos/ChaosGameplayEventDispatcher.h
@@ -15,6 +15,7 @@ namespace Chaos
 	struct FBreakingEventData;
 	struct FSleepingEventData;
 	struct FRemovalEventData;
+	struct FCrumblingEventData;
 }
 
 USTRUCT(BlueprintType)
@@ -71,6 +72,53 @@ public:
 };
 
 
+USTRUCT(BlueprintType)
+struct CHAOSSOLVERENGINE_API FChaosCrumblingEvent
+{
+	GENERATED_BODY()
+
+public:
+	FChaosCrumblingEvent()
+		: Component(nullptr)
+		, Location(FVector::ZeroVector)
+		, Orientation(FQuat::Identity)
+		, LinearVelocity(FVector::ZeroVector)
+		, AngularVelocity(FVector::ZeroVector)
+	{}
+
+	/** primitive component involved in the crumble event */
+	UPROPERTY(BlueprintReadOnly, Category = "Crumble Event")
+	TObjectPtr<UPrimitiveComponent> Component = nullptr;
+
+	/** World location of the crumbling cluster */
+	UPROPERTY(BlueprintReadOnly, Category = "Crumble Event")
+	FVector Location;
+
+	/** World orientation of the crumbling cluster */
+	UPROPERTY(BlueprintReadOnly, Category = "Crumble Event")
+	FQuat Orientation;
+
+	/** Linear Velocity of the crumbling cluster */
+	UPROPERTY(BlueprintReadOnly, Category = "Crumble Event")
+	FVector LinearVelocity;
+
+	/** Angular Velocity of the crumbling cluster  */
+	UPROPERTY(BlueprintReadOnly, Category = "Crumble Event")
+	FVector AngularVelocity;
+
+	/** Mass of the crumbling cluster  */
+	UPROPERTY(BlueprintReadOnly, Category = "Crumble Event")
+	float Mass;
+
+	/** Local bounding box of the crumbling cluster  */
+	UPROPERTY(BlueprintReadOnly, Category = "Crumble Event")
+	FBox LocalBounds;
+	
+	/** List of children indices released (optional : see geometry collection component bCrumblingEventIncludesChildren) */
+	UPROPERTY(BlueprintReadOnly, Category = "Crumble Event")
+	TArray<int32> Children;
+};
+
 typedef TFunction<void(const FChaosBreakEvent&)> FOnBreakEventCallback;
 
 /** UStruct wrapper so we can store the TFunction in a TMap */
@@ -95,6 +143,18 @@ public:
 	FOnRemovalEventCallback RemovalEventCallback;
 };
 
+typedef TFunction<void(const FChaosCrumblingEvent&)> FOnCrumblingEventCallback;
+
+/** UStruct wrapper so we can store the TFunction in a TMap */
+USTRUCT()
+struct CHAOSSOLVERENGINE_API FCrumblingEventCallbackWrapper
+{
+	GENERATED_BODY()
+
+public:
+	FOnCrumblingEventCallback CrumblingEventCallback;
+};
+
 /** UStruct wrapper so we can store the TSet in a TMap */
 USTRUCT()
 struct FChaosHandlerSet
@@ -177,6 +237,9 @@ public:
 	void RegisterForRemovalEvents(UPrimitiveComponent* Component, FOnRemovalEventCallback InFunc);
 	void UnRegisterForRemovalEvents(UPrimitiveComponent* Component);
 
+	void RegisterForCrumblingEvents(UPrimitiveComponent* Component, FOnCrumblingEventCallback InFunc);
+	void UnRegisterForCrumblingEvents(UPrimitiveComponent* Component);
+	
 private:
 
  	UPROPERTY()
@@ -188,9 +251,14 @@ private:
 	UPROPERTY()
 	TMap<TObjectPtr<UPrimitiveComponent>, FRemovalEventCallbackWrapper> RemovalEventRegistrations;
 
+	UPROPERTY()
+	TMap<TObjectPtr<UPrimitiveComponent>, FCrumblingEventCallbackWrapper> CrumblingEventRegistrations;
+	
+
 	float LastCollisionDataTime = -1.f;
 	float LastBreakingDataTime = -1.f;
 	float LastRemovalDataTime = -1.f;
+	float LastCrumblingDataTime = -1.f;
 
 	void DispatchPendingCollisionNotifies();
 	void DispatchPendingWakeNotifies();
@@ -204,6 +272,7 @@ private:
 	void HandleSleepingEvents(const Chaos::FSleepingEventData& SleepingData);
 	void AddPendingSleepingNotify(FBodyInstance* BodyInstance, ESleepEvent SleepEventType);
 	void HandleRemovalEvents(const Chaos::FRemovalEventData& RemovalData);
+	void HandleCrumblingEvents(const Chaos::FCrumblingEventData& CrumblingData);
 
 };
 
diff --git a/Engine/Source/Runtime/Experimental/GeometryCollectionEngine/Private/GeometryCollection/GeometryCollectionComponent.cpp b/Engine/Source/Runtime/Experimental/GeometryCollectionEngine/Private/GeometryCollection/GeometryCollectionComponent.cpp
--- a/Engine/Source/Runtime/Experimental/GeometryCollectionEngine/Private/GeometryCollection/GeometryCollectionComponent.cpp
+++ b/Engine/Source/Runtime/Experimental/GeometryCollectionEngine/Private/GeometryCollection/GeometryCollectionComponent.cpp
@@ -244,6 +244,7 @@ UGeometryCollectionComponent::UGeometryCollectionComponent(const FObjectInitiali
 	, bNotifyBreaks(false)
 	, bNotifyCollisions(false)
 	, bNotifyRemovals(false)
+	, bNotifyCrumblings(false)
 	, bStoreVelocities(false)
 	, bShowBoneColors(false)
 #if WITH_EDITORONLY_DATA 
@@ -713,6 +714,15 @@ void UGeometryCollectionComponent::SetNotifyRemovals(bool bNewNotifyRemovals)
 	}
 }
 
+void UGeometryCollectionComponent::SetNotifyCrumblings(bool bNewNotifyCrumblings)
+{
+	if (bNotifyCrumblings != bNewNotifyCrumblings)
+	{
+		bNotifyCrumblings = bNewNotifyCrumblings;
+		UpdateBreakEventRegistration();
+	}
+}
+
 FBodyInstance* UGeometryCollectionComponent::GetBodyInstance(FName BoneName /*= NAME_None*/, bool bGetWelded /*= true*/, int32 Index /*=INDEX_NONE*/) const
 {
 	return nullptr;// const_cast<FBodyInstance*>(&DummyBodyInstance);
@@ -839,6 +849,15 @@ void UGeometryCollectionComponent::DispatchRemovalEvent(const FChaosRemovalEvent
 	}
 }
 
+void UGeometryCollectionComponent::DispatchCrumblingEvent(const FChaosCrumblingEvent& Event)
+{
+	// bp
+	if (OnChaosCrumblingEvent.IsBound())
+	{
+		OnChaosCrumblingEvent.Broadcast(Event);
+	}
+}
+
 bool UGeometryCollectionComponent::DoCustomNavigableGeometryExport(FNavigableGeometryExport& GeomExport) const
 {
 	if(!RestCollection)
@@ -1271,6 +1290,14 @@ static void DispatchGeometryCollectionRemovalEvent(const FChaosRemovalEvent& Eve
 	}
 }
 
+static void DispatchGeometryCollectionCrumblingEvent(const FChaosCrumblingEvent& Event)
+{
+	if (UGeometryCollectionComponent* const GC = Cast<UGeometryCollectionComponent>(Event.Component))
+	{
+		GC->DispatchCrumblingEvent(Event);
+	}
+}
+
 void UGeometryCollectionComponent::DispatchChaosPhysicsCollisionBlueprintEvents(const FChaosPhysicsCollisionInfo& CollisionInfo)
 {
 	ReceivePhysicsCollision(CollisionInfo);
@@ -1280,7 +1307,7 @@ void UGeometryCollectionComponent::DispatchChaosPhysicsCollisionBlueprintEvents(
 // call when first registering
 void UGeometryCollectionComponent::RegisterForEvents()
 {
-	if (BodyInstance.bNotifyRigidBodyCollision || bNotifyBreaks || bNotifyCollisions || bNotifyRemovals)
+	if (BodyInstance.bNotifyRigidBodyCollision || bNotifyBreaks || bNotifyCollisions || bNotifyRemovals || bNotifyCrumblings)
 	{
 		Chaos::FPhysicsSolver* Solver = GetWorld()->GetPhysicsScene()->GetSolver();
 		if (Solver)
@@ -1316,6 +1343,16 @@ void UGeometryCollectionComponent::RegisterForEvents()
 					});
 
 			}
+
+			if (bNotifyCrumblings)
+			{
+				EventDispatcher->RegisterForCrumblingEvents(this, &DispatchGeometryCollectionCrumblingEvent);
+
+				Solver->EnqueueCommandImmediate([Solver]()
+					{
+						Solver->SetGenerateBreakingData(true);
+					});
+			}
 		}
 	}
 }
@@ -1356,6 +1393,18 @@ void UGeometryCollectionComponent::UpdateRemovalEventRegistration()
 	}
 }
 
+void UGeometryCollectionComponent::UpdateCrumblingEventRegistration()
+{
+	if (bNotifyCrumblings)
+	{
+		EventDispatcher->RegisterForCrumblingEvents(this, &DispatchGeometryCollectionCrumblingEvent);
+	}
+	else
+	{
+		EventDispatcher->UnRegisterForCrumblingEvents(this);
+	}
+}
+
 void ActivateClusters(Chaos::FRigidClustering& Clustering, Chaos::TPBDRigidClusteredParticleHandle<Chaos::FReal, 3>* Cluster)
 {
 	if(!Cluster)
@@ -2301,6 +2350,8 @@ void UGeometryCollectionComponent::RegisterAndInitializePhysicsProxy()
 		SimulationParameters.bGenerateCollisionData = bNotifyCollisions;
 		SimulationParameters.bGenerateTrailingData = bNotifyTrailing;
 		SimulationParameters.bGenerateRemovalsData = bNotifyRemovals;
+		SimulationParameters.bGenerateCrumblingData = bNotifyCrumblings;
+		SimulationParameters.bGenerateCrumblingChildrenData = bCrumblingEventIncludesChildren;
 		SimulationParameters.RemoveOnFractureEnabled = SimulationParameters.Shared.RemoveOnFractureIndices.Num() > 0;
 		SimulationParameters.EnableGravity = BodyInstance.bEnableGravity;
 		SimulationParameters.UseInertiaConditioning = BodyInstance.IsInertiaConditioningEnabled();
@@ -3868,4 +3919,4 @@ void UGeometryCollectionComponent::PostLoad()
 		BodyInstance.SetPhysMaterialOverride(PhysicalMaterialOverride_DEPRECATED.Get());
 		PhysicalMaterialOverride_DEPRECATED = nullptr;
 	}
-}
\ No newline at end of file
+}
diff --git a/Engine/Source/Runtime/Experimental/GeometryCollectionEngine/Public/GeometryCollection/GeometryCollectionComponent.h b/Engine/Source/Runtime/Experimental/GeometryCollectionEngine/Public/GeometryCollection/GeometryCollectionComponent.h
--- a/Engine/Source/Runtime/Experimental/GeometryCollectionEngine/Public/GeometryCollection/GeometryCollectionComponent.h
+++ b/Engine/Source/Runtime/Experimental/GeometryCollectionEngine/Public/GeometryCollection/GeometryCollectionComponent.h
@@ -44,6 +44,8 @@ DECLARE_DYNAMIC_MULTICAST_DELEGATE_OneParam(FOnChaosBreakEvent, const FChaosBrea
 
 DECLARE_DYNAMIC_MULTICAST_DELEGATE_OneParam(FOnChaosRemovalEvent, const FChaosRemovalEvent&, RemovalEvent);
 
+DECLARE_DYNAMIC_MULTICAST_DELEGATE_OneParam(FOnChaosCrumblingEvent, const FChaosCrumblingEvent&, CrumbleEvent);
+
 namespace GeometryCollection
 {
 	enum class ESelectionMode : uint8
@@ -798,6 +800,10 @@ public:
 	UFUNCTION(BlueprintCallable, Category = "Physics")
 	void SetNotifyRemovals(bool bNewNotifyRemovals);
 
+	/** Changes whether or not this component will get future break notifications. */
+	UFUNCTION(BlueprintCallable, Category = "Physics")
+	void SetNotifyCrumblings(bool bNewNotifyCrumblings);
+	
 	/** Overrideable native notification */
 	virtual void NotifyBreak(const FChaosBreakEvent& Event) {};
 
@@ -809,11 +815,16 @@ public:
 
 	UPROPERTY(BlueprintAssignable, Category = "Chaos")
 	FOnChaosRemovalEvent OnChaosRemovalEvent;
-	
+
+	UPROPERTY(BlueprintAssignable, Category = "Chaos")
+	FOnChaosCrumblingEvent OnChaosCrumblingEvent;
+
 	void DispatchBreakEvent(const FChaosBreakEvent& Event);
 
 	void DispatchRemovalEvent(const FChaosRemovalEvent& Event);
 
+	void DispatchCrumblingEvent(const FChaosCrumblingEvent& Event);
+
 	UPROPERTY(Transient, VisibleAnywhere, BlueprintReadWrite, Interp, Category = "Chaos")
 	float DesiredCacheTime;
 
@@ -883,6 +894,14 @@ public:
 	UPROPERTY(EditAnywhere, BlueprintReadOnly, Category = "ChaosPhysics|General")
 	bool bNotifyRemovals;
 
+	/** If true, this component will generate crumbling events that other systems may subscribe to. */
+	UPROPERTY(EditAnywhere, BlueprintReadOnly, Category = "ChaosPhysics|General")
+	bool bNotifyCrumblings;
+
+	/** If this and bNotifyCrumblings are true, the crumbling events will contain released children indices. */
+	UPROPERTY(EditAnywhere, BlueprintReadOnly, Category = "ChaosPhysics|General", meta = (EditCondition = "bNotifyCrumblings"))
+	bool bCrumblingEventIncludesChildren;
+	
 	/** If true, this component will save linear and angular velocities on its DynamicCollection. */
 	UPROPERTY(EditAnywhere, BlueprintReadOnly, Category = "ChaosPhysics|General")
 	bool bStoreVelocities;
@@ -923,9 +942,8 @@ protected:
 	void UpdateRBCollisionEventRegistration();
 	void UpdateBreakEventRegistration();
 	void UpdateRemovalEventRegistration();
-
-
-
+	void UpdateCrumblingEventRegistration();
+	
 	/* Per-instance override to enable/disable replication for the geometry collection */
 	UPROPERTY(EditAnywhere, BlueprintReadOnly, Category=Network)
 	bool bEnableReplication;
