Fracture tool : fix issues with setting removal parameters
- when timers where set to 0, enable and cluster crumbling would not be set properly

diff --git a/Engine/Source/Runtime/Experimental/GeometryCollectionEngine/Public/GeometryCollection/GeometryCollectionEngineRemoval.h b/Engine/Source/Runtime/Experimental/GeometryCollectionEngine/Public/GeometryCollection/GeometryCollectionEngineRemoval.h
--- a/Engine/Source/Runtime/Experimental/GeometryCollectionEngine/Public/GeometryCollection/GeometryCollectionEngineRemoval.h
+++ b/Engine/Source/Runtime/Experimental/GeometryCollectionEngine/Public/GeometryCollection/GeometryCollectionEngineRemoval.h
@@ -30,8 +30,8 @@ public:
 		PackedData.Z = FMath::Min(FMath::Abs(RemovalTimer.X), FMath::Abs(RemovalTimer.Y)); // Min removal timer
 		PackedData.W = FMath::Max(FMath::Abs(RemovalTimer.X), FMath::Abs(RemovalTimer.Y)); // Max removal timer
 
-		PackedData.X *= bEnable? +1: -1;
-		PackedData.Z *= bClusterCrumbling? -1: +1;
+		PackedData.X = bEnable? PackedData.X: -1.0;
+		PackedData.Z = bClusterCrumbling? -1.0: PackedData.Z;
 	}
 	
 	bool IsEnabled() const { return PackedData.X >= 0; }
