Iris experimental replication system.

diff --git a/Engine/Source/Runtime/Engine/Private/RepLayout.cpp b/Engine/Source/Runtime/Engine/Private/RepLayout.cpp
--- a/Engine/Source/Runtime/Engine/Private/RepLayout.cpp
+++ b/Engine/Source/Runtime/Engine/Private/RepLayout.cpp
@@ -26,6 +26,9 @@
 #include "Net/Core/Trace/NetTrace.h"
 #include "UObject/EnumProperty.h"
 #include "UObject/UnrealType.h"
+#if UE_WITH_IRIS
+#include "Net/Iris/ReplicationSystem/ReplicationSystemUtil.h"
+#endif // UE_WITH_IRIS
 
 DECLARE_CYCLE_STAT(TEXT("RepLayout InitFromObjectClass"), STAT_RepLayout_InitFromObjectClass, STATGROUP_Game);
 DECLARE_CYCLE_STAT(TEXT("RepLayout BuildShadowOffsets"), STAT_RepLayout_BuildShadowOffsets, STATGROUP_Game);
@@ -199,7 +202,15 @@ namespace UE_RepLayout_Private
 	const UEPushModelPrivate::FPushModelPerNetDriverHandle ConditionallyAddPushModelObject(const UObject* const Object, const TSharedRef<const FRepLayout>& InRepLayout)
 	{
 		const int32 NumReplicatedProperties = InRepLayout->GetNumParents();
-		if (UEPushModelPrivate::IsPushModelEnabled() &&
+
+#if UE_WITH_IRIS
+		// Implement shared PushModelids/NetHandles for Iris and PushModel to avoid conflicts when we mix systems - JIRA: UE-158304
+		const bool bShouldUseIrisReplication = UE::Net::ShouldUseIrisReplication();
+#else
+		const bool bShouldUseIrisReplication = false;
+#endif
+
+		if (UEPushModelPrivate::IsPushModelEnabled() && !bShouldUseIrisReplication &&
 			NumReplicatedProperties > 0 &&
 			EnumHasAnyFlags(InRepLayout->GetFlags(), ERepLayoutFlags::FullPushSupport | ERepLayoutFlags::PartialPushSupport))
 		{
@@ -967,25 +978,25 @@ static uint32 GetRepLayoutCmdCompatibleChecksum(
 #if (WITH_PUSH_MODEL)
 struct FNetPrivatePushIdHelper
 {
-	static void SetNetPushID(UObject* InObject, const UEPushModelPrivate::FNetPushObjectId ObjectId)
+	static void SetNetPushID(UObject* InObject, const UEPushModelPrivate::FNetLegacyPushObjectId ObjectId)
 	{
-		const UEPushModelPrivate::FNetPushObjectId CurrentId = InObject->GetNetPushIdDynamic();
-		if (CurrentId != ObjectId)
+		const UEPushModelPrivate::FNetPushObjectId CurrentId(InObject->GetNetPushIdDynamic());
+		if (CurrentId.GetLegacyPushObjectId() != ObjectId)
 		{
-			if (CurrentId != INDEX_NONE)
+			if (CurrentId.IsValid())
 			{
 				UE_LOG(LogRep, Error, TEXT("SetNetPushID: %s already has a push id. Existing ID = %s, New ID = %s"),
 					*InObject->GetPathName(),
 					*UEPushModelPrivate::ToString(CurrentId),
-					*UEPushModelPrivate::ToString(ObjectId));
+					*UEPushModelPrivate::ToString(UEPushModelPrivate::FNetPushObjectId(ObjectId)));
 
-				if (!UEPushModelPrivate::ValidateObjectIdReassignment(CurrentId, ObjectId))
+				if (!UEPushModelPrivate::ValidateObjectIdReassignment(CurrentId.GetLegacyPushObjectId(), ObjectId))
 				{
 					return;
 				}
 			}
 
-			FObjectNetPushIdHelper::SetNetPushIdDynamic(InObject, ObjectId);
+			FObjectNetPushIdHelper::SetNetPushIdDynamic(InObject, UEPushModelPrivate::FNetPushObjectId(ObjectId).GetValue());
 		}
 	}
 
@@ -7865,6 +7876,8 @@ ERepLayoutResult FRepLayout::DeltaSerializeFastArrayProperty(FFastArrayDeltaSeri
 		// Now that we have our changelists setup, we can send the data.
 		for (int32 i = 0; i < ChangedElements.Num(); ++i)
 		{
+			UE_NET_TRACE_SCOPE(ChangedElement, Writer, GetTraceCollector(Writer), ENetTraceVerbosity::Trace);
+
 			const auto& IDIndexPair = ChangedElements[i];
 			uint32 ID = ChangedElements[i].ID;
 			Writer << ID;
diff --git a/Engine/Source/Runtime/Engine/Private/UnrealEngine.cpp b/Engine/Source/Runtime/Engine/Private/UnrealEngine.cpp
--- a/Engine/Source/Runtime/Engine/Private/UnrealEngine.cpp
+++ b/Engine/Source/Runtime/Engine/Private/UnrealEngine.cpp
@@ -140,6 +140,10 @@ UnrealEngine.cpp: Implements the UEngine class and helpers.
 #include "GPUSkinCache.h"
 #include "Chaos/TriangleMeshImplicitObject.h"
 
+#if UE_WITH_IRIS
+#include "Iris/IrisConfig.h"
+#endif
+
 #include "Particles/Spawn/ParticleModuleSpawn.h"
 #include "Particles/TypeData/ParticleModuleTypeDataMesh.h"
 #include "Particles/ParticleLODLevel.h"
@@ -13039,233 +13043,298 @@ UNetDriver* UEngine::FindNamedNetDriver(const UPendingNetGame* InPendingNetGame,
 	return FindNamedNetDriver_Local(GetWorldContextFromPendingNetGameChecked(InPendingNetGame).ActiveNetDrivers, NetDriverName);
 }
 
-UNetDriver* CreateNetDriver_Local(UEngine* Engine, FWorldContext& Context, FName NetDriverDefinition)
+namespace UE::Private
 {
-	UNetDriver* ReturnVal = nullptr;
-	FNetDriverDefinition* Definition = nullptr;
-	auto FindNetDriverDefPred =
-		[NetDriverDefinition](const FNetDriverDefinition& CurDef)
+	/**
+	 * Look for a config setting telling if this new NetDriver should be enabled with Iris or not
+	 * To enable iris for a netdriver you need to add the proper configuration in Engine.ini like so:
+	 *		[/Script/Engine.Engine]
+	 *		+IrisNetDriverConfigs=(NetDriverName="MyNetDriver",bEnableIris=true)
+	 *		+IrisNetDriverConfigs=(NetDriverDefinition="MyNetDriverDef",bEnableIris=false)
+	 *		+IrisNetDriverConfigs=(NetDriverWildcardName="SecondaryNetDriver*",bEnableIris=true);
+	 *
+	 * Priority order for the IrisNetDriverConfigs are:
+	 *		1. NetDriverName exact match
+	 *		2. NetDriverName wildcard match
+	 *		3. NetDriverDefinition match
+	 *
+	 */
+	bool IsNetDriverUsingIris(UEngine* Engine, FName InNetDriverDefinition, FName InNetDriverName)
 	{
-		return CurDef.DefName == NetDriverDefinition;
-	};
+#if UE_WITH_IRIS
+		// NetDrivers can only use Iris if the global setting allows it.
+		if (UE::Net::ShouldUseIrisReplication() == false)
+		{
+			return false;
+		}
 
-#if !UE_BUILD_SHIPPING
-	/**
-	* Commandline override for the net driver.
-	*
-	* Format: (NOTE: Use quotes whenever the ',' character is used)
-	*	Override the main/game net driver (most common usage):
-	*		-NetDriverOverrides=DriverClassName
-	*
-	*	Override a specific/named net driver:
-	*		-NetDriverOverrides="DefName,DriverClassName"
-	*
-	*	Override a specific driver, including fallback driver:
-	*		-NetDriverOverrides="DefName,DriverClassName,DriverClassNameFallback"
-	*
-	*	Override multiple net drivers:
-	*		-NetDriverOverrides="DriverClassName;DefName2,DriverClassName2"
-	*
-	*
-	* Example:
-	*	Use WebSocket for the main game net driver:
-	*		-NetDriverOverrides=/Script/WebSocketNetworking.WebSocketNetDriver
-	*
-	*	Use WebSocket for the main game net driver, and the party beacon net driver
-	*		-NetDriverOverrides="/Script/WebSocketNetworking.WebSocketNetDriver;BeaconNetDriver,/Script/WebSocketNetworking.WebSocketNetDriver"
-	*/
+		FIrisNetDriverConfig* IrisConfig = nullptr;
 
-	static TArray<FNetDriverDefinition> NetDriverOverrides = TArray<FNetDriverDefinition>();
-	FString OverrideCmdLine;
+		// Search for the exact name match
+		if (IrisConfig == nullptr)
+		{
+			IrisConfig = Engine->IrisNetDriverConfigs.FindByPredicate([InNetDriverName](const FIrisNetDriverConfig& Config)
+				{
+					return Config.NetDriverName.IsNone() ? false : Config.NetDriverName == InNetDriverName;
+				});
+		}
+
+		// Search for a wildcard match
+		if (IrisConfig == nullptr)
+		{
+			const FString TempNetDriverName = InNetDriverName.ToString();
+			IrisConfig = Engine->IrisNetDriverConfigs.FindByPredicate([TempNetDriverName](const FIrisNetDriverConfig& Config)
+				{
+					return Config.NetDriverWildcardName.IsEmpty() ? false : TempNetDriverName.MatchesWildcard(Config.NetDriverWildcardName);
+				});
+		}
+
+		// Search for the definition match
+		if (IrisConfig == nullptr)
+		{
+			IrisConfig = Engine->IrisNetDriverConfigs.FindByPredicate([InNetDriverDefinition](const FIrisNetDriverConfig& Config)
+				{
+					return Config.NetDriverDefinition.IsNone() ? false : Config.NetDriverDefinition == InNetDriverDefinition;
+				});
+		}
+
+		return IrisConfig ? IrisConfig->bEnableIris : false;
+#else
+		return false;
+#endif //UE_WITH_IRIS
+	}
 
-	if (NetDriverOverrides.Num() == 0 && FParse::Value(FCommandLine::Get(), TEXT("NetDriverOverrides="), OverrideCmdLine))
+	UNetDriver* CreateNetDriver_Local(UEngine* Engine, FWorldContext& Context, FName NetDriverDefinition, FName InNetDriverName)
 	{
-		TArray<FString> OverrideEntries;
-		OverrideCmdLine.ParseIntoArray(OverrideEntries, TEXT(";"), false);
+		UNetDriver* ReturnVal = nullptr;
+		FNetDriverDefinition* Definition = nullptr;
+		auto FindNetDriverDefPred =
+			[NetDriverDefinition](const FNetDriverDefinition& CurDef)
+		{
+			return CurDef.DefName == NetDriverDefinition;
+		};
 
-		UE_LOG(LogNet, Log, TEXT("NetDriverOverrides:"));
+#if !UE_BUILD_SHIPPING
+		/**
+		* Commandline override for the net driver.
+		*
+		* Format: (NOTE: Use quotes whenever the ',' character is used)
+		*	Override the main/game net driver (most common usage):
+		*		-NetDriverOverrides=DriverClassName
+		*
+		*	Override a specific/named net driver:
+		*		-NetDriverOverrides="DefName,DriverClassName"
+		*
+		*	Override a specific driver, including fallback driver:
+		*		-NetDriverOverrides="DefName,DriverClassName,DriverClassNameFallback"
+		*
+		*	Override multiple net drivers:
+		*		-NetDriverOverrides="DriverClassName;DefName2,DriverClassName2"
+		*
+		*
+		* Example:
+		*	Use WebSocket for the main game net driver:
+		*		-NetDriverOverrides=/Script/WebSocketNetworking.WebSocketNetDriver
+		*
+		*	Use WebSocket for the main game net driver, and the party beacon net driver
+		*		-NetDriverOverrides="/Script/WebSocketNetworking.WebSocketNetDriver;BeaconNetDriver,/Script/WebSocketNetworking.WebSocketNetDriver"
+		*/
+
+		static TArray<FNetDriverDefinition> NetDriverOverrides = TArray<FNetDriverDefinition>();
+		FString OverrideCmdLine;
 
-		for (FString& CurOverrideEntry : OverrideEntries)
+		if (NetDriverOverrides.Num() == 0 && FParse::Value(FCommandLine::Get(), TEXT("NetDriverOverrides="), OverrideCmdLine))
 		{
-			TArray<FString> CurEntryParms;
-			CurOverrideEntry.ParseIntoArray(CurEntryParms, TEXT(","), false);
+			TArray<FString> OverrideEntries;
+			OverrideCmdLine.ParseIntoArray(OverrideEntries, TEXT(";"), false);
+
+			UE_LOG(LogNet, Log, TEXT("NetDriverOverrides:"));
 
-			if (CurEntryParms.Num() == 0)
+			for (FString& CurOverrideEntry : OverrideEntries)
 			{
-				continue;
-			}
+				TArray<FString> CurEntryParms;
+				CurOverrideEntry.ParseIntoArray(CurEntryParms, TEXT(","), false);
 
+				if (CurEntryParms.Num() == 0)
+				{
+					continue;
+				}
 
-			FString TargetDefName = (CurEntryParms.Num() > 1 ? CurEntryParms[0] : FName(NAME_GameNetDriver).ToString());
-			FString OverrideClass = (CurEntryParms.Num() > 1 ? CurEntryParms[1] : CurEntryParms[0]);
-			FString OverrideFallback;
-			auto FindTargetDefPred =
-				[TargetDefName](const FNetDriverDefinition& CurDef)
-			{
-				return CurDef.DefName == *TargetDefName;
-			};
 
-			if (CurEntryParms.Num() > 2)
-			{
-				OverrideFallback = CurEntryParms[2];
-			}
-			else if (FNetDriverDefinition* FallbackDef = Engine->NetDriverDefinitions.FindByPredicate(FindTargetDefPred))
-			{
-				OverrideFallback = FallbackDef->DriverClassNameFallback.ToString();
-			}
-			else
-			{
-				OverrideFallback = TEXT("/Script/OnlineSubsystemUtils.IpNetDriver");
-			}
+				FString TargetDefName = (CurEntryParms.Num() > 1 ? CurEntryParms[0] : FName(NAME_GameNetDriver).ToString());
+				FString OverrideClass = (CurEntryParms.Num() > 1 ? CurEntryParms[1] : CurEntryParms[0]);
+				FString OverrideFallback;
+				auto FindTargetDefPred =
+					[TargetDefName](const FNetDriverDefinition& CurDef)
+				{
+					return CurDef.DefName == *TargetDefName;
+				};
 
-			if (TargetDefName.Len() == 0 || OverrideClass.Len() == 0 || OverrideFallback.Len() == 0)
-			{
-				continue;
-			}
+				if (CurEntryParms.Num() > 2)
+				{
+					OverrideFallback = CurEntryParms[2];
+				}
+				else if (FNetDriverDefinition* FallbackDef = Engine->NetDriverDefinitions.FindByPredicate(FindTargetDefPred))
+				{
+					OverrideFallback = FallbackDef->DriverClassNameFallback.ToString();
+				}
+				else
+				{
+					OverrideFallback = TEXT("/Script/OnlineSubsystemUtils.IpNetDriver");
+				}
+
+				if (TargetDefName.Len() == 0 || OverrideClass.Len() == 0 || OverrideFallback.Len() == 0)
+				{
+					continue;
+				}
 
 
-			if (!NetDriverOverrides.ContainsByPredicate(FindTargetDefPred))
-			{
-				FNetDriverDefinition NewDef;
+				if (!NetDriverOverrides.ContainsByPredicate(FindTargetDefPred))
+				{
+					FNetDriverDefinition NewDef;
 
-				NewDef.DefName = *TargetDefName;
-				NewDef.DriverClassName = FName(*OverrideClass);
-				NewDef.DriverClassNameFallback = FName(*OverrideFallback);
+					NewDef.DefName = *TargetDefName;
+					NewDef.DriverClassName = FName(*OverrideClass);
+					NewDef.DriverClassNameFallback = FName(*OverrideFallback);
 
-				NetDriverOverrides.Add(NewDef);
+					NetDriverOverrides.Add(NewDef);
 
-				UE_LOG(LogNet, Log, TEXT("- DefName: %s, DriverClassName: %s, DriverClassNameFallback: %s"), *TargetDefName,
-					*OverrideClass, *OverrideFallback);
+					UE_LOG(LogNet, Log, TEXT("- DefName: %s, DriverClassName: %s, DriverClassNameFallback: %s"), *TargetDefName,
+						*OverrideClass, *OverrideFallback);
+				}
 			}
 		}
-	}
 
 
-	Definition = NetDriverOverrides.FindByPredicate(FindNetDriverDefPred);
+		Definition = NetDriverOverrides.FindByPredicate(FindNetDriverDefPred);
 
-	if (Definition != nullptr)
-	{
-		UE_LOG(LogNet, Log, TEXT("Overriding NetDriver '%s' with class: %s"), *NetDriverDefinition.ToString(),
-			*Definition->DriverClassName.ToString());
-	}
-	else
+		if (Definition != nullptr)
+		{
+			UE_LOG(LogNet, Log, TEXT("Overriding NetDriver '%s' with class: %s"), *NetDriverDefinition.ToString(),
+				*Definition->DriverClassName.ToString());
+		}
+		else
 #endif
-	{
-		Definition = Engine->NetDriverDefinitions.FindByPredicate(FindNetDriverDefPred);
-	}
-
-	if (Definition != nullptr)
-	{
-		UClass* NetDriverClass = StaticLoadClass(UNetDriver::StaticClass(), nullptr, *Definition->DriverClassName.ToString(), nullptr,
-			LOAD_Quiet);
-
-		// if it fails, then fall back to standard fallback
-		if (NetDriverClass == nullptr || !NetDriverClass->GetDefaultObject<UNetDriver>()->IsAvailable())
 		{
-			NetDriverClass = StaticLoadClass(UNetDriver::StaticClass(), nullptr, *Definition->DriverClassNameFallback.ToString(),
-				nullptr, LOAD_None);
+			Definition = Engine->NetDriverDefinitions.FindByPredicate(FindNetDriverDefPred);
 		}
 
-		if (NetDriverClass != nullptr)
+		if (Definition != nullptr)
 		{
-			ReturnVal = NewObject<UNetDriver>(GetTransientPackage(), NetDriverClass);
+			UClass* NetDriverClass = StaticLoadClass(UNetDriver::StaticClass(), nullptr, *Definition->DriverClassName.ToString(), nullptr,
+				LOAD_Quiet);
 
-			check(ReturnVal != nullptr);
+			// if it fails, then fall back to standard fallback
+			if (NetDriverClass == nullptr || !NetDriverClass->GetDefaultObject<UNetDriver>()->IsAvailable())
+			{
+				NetDriverClass = StaticLoadClass(UNetDriver::StaticClass(), nullptr, *Definition->DriverClassNameFallback.ToString(),
+					nullptr, LOAD_None);
+			}
 
-			ReturnVal->SetNetDriverName(ReturnVal->GetFName());
-			ReturnVal->SetNetDriverDefinition(NetDriverDefinition);
+			if (NetDriverClass != nullptr)
+			{
+				ReturnVal = NewObject<UNetDriver>(GetTransientPackage(), NetDriverClass);
 
-			new(Context.ActiveNetDrivers) FNamedNetDriver(ReturnVal, Definition);
-		}
-	}
+				check(ReturnVal != nullptr);
 
+				const FName DriverName = InNetDriverName.IsNone() ? ReturnVal->GetFName() : InNetDriverName;
+				const bool bInitializeWithIris = IsNetDriverUsingIris(Engine, NetDriverDefinition, DriverName);
 
-	if (ReturnVal == nullptr)
-	{
-		UE_LOG(LogNet, Log, TEXT("CreateNamedNetDriver failed to create driver from definition %s"), *NetDriverDefinition.ToString());
-	}
+				ReturnVal->SetNetDriverName(DriverName);
+				ReturnVal->SetNetDriverDefinition(NetDriverDefinition);
+				ReturnVal->PostCreation(bInitializeWithIris);
 
-	return ReturnVal;
-}
+				new(Context.ActiveNetDrivers) FNamedNetDriver(ReturnVal, Definition);
+			}
+		}
 
-UNetDriver* UEngine::CreateNetDriver(UWorld *InWorld, FName NetDriverDefinition)
-{
-	LLM_SCOPE(ELLMTag::Networking);
-	return CreateNetDriver_Local(this, GetWorldContextFromWorldChecked(InWorld), NetDriverDefinition);
-}
 
-bool CreateNamedNetDriver_Local(UEngine *Engine, FWorldContext &Context, FName NetDriverName, FName NetDriverDefinition)
-{
-	UNetDriver* NetDriver = FindNamedNetDriver_Local(Context.ActiveNetDrivers, NetDriverName);
-	if (NetDriver == nullptr)
-	{
-		NetDriver = CreateNetDriver_Local(Engine, Context, NetDriverDefinition);
-		if (NetDriver)
+		if (ReturnVal == nullptr)
 		{
-			NetDriver->SetNetDriverName(NetDriverName);
-			return true;
+			UE_LOG(LogNet, Log, TEXT("CreateNamedNetDriver failed to create driver from definition %s"), *NetDriverDefinition.ToString());
 		}
+
+		return ReturnVal;
 	}
 
-	if (NetDriver)
+	bool CreateNamedNetDriver_Local(UEngine* Engine, FWorldContext& Context, FName NetDriverName, FName NetDriverDefinition)
 	{
-		UE_LOG(LogNet, Log, TEXT("CreateNamedNetDriver %s already exists as %s"), *NetDriverName.ToString(), *NetDriver->GetName());
+		UNetDriver* NetDriver = FindNamedNetDriver_Local(Context.ActiveNetDrivers, NetDriverName);
+		if (NetDriver == nullptr)
+		{
+			NetDriver = CreateNetDriver_Local(Engine, Context, NetDriverDefinition, NetDriverName);
+			if (NetDriver)
+			{
+				return true;
+			}
+		}
+
+		if (NetDriver)
+		{
+			UE_LOG(LogNet, Log, TEXT("CreateNamedNetDriver %s already exists as %s"), *NetDriverName.ToString(), *NetDriver->GetName());
+		}
+		else
+		{
+			UE_LOG(LogNet, Log, TEXT("CreateNamedNetDriver failed to create driver %s from definition %s"), *NetDriverName.ToString(), *NetDriverDefinition.ToString());
+		}
+
+		return false;
 	}
-	else
+
+	void DestroyNamedNetDriver_Local(FWorldContext& Context, FName NetDriverName)
 	{
-		UE_LOG(LogNet, Log, TEXT("CreateNamedNetDriver failed to create driver %s from definition %s"), *NetDriverName.ToString(), *NetDriverDefinition.ToString());
+		for (int32 Index = 0; Index < Context.ActiveNetDrivers.Num(); Index++)
+		{
+			FNamedNetDriver& NamedNetDriver = Context.ActiveNetDrivers[Index];
+			UNetDriver* NetDriver = NamedNetDriver.NetDriver;
+			if (NetDriver && NetDriver->NetDriverName == NetDriverName)
+			{
+				UE_LOG(LogNet, Log, TEXT("DestroyNamedNetDriver %s [%s]"), *NetDriver->GetName(), *NetDriverName.ToString());
+				
+				ensureMsgf(!(NetDriver->IsInTick() && NetDriverName != NAME_PendingNetDriver), TEXT("Attempting to destroy NetDriver %s [%s] while it is ticking."), *NetDriver->GetName(), *NetDriverName.ToString());
+
+				Context.ActiveNetDrivers.RemoveAtSwap(Index);
+				Context.GarbageObjectsToVerify.Add(NetDriver);
+				NetDriver->SetWorld(NULL);
+				NetDriver->Shutdown();
+				NetDriver->LowLevelDestroy();
+				NetDriver->MarkAsGarbage();
+
+				if (UWorld* World = Context.World())
+				{
+					World->ClearNetDriver(NetDriver);
+				}
+
+				break;
+			}
+		}
 	}
+} // end namespace UE::Private
 
-	return false;
+UNetDriver* UEngine::CreateNetDriver(UWorld *InWorld, FName NetDriverDefinition)
+{
+	LLM_SCOPE(ELLMTag::Networking);
+	return UE::Private::CreateNetDriver_Local(this, GetWorldContextFromWorldChecked(InWorld), NetDriverDefinition, NAME_None);
 }
 
 bool UEngine::CreateNamedNetDriver(UWorld *InWorld, FName NetDriverName, FName NetDriverDefinition)
 {
-	return CreateNamedNetDriver_Local( this, GetWorldContextFromWorldChecked(InWorld), NetDriverName, NetDriverDefinition );
+	return UE::Private::CreateNamedNetDriver_Local(this, GetWorldContextFromWorldChecked(InWorld), NetDriverName, NetDriverDefinition);
 }
 
 bool UEngine::CreateNamedNetDriver(UPendingNetGame *PendingNetGame, FName NetDriverName, FName NetDriverDefinition)
 {
-	return CreateNamedNetDriver_Local( this, GetWorldContextFromPendingNetGameChecked(PendingNetGame), NetDriverName, NetDriverDefinition);
-}
-
-void DestroyNamedNetDriver_Local(FWorldContext &Context, FName NetDriverName)
-{
-	for (int32 Index = 0; Index < Context.ActiveNetDrivers.Num(); Index++)
-	{
-		FNamedNetDriver& NamedNetDriver = Context.ActiveNetDrivers[Index];
-		UNetDriver* NetDriver = NamedNetDriver.NetDriver;
-		if (NetDriver && NetDriver->NetDriverName == NetDriverName)
-		{
-			UE_LOG(LogNet, Log, TEXT("DestroyNamedNetDriver %s [%s]"), *NetDriver->GetName(), *NetDriverName.ToString());
-			
-			ensureMsgf(!(NetDriver->IsInTick() && NetDriverName != NAME_PendingNetDriver), TEXT("Attempting to destroy NetDriver %s [%s] while it is ticking."), *NetDriver->GetName(), *NetDriverName.ToString());
-
-			Context.ActiveNetDrivers.RemoveAtSwap(Index);
-			Context.GarbageObjectsToVerify.Add(NetDriver);
-			NetDriver->SetWorld(NULL);
-			NetDriver->Shutdown();
-			NetDriver->LowLevelDestroy();
-			NetDriver->MarkAsGarbage();
-
-			if (UWorld* World = Context.World())
-			{
-				World->ClearNetDriver(NetDriver);
-			}
-
-			break;
-		}
-	}
+	return UE::Private::CreateNamedNetDriver_Local(this, GetWorldContextFromPendingNetGameChecked(PendingNetGame), NetDriverName, NetDriverDefinition);
 }
 
 void UEngine::DestroyNamedNetDriver(UWorld *InWorld, FName NetDriverName)
 {
-	DestroyNamedNetDriver_Local( GetWorldContextFromWorldChecked(InWorld), NetDriverName);
+	UE::Private::DestroyNamedNetDriver_Local(GetWorldContextFromWorldChecked(InWorld), NetDriverName);
 }
 
 void UEngine::DestroyNamedNetDriver(UPendingNetGame *PendingNetGame, FName NetDriverName)
 {
-	DestroyNamedNetDriver_Local( GetWorldContextFromPendingNetGameChecked(PendingNetGame), NetDriverName );
+	UE::Private::DestroyNamedNetDriver_Local(GetWorldContextFromPendingNetGameChecked(PendingNetGame), NetDriverName);
 }
 
 ENetMode UEngine::GetNetMode(const UWorld *World) const
@@ -13824,7 +13893,7 @@ EBrowseReturnVal::Type UEngine::Browse( FWorldContext& WorldContext, FURL URL, F
 				{
 					if (ActiveNetDriver.NetDriver)
 					{
-						DestroyNamedNetDriver_Local(WorldContext, ActiveNetDriver.NetDriver->NetDriverName);
+						UE::Private::DestroyNamedNetDriver_Local(WorldContext, ActiveNetDriver.NetDriver->NetDriverName);
 					}
 				}
 				CollectGarbage(GARBAGE_COLLECTION_KEEPFLAGS);
@@ -13963,7 +14032,7 @@ void UEngine::CancelPending(FWorldContext &Context)
 	if (Context.PendingNetGame && Context.PendingNetGame->NetDriver && Context.PendingNetGame->NetDriver->ServerConnection)
 	{
 		Context.PendingNetGame->NetDriver->ServerConnection->Close();
-		DestroyNamedNetDriver_Local(Context, Context.PendingNetGame->NetDriver->NetDriverName);
+		UE::Private::DestroyNamedNetDriver_Local(Context, Context.PendingNetGame->NetDriver->NetDriverName);
 		Context.PendingNetGame->NetDriver = NULL;
 	}
 
diff --git a/Engine/Source/Runtime/Engine/Private/UnrealNetwork.cpp b/Engine/Source/Runtime/Engine/Private/UnrealNetwork.cpp
--- a/Engine/Source/Runtime/Engine/Private/UnrealNetwork.cpp
+++ b/Engine/Source/Runtime/Engine/Private/UnrealNetwork.cpp
@@ -47,6 +47,10 @@ void RegisterReplicatedLifetimeProperty(
 
 		FLifetimeProperty LifetimeProp(RepIndex, Params.Condition, Params.RepNotifyCondition, Params.bIsPushBased);
 
+#if UE_WITH_IRIS
+		LifetimeProp.CreateAndRegisterReplicationFragmentFunction = Params.CreateAndRegisterReplicationFragmentFunction;
+#endif
+
 		if (RegisteredPropertyPtr)
 		{
 			// Disabled properties can be re-enabled via DOREPLIFETIME
diff --git a/Engine/Source/Runtime/Engine/Private/World.cpp b/Engine/Source/Runtime/Engine/Private/World.cpp
--- a/Engine/Source/Runtime/Engine/Private/World.cpp
+++ b/Engine/Source/Runtime/Engine/Private/World.cpp
@@ -139,6 +139,12 @@
 #include "ObjectTrace.h"
 #include "ReplaySubsystem.h"
 #include "Net/NetPing.h"
+
+#if UE_WITH_IRIS
+#include "Iris/ReplicationSystem/ReplicationSystem.h"
+#include "Net/Iris/ReplicationSystem/ActorReplicationBridge.h"
+#endif // UE_WITH_IRIS
+
 #include "ChaosSolversModule.h"
 
 DEFINE_LOG_CATEGORY_STATIC(LogWorld, Log, All);
@@ -1032,6 +1038,15 @@ void UWorld::BeginDestroy()
 
 	FAudioDeviceHandle EmptyHandle;
 	SetAudioDevice(EmptyHandle);
+
+#if UE_WITH_IRIS
+	if (IrisSystemHolder.IsHolding())
+	{
+		UE::Net::FReplicationSystemFactory::DestroyReplicationSystem(IrisSystemHolder.ReplicationSystem);
+
+		IrisSystemHolder.Clear();
+	}
+#endif // UE_WITH_IRIS
 }
 
 void UWorld::ReleasePhysicsScene()
@@ -5946,12 +5961,27 @@ void UWorld::NotifyControlMessage(UNetConnection* Connection, uint8 MessageType,
 				uint32 LocalNetworkVersion = FNetworkVersion::GetLocalNetworkVersion();
 				FString EncryptionToken;
 
-				if (FNetControlMessage<NMT_Hello>::Receive(Bunch, IsLittleEndian, RemoteNetworkVersion, EncryptionToken))
+				EEngineNetworkRuntimeFeatures LocalNetworkFeatures = NetDriver->GetNetworkRuntimeFeatures();
+				EEngineNetworkRuntimeFeatures RemoteNetworkFeatures = EEngineNetworkRuntimeFeatures::None;
+
+				if (FNetControlMessage<NMT_Hello>::Receive(Bunch, IsLittleEndian, RemoteNetworkVersion, EncryptionToken, RemoteNetworkFeatures))
 				{
-					if (!FNetworkVersion::IsNetworkCompatible(LocalNetworkVersion, RemoteNetworkVersion))
+					const bool bIsCompatible = FNetworkVersion::IsNetworkCompatible(LocalNetworkVersion, RemoteNetworkVersion) && FNetworkVersion::AreNetworkRuntimeFeaturesCompatible(LocalNetworkFeatures, RemoteNetworkFeatures);
+					if (!bIsCompatible)
 					{
-						UE_LOG(LogNet, Log, TEXT("NotifyControlMessage: Client connecting with invalid version. LocalNetworkVersion: %i, RemoteNetworkVersion: %i"), LocalNetworkVersion, RemoteNetworkVersion);
-						FNetControlMessage<NMT_Upgrade>::Send(Connection, LocalNetworkVersion);
+						TStringBuilder<128> LocalNetFeaturesDescription;
+						TStringBuilder<128> RemoteNetFeaturesDescription;
+
+						FNetworkVersion::DescribeNetworkRuntimeFeaturesBitset(LocalNetworkFeatures, LocalNetFeaturesDescription);
+						FNetworkVersion::DescribeNetworkRuntimeFeaturesBitset(RemoteNetworkFeatures, RemoteNetFeaturesDescription);
+
+						UE_LOG(LogNet, Log, TEXT("NotifyControlMessage: Client %s connecting with invalid version. LocalNetworkVersion: %u, RemoteNetworkVersion: %u, LocalNetFeatures=%s, RemoteNetFeatures=%s"),
+							*Connection->GetName(), 
+							LocalNetworkVersion, RemoteNetworkVersion,
+							LocalNetFeaturesDescription.ToString(), RemoteNetFeaturesDescription.ToString()
+						);
+
+						FNetControlMessage<NMT_Upgrade>::Send(Connection, LocalNetworkVersion, LocalNetworkFeatures);
 						Connection->FlushNet(true);
 						Connection->Close(ENetCloseResult::Upgrade);
 #if USE_SERVER_PERF_COUNTERS
@@ -6339,6 +6369,16 @@ bool UWorld::Listen( FURL& InURL )
 	if (GEngine->CreateNamedNetDriver(this, NAME_GameNetDriver, NAME_GameNetDriver))
 	{
 		NetDriver = GEngine->FindNamedNetDriver(this, NAME_GameNetDriver);
+
+#if UE_WITH_IRIS
+		if (IrisSystemHolder.IsHolding())
+		{
+			NetDriver->RestoreIrisSystem(IrisSystemHolder.ReplicationSystem);
+			
+			IrisSystemHolder.Clear();
+		}
+#endif // UE_WITH_IRIS
+
 		NetDriver->SetWorld(this);
 		FLevelCollection* const SourceCollection = FindCollectionByType(ELevelCollectionType::DynamicSourceLevels);
 		if (SourceCollection)
@@ -8412,6 +8452,17 @@ void UWorld::DuplicateRequestedLevels(const FName MapName)
 	}
 }
 
+#if UE_WITH_IRIS
+void UWorld::StoreIrisAndClearReferences()
+{
+	if (NetDriver)
+	{
+		IrisSystemHolder.ReplicationSystem = NetDriver->GetReplicationSystem();
+		NetDriver->ClearIrisSystem();
+	}
+}
+#endif // UE_WITH_IRIS
+
 #if WITH_EDITOR
 void UWorld::ChangeFeatureLevel(ERHIFeatureLevel::Type InFeatureLevel, bool bShowSlowProgressDialog )
 {
diff --git a/Engine/Source/Runtime/Engine/Public/Net/DataChannel.h b/Engine/Source/Runtime/Engine/Public/Net/DataChannel.h
--- a/Engine/Source/Runtime/Engine/Public/Net/DataChannel.h
+++ b/Engine/Source/Runtime/Engine/Public/Net/DataChannel.h
@@ -150,9 +150,9 @@ public: \
 #define IMPLEMENT_CONTROL_CHANNEL_MESSAGE(Name) static uint8 Dummy##_FNetControlMessage_##Name = FNetControlMessage<NMT_##Name>::Initialize();
 
 // message type definitions
-DEFINE_CONTROL_CHANNEL_MESSAGE(Hello, 0, uint8, uint32, FString); // initial client connection message
+DEFINE_CONTROL_CHANNEL_MESSAGE(Hello, 0, uint8, uint32, FString, uint16); // initial client connection message
 DEFINE_CONTROL_CHANNEL_MESSAGE(Welcome, 1, FString, FString, FString); // server tells client they're ok'ed to load the server's level
-DEFINE_CONTROL_CHANNEL_MESSAGE(Upgrade, 2, uint32); // server tells client their version is incompatible
+DEFINE_CONTROL_CHANNEL_MESSAGE(Upgrade, 2, uint32, uint16); // server tells client their version is incompatible
 DEFINE_CONTROL_CHANNEL_MESSAGE(Challenge, 3, FString); // server sends client challenge string to verify integrity
 DEFINE_CONTROL_CHANNEL_MESSAGE(Netspeed, 4, int32); // client sends requested transfer rate
 DEFINE_CONTROL_CHANNEL_MESSAGE(Login, 5, FString, FString, FUniqueNetIdRepl, FString); // client requests to be admitted to the game
diff --git a/Engine/Source/Runtime/Engine/Public/Net/Iris/ReplicationSystem/ActorReplicationBridge.h b/Engine/Source/Runtime/Engine/Public/Net/Iris/ReplicationSystem/ActorReplicationBridge.h
new file mode 100644
--- /dev/null
+++ b/Engine/Source/Runtime/Engine/Public/Net/Iris/ReplicationSystem/ActorReplicationBridge.h
@@ -0,0 +1,105 @@
+// Copyright Epic Games, Inc. All Rights Reserved.
+
+#pragma once
+
+#include "Iris/ReplicationSystem/ObjectReplicationBridge.h"
+#include "Engine/EngineTypes.h"
+#include "ActorReplicationBridge.generated.h"
+
+#if UE_WITH_IRIS
+
+class UActorComponent;
+class UNetDriver;
+class UWorld;
+class UNetDriver;
+class UIrisObjectReferencePackageMap;
+
+namespace UE::Net::Private
+{
+	struct FActorReplicationBridgeCreationHeader;
+	struct FActorCreationHeader;
+	struct FSubObjectCreationHeader;
+}
+
+namespace UE::Net
+{
+
+// If actor should be replicated using IRIS or old replication system
+ENGINE_API bool ShouldUseIrisReplication(const UObject* Actor);
+
+}
+
+/** Parameters passed to UActorReplicationBridge::BeginReplication. */
+struct FActorBeginReplicationParams
+{
+	FActorBeginReplicationParams() : bIncludeInLevelGroupFilter(1U) {}
+
+	uint32 bIncludeInLevelGroupFilter : 1U;
+};
+
+#endif // UE_WITH_IRIS
+
+UCLASS(Transient, MinimalAPI)
+class UActorReplicationBridge final : public UObjectReplicationBridge
+{
+	GENERATED_BODY()
+
+public:
+	ENGINE_API UActorReplicationBridge();
+
+#if UE_WITH_IRIS
+
+	ENGINE_API virtual ~UActorReplicationBridge();
+
+	ENGINE_API static UActorReplicationBridge* Create(UNetDriver* NetDriver);
+
+	/** Sets the net driver for the bridge. */
+	ENGINE_API virtual void SetNetDriver(UNetDriver* const InNetDriver) override;
+	
+	/** Get net driver used by the bridge .*/
+	inline UNetDriver* GetNetDriver() const { return NetDriver; }
+
+	/** Begin replication of an actor and its registered ActorComponents and SubObjects. */
+	ENGINE_API FNetHandle BeginReplication(AActor* Instance, const FActorBeginReplicationParams& Params);
+
+	/** Stop replicating an actor. Will destroy handle for actor and registered subobjects. */
+	ENGINE_API void EndReplication(AActor* Actor, EEndPlayReason::Type EndPlayReason);
+		
+	/**
+	 * Begin replication of an ActorComponent and its registered SubObjects, 
+	 * if the ActorComponent already is replicated any set NetObjectConditions will be updated.
+	*/
+	ENGINE_API FNetHandle BeginReplication(FNetHandle OwnerHandle, UActorComponent* ActorComponent);
+
+	/** Stop replicating an ActorComponent and its associated SubObjects. */
+	ENGINE_API void EndReplicationForActorComponent(UActorComponent* ActorComponent);
+	
+	/** Get object reference packagemap. Used in special cases where serialization hasn't been converted to use NetSerializers.  */
+	ENGINE_API UIrisObjectReferencePackageMap* GetObjectReferencePackageMap() const { return ObjectReferencePackageMap; }
+	
+protected:
+	// UObjectReplicationBridge
+	virtual void Initialize(UReplicationSystem* ReplicationSystem) override;
+	virtual bool WriteCreationHeader(UE::Net::FNetSerializationContext& Context, FNetHandle Handle) override;
+	virtual FCreationHeader* ReadCreationHeader(UE::Net::FNetSerializationContext& Context) override;
+	virtual UObject* BeginInstantiateFromRemote(FNetHandle SubObjectOwnerNetHandle, const UE::Net::FNetObjectResolveContext& ResolveContext, const FCreationHeader* InHeader) override;
+	virtual bool OnInstantiatedFromRemote(UObject* Instance, const FCreationHeader* InHeader, uint32 ConnectionId) const override;
+	virtual void EndInstantiateFromRemote(FNetHandle Handle) override;
+	virtual void DestroyInstanceFromRemote(UObject* Instance, bool bTearOff) override;
+	virtual void GetInitialDependencies(FNetHandle Handle, FNetDependencyInfoArray& OutDependencies) const override;
+	virtual bool RemapPathForPIE(FString& Path, bool bReading) const override;
+	virtual bool ObjectLevelHasFinishedLoading(UObject* Object) const override;
+
+private:
+	void GetActorCreationHeader(const AActor* Actor, UE::Net::Private::FActorCreationHeader& Header) const;
+	void GetSubObjectCreationHeader(const UObject* Object, UE::Net::Private::FSubObjectCreationHeader& Header) const;
+
+	uint32 GetPollFramePeriod(float PollFrequency) const;
+	float GetMinSupportedNetUpdateFrequency() const;
+
+	UNetDriver* NetDriver;
+	float MaxPollFrequency;
+
+	UIrisObjectReferencePackageMap* ObjectReferencePackageMap;
+#endif // UE_WITH_IRIS
+};
diff --git a/Engine/Source/Runtime/Engine/Public/Net/Iris/ReplicationSystem/ReplicationSystemUtil.h b/Engine/Source/Runtime/Engine/Public/Net/Iris/ReplicationSystem/ReplicationSystemUtil.h
new file mode 100644
--- /dev/null
+++ b/Engine/Source/Runtime/Engine/Public/Net/Iris/ReplicationSystem/ReplicationSystemUtil.h
@@ -0,0 +1,118 @@
+// Copyright Epic Games, Inc. All Rights Reserved.
+
+#pragma once
+
+#if UE_WITH_IRIS
+
+#include "Iris/ReplicationSystem/NetHandle.h"
+#include "Engine/EngineTypes.h"
+
+class AActor;
+class UActorComponent;
+struct FActorBeginReplicationParams;
+class UActorReplicationBridge;
+class UObject;
+class FRepChangedPropertyTracker;
+class UReplicationSystem;
+class UWorld;
+class UNetConnection;
+
+namespace UE::Net
+{
+
+// Helper methods to interact with ReplicationSystem from engine code
+struct FReplicationSystemUtil
+{
+	/** Returns the UReplicationSystem the actor is replicated by, or null if it isn't. */
+	ENGINE_API static UReplicationSystem* GetReplicationSystem(const AActor* Actor);
+
+	/** Returns the UActorReplicationBridge of the UReplicationSystem the actor is replicated by, or null if it isn't. */
+	ENGINE_API static UActorReplicationBridge* GetActorReplicationBridge(const AActor* Actor);
+
+	/** Returns the UReplicationSystem of the UNetDriver the UNetConnection belongs to. */
+	ENGINE_API static UActorReplicationBridge* GetActorReplicationBridge(const UNetConnection* NetConnection);
+
+	/** Get the NetHandle for the actor. */
+	ENGINE_API static FNetHandle GetNetHandle(const AActor* Actor);
+
+	/** Get the NetHandle for an actor component. */
+	ENGINE_API static FNetHandle GetNetHandle(const UActorComponent* SubObject);
+
+	/** 
+		Begins replication an actor and all of its registered subobjects returns a NetHandle for the actor if it succeeds. 
+	*/
+	ENGINE_API static FNetHandle BeginReplication(AActor* Actor, const FActorBeginReplicationParams& Params);
+
+	/** 
+		Begins replication an actor and all of its registered subobjects returns a NetHandle for the actor if it succeeds. 
+	*/
+	ENGINE_API static FNetHandle BeginReplication(AActor* Actor);
+	
+	/** Stop replicating an actor. Will destroy handle for actor and registered subobjects. */
+	ENGINE_API static void EndReplication(AActor* Actor, EEndPlayReason::Type EndPlayReason);
+
+	/** Create NetHandle for actor component and add it as a subobject to the actor handle. */
+	ENGINE_API static FNetHandle BeginReplicationForActorComponent(FNetHandle ActorHandle, UActorComponent* ActorComponent);
+
+	/** Create NetHandle for actor component and add it as a subobject to the actor. */
+	ENGINE_API static FNetHandle BeginReplicationForActorComponent(const AActor* Actor, UActorComponent* ActorComponent);
+
+	/** Stop replicating an ActorComponent and its associated SubObjects. */
+	ENGINE_API static void EndReplicationForActorComponent(UActorComponent* SubObject);
+
+	/** Create NetHandle for SubObject and add it as a subobject to the actor. */
+	ENGINE_API static FNetHandle BeginReplicationForActorSubObject(const AActor* Actor, UObject* SubObject, ELifetimeCondition NetCondition);
+
+	/** Stop replicating a SubObject ActorComponent and its associated SubObjects. */
+	ENGINE_API static void EndReplicationForActorSubObject(const AActor* Actor, UObject* SubObject);
+
+	/** Create NetHandle for SubObject and add it as a subobject to the actor but only replicated if the ActorComponent replicates. */
+	ENGINE_API static FNetHandle BeginReplicationForActorComponentSubObject(UActorComponent* ActorComponent, UObject* SubObject, ELifetimeCondition Condition);
+
+	/** Stop replicating an ActorComponent and its associated SubObjects. */
+	ENGINE_API static void EndReplicationForActorComponentSubObject(UActorComponent* ActorComponent, UObject* SubObject);
+	
+	/** 
+	 * Set subobject NetCondition for a subobject, the conditions is used to determine if the SubObject should be replicated or not.
+	 * @note: As the filtering is done at the serialization level it is typically more efficient to use a separate object for connection 
+	 * specific data as filtering can be done at a higher level.
+	 */
+	ENGINE_API static void SetActorComponentNetCondition(const UActorComponent* SubObject, ELifetimeCondition Condition);
+
+	/** Update group memberships used by COND_Group sub object filtering for the specified SubObject */
+	ENGINE_API static void UpdateSubObjectGroupMemberships(const UObject* SubObject, const UWorld* World);
+
+	/** Update replication status for all NetGroups used by COND_Group sub object filtering for the provided PlayerController */
+	ENGINE_API static void UpdateSubObjectGroupMemberships(const APlayerController* PC);
+
+	/** Update replication status for the PlayerController to not include the specified NetGroup */
+	ENGINE_API static void RemoveSubObjectGroupMembership(const APlayerController* PC, const FName NetGroup);
+
+	/**
+	 * Adds a dependent actor. A dependent actor can replicate separately or if a parent replicates.
+	 * Dependent actors cannot be filtered out by dynamic filtering unless the parent is also filtered out.
+	 * @note There is no guarantee that the data will end up in the same packet so it is a very loose form of dependency.
+	 */
+	ENGINE_API static void AddDependentActor(const AActor* Parent, AActor* Child);
+
+	/** Remove dependent actor from parent. The dependent actor will function as a standard standalone replicated actor. */
+	ENGINE_API static void RemoveDependentActor(const AActor* Parent, AActor* Child);
+
+	/** Begin replication for all networked actors in the world. */
+	ENGINE_API static void BeginReplicationForActorsInWorld(UWorld* World);
+
+	/** Notify the ReplicationSystem of a dormancy change. */
+	ENGINE_API static void NotifyActorDormancyChange(AActor* Actor, ENetDormancy OldDormancyState);
+
+	/** Trigger replication of dirty state for actor wanting to be dormant. */
+	ENGINE_API static void FlushNetDormancy(AActor* Actor, bool bWasDormInitial);
+
+private:
+	friend FRepChangedPropertyTracker;
+
+	ENGINE_API static void SetPropertyCustomCondition(const UObject* Object, uint16 RepIndex, bool bIsActive);
+};
+
+}
+
+#endif
diff --git a/Engine/Source/Runtime/Engine/Public/Net/Subsystems/NetworkSubsystem.h b/Engine/Source/Runtime/Engine/Public/Net/Subsystems/NetworkSubsystem.h
--- a/Engine/Source/Runtime/Engine/Public/Net/Subsystems/NetworkSubsystem.h
+++ b/Engine/Source/Runtime/Engine/Public/Net/Subsystems/NetworkSubsystem.h
@@ -17,6 +17,9 @@ public:
 	/** Access the NetConditionGroupManager */
 	UE::Net::FNetConditionGroupManager& GetNetConditionGroupManager() { return GroupsManager;  }
 
+	/** Const access the NetConditionGroupManager */
+	const UE::Net::FNetConditionGroupManager& GetNetConditionGroupManager() const { return GroupsManager; }
+
 protected:
 	
 	/** Allow the subsystem on every world type. It's essentially free until it's used but accessed by any world that does netdriver replication. */
diff --git a/Engine/Source/Runtime/Engine/Public/Net/UnrealNetwork.h b/Engine/Source/Runtime/Engine/Public/Net/UnrealNetwork.h
--- a/Engine/Source/Runtime/Engine/Public/Net/UnrealNetwork.h
+++ b/Engine/Source/Runtime/Engine/Public/Net/UnrealNetwork.h
@@ -120,6 +120,12 @@ struct ENGINE_API FDoRepLifetimeParams
 	
 	/** Whether or not this property uses Push Model. See PushModel.h */
 	bool bIsPushBased = false;
+
+#if UE_WITH_IRIS
+	/** Function to create and register a ReplicationFragment for the property */
+	UE::Net::CreateAndRegisterReplicationFragmentFunc CreateAndRegisterReplicationFragmentFunction = nullptr;
+#endif
+
 };
 
 namespace NetworkingPrivate
@@ -219,7 +225,10 @@ static FProperty* GetReplicatedProperty(const UClass* CallingClass, const UClass
 	static const bool bIsValid_##c_##v = ValidateReplicatedClassInheritance(StaticClass(), c::StaticClass(), TEXT(#v)); \
 	const TCHAR* DoRepPropertyName_##c_##v(TEXT(#v)); \
 	const NetworkingPrivate::FRepPropertyDescriptor PropertyDescriptor_##c_##v(DoRepPropertyName_##c_##v, (int32)c::ENetFields_Private::v, 1); \
-	RegisterReplicatedLifetimeProperty(PropertyDescriptor_##c_##v, OutLifetimeProps, params); \
+\
+	PRAGMA_DISABLE_DEPRECATION_WARNINGS \
+	RegisterReplicatedLifetimeProperty(PropertyDescriptor_##c_##v, OutLifetimeProps, FixupParams<decltype(c::v)>(params)); \
+	PRAGMA_ENABLE_DEPRECATION_WARNINGS \
 }
 
 #define DOREPLIFETIME_WITH_PARAMS_FAST_STATIC_ARRAY(c,v,params) \
@@ -233,7 +242,9 @@ static FProperty* GetReplicatedProperty(const UClass* CallingClass, const UClass
 #define DOREPLIFETIME_WITH_PARAMS(c,v,params) \
 { \
 	FProperty* ReplicatedProperty = GetReplicatedProperty(StaticClass(), c::StaticClass(),GET_MEMBER_NAME_CHECKED(c,v)); \
-	RegisterReplicatedLifetimeProperty(ReplicatedProperty, OutLifetimeProps, params); \
+	PRAGMA_DISABLE_DEPRECATION_WARNINGS \
+	RegisterReplicatedLifetimeProperty(ReplicatedProperty, OutLifetimeProps, FixupParams<decltype(c::v)>(params)); \
+	PRAGMA_ENABLE_DEPRECATION_WARNINGS \
 }
 
 #define DOREPLIFETIME(c,v) DOREPLIFETIME_WITH_PARAMS(c,v,FDoRepLifetimeParams())
@@ -326,6 +337,31 @@ ENGINE_API void RegisterReplicatedLifetimeProperty(
 	TArray<FLifetimeProperty>& OutLifetimeProps,
 	const FDoRepLifetimeParams& Params);
 
+/*-----------------------------------------------------------------------------
+	Capture additional parameters required to bind FastArrays when using Iris
+-----------------------------------------------------------------------------*/
+struct CGetFastArrayCreateReplicationFragmentFuncable
+{
+	template <typename T, typename...>
+	auto Requires(T* FastArray) -> decltype(FastArray->GetFastArrayCreateReplicationFragmentFunction());
+};
+
+#if UE_WITH_IRIS
+template<typename T>
+inline typename TEnableIf<TModels<CGetFastArrayCreateReplicationFragmentFuncable, T>::Value, const FDoRepLifetimeParams>::Type FixupParams(const FDoRepLifetimeParams& Params)
+{
+	FDoRepLifetimeParams NewParams(Params);
+	NewParams.CreateAndRegisterReplicationFragmentFunction = T::GetFastArrayCreateReplicationFragmentFunction();
+	return NewParams;
+}
+#endif
+
+template<typename T>
+inline typename TEnableIf<!TModels<CGetFastArrayCreateReplicationFragmentFuncable, T>::Value, const FDoRepLifetimeParams&>::Type FixupParams(const FDoRepLifetimeParams& Params)
+{
+	return Params;
+}
+
 /*-----------------------------------------------------------------------------
 	Disable macros.
 	Use these macros to state that properties tagged replicated 
diff --git a/Engine/Source/Runtime/Experimental/Iris/Core/IrisCore.Build.cs b/Engine/Source/Runtime/Experimental/Iris/Core/IrisCore.Build.cs
new file mode 100644
--- /dev/null
+++ b/Engine/Source/Runtime/Experimental/Iris/Core/IrisCore.Build.cs
@@ -0,0 +1,37 @@
+// Copyright Epic Games, Inc. All Rights Reserved.
+
+namespace UnrealBuildTool.Rules
+{
+	public class IrisCore : ModuleRules
+	{
+        public IrisCore(ReadOnlyTargetRules Target) : base(Target)
+		{
+            PublicIncludePaths.Add(ModuleDirectory + "/Public");
+
+			PublicDependencyModuleNames.AddRange(
+				new string[]
+				{
+					"Core",
+					"CoreUObject",
+					"NetCore",
+                }
+            );
+
+			if (Target.bUseIris == true)
+			{
+				PublicDefinitions.Add("UE_WITH_IRIS=1");
+			}
+			else
+			{
+				PublicDefinitions.Add("UE_WITH_IRIS=0");
+
+				// This module should not be precompiled if Iris is disabled
+				if (Target.bBuildAllModules == true)
+				{
+					PrecompileForTargets = PrecompileTargetsType.None;
+				}
+			}
+
+		}
+	}
+}
diff --git a/Engine/Source/Runtime/Experimental/Iris/Core/Private/Iris/Core/IrisDebugging.cpp b/Engine/Source/Runtime/Experimental/Iris/Core/Private/Iris/Core/IrisDebugging.cpp
new file mode 100644
--- /dev/null
+++ b/Engine/Source/Runtime/Experimental/Iris/Core/Private/Iris/Core/IrisDebugging.cpp
@@ -0,0 +1,384 @@
+// Copyright Epic Games, Inc. All Rights Reserved.
+
+#include "Iris/Core/IrisDebugging.h"
+
+#if !UE_BUILD_SHIPPING
+
+#include "UObject/Field.h"
+#include "HAL/IConsoleManager.h"
+#include "GenericPlatform/GenericPlatformMisc.h"
+#include "Iris/Core/IrisLog.h"
+#include "Iris/ReplicationSystem/ObjectReplicationBridge.h"
+#include "Iris/ReplicationSystem/ReplicationOperations.h"
+#include "Iris/ReplicationSystem/ReplicationSystem.h"
+#include "Iris/ReplicationSystem/ReplicationSystemInternal.h"
+#include "Iris/ReplicationSystem/NetHandleManager.h"
+#include "Iris/ReplicationSystem/ReplicationProtocolManager.h"
+#include "Iris/ReplicationState/PropertyReplicationState.h"
+#include "Iris/Serialization/NetSerializationContext.h"
+#include "Iris/Serialization/InternalNetSerializationContext.h"
+#include "Containers/StringFwd.h"
+
+namespace UE::Net::IrisDebugHelper
+{
+
+namespace IrisDebugHelperInternal
+{
+	static FString GIrisDebugName;
+	FAutoConsoleVariableRef NetIrisDebugName(
+		TEXT("Net.Iris.DebugName"),
+		GIrisDebugName,
+		TEXT("Set a class name or object name to break on."),
+		ECVF_Default);
+
+	static FString GIrisDebugRPCName;
+	FAutoConsoleVariableRef NetRPCSetDebugRPCName(
+		TEXT("Net.Iris.DebugRPCName"),
+		GIrisDebugRPCName,
+		TEXT("Set the Name of an RPC to break on."),
+		ECVF_Default);
+
+	static FNetHandle GIrisDebugNetHandle;
+	static FAutoConsoleCommand NetIrisDebugNetHandle(
+		TEXT("Net.Iris.DebugNetHandle"), 
+		TEXT("Specify an handle index that we will break on (or none to turn off)."), 
+		FConsoleCommandWithArgsDelegate::CreateLambda([](const TArray<FString>& Args)
+		{
+			if (Args.Num() > 0)
+			{
+				uint32 NetId = 0;
+				LexFromString(NetId, *Args[0]);
+				GIrisDebugNetHandle = Private::FNetHandleManager::MakeNetHandleFromId(NetId);
+			}
+			else
+			{
+				GIrisDebugNetHandle = FNetHandle();
+			}
+		}));
+}; // namespace IrisDebugHelperInternal
+
+bool BreakOnObjectName(UObject* Object)
+{
+	if (IrisDebugHelperInternal::GIrisDebugName.IsEmpty() == false && GetNameSafe(Object).Contains(IrisDebugHelperInternal::GIrisDebugName))
+	{
+		UE_DEBUG_BREAK();
+		return true;
+	}
+
+	return false;
+}
+
+bool BreakOnNetHandle(FNetHandle NetHandle)
+{
+	if (IrisDebugHelperInternal::GIrisDebugNetHandle.IsValid() && IrisDebugHelperInternal::GIrisDebugNetHandle == NetHandle)
+	{
+		UE_DEBUG_BREAK();
+		return true;
+	}
+
+	return false;
+}
+
+bool BreakOnRPCName(FName RPCName)
+{
+	if (IrisDebugHelperInternal::GIrisDebugRPCName.IsEmpty() == false && RPCName.GetPlainNameString().Contains(IrisDebugHelperInternal::GIrisDebugRPCName))
+	{
+		UE_DEBUG_BREAK();
+		return true;
+	}
+
+	return false;
+}
+
+void SetIrisDebugObjectName(const ANSICHAR* NameBuffer)
+{
+	if (NameBuffer)
+	{
+		IrisDebugHelperInternal::GIrisDebugName = ANSI_TO_TCHAR(NameBuffer);
+	}
+	else
+	{
+		IrisDebugHelperInternal::GIrisDebugName = FString();
+	}
+}
+
+void SetIrisDebugNetHandle(uint32 NetHandleId)
+{
+	IrisDebugHelperInternal::GIrisDebugNetHandle = Private::FNetHandleManager::MakeNetHandleFromId(NetHandleId);
+}
+
+void SetIrisDebugRPCName(const ANSICHAR* NameBuffer)
+{
+	if (NameBuffer)
+	{
+		IrisDebugHelperInternal::GIrisDebugRPCName = ANSI_TO_TCHAR(NameBuffer);
+	}
+	else
+	{
+		IrisDebugHelperInternal::GIrisDebugRPCName = FString();
+	}
+}
+
+UReplicationSystem* GetReplicationSystemForDebug(uint32 Id)
+{
+	return GetReplicationSystem(Id);
+}
+
+#define UE_NET_FORCE_REFERENCE_DEBUGFUNCTION(FunctionName) FunctionReferenceAccumulator += uint64(&FunctionName);
+
+uint64 Init()
+{
+	uint64 FunctionReferenceAccumulator = uint64(0);
+
+	UE_NET_FORCE_REFERENCE_DEBUGFUNCTION(GetReplicationSystemForDebug);
+	UE_NET_FORCE_REFERENCE_DEBUGFUNCTION(DebugOutputNetObjectState);
+	UE_NET_FORCE_REFERENCE_DEBUGFUNCTION(DebugNetObjectStateToString);
+	UE_NET_FORCE_REFERENCE_DEBUGFUNCTION(DebugOutputNetObjectProtocolReferences);
+	UE_NET_FORCE_REFERENCE_DEBUGFUNCTION(DebugNetObject);
+	UE_NET_FORCE_REFERENCE_DEBUGFUNCTION(DebugNetObjectById);
+	UE_NET_FORCE_REFERENCE_DEBUGFUNCTION(DebugNetHandle);
+	UE_NET_FORCE_REFERENCE_DEBUGFUNCTION(DebugNetHandleById);
+	UE_NET_FORCE_REFERENCE_DEBUGFUNCTION(DebugNetObjectProtocolReferencesToString);
+	UE_NET_FORCE_REFERENCE_DEBUGFUNCTION(SetIrisDebugObjectName);
+	UE_NET_FORCE_REFERENCE_DEBUGFUNCTION(SetIrisDebugNetHandle);
+	UE_NET_FORCE_REFERENCE_DEBUGFUNCTION(SetIrisDebugRPCName);
+	
+	return FunctionReferenceAccumulator;
+}
+
+#undef UE_NET_FORCE_REFERENCE_DEBUGFUNCTION
+
+void NetObjectStateToString(FStringBuilderBase& StringBuilder, FNetHandle NetHandle)
+{
+	using namespace UE::Net::Private;
+
+	if (!NetHandle.IsValid())
+	{
+		return;
+	}
+
+	UReplicationSystem* ReplicationSystem = GetReplicationSystem(NetHandle.GetReplicationSystemId());
+	if (!ReplicationSystem)
+	{
+		return;
+	}
+
+	FReplicationSystemInternal* ReplicationSystemInternal = ReplicationSystem->GetReplicationSystemInternal();
+	const FNetHandleManager& NetHandleManager = ReplicationSystemInternal->GetNetHandleManager();
+
+	const uint32 NetObjectInternalIndex = NetHandleManager.GetInternalIndex(NetHandle);
+	if (!NetObjectInternalIndex)
+	{
+		return;
+	}
+
+	const bool bIsServer = ReplicationSystem->IsServer();
+	const FNetHandleManager::FReplicatedObjectData& ReplicatedObjectData = NetHandleManager.GetReplicatedObjectDataNoCheck(NetObjectInternalIndex);	
+	const uint8* InternalStateBuffer = bIsServer ? NetHandleManager.GetReplicatedObjectStateBufferNoCheck(NetObjectInternalIndex) : ReplicatedObjectData.ReceiveStateBuffer;
+	if (ReplicatedObjectData.InstanceProtocol == nullptr || InternalStateBuffer == nullptr)
+	{
+		return;
+	}
+
+	const FReplicationProtocol* Protocol = ReplicatedObjectData.Protocol;
+
+	// In order to be able to output object references we need the TokenStoreState, for the server we just use the local one but if we are a client we must use the remote token store state
+	// If this is a client handle we assume that we only have a single connections and use the first valid connection to get the remote token store.
+	FReplicationConnections& Connections = ReplicationSystemInternal->GetConnections();
+	const uint32 FirstValidConnectionId = Connections.GetValidConnections().FindFirstOne();
+	FNetTokenStoreState* TokenStoreState = (bIsServer || (FirstValidConnectionId == FNetBitArray::InvalidIndex)) ? ReplicationSystemInternal->GetNetTokenStore().GetLocalNetTokenStoreState() : &Connections.GetRemoteNetTokenStoreState(FirstValidConnectionId);
+
+	// Setup Context
+	FInternalNetSerializationContext InternalContext(ReplicationSystem, TokenStoreState);
+	FNetSerializationContext NetSerializationContext;
+	NetSerializationContext.SetInternalContext(&InternalContext);
+
+	StringBuilder << TEXT("State for ") << NetHandle;
+	StringBuilder.Appendf(TEXT(" of type %s\n"), Protocol->DebugName->Name);
+
+	FReplicationInstanceOperations::OutputInternalStateToString(NetSerializationContext, StringBuilder, nullptr, InternalStateBuffer, ReplicatedObjectData.InstanceProtocol, ReplicatedObjectData.Protocol);
+}
+
+void DebugOutputNetObjectState(uint32 NetHandleId, uint32 ReplicationSystemId)
+{
+	const FNetHandle NetHandle = Private::FNetHandleManager::MakeNetHandle(NetHandleId, ReplicationSystemId);
+
+	TStringBuilder<4096> StringBuilder;
+
+	NetObjectStateToString(StringBuilder, NetHandle);
+
+	FPlatformMisc::LowLevelOutputDebugString(StringBuilder.ToString());
+}
+
+const TCHAR* DebugNetObjectStateToString(uint32 NetHandleId, uint32 ReplicationSystemId)
+{
+	static TStringBuilder<4096> StringBuilder;
+
+	StringBuilder.Reset();
+	const FNetHandle NetHandle = Private::FNetHandleManager::MakeNetHandle(NetHandleId, ReplicationSystemId);
+	NetObjectStateToString(StringBuilder, NetHandle);
+
+	return StringBuilder.ToString();
+}
+
+FNetReplicatedObjectDebugInfo DebugNetObject(UObject* Instance)
+{
+	using namespace UE::Net::Private;
+	FNetReplicatedObjectDebugInfo Info = {};
+
+	// See if we can find the instance in any replication system
+	for (uint32 RepSystemIt = 0; RepSystemIt < FReplicationSystemFactory::MaxReplicationSystemCount; ++RepSystemIt)
+	{
+		const UReplicationSystem* ReplicationSystem = GetReplicationSystemForDebug(RepSystemIt);
+		if (ReplicationSystem)
+		{
+			if (const UObjectReplicationBridge* Bridge = ReplicationSystem->GetReplicationBridgeAs<UObjectReplicationBridge>())
+			{
+				FNetHandle Handle = Bridge->GetReplicatedHandle(Instance);
+				if (Handle.IsValid())
+				{
+					const FNetHandleManager& NetHandleManager = ReplicationSystem->GetReplicationSystemInternal()->GetNetHandleManager();
+					Info.InternalNetHandleIndex = NetHandleManager.GetInternalIndex(Handle);
+
+					const FNetHandleManager::FReplicatedObjectData& ObjectData = NetHandleManager.GetReplicatedObjectDataNoCheck(Info.InternalNetHandleIndex);				
+
+					Info.Handle = &ObjectData.Handle;
+					Info.Protocol = ObjectData.Protocol;
+					Info.InstanceProtocol = ObjectData.InstanceProtocol;
+					Info.ReplicationSystem = ReplicationSystem;
+
+					return Info;
+				}
+			}
+		}
+	}
+
+	return Info;
+}
+
+FNetReplicatedObjectDebugInfo DebugNetObjectById(UObject* Instance, uint32 ReplicationSystemId)
+{
+	using namespace UE::Net::Private;
+	FNetReplicatedObjectDebugInfo Info = {};
+
+	const UReplicationSystem* ReplicationSystem = GetReplicationSystemForDebug(ReplicationSystemId);
+	if (ReplicationSystem)
+	{
+		if (const UObjectReplicationBridge* Bridge = ReplicationSystem->GetReplicationBridgeAs<UObjectReplicationBridge>())
+		{
+			FNetHandle Handle = Bridge->GetReplicatedHandle(Instance);
+			if (Handle.IsValid())
+			{
+				const FNetHandleManager& NetHandleManager = ReplicationSystem->GetReplicationSystemInternal()->GetNetHandleManager();
+				Info.InternalNetHandleIndex = NetHandleManager.GetInternalIndex(Handle);
+
+				const FNetHandleManager::FReplicatedObjectData& ObjectData = NetHandleManager.GetReplicatedObjectDataNoCheck(Info.InternalNetHandleIndex);
+
+				Info.Handle = &ObjectData.Handle;
+				Info.Protocol = ObjectData.Protocol;
+				Info.InstanceProtocol = ObjectData.InstanceProtocol;
+				Info.ReplicationSystem = ReplicationSystem;
+
+				return Info;
+			}
+		}
+	}
+
+	return Info;
+}
+
+FNetReplicatedObjectDebugInfo DebugNetHandle(FNetHandle Handle)
+{
+	return DebugNetHandleById(Handle.GetId(), Handle.GetReplicationSystemId());
+}
+
+FNetReplicatedObjectDebugInfo DebugNetHandleById(uint32 NetHandleId, uint32 ReplicationSystemId)
+{
+	using namespace UE::Net::Private;
+	FNetReplicatedObjectDebugInfo Info = {};
+
+	// See if we can find the instance in any replication system
+	{
+		const UReplicationSystem* ReplicationSystem = GetReplicationSystemForDebug(ReplicationSystemId);
+		if (ReplicationSystem)
+		{
+			const FNetHandleManager& NetHandleManager = ReplicationSystem->GetReplicationSystemInternal()->GetNetHandleManager();
+			FNetHandle IncompleteHandle = FNetHandleManager::MakeNetHandle(NetHandleId, ReplicationSystemId);
+			const uint32 InternalNetHandleIndex = Info.InternalNetHandleIndex = NetHandleManager.GetInternalIndex(IncompleteHandle);
+			if (InternalNetHandleIndex != FNetHandleManager::InvalidInternalIndex)
+			{
+				const FNetHandleManager::FReplicatedObjectData& ObjectData = NetHandleManager.GetReplicatedObjectDataNoCheck(InternalNetHandleIndex);
+				if (ObjectData.Handle == IncompleteHandle)
+				{
+					Info.Handle = &ObjectData.Handle;
+					Info.InternalNetHandleIndex = InternalNetHandleIndex;
+					Info.Protocol = ObjectData.Protocol;
+					Info.InstanceProtocol = ObjectData.InstanceProtocol;
+					Info.ReplicationSystem = ReplicationSystem;
+
+					return Info;
+				}
+			}
+		}
+	}
+
+	return Info;
+}
+
+void NetObjectProtocolReferencesToString(FStringBuilderBase& StringBuilder, uint64 ProtocolId, uint32 ReplicationSystemId)
+{
+	using namespace UE::Net::Private;
+
+	UReplicationSystem* ReplicationSystem = GetReplicationSystem(ReplicationSystemId);
+	if (!ReplicationSystem)
+	{
+		return;
+	}
+
+	FReplicationSystemInternal* ReplicationSystemInternal = ReplicationSystem->GetReplicationSystemInternal();
+	const FNetHandleManager& NetHandleManager = ReplicationSystemInternal->GetNetHandleManager();
+	const FReplicationProtocolManager& ProtocolManager = ReplicationSystemInternal->GetReplicationProtocolManager();
+
+	// As there might be multiple protocols sharing the same identifier
+	ProtocolManager.ForEachProtocol(ProtocolId, [&StringBuilder, &NetHandleManager](const FReplicationProtocol* Protocol, const UObject* ArchetypeOrCDOUsedAsKey)
+		{
+			StringBuilder << TEXT("Protocol: ") << ToCStr(Protocol->DebugName);
+			StringBuilder.Appendf(TEXT("Id: 0x%" UINT64_x_FMT " Created From : 0x%p"), Protocol->ProtocolIdentifier, ArchetypeOrCDOUsedAsKey) << TEXT(" Used by : \n");
+
+			auto&& FindMatchingProtocol = [&StringBuilder, &Protocol, &NetHandleManager](uint32 InternalObjectIndex)
+			{
+				const FNetHandleManager::FReplicatedObjectData& ObjectData = NetHandleManager.GetReplicatedObjectDataNoCheck(InternalObjectIndex);
+				if (ObjectData.Protocol == Protocol)
+				{
+					StringBuilder << ObjectData.Handle << TEXT(" ( InternalIndex: ") << InternalObjectIndex << TEXT(", )\n");
+				}
+			};
+			
+			NetHandleManager.GetAssignedInternalIndices().ForAllSetBits(FindMatchingProtocol);
+		}
+	);
+}
+
+void DebugOutputNetObjectProtocolReferences(uint64 ProtocolId, uint32 ReplicationSystemId)
+{
+	TStringBuilder<4096> StringBuilder;
+	StringBuilder.Reset();
+
+	NetObjectProtocolReferencesToString(StringBuilder, ProtocolId, ReplicationSystemId);
+
+	FPlatformMisc::LowLevelOutputDebugString(StringBuilder.ToString());
+}
+
+const TCHAR* DebugNetObjectProtocolReferencesToString(uint64 ProtocolId, uint32 ReplicationSystemId)
+{
+	static TStringBuilder<4096> StringBuilder;
+	StringBuilder.Reset();
+
+	NetObjectProtocolReferencesToString(StringBuilder, ProtocolId, ReplicationSystemId);
+
+	return StringBuilder.ToString();
+}
+
+} // End of namespaces
+
+#endif
\ No newline at end of file
diff --git a/Engine/Source/Runtime/Experimental/Iris/Core/Private/Iris/DataStream/DataStreamDefinitions.cpp b/Engine/Source/Runtime/Experimental/Iris/Core/Private/Iris/DataStream/DataStreamDefinitions.cpp
new file mode 100644
--- /dev/null
+++ b/Engine/Source/Runtime/Experimental/Iris/Core/Private/Iris/DataStream/DataStreamDefinitions.cpp
@@ -0,0 +1,47 @@
+// Copyright Epic Games, Inc. All Rights Reserved.
+
+#include "Iris/DataStream/DataStreamDefinitions.h"
+#include "Iris/DataStream/DataStream.h"
+#include "Iris/Core/IrisLog.h"
+#include "UObject/UObjectGlobals.h"
+
+UDataStreamDefinitions::UDataStreamDefinitions()
+: bFixupComplete(false)
+{
+}
+
+void UDataStreamDefinitions::FixupDefinitions()
+{
+	if (bFixupComplete)
+	{
+		return;
+	}
+
+	for (FDataStreamDefinition& Definition : DataStreamDefinitions)
+	{
+		UE_CLOG(DataStreamDefinitions.ContainsByPredicate([Name = Definition.DataStreamName, &Definition](const FDataStreamDefinition& ExistingDefinition) { return Name == ExistingDefinition.DataStreamName && &Definition != &ExistingDefinition; }), LogIris, Error, TEXT("DataStream name is defined multiple times: %s."), *Definition.DataStreamName.GetPlainNameString());
+		UE_CLOG(!StaticEnum<EDataStreamSendStatus>()->IsValidEnumValue(int8(Definition.DefaultSendStatus)), LogIris, Error, TEXT("Invalid DataStreamSendStatus %u for DataStream %s."), unsigned(Definition.DefaultSendStatus), *Definition.DataStreamName.GetPlainNameString());
+
+		Definition.Class = StaticLoadClass(UDataStream::StaticClass(), nullptr, *Definition.ClassName.ToString(), nullptr, LOAD_Quiet);
+
+		UE_CLOG(Definition.Class == nullptr, LogIris, Error, TEXT("DataStream class could not be loaded: %s"), *Definition.ClassName.GetPlainNameString());
+	}
+
+	bFixupComplete = true;
+}
+
+const FDataStreamDefinition* UDataStreamDefinitions::FindDefinition(const FName Name) const
+{
+	return DataStreamDefinitions.FindByPredicate([Name](const FDataStreamDefinition& Definition) { return Name == Definition.DataStreamName; });
+}
+
+void UDataStreamDefinitions::GetAutoCreateStreamNames(TArray<FName>& OutStreamNames) const
+{
+	for (const FDataStreamDefinition& Definition : DataStreamDefinitions)
+	{
+		if (Definition.bAutoCreate)
+		{
+			OutStreamNames.Add(Definition.DataStreamName);
+		}
+	}
+}
diff --git a/Engine/Source/Runtime/Experimental/Iris/Core/Private/Iris/IrisConfig.cpp b/Engine/Source/Runtime/Experimental/Iris/Core/Private/Iris/IrisConfig.cpp
new file mode 100644
--- /dev/null
+++ b/Engine/Source/Runtime/Experimental/Iris/Core/Private/Iris/IrisConfig.cpp
@@ -0,0 +1,22 @@
+// Copyright Epic Games, Inc. All Rights Reserved.
+
+#include "Iris/IrisConfig.h"
+#include "HAL/IConsoleManager.h"
+
+namespace UE::Net
+{
+
+static int32 CVarUseIrisReplication = 0;
+static FAutoConsoleVariableRef CVarUseIrisReplicationRef(TEXT("net.Iris.UseIrisReplication"), CVarUseIrisReplication, TEXT("Enables Iris replication system. 0 will fallback to legacy replicationsystem."), ECVF_Default );
+
+bool ShouldUseIrisReplication()
+{
+	return CVarUseIrisReplication > 0;
+}
+
+void SetUseIrisReplication(bool EnableIrisReplication)
+{
+	CVarUseIrisReplication = EnableIrisReplication ? 1 : 0;
+}
+
+}
diff --git a/Engine/Source/Runtime/Experimental/Iris/Core/Private/Iris/IrisCoreModule.cpp b/Engine/Source/Runtime/Experimental/Iris/Core/Private/Iris/IrisCoreModule.cpp
new file mode 100644
--- /dev/null
+++ b/Engine/Source/Runtime/Experimental/Iris/Core/Private/Iris/IrisCoreModule.cpp
@@ -0,0 +1,121 @@
+// Copyright Epic Games, Inc. All Rights Reserved.
+
+#include "CoreMinimal.h"
+#include "Modules/ModuleInterface.h"
+#include "Modules/ModuleManager.h"
+#include "Iris/ReplicationState/DefaultPropertyNetSerializerInfos.h"
+#include "Iris/ReplicationState/PropertyNetSerializerInfoRegistry.h"
+#include "Iris/ReplicationSystem/LegacyPushModel.h"
+#include "Iris/Serialization/InternalNetSerializerDelegates.h"
+#include "Iris/IrisConfigInternal.h"
+
+#include "Misc/CommandLine.h"
+
+// $IRIS: $TODO: Remove when UE-158358 is completed
+#include "Iris/ReplicationSystem/ReplicationSystem.h"
+#include "Iris/ReplicationSystem/ObjectReplicationBridge.h"
+#include "Net/Core/PropertyConditions/RepChangedPropertyTracker.h"
+// End Remove when UE-158358 is completed
+
+class FIrisCoreModule : public IModuleInterface
+{
+private:
+
+	// $IRIS: $TODO: Make proper impl of support for SetPropertyCustomCondition - UE-158358
+	static void SetPropertyCustomCondition(const UObject* Owner, uint16 RepIndex, bool bEnable)
+	{
+		using namespace UE::Net;
+		// See if we can find the instance in any replication system
+		for (uint32 RepSystemIt = 0; RepSystemIt < FReplicationSystemFactory::MaxReplicationSystemCount; ++RepSystemIt)
+		{
+			UReplicationSystem* ReplicationSystem = GetReplicationSystem(RepSystemIt);
+			if (ReplicationSystem)
+			{
+				if (const UObjectReplicationBridge* Bridge = ReplicationSystem->GetReplicationBridgeAs<UObjectReplicationBridge>())
+				{
+					const FNetHandle Handle = Bridge->GetReplicatedHandle(Owner);
+					if (Handle.IsValid())
+					{
+						ReplicationSystem->SetPropertyCustomCondition(Handle, Owner, RepIndex, bEnable);
+					}
+				}
+			}
+		}
+	}
+
+	void RegisterPropertyNetSerializerSelectorTypes()
+	{
+		using namespace UE::Net;
+		using namespace UE::Net::Private;
+
+		FPropertyNetSerializerInfoRegistry::Reset();
+
+		FInternalNetSerializerDelegates::BroadcastPreFreezeNetSerializerRegistry();
+		RegisterDefaultPropertyNetSerializerInfos();
+
+		FPropertyNetSerializerInfoRegistry::Freeze();
+		FInternalNetSerializerDelegates::BroadcastPostFreezeNetSerializerRegistry();
+	}
+
+	virtual void StartupModule() override
+	{
+		// Check command line if we should disable Iris as we need to do that early
+		int32 UseIrisReplication;
+		if(FParse::Value(FCommandLine::Get(), TEXT("UseIrisReplication="), UseIrisReplication))
+		{
+			UE::Net::SetUseIrisReplication(UseIrisReplication > 0);
+		}
+		else
+		{
+			UE::Net::SetUseIrisReplication(true);
+		}
+		
+		RegisterPropertyNetSerializerSelectorTypes();
+
+		UE_NET_IRIS_INIT_LEGACY_PUSH_MODEL();
+
+		UE::Net::Private::SetIrisSetPropertyCustomConditionDelegate(UE::Net::Private::FIrisSetPropertyCustomCondition::CreateStatic(&FIrisCoreModule::SetPropertyCustomCondition));
+		
+		ModulesChangedHandle = FModuleManager::Get().OnModulesChanged().AddRaw(this, &FIrisCoreModule::OnModulesChanged);
+	}
+
+	virtual void ShutdownModule() override
+	{
+		if (ModulesChangedHandle.IsValid())
+		{
+			FModuleManager::Get().OnModulesChanged().Remove(ModulesChangedHandle);
+			ModulesChangedHandle.Reset();
+		}
+	
+		UE::Net::Private::SetIrisSetPropertyCustomConditionDelegate({});
+
+		UE_NET_IRIS_SHUTDOWN_LEGACY_PUSH_MODEL();
+	}
+
+	virtual bool SupportsDynamicReloading() override
+	{
+		return false;
+	}
+
+	void OnModulesChanged(FName ModuleThatChanged, EModuleChangeReason ReasonForChange)
+	{
+		switch (ReasonForChange)
+		{
+			case EModuleChangeReason::ModuleLoaded:
+			{
+				UE::Net::Private::FInternalNetSerializerDelegates::BroadcastPreFreezeNetSerializerRegistry();
+				UE::Net::Private::FInternalNetSerializerDelegates::BroadcastPostFreezeNetSerializerRegistry();
+			}
+			break;
+
+			default:
+			{
+			}
+			break;
+		}
+	}
+
+private:
+	FDelegateHandle ModulesChangedHandle;
+};
+IMPLEMENT_MODULE(FIrisCoreModule, IrisCore);
diff --git a/Engine/Source/Runtime/Experimental/Iris/Core/Private/Iris/ReplicationSystem/AttachmentReplication.cpp b/Engine/Source/Runtime/Experimental/Iris/Core/Private/Iris/ReplicationSystem/AttachmentReplication.cpp
new file mode 100644
--- /dev/null
+++ b/Engine/Source/Runtime/Experimental/Iris/Core/Private/Iris/ReplicationSystem/AttachmentReplication.cpp
@@ -0,0 +1,952 @@
+// Copyright Epic Games, Inc. All Rights Reserved.
+
+#include "Iris/ReplicationSystem/AttachmentReplication.h"
+#include "HAL/IConsoleManager.h"
+#include "Iris/Core/IrisLog.h"
+#include "Net/Core/Trace/NetTrace.h"
+#include "Iris/ReplicationSystem/ReplicationSystem.h"
+#include "Iris/ReplicationSystem/ReplicationSystemInternal.h"
+#include "Iris/ReplicationSystem/NetBlob/NetBlobAssembler.h"
+#include "Iris/ReplicationSystem/NetBlob/SequentialPartialNetBlobHandler.h"
+#include "Iris/Serialization/InternalNetSerializationContext.h"
+#include "Iris/Serialization/NetExportContext.h"
+#include "Iris/Serialization/NetSerializationContext.h"
+#include "Iris/Serialization/NetBitStreamReader.h"
+#include "Iris/Serialization/NetBitStreamUtil.h"
+#include "Iris/Serialization/NetBitStreamWriter.h"
+
+namespace UE::Net::Private
+{
+
+namespace AttachmentReplicationCVars
+{
+	static int32 UnreliableRPCQueueSize = 10;
+	FAutoConsoleVariableRef CVarUnreliableRPCQueueSize(TEXT("net.UnreliableRPCQueueSize"), UnreliableRPCQueueSize, TEXT("Maximum number of unreliable RPCs queued per object. If more RPCs are queued then older ones will be dropped."));
+
+	static int32 ReliableRPCQueueSize = 4096;
+	FAutoConsoleVariableRef CVarReliableRPCQueueSize(TEXT("net.ReliableRPCQueueSize"), ReliableRPCQueueSize, TEXT("Maximum number of reliable RPCs queued per object. This is in addition to the 256 that are in the send window. This is to support very large RPCs that are split into smaller pieces."));
+
+	static int32 ClientToServerUnreliableRPCQueueSize = 16;
+	FAutoConsoleVariableRef CVarClientToServerUnreliableRPCQueueSize(TEXT("net.ClientToServerUnreliableRPCQueueSize"), ClientToServerUnreliableRPCQueueSize, TEXT( "Maximum number of unreliable RPCs queued for sending from the client to the server. If more RPCs are queued then older ones will be dropped."));
+
+	static int32 MaxSimultaneousObjectsWithRPCs = 4096;
+	FAutoConsoleVariableRef CVarMaxSimultaneousObjectsWithRPCs(TEXT("net.MaxSimultaneousObjectsWithRPCs"), MaxSimultaneousObjectsWithRPCs, TEXT("Maximum number of objects that can have unsent RPCs at the same time. "));
+}
+
+static const FName NetError_UnreliableQueueFull("Unreliable attachment queue full");
+static const FName NetError_UnsupportedNetBlob("Unsupported NetBlob type");
+static const FName NetError_TooManyObjectsWithFunctionCalls("Too many objects being targeted by RPCs at the same time.");
+
+// SendQueue and Writer
+class FNetObjectAttachmentSendQueue::FReliableSendQueue
+{
+public:
+	FReliableSendQueue()
+	: MaxPreQueueCount(AttachmentReplicationCVars::ReliableRPCQueueSize)
+	{
+	}
+
+	bool HasUnsentBlobs() const
+	{
+		return PreQueue.Num() > 0 || ReliableQueue.HasUnsentBlobs();
+	}
+
+	bool IsSafeToDestroy() const
+	{
+		return ReliableQueue.IsSafeToDestroy();
+	}
+
+	bool IsAllSentAndAcked() const
+	{
+		return PreQueue.Num() == 0 && ReliableQueue.IsAllSentAndAcked();
+	}
+
+	uint32 GetUnsentBlobCount() const
+	{
+		return ReliableQueue.GetUnsentBlobCount() + static_cast<uint32>(PreQueue.Num());
+	}
+
+	bool Enqueue(TArrayView<const TRefCountPtr<FNetBlob>> Attachments)
+	{
+		// If we have blobs in the pre queue then continue adding to that.
+		const uint32 PreQueueCount = static_cast<uint32>(PreQueue.Num());
+		const uint32 AttachmentCount = static_cast<uint32>(Attachments.Num());
+
+		// Assume that we won't fit all the attachments unless they can fit in the pre-queue.
+		if (AttachmentCount + PreQueueCount > MaxPreQueueCount)
+		{
+			return false;
+		}
+
+		uint32 QueuedCount = 0;
+		if (PreQueueCount == 0)
+		{
+			for (const TRefCountPtr<FNetBlob>& Attachment : Attachments)
+			{
+				if (!ReliableQueue.Enqueue(Attachment))
+				{
+					break;
+				}
+
+				++QueuedCount;
+			}
+		}
+
+		const uint32 AddToPreQueueCount = AttachmentCount - QueuedCount;
+		if (AddToPreQueueCount > 0)
+		{
+			PreQueue.Append(Attachments.GetData() + QueuedCount, AddToPreQueueCount);
+		}
+
+		return true;
+	}
+
+	uint32 Serialize(FNetSerializationContext& Context, FNetHandle NetHandle, FReliableNetBlobQueue::ReplicationRecord& OutRecord)
+	{
+		if (NetHandle.IsValid())
+		{
+			return ReliableQueue.SerializeWithObject(Context, NetHandle, OutRecord);
+		}
+		else
+		{
+			return ReliableQueue.Serialize(Context, OutRecord);
+		}
+	}
+
+	void ProcessPacketDeliveryStatus(EPacketDeliveryStatus Status, FReliableNetBlobQueue::ReplicationRecord Record)
+	{
+		ReliableQueue.ProcessPacketDeliveryStatus(Status, Record);
+		if (Status == EPacketDeliveryStatus::Delivered && Record != FReliableNetBlobQueue::InvalidReplicationRecord)
+		{
+			PopulateQueueFromPreQueue();
+		}
+	}
+
+	void CommitReplicationRecord(FReliableNetBlobQueue::ReplicationRecord Record)
+	{
+		ReliableQueue.CommitReplicationRecord(Record);
+	}
+
+private:
+	void PopulateQueueFromPreQueue()
+	{
+		const uint32 AttachmentCount = static_cast<uint32>(PreQueue.Num());
+		if (AttachmentCount == 0)
+		{
+			return;
+		}
+
+		// Similar to enqueueing, but without all the checks as we've already passed them.
+		TArray<TRefCountPtr<FNetBlob>> PrevPreQueue = MoveTemp(PreQueue);
+		TArrayView<const TRefCountPtr<FNetBlob>> Attachments = MakeArrayView(PrevPreQueue);
+
+		uint32 QueuedCount = 0;
+		for (const TRefCountPtr<FNetBlob>& Attachment : Attachments)
+		{
+			if (!ReliableQueue.Enqueue(Attachment))
+			{
+				break;
+			}
+
+			++QueuedCount;
+		}
+
+		const uint32 AddToPreQueueCount = AttachmentCount - QueuedCount;
+		if (AddToPreQueueCount > 0)
+		{
+			PreQueue.Append(Attachments.GetData() + QueuedCount, AddToPreQueueCount);
+		}
+	}
+
+	FReliableNetBlobQueue ReliableQueue;
+	TArray<TRefCountPtr<FNetBlob>> PreQueue;
+	const uint32 MaxPreQueueCount;
+};
+
+FNetObjectAttachmentSendQueue::FNetObjectAttachmentSendQueue()
+: ReliableQueue(nullptr)
+, MaxUnreliableCount(AttachmentReplicationCVars::UnreliableRPCQueueSize)
+{
+}
+
+FNetObjectAttachmentSendQueue::~FNetObjectAttachmentSendQueue()
+{
+	delete ReliableQueue;
+}
+
+bool FNetObjectAttachmentSendQueue::Enqueue(TArrayView<const TRefCountPtr<FNetBlob>> Attachments)
+{
+	const FNetBlobCreationInfo& CreationInfo = Attachments[0]->GetCreationInfo();
+	if (EnumHasAnyFlags(CreationInfo.Flags, ENetBlobFlags::Reliable))
+	{
+		if (ReliableQueue == nullptr)
+		{
+			ReliableQueue = new FReliableSendQueue();
+		}
+
+		return ReliableQueue->Enqueue(Attachments);
+	}
+	else
+	{
+		const uint32 TotalCountNeeded = UnreliableQueue.Count() + Attachments.Num();
+		if (TotalCountNeeded > MaxUnreliableCount)
+		{
+			UE_LOG(LogIris, Verbose, TEXT("Dropping old RPC due to too many unreliable Attachments: %d max: %u"), UnreliableQueue.Count(), MaxUnreliableCount);
+			UnreliableQueue.PopNoCheck(TotalCountNeeded - MaxUnreliableCount);
+		}
+
+		for (const TRefCountPtr<FNetBlob>& Attachment : Attachments)
+		{
+			UnreliableQueue.Enqueue(Attachment);
+		}
+	}
+
+	return true;
+}
+
+void FNetObjectAttachmentSendQueue::DropUnreliable(bool &bOutHasUnsent)
+{
+	UnreliableQueue.Empty();
+	bOutHasUnsent = (ReliableQueue != nullptr && ReliableQueue->HasUnsentBlobs());
+}
+
+bool FNetObjectAttachmentSendQueue::HasUnsent() const
+{
+	return !UnreliableQueue.IsEmpty() || (ReliableQueue != nullptr && ReliableQueue->HasUnsentBlobs());
+}
+
+bool FNetObjectAttachmentSendQueue::IsAllSentAndAcked() const
+{
+	return UnreliableQueue.IsEmpty() && (ReliableQueue == nullptr || ReliableQueue->IsAllSentAndAcked());
+}
+
+bool FNetObjectAttachmentSendQueue::IsSafeToDestroy() const
+{
+	return UnreliableQueue.IsEmpty() && (ReliableQueue == nullptr || ReliableQueue->IsSafeToDestroy());
+}
+
+void FNetObjectAttachmentSendQueue::SetUnreliableQueueCapacity(uint32 QueueCapacity)
+{
+	// MaxUnreliableCount is what prevents the queue from growing too large.
+	MaxUnreliableCount = QueueCapacity;
+	
+	const uint32 UnreliableCount = UnreliableQueue.Count();
+	if (QueueCapacity >= UnreliableCount)
+	{
+		return;
+	}
+
+	const uint32 DropCount = UnreliableCount - QueueCapacity;
+	UE_LOG(LogIris, Warning, TEXT("Dropping %u attachments due to change in unreliable queue capacity to %u"), DropCount, QueueCapacity);
+	UnreliableQueue.PopNoCheck(DropCount);
+}
+
+void FNetObjectAttachmentSendQueue::Serialize(FNetSerializationContext& Context, FNetHandle NetHandle, FNetObjectAttachmentSendQueue::ReplicationRecord& OutRecord, bool& bOutHasUnsentAttachments)
+{
+	FNetBitStreamWriter& Writer = *Context.GetBitStreamWriter();
+
+	FInternalRecord ReplicationRecord;
+	const uint32 UnsentReliableCount = (ReliableQueue != nullptr ? ReliableQueue->GetUnsentBlobCount() : 0U);
+	const bool bHasReliableAttachments = UnsentReliableCount > 0;
+	const bool bHasUnreliableAttachments = !UnreliableQueue.IsEmpty();
+	if (!(bHasReliableAttachments || bHasUnreliableAttachments))
+	{
+		// Ideally we shouldn't get here, but we can handle it. Important to overflow to report a soft error.
+		Writer.DoOverflow();
+		OutRecord = ReplicationRecord.CombinedRecord;
+		bOutHasUnsentAttachments = false;
+		return;
+	}
+
+	UE_NET_TRACE_SCOPE(RPCs, Writer, Context.GetTraceCollector(), ENetTraceVerbosity::Trace);
+
+	Writer.WriteBool(bHasReliableAttachments);
+	const uint32 HasUnreliableAttachmentsWritePos = Writer.GetPosBits();
+	Writer.WriteBool(bHasUnreliableAttachments);
+
+	if (Writer.IsOverflown())
+	{
+		OutRecord = ReplicationRecord.CombinedRecord;
+		bOutHasUnsentAttachments = true;
+		return;
+	}
+
+	uint32 SerializedReliableCount = 0;
+	if (bHasReliableAttachments)
+	{
+		UE_NET_TRACE_SCOPE(Reliable, Writer, Context.GetTraceCollector(), ENetTraceVerbosity::Trace);
+		SerializedReliableCount = SerializeReliable(Context, NetHandle, ReplicationRecord.ReliableRecord);
+		// If we couldn't fit any reliable attachments then don't even try unreliable
+		if (SerializedReliableCount == 0)
+		{
+			Writer.DoOverflow();
+			OutRecord = ReplicationRecord.CombinedRecord;
+			bOutHasUnsentAttachments = true;
+			return;
+		}
+	}
+
+	uint32 SerializedUnreliableCount = 0;
+	if (bHasUnreliableAttachments)
+	{
+		UE_NET_TRACE_SCOPE(Unreliable, Writer, Context.GetTraceCollector(), ENetTraceVerbosity::Trace);
+		SerializedUnreliableCount = SerializeUnreliable(Context, NetHandle, ReplicationRecord.UnreliableRecord);
+		if (SerializedUnreliableCount == 0)
+		{
+			// If we didn't manage to send anything then inform the caller of this via overflowing the bitstream
+			if (SerializedReliableCount == 0)
+			{
+				Writer.DoOverflow();
+			}
+			// Patch up information saying there were unreliable attachments because we couldn't fit any in the bitstream
+			else
+			{
+				FNetBitStreamWriteScope WriteScope(Writer, HasUnreliableAttachmentsWritePos);
+				Writer.WriteBool(false);
+			}
+		}
+	}
+
+	OutRecord = ReplicationRecord.CombinedRecord;
+	bOutHasUnsentAttachments = (SerializedReliableCount < UnsentReliableCount) || (SerializedUnreliableCount < UnreliableQueue.Count());
+}
+
+uint32 FNetObjectAttachmentSendQueue::SerializeReliable(FNetSerializationContext& Context, FNetHandle NetHandle, FReliableNetBlobQueue::ReplicationRecord& OutRecord)
+{
+	return ReliableQueue->Serialize(Context, NetHandle, OutRecord);
+}
+
+uint32 FNetObjectAttachmentSendQueue::SerializeUnreliable(FNetSerializationContext& Context, FNetHandle NetHandle, uint32& OutRecord)
+{
+	FNetBitStreamWriter* Writer = Context.GetBitStreamWriter();
+	const FObjectReferenceCache* ObjectReferenceCache = Context.GetInternalContext()->ObjectReferenceCache;
+
+	uint32 SerializedUnreliableCount = 0;
+	uint32 PrevHasMoreAttachmentsWritePos = 0;
+	const bool bSerializeWithObject = NetHandle.IsValid();
+	for (uint32 AttachmentIt = 0, AttachmentEndIt = UnreliableQueue.Count(); AttachmentIt != AttachmentEndIt; ++AttachmentIt)
+	{
+		FNetBitStreamRollbackScope RollbackScope(*Writer);
+		FNetExportRollbackScope ExportScope(Context);
+
+		const TRefCountPtr<FNetBlob>& Attachment = UnreliableQueue.PeekAtOffsetNoCheck(AttachmentIt);
+
+		// Write exports, if attachment does not fit we will rollback exports as well.
+		ObjectReferenceCache->WriteExports(Context, Attachment->CallGetExports());
+
+		Attachment->SerializeCreationInfo(Context, Attachment->GetCreationInfo());
+		if (bSerializeWithObject)
+		{
+			Attachment->SerializeWithObject(Context, NetHandle);
+		}
+		else
+		{
+			Attachment->Serialize(Context);
+		}
+
+		const uint32 HasMoreAttachmentsWritePos = Writer->GetPosBits();
+		const bool bHasMoreAttachments = AttachmentIt + 1 != AttachmentEndIt;
+		Writer->WriteBool(bHasMoreAttachments);
+
+		if (Writer->IsOverflown())
+		{
+			// Rollback exports
+			ExportScope.Rollback();
+
+			// We need to rollback before we update the continuation bit
+			RollbackScope.Rollback();
+			if (SerializedUnreliableCount > 0)
+			{
+				FNetBitStreamWriteScope WriteScope(*Writer, PrevHasMoreAttachmentsWritePos);
+				Writer->WriteBool(false);
+			}
+
+			break;
+		}
+		else
+		{
+			++SerializedUnreliableCount;
+			PrevHasMoreAttachmentsWritePos = HasMoreAttachmentsWritePos;
+		}
+	}
+
+	OutRecord = SerializedUnreliableCount;
+	return SerializedUnreliableCount;
+}
+
+void FNetObjectAttachmentSendQueue::CommitReplicationRecord(FNetObjectAttachmentSendQueue::ReplicationRecord Record)
+{
+	FInternalRecord InternalRecord;
+	InternalRecord.CombinedRecord = Record;
+	if (InternalRecord.UnreliableRecord)
+	{
+		UnreliableQueue.PopNoCheck(InternalRecord.UnreliableRecord);
+	}
+	if (InternalRecord.ReliableRecord)
+	{
+		ReliableQueue->CommitReplicationRecord(InternalRecord.ReliableRecord);
+	}
+}
+
+void FNetObjectAttachmentSendQueue::OnPacketDelivered(FNetObjectAttachmentSendQueue::ReplicationRecord Record)
+{
+	if (ReliableQueue == nullptr)
+	{
+		return;
+	}
+
+	FInternalRecord InternalRecord;
+	InternalRecord.CombinedRecord = Record;
+	ReliableQueue->ProcessPacketDeliveryStatus(EPacketDeliveryStatus::Delivered, InternalRecord.ReliableRecord);
+}
+
+void FNetObjectAttachmentSendQueue::OnPacketLost(FNetObjectAttachmentSendQueue::ReplicationRecord Record)
+{
+	if (ReliableQueue == nullptr)
+	{
+		return;
+	}
+
+	FInternalRecord InternalRecord;
+	InternalRecord.CombinedRecord = Record;
+	ReliableQueue->ProcessPacketDeliveryStatus(EPacketDeliveryStatus::Lost, InternalRecord.ReliableRecord);
+}
+
+bool FNetObjectAttachmentsWriter::Enqueue(ENetObjectAttachmentType Type, uint32 ObjectIndex, TArrayView<const TRefCountPtr<FNetBlob>> Attachments)
+{
+	FNetObjectAttachmentSendQueue* Queue = GetOrCreateQueue(Type, ObjectIndex);
+	if (!ensureMsgf(Queue != nullptr, TEXT("Too many objects being targeted by Attachments simultaneously: %d"), ObjectToQueue.Num()))
+	{
+		return false;
+	}
+
+	return Queue->Enqueue(Attachments);
+}
+
+bool FNetObjectAttachmentsWriter::HasUnsentAttachments(ENetObjectAttachmentType Type, uint32 ObjectIndex) const
+{
+	const FNetObjectAttachmentSendQueue* Queue = GetQueue(Type, ObjectIndex);
+	if (Queue == nullptr)
+	{
+		return false;
+	}
+
+	return Queue->HasUnsent();
+}
+
+bool FNetObjectAttachmentsWriter::IsAllSentAndAcked(ENetObjectAttachmentType Type, uint32 ObjectIndex) const
+{
+	const FNetObjectAttachmentSendQueue* Queue = GetQueue(Type, ObjectIndex);
+	if (Queue == nullptr)
+	{
+		return false;
+	}
+
+	return Queue->IsAllSentAndAcked();
+}
+
+bool FNetObjectAttachmentsWriter::IsSafeToDestroy(ENetObjectAttachmentType Type, uint32 ObjectIndex) const
+{
+	const FNetObjectAttachmentSendQueue* Queue = GetQueue(Type, ObjectIndex);
+	if (Queue == nullptr)
+	{
+		return true;
+	}
+
+	return Queue->IsSafeToDestroy();
+}
+
+void FNetObjectAttachmentsWriter::DropAllAttachments(ENetObjectAttachmentType Type, uint32 ObjectIndex)
+{
+	ObjectToQueue.Remove(ObjectIndex);
+}
+
+void FNetObjectAttachmentsWriter::DropUnreliableAttachments(ENetObjectAttachmentType Type, uint32 ObjectIndex, bool &bOutHasUnsentAttachments)
+{
+	FNetObjectAttachmentSendQueue* Queue = GetQueue(ENetObjectAttachmentType::Normal, ObjectIndex);
+	if (Queue == nullptr)
+	{
+		bOutHasUnsentAttachments = false;
+		return;
+	}
+
+	Queue->DropUnreliable(bOutHasUnsentAttachments);
+	if (!bOutHasUnsentAttachments && Queue->IsSafeToDestroy())
+	{
+		ObjectToQueue.Remove(ObjectIndex);
+	}
+}
+
+void FNetObjectAttachmentsWriter::Serialize(FNetSerializationContext& Context, ENetObjectAttachmentType Type, uint32 ObjectIndex, const FNetHandle NetHandle,  FNetObjectAttachmentsWriter::ReplicationRecord& OutRecord, bool& bOutHasUnsentAttachments)
+{
+	FNetObjectAttachmentSendQueue* Queue = GetQueue(Type, ObjectIndex);
+	// If this ensure fires we have bad logic for keeping track of whether there are attachments or not
+	if (!ensure(Queue != nullptr))
+	{
+		OutRecord = 0;
+		bOutHasUnsentAttachments = false;
+		return;
+	}
+
+	Queue->Serialize(Context, NetHandle, OutRecord, bOutHasUnsentAttachments);
+}
+
+void FNetObjectAttachmentsWriter::CommitReplicationRecord(ENetObjectAttachmentType Type, uint32 ObjectIndex, FNetObjectAttachmentsWriter::ReplicationRecord Record)
+{
+	FNetObjectAttachmentSendQueue* Queue = GetQueue(Type, ObjectIndex);
+	Queue->CommitReplicationRecord(Record);
+}
+
+void FNetObjectAttachmentsWriter::OnPacketDelivered(ENetObjectAttachmentType Type, uint32 ObjectIndex, FNetObjectAttachmentsWriter::ReplicationRecord Record)
+{
+	FNetObjectAttachmentSendQueue* Queue = GetQueue(Type, ObjectIndex);
+	if (Queue == nullptr)
+	{
+		return;
+	}
+
+	Queue->OnPacketDelivered(Record);
+}
+
+void FNetObjectAttachmentsWriter::OnPacketLost(ENetObjectAttachmentType Type, uint32 ObjectIndex, ReplicationRecord Record)
+{
+	FNetObjectAttachmentSendQueue* Queue = GetQueue(Type, ObjectIndex);
+	if (Queue == nullptr)
+	{
+		return;
+	}
+
+	Queue->OnPacketLost(Record);
+}
+
+FNetObjectAttachmentSendQueue* FNetObjectAttachmentsWriter::GetQueue(ENetObjectAttachmentType Type, uint32 ObjectIndex)
+{
+	return Type == ENetObjectAttachmentType::Normal ? ObjectToQueue.Find(ObjectIndex) : SpecialQueues[uint32(Type)].Get();
+}
+
+const FNetObjectAttachmentSendQueue* FNetObjectAttachmentsWriter::GetQueue(ENetObjectAttachmentType Type, uint32 ObjectIndex) const
+{
+	return Type == ENetObjectAttachmentType::Normal ? ObjectToQueue.Find(ObjectIndex) : SpecialQueues[uint32(Type)].Get();
+}
+
+FNetObjectAttachmentSendQueue* FNetObjectAttachmentsWriter::GetOrCreateQueue(ENetObjectAttachmentType Type, uint32 ObjectIndex)
+{
+	if (Type != ENetObjectAttachmentType::Normal)
+	{
+		checkSlow(ObjectIndex == 0U);
+		TUniquePtr<FNetObjectAttachmentSendQueue>& QueuePtr = SpecialQueues[uint32(Type)];
+		if (FNetObjectAttachmentSendQueue* Queue = QueuePtr.Get())
+		{
+			return Queue;
+		}
+
+		QueuePtr = MakeUnique<FNetObjectAttachmentSendQueue>();
+
+		if (Type == ENetObjectAttachmentType::OutOfBand)
+		{
+			QueuePtr->SetUnreliableQueueCapacity(AttachmentReplicationCVars::ClientToServerUnreliableRPCQueueSize);
+		}
+
+		return QueuePtr.Get();
+	}
+	else
+	{
+		FNetObjectAttachmentSendQueue* Queue = ObjectToQueue.Find(ObjectIndex);
+		if (Queue != nullptr)
+		{
+			return Queue;
+		}
+
+		if (ObjectToQueue.Num() >= AttachmentReplicationCVars::MaxSimultaneousObjectsWithRPCs)
+		{
+			return nullptr;
+		}
+
+		Queue = &ObjectToQueue.Add(ObjectIndex);
+
+		return Queue;
+	}
+}
+
+// ReceiveQueue and Reader
+class FNetObjectAttachmentReceiveQueue::FDeferredProcessingQueue
+{
+public:
+	bool IsEmpty() const { return Queue.IsEmpty(); }
+	bool HasUnprocessed() const { return !Queue.IsEmpty(); }
+	bool IsSafeToDestroy() const { return Queue.IsEmpty() && !NetBlobAssembler.IsValid(); }
+
+	void Enqueue(FNetSerializationContext& Context, FNetHandle NetHandle, const TRefCountPtr<FNetBlob>& NetBlob, bool bIsPartialNetBlob)
+	{
+		if (bIsPartialNetBlob)
+		{
+			if (!NetBlobAssembler.IsValid())
+			{
+				NetBlobAssembler = MakeUnique<FNetBlobAssembler>();
+			}
+
+			NetBlobAssembler->AddPartialNetBlob(Context, NetHandle, reinterpret_cast<const TRefCountPtr<FPartialNetBlob>&>(NetBlob));
+			if (NetBlobAssembler->IsReadyToAssemble())
+			{
+				const TRefCountPtr<FNetBlob>& AssembledBlob = NetBlobAssembler->Assemble(Context);
+				if (AssembledBlob.IsValid())
+				{
+					Queue.Enqueue(AssembledBlob);
+				}
+				NetBlobAssembler.Reset();
+			}
+		}
+		else
+		{
+			Queue.Enqueue(NetBlob);
+		}
+	}
+
+	const TRefCountPtr<FNetBlob>* Peek() const
+	{
+		return &Queue.Peek();
+	}
+
+	void Pop()
+	{
+		return Queue.Pop();
+	}
+
+private:
+	TUniquePtr<FNetBlobAssembler> NetBlobAssembler;
+	TResizableCircularQueue<TRefCountPtr<FNetBlob>> Queue;
+};
+
+FNetObjectAttachmentReceiveQueue::FNetObjectAttachmentReceiveQueue()
+: ReliableQueue(nullptr)
+, DeferredProcessingQueue(nullptr)
+, MaxUnreliableCount(AttachmentReplicationCVars::UnreliableRPCQueueSize)
+, PartialNetBlobType(InvalidNetBlobType)
+{
+}
+
+FNetObjectAttachmentReceiveQueue::~FNetObjectAttachmentReceiveQueue()
+{
+	delete ReliableQueue;
+	delete DeferredProcessingQueue;
+}
+
+bool FNetObjectAttachmentReceiveQueue::IsSafeToDestroy() const
+{
+	return UnreliableQueue.IsEmpty() && IsDeferredProcessingQueueSafeToDestroy() && (ReliableQueue == nullptr || ReliableQueue->IsSafeToDestroy());
+}
+
+bool FNetObjectAttachmentReceiveQueue::HasUnprocessed() const
+{
+	return !UnreliableQueue.IsEmpty() || HasDeferredProcessingQueueUnprocessed();
+}
+
+const TRefCountPtr<FNetBlob>* FNetObjectAttachmentReceiveQueue::PeekReliable() const
+{
+	if (IsDeferredProcessingQueueEmpty())
+	{
+		return nullptr;
+	}
+	else
+	{
+		return DeferredProcessingQueue->Peek();
+	}
+}
+
+void FNetObjectAttachmentReceiveQueue::PopReliable()
+{
+	DeferredProcessingQueue->Pop();
+}
+
+const TRefCountPtr<FNetBlob>* FNetObjectAttachmentReceiveQueue::PeekUnreliable() const
+{
+	if (UnreliableQueue.IsEmpty())
+	{
+		return nullptr;
+	}
+	else
+	{
+		return &UnreliableQueue.Peek();
+	}
+}
+
+void FNetObjectAttachmentReceiveQueue::PopUnreliable()
+{
+	UnreliableQueue.Pop();
+}
+
+void FNetObjectAttachmentReceiveQueue::SetUnreliableQueueCapacity(uint32 QueueCapacity)
+{
+	MaxUnreliableCount = QueueCapacity;
+	
+	const uint32 UnreliableCount = UnreliableQueue.Count();
+	if (QueueCapacity >= UnreliableCount)
+	{
+		return;
+	}
+
+	const uint32 DropCount = UnreliableCount - QueueCapacity;
+	UE_LOG(LogIris, Warning, TEXT("Dropping %u attachments to due to change in unreliable queue capacity to %u"), DropCount, QueueCapacity);
+	UnreliableQueue.Pop(DropCount);
+}
+
+bool FNetObjectAttachmentReceiveQueue::IsDeferredProcessingQueueEmpty() const
+{
+	return DeferredProcessingQueue == nullptr || DeferredProcessingQueue->IsEmpty();
+}
+
+bool FNetObjectAttachmentReceiveQueue::IsDeferredProcessingQueueSafeToDestroy() const
+{
+	return DeferredProcessingQueue == nullptr || DeferredProcessingQueue->IsSafeToDestroy();
+}
+
+bool FNetObjectAttachmentReceiveQueue::HasDeferredProcessingQueueUnprocessed() const
+{
+	return DeferredProcessingQueue != nullptr && DeferredProcessingQueue->HasUnprocessed();
+}
+
+bool FNetObjectAttachmentReceiveQueue::IsPartialNetBlob(const TRefCountPtr<FNetBlob>& Blob) const
+{
+	return Blob.IsValid() && Blob->GetCreationInfo().Type == PartialNetBlobType;
+}
+
+void FNetObjectAttachmentReceiveQueue::Deserialize(FNetSerializationContext& Context, FNetHandle NetHandle)
+{
+	FNetBitStreamReader& Reader = *Context.GetBitStreamReader();
+
+	UE_NET_TRACE_SCOPE(RPCs, Reader, Context.GetTraceCollector(), ENetTraceVerbosity::Trace);
+
+	const bool bHasReliableAttachments = Reader.ReadBool();
+	const bool bHasUnreliableAttachments = Reader.ReadBool();
+
+	if (Reader.IsOverflown())
+	{
+		return;
+	}
+
+	if (bHasReliableAttachments)
+	{
+		UE_NET_TRACE_SCOPE(Reliable, Reader, Context.GetTraceCollector(), ENetTraceVerbosity::Trace);
+		DeserializeReliable(Context, NetHandle);
+		if (Context.HasErrorOrOverflow())
+		{
+			return;
+		}
+	}
+
+	if (bHasUnreliableAttachments)
+	{
+		UE_NET_TRACE_SCOPE(Unreliable, Reader, Context.GetTraceCollector(), ENetTraceVerbosity::Trace);
+		DeserializeUnreliable(Context, NetHandle);
+		if (Context.HasErrorOrOverflow())
+		{
+			return;
+		}
+	}
+}
+
+uint32 FNetObjectAttachmentReceiveQueue::DeserializeReliable(FNetSerializationContext& Context, FNetHandle NetHandle)
+{
+	if (ReliableQueue == nullptr)
+	{
+		ReliableQueue = new FReliableNetBlobQueue();
+		checkSlow(DeferredProcessingQueue == nullptr);
+		DeferredProcessingQueue = new FDeferredProcessingQueue();
+	}
+
+	uint32 DeserializedReliableCount = 0;
+	if (NetHandle.IsValid())
+	{
+		DeserializedReliableCount = ReliableQueue->DeserializeWithObject(Context, NetHandle);
+	}
+	else
+	{
+		DeserializedReliableCount = ReliableQueue->Deserialize(Context);
+	}
+
+	if (Context.HasErrorOrOverflow())
+	{
+		return DeserializedReliableCount;
+	}
+
+	while (const TRefCountPtr<FNetBlob>* Attachment = ReliableQueue->Peek())
+	{
+		DeferredProcessingQueue->Enqueue(Context, NetHandle, *Attachment, IsPartialNetBlob(*Attachment));
+		ReliableQueue->Pop();
+
+		if (Context.HasErrorOrOverflow())
+		{
+			return DeserializedReliableCount;
+		}
+	}
+
+	return DeserializedReliableCount;
+}
+
+uint32 FNetObjectAttachmentReceiveQueue::DeserializeUnreliable(FNetSerializationContext& Context, FNetHandle NetHandle)
+{
+	INetBlobReceiver* BlobReceiver = Context.GetNetBlobReceiver();
+	checkSlow(BlobReceiver != nullptr);
+
+	FNetBitStreamReader& Reader = *Context.GetBitStreamReader();
+
+	uint32 DeserializedUnreliableCount = 0;
+	const bool bDeserializeWithObject = NetHandle.IsValid();
+	bool bHasMoreAttachments = false;
+	do
+	{
+		// If the unreliable queue overflows this could be a carefully crafted malicious packet.
+		if (UnreliableQueue.Count() == MaxUnreliableCount)
+		{
+			UE_LOG(LogIris, Error, TEXT("Unreliable queue is full for %s"), *NetHandle.ToString());
+			Context.SetError(NetError_UnreliableQueueFull);
+			break;
+		}
+
+		Context.GetInternalContext()->ObjectReferenceCache->ReadExports(Context);
+
+		FNetBlobCreationInfo CreationInfo;
+		FNetBlob::DeserializeCreationInfo(Context, CreationInfo);
+		const TRefCountPtr<FNetBlob>& Attachment = BlobReceiver->CreateNetBlob(CreationInfo);
+		if (!Attachment.IsValid())
+		{
+			Context.SetError(GNetError_UnsupportedNetBlob);
+			break;
+		}
+
+		if (bDeserializeWithObject)
+		{
+			Attachment->DeserializeWithObject(Context, NetHandle);
+		}
+		else
+		{
+			Attachment->Deserialize(Context);
+		}
+		
+		bHasMoreAttachments = Reader.ReadBool();
+
+		// Break if something went wrong with deserialization.
+		if (Context.HasErrorOrOverflow())
+		{
+			UE_LOG(LogIris, Error, TEXT("Failed to deserialize unreliable attachments for %s"), *NetHandle.ToString());
+			break;
+		}
+
+		UnreliableQueue.Enqueue(Attachment);
+		++DeserializedUnreliableCount;
+	} while (bHasMoreAttachments);
+
+	return DeserializedUnreliableCount;
+}
+
+FNetObjectAttachmentsReader::FNetObjectAttachmentsReader()
+: PartialNetBlobType(InvalidNetBlobType)
+{
+}
+
+FNetObjectAttachmentsReader::~FNetObjectAttachmentsReader()
+{
+}
+
+bool FNetObjectAttachmentsReader::HasUnprocessedAttachments(ENetObjectAttachmentType Type, uint32 ObjectIndex) const
+{
+	const FNetObjectAttachmentReceiveQueue* Queue = GetQueue(Type, ObjectIndex);
+	if (Queue == nullptr)
+	{
+		return false;
+	}
+
+	return Queue->HasUnprocessed();
+}
+
+void FNetObjectAttachmentsReader::DropAllAttachments(ENetObjectAttachmentType Type, uint32 ObjectIndex)
+{
+	ObjectToQueue.Remove(ObjectIndex);
+}
+
+void FNetObjectAttachmentsReader::Deserialize(FNetSerializationContext& Context, ENetObjectAttachmentType Type, uint32 ObjectIndex, FNetHandle NetHandle)
+{
+	FNetObjectAttachmentReceiveQueue* Queue = GetOrCreateQueue(Type, ObjectIndex);
+	if (!ensure(Queue != nullptr))
+	{
+		Context.SetError(NetError_TooManyObjectsWithFunctionCalls);
+		return;
+	}
+
+	Queue->Deserialize(Context, NetHandle);
+}
+
+FNetObjectAttachmentReceiveQueue* FNetObjectAttachmentsReader::GetQueue(ENetObjectAttachmentType Type, uint32 ObjectIndex)
+{
+	return Type == ENetObjectAttachmentType::Normal ? ObjectToQueue.Find(ObjectIndex) : SpecialQueues[uint32(Type)].Get();
+}
+
+const FNetObjectAttachmentReceiveQueue* FNetObjectAttachmentsReader::GetQueue(ENetObjectAttachmentType Type, uint32 ObjectIndex) const
+{
+	return Type == ENetObjectAttachmentType::Normal ? ObjectToQueue.Find(ObjectIndex) : SpecialQueues[uint32(Type)].Get();
+}
+
+FNetObjectAttachmentReceiveQueue* FNetObjectAttachmentsReader::GetOrCreateQueue(ENetObjectAttachmentType Type, uint32 ObjectIndex)
+{
+	if (Type != ENetObjectAttachmentType::Normal)
+	{
+		checkSlow(ObjectIndex == 0U);
+		TUniquePtr<FNetObjectAttachmentReceiveQueue>& QueuePtr = SpecialQueues[uint32(Type)];
+		FNetObjectAttachmentReceiveQueue* Queue = QueuePtr.Get();
+		if (Queue != nullptr)
+		{
+			return Queue;
+		}
+
+		QueuePtr = MakeUnique<FNetObjectAttachmentReceiveQueue>();
+		Queue = QueuePtr.Get();
+		Queue->SetPartialNetBlobType(PartialNetBlobType);
+
+		if (Type == ENetObjectAttachmentType::OutOfBand)
+		{
+			Queue->SetUnreliableQueueCapacity(AttachmentReplicationCVars::ClientToServerUnreliableRPCQueueSize);
+		}
+
+		return Queue;
+	}
+	else
+	{
+		FNetObjectAttachmentReceiveQueue* Queue = ObjectToQueue.Find(ObjectIndex);
+		if (Queue != nullptr)
+		{
+			return Queue;
+		}
+
+		if (ObjectToQueue.Num() >= AttachmentReplicationCVars::MaxSimultaneousObjectsWithRPCs)
+		{
+			return nullptr;
+		}
+
+		Queue = &ObjectToQueue.Add(ObjectIndex);
+		Queue->SetPartialNetBlobType(PartialNetBlobType);
+
+		return Queue;
+	}
+
+
+	FNetObjectAttachmentReceiveQueue* Queue = ObjectToQueue.Find(ObjectIndex);
+	if (Queue != nullptr)
+	{
+		return Queue;
+	}
+
+	if (ObjectToQueue.Num() >= AttachmentReplicationCVars::MaxSimultaneousObjectsWithRPCs)
+	{
+		return nullptr;
+	}
+
+	Queue = &ObjectToQueue.Add(ObjectIndex);
+	if (ObjectIndex == 0)
+	{
+		Queue->SetUnreliableQueueCapacity(AttachmentReplicationCVars::ClientToServerUnreliableRPCQueueSize);
+	}
+
+	return Queue;
+}
+
+}
