Iris experimental replication system.

diff --git a/Engine/Source/Runtime/Experimental/Iris/Core/Private/Iris/ReplicationSystem/AttachmentReplication.h b/Engine/Source/Runtime/Experimental/Iris/Core/Private/Iris/ReplicationSystem/AttachmentReplication.h
new file mode 100644
--- /dev/null
+++ b/Engine/Source/Runtime/Experimental/Iris/Core/Private/Iris/ReplicationSystem/AttachmentReplication.h
@@ -0,0 +1,189 @@
+// Copyright Epic Games, Inc. All Rights Reserved.
+
+#pragma once
+
+#include "Net/Core/Misc/ResizableCircularQueue.h"
+#include "Iris/ReplicationSystem/NetBlob/NetBlob.h"
+#include "Iris/ReplicationSystem/NetBlob/ReliableNetBlobQueue.h"
+#include "Containers/Map.h"
+
+namespace UE::Net::Private
+{
+	class FNetBlobHandlerManager;
+	class FNetObjectAttachmentsReader;
+	class FNetObjectAttachmentsWriter;
+}
+
+namespace UE::Net::Private
+{
+
+enum class ENetObjectAttachmentType : uint32
+{
+	// For normal attachments a valid ObjectIndex is required.
+	Normal,
+
+	// For special attachments like OutOfBand and HugeObject the ObjectIndex is assumed to be zero.
+	OutOfBand,
+	HugeObject,
+
+	// For internal use only
+	InternalCount
+};
+
+class FNetObjectAttachmentSendQueue
+{
+public:
+	typedef uint64 ReplicationRecord;
+	static constexpr uint64 InvalidReplicationRecord = uint64(0);
+
+public:
+	FNetObjectAttachmentSendQueue();
+	~FNetObjectAttachmentSendQueue();
+
+	bool Enqueue(TArrayView<const TRefCountPtr<FNetBlob>> Attachments);
+
+	void DropUnreliable(bool &bOutHasUnsentAttachments);
+
+	bool HasUnsent() const;
+
+	bool IsSafeToDestroy() const;
+
+	bool IsAllSentAndAcked() const;
+
+	void SetUnreliableQueueCapacity(uint32 QueueCapacity);
+
+private:
+	friend FNetObjectAttachmentsWriter;
+	class FReliableSendQueue;
+
+	struct FInternalRecord
+	{
+		union
+		{
+			struct
+			{
+				uint32 UnreliableRecord;
+				FReliableNetBlobQueue::ReplicationRecord ReliableRecord;
+			};
+			ReplicationRecord CombinedRecord;
+		};
+
+		FInternalRecord() { CombinedRecord = InvalidReplicationRecord; };
+	};
+
+	void Serialize(FNetSerializationContext& Context, FNetHandle NetHandle, ReplicationRecord& OutRecord, bool& bOutHasUnprocessedAttachments);
+	uint32 SerializeReliable(FNetSerializationContext& Context, FNetHandle NetHandle, FReliableNetBlobQueue::ReplicationRecord& OutRecord);
+	uint32 SerializeUnreliable(FNetSerializationContext& Context, FNetHandle NetHandle, uint32& OutRecord);
+
+	void CommitReplicationRecord(ReplicationRecord Record);
+
+	void OnPacketDelivered(ReplicationRecord Record);
+	void OnPacketLost(ReplicationRecord Record);
+
+	FReliableSendQueue* ReliableQueue;
+	TResizableCircularQueue<TRefCountPtr<FNetBlob>> UnreliableQueue;
+	uint32 MaxUnreliableCount;
+};
+
+class FNetObjectAttachmentsWriter
+{
+public:
+	typedef FNetObjectAttachmentSendQueue::ReplicationRecord ReplicationRecord;
+
+public:
+	bool Enqueue(ENetObjectAttachmentType Type, uint32 ObjectIndex, TArrayView<const TRefCountPtr<FNetBlob>> Attachments);
+
+	bool HasUnsentAttachments(ENetObjectAttachmentType Type, uint32 ObjectIndex) const;
+	// Whether all queued attachments have been sent and that all reliable ones have been acked.
+	bool IsAllSentAndAcked(ENetObjectAttachmentType Type, uint32 ObjectIndex) const;
+	// Whether the queue can be destroyed without causing issues if more attachments are queued to this instance.
+	bool IsSafeToDestroy(ENetObjectAttachmentType Type, uint32 ObjectIndex) const;
+
+	void DropAllAttachments(ENetObjectAttachmentType Type, uint32 ObjectIndex);
+	void DropUnreliableAttachments(ENetObjectAttachmentType Type, uint32 ObjectIndex, bool& bOutHasUnsentAttachments);
+
+	void Serialize(FNetSerializationContext& Context, ENetObjectAttachmentType Type, uint32 ObjectIndex, FNetHandle NetHandle, ReplicationRecord& OutRecord, bool& bOutHasUnsentAttachments);
+
+	void CommitReplicationRecord(ENetObjectAttachmentType Type, uint32 ObjectIndex, ReplicationRecord Record);
+
+	void OnPacketDelivered(ENetObjectAttachmentType Type, uint32 ObjectIndex, ReplicationRecord Record);
+	void OnPacketLost(ENetObjectAttachmentType Type, uint32 ObjectIndex, ReplicationRecord Record);
+
+private:
+	bool NetBlobMightNeedSplitting(const TRefCountPtr<FNetObjectAttachment>& Attachment) const;
+
+	FNetObjectAttachmentSendQueue* GetQueue(ENetObjectAttachmentType Type, uint32 ObjectIndex);
+	const FNetObjectAttachmentSendQueue* GetQueue(ENetObjectAttachmentType Type, uint32 ObjectIndex) const;
+
+	FNetObjectAttachmentSendQueue* GetOrCreateQueue(ENetObjectAttachmentType Reason, uint32 ObjectIndex);
+
+	TMap<uint32, FNetObjectAttachmentSendQueue> ObjectToQueue;
+	TUniquePtr<FNetObjectAttachmentSendQueue> SpecialQueues[uint32(ENetObjectAttachmentType::InternalCount)];
+};
+
+class FNetObjectAttachmentReceiveQueue
+{
+public:
+	FNetObjectAttachmentReceiveQueue();
+	~FNetObjectAttachmentReceiveQueue();
+
+	void SetPartialNetBlobType(FNetBlobType InPartialNetBlobType) { PartialNetBlobType = InPartialNetBlobType; }
+
+	bool IsSafeToDestroy() const;
+	bool HasUnprocessed() const;
+
+	const TRefCountPtr<FNetBlob>* PeekReliable() const;
+	void PopReliable();
+
+	const TRefCountPtr<FNetBlob>* PeekUnreliable() const;
+	void PopUnreliable();
+
+	void SetUnreliableQueueCapacity(uint32 QueueCapacity);
+
+private:
+	friend FNetObjectAttachmentsReader;
+	class FDeferredProcessingQueue;
+
+	bool IsDeferredProcessingQueueEmpty() const;
+	bool IsDeferredProcessingQueueSafeToDestroy() const;
+	bool HasDeferredProcessingQueueUnprocessed() const;
+	bool IsPartialNetBlob(const TRefCountPtr<FNetBlob>& Blob) const;
+
+	void Deserialize(FNetSerializationContext& Context, FNetHandle NetHandle);
+	uint32 DeserializeReliable(FNetSerializationContext& Context, FNetHandle NetHandle);
+	uint32 DeserializeUnreliable(FNetSerializationContext& Context, FNetHandle NetHandle);
+
+	TResizableCircularQueue<TRefCountPtr<FNetBlob>> UnreliableQueue;
+	FReliableNetBlobQueue* ReliableQueue;
+	FDeferredProcessingQueue* DeferredProcessingQueue;
+	uint32 MaxUnreliableCount;
+	FNetBlobType PartialNetBlobType;
+};
+
+class FNetObjectAttachmentsReader
+{
+public:
+	FNetObjectAttachmentsReader();
+	~FNetObjectAttachmentsReader();
+
+	void SetPartialNetBlobType(FNetBlobType InPartialNetBlobType) { PartialNetBlobType = InPartialNetBlobType; }
+
+	bool HasUnprocessedAttachments(ENetObjectAttachmentType Type, uint32 ObjectIndex) const;
+
+	void DropAllAttachments(ENetObjectAttachmentType Type, uint32 ObjectIndex);
+
+	void Deserialize(FNetSerializationContext& Context, ENetObjectAttachmentType Type, uint32 ObjectIndex, FNetHandle NetHandle);
+
+	FNetObjectAttachmentReceiveQueue* GetQueue(ENetObjectAttachmentType Type, uint32 ObjectIndex);
+
+private:
+	const FNetObjectAttachmentReceiveQueue* GetQueue(ENetObjectAttachmentType Type, uint32 ObjectIndex) const;
+
+	FNetObjectAttachmentReceiveQueue* GetOrCreateQueue(ENetObjectAttachmentType Type, uint32 ObjectIndex);
+
+	TMap<uint32, FNetObjectAttachmentReceiveQueue> ObjectToQueue;
+	TUniquePtr<FNetObjectAttachmentReceiveQueue> SpecialQueues[uint32(ENetObjectAttachmentType::InternalCount)];
+	FNetBlobType PartialNetBlobType;
+};
+
+}
diff --git a/Engine/Source/Runtime/Experimental/Iris/Core/Private/Iris/ReplicationSystem/Conditionals/ReplicationConditionals.cpp b/Engine/Source/Runtime/Experimental/Iris/Core/Private/Iris/ReplicationSystem/Conditionals/ReplicationConditionals.cpp
new file mode 100644
--- /dev/null
+++ b/Engine/Source/Runtime/Experimental/Iris/Core/Private/Iris/ReplicationSystem/Conditionals/ReplicationConditionals.cpp
@@ -0,0 +1,567 @@
+// Copyright Epic Games, Inc. All Rights Reserved.
+
+#include "Iris/ReplicationSystem/Conditionals/ReplicationConditionals.h"
+#include "Iris/Core/IrisLog.h"
+#include "Iris/Core/IrisProfiler.h"
+#include "Iris/ReplicationState/ReplicationStateUtil.h"
+#include "Iris/ReplicationSystem/NetHandleManager.h"
+#include "Iris/ReplicationSystem/ReplicationConnections.h"
+#include "Iris/ReplicationSystem/ReplicationFragment.h"
+#include "Iris/ReplicationSystem/ReplicationProtocol.h"
+#include "Iris/ReplicationSystem/Conditionals/ReplicationCondition.h"
+#include "Iris/ReplicationSystem/DeltaCompression/DeltaCompressionBaselineInvalidationTracker.h"
+#include "Iris/ReplicationSystem/Filtering/NetObjectGroups.h"
+#include "Iris/ReplicationSystem/Filtering/ReplicationFiltering.h"
+#include "Containers/ArrayView.h"
+#include "UObject/CoreNetTypes.h"
+
+namespace UE::Net::Private
+{
+
+FReplicationConditionals::FReplicationConditionals()
+{
+}
+
+void FReplicationConditionals::Init(FReplicationConditionalsInitParams& Params)
+{
+#if DO_CHECK
+	// Verify we can handle max connection count
+	{
+		const FPerObjectInfo ObjectInfo{static_cast<decltype(FPerObjectInfo::AutonomousConnectionId)>(Params.MaxConnectionCount)};
+		check(ObjectInfo.AutonomousConnectionId == Params.MaxConnectionCount);
+	}
+#endif
+
+	NetHandleManager = Params.NetHandleManager;
+	ReplicationFiltering = Params.ReplicationFiltering;
+	ReplicationConnections = Params.ReplicationConnections;
+	BaselineInvalidationTracker = Params.BaselineInvalidationTracker;
+	NetObjectGroups = Params.NetObjectGroups;
+	MaxObjectCount = Params.MaxObjectCount;
+	MaxConnectionCount = Params.MaxConnectionCount;
+
+	PerObjectInfos.SetNumZeroed(MaxObjectCount);
+	ConnectionInfos.SetNum(MaxConnectionCount + 1U);
+}
+
+bool FReplicationConditionals::SetConditionConnectionFilter(FInternalNetHandle ObjectIndex, EReplicationCondition Condition, uint32 ConnectionId, bool bEnable)
+{
+	if (ConnectionId >= MaxConnectionCount)
+	{
+		return false;
+	}
+
+	if (!ensure(Condition == EReplicationCondition::RoleAutonomous))
+	{
+		UE_LOG(LogIris, Error, TEXT("Only EReplicationCondition::RoleAutonomous supports connection filtering, got '%u'."), uint32(Condition));
+		return false;
+	}
+
+	const uint32 AutonomousConnectionId = (ConnectionId == 0U || !bEnable) ? 0U : ConnectionId;
+	FPerObjectInfo* ObjectInfo = GetPerObjectInfo(ObjectIndex);
+	if (ObjectInfo->AutonomousConnectionId != AutonomousConnectionId)
+	{
+		const uint32 ConnIdForBaselineInvalidation = (bEnable ? ConnectionId : ObjectInfo->AutonomousConnectionId);
+		ObjectInfo->AutonomousConnectionId = AutonomousConnectionId;
+
+		BaselineInvalidationTracker->InvalidateBaselines(ObjectIndex, ConnIdForBaselineInvalidation);
+	}
+
+	return true;
+}
+
+void FReplicationConditionals::AddConnection(uint32 ConnectionId)
+{
+	// Init connection info
+	FPerConnectionInfo& ConnectionInfo = ConnectionInfos[ConnectionId];
+	ConnectionInfo.ObjectConditionals.SetNumZeroed(MaxObjectCount);
+}
+
+void FReplicationConditionals::RemoveConnection(uint32 ConnectionId)
+{
+	// Reset connection info
+	FPerConnectionInfo& ConnectionInfo = ConnectionInfos[ConnectionId];
+	ConnectionInfo.ObjectConditionals.Empty();
+}
+
+bool FReplicationConditionals::SetCondition(FInternalNetHandle ObjectIndex, EReplicationCondition Condition, bool bEnable)
+{
+	if (!ensure(Condition != EReplicationCondition::RoleAutonomous))
+	{
+		UE_LOG(LogIris, Error, TEXT("%s"), TEXT("EReplicationCondition::RoleAutonomous requires a connection."));
+		return false;
+	}
+
+	if (Condition == EReplicationCondition::ReplicatePhysics)
+	{
+		FPerObjectInfo* ObjectInfo = GetPerObjectInfo(ObjectIndex);
+		if (bEnable && !ObjectInfo->bRepPhysics)
+		{
+			// We only care to track this change if the condition is enabled.
+			BaselineInvalidationTracker->InvalidateBaselines(ObjectIndex, BaselineInvalidationTracker->InvalidateBaselineForAllConnections);
+		}
+		ObjectInfo->bRepPhysics = bEnable ? 1U : 0U;
+		return true;
+	}
+
+	ensureMsgf(false, TEXT("Unhandled EReplicationCondition '%u'"), uint32(Condition));
+	return false;
+}
+
+// N.B. Calls can come for properties that have been disabled. We must handle such cases gracefully.
+bool FReplicationConditionals::SetPropertyCustomCondition(FInternalNetHandle ObjectIndex, const void* Owner, uint16 RepIndex, bool bIsActive)
+{
+	IRIS_PROFILER_SCOPE(FReplicationConditionals_SetPropertyCustomCondition);
+
+	const FNetHandleManager::FReplicatedObjectData& ReplicatedObjectData = NetHandleManager->GetReplicatedObjectDataNoCheck(ObjectIndex);
+	const FReplicationProtocol* Protocol = ReplicatedObjectData.Protocol;
+
+	if (NetHandleManager->GetReplicatedObjectStateBufferNoCheck(ObjectIndex) == nullptr || !EnumHasAnyFlags(Protocol->ProtocolTraits, EReplicationProtocolTraits::HasLifetimeConditionals))
+	{
+		return false;
+	}
+
+	if (Protocol->LifetimeConditionalsStateCount == 1U)
+	{
+		const SIZE_T StateIndex = Protocol->FirstLifetimeConditionalsStateIndex;
+		const FReplicationInstanceProtocol* InstanceProtocol = ReplicatedObjectData.InstanceProtocol;
+		const FReplicationInstanceProtocol::FFragmentData& Fragment = InstanceProtocol->FragmentData[StateIndex];
+		const FReplicationStateDescriptor* StateDescriptor = Protocol->ReplicationStateDescriptors[StateIndex];
+
+		// Note: In this optimized code path we assume the passed Owner is the fragment owner. No checks.
+		if (RepIndex >= StateDescriptor->RepIndexCount)
+		{
+			UE_LOG(LogIris, Warning, TEXT("Trying to change non-existing custom conditional for RepIndex %u in protocol %s"), RepIndex, ToCStr(Protocol->DebugName));
+			return false;
+		}
+
+		// Modify the external state changemasks accordingly.
+		const FReplicationStateMemberRepIndexToMemberIndexDescriptor& RepIndexToMemberIndexDescriptor = StateDescriptor->MemberRepIndexToMemberIndexDescriptors[RepIndex];
+		if (RepIndexToMemberIndexDescriptor.MemberIndex == FReplicationStateMemberRepIndexToMemberIndexDescriptor::InvalidEntry)
+		{
+			UE_LOG(LogIris, Warning, TEXT("Trying to change non-existing custom conditional for RepIndex %u in protocol %s"), RepIndex, ToCStr(Protocol->DebugName));
+			return false;
+		}
+
+		const FReplicationStateMemberChangeMaskDescriptor& ChangeMaskDescriptor = StateDescriptor->MemberChangeMaskDescriptors[RepIndexToMemberIndexDescriptor.MemberIndex];
+		checkSlow(ChangeMaskDescriptor.BitCount == 1);
+		FNetBitArrayView ConditionalChangeMask = UE::Net::Private::GetMemberConditionalChangeMask(Fragment.ExternalSrcBuffer, StateDescriptor);
+
+		if (bIsActive)
+		{
+			ConditionalChangeMask.SetBits(ChangeMaskDescriptor.BitOffset, ChangeMaskDescriptor.BitCount);
+
+			// If a condition is enabled we also mark the corresponding regular changemask as dirty.
+			FNetBitArrayView MemberChangeMask = UE::Net::Private::GetMemberChangeMask(Fragment.ExternalSrcBuffer, StateDescriptor);
+			const FReplicationStateHeader& ReplicationStateHeader = UE::Net::Private::GetReplicationStateHeader(Fragment.ExternalSrcBuffer, StateDescriptor);
+			MarkDirty(ReplicationStateHeader, MemberChangeMask, ChangeMaskDescriptor);
+
+			// Enabled conditions causes new properties to be replicated which most likely have incorrect values at the receiving end.
+			BaselineInvalidationTracker->InvalidateBaselines(ObjectIndex, BaselineInvalidationTracker->InvalidateBaselineForAllConnections);
+		}
+		else
+		{
+			ConditionalChangeMask.ClearBits(ChangeMaskDescriptor.BitOffset, ChangeMaskDescriptor.BitCount);
+		}
+	}
+	else
+	{
+		// Set up fragment owner collector once.
+		constexpr uint32 MaxFragmentOwnerCount = 1U;
+		UObject* FragmentOwners[MaxFragmentOwnerCount] = {};
+		FReplicationStateOwnerCollector FragmentOwnerCollector(FragmentOwners, MaxFragmentOwnerCount);
+
+		const FReplicationInstanceProtocol* InstanceProtocol = ReplicatedObjectData.InstanceProtocol;
+
+		const SIZE_T FirstRelevantStateIndex = Protocol->FirstLifetimeConditionalsStateIndex;
+		for (const FReplicationStateDescriptor*& StateDescriptor : MakeArrayView(Protocol->ReplicationStateDescriptors + FirstRelevantStateIndex, Protocol->ReplicationStateCount - FirstRelevantStateIndex))
+		{
+			if (EnumHasAnyFlags(StateDescriptor->Traits, EReplicationStateTraits::HasLifetimeConditionals))
+			{
+				const SIZE_T StateIndex = &StateDescriptor - Protocol->ReplicationStateDescriptors;
+
+				// Is the passed Owner the owner of the fragment?
+				{
+					const FReplicationFragment* ReplicationFragment = InstanceProtocol->Fragments[StateIndex];
+
+					FragmentOwnerCollector.Reset();
+					ReplicationFragment->CollectOwner(&FragmentOwnerCollector);
+					if (FragmentOwnerCollector.GetOwnerCount() == 0U || FragmentOwnerCollector.GetOwners()[0] != Owner)
+					{
+						// Not the right owner.
+						continue;
+					}
+				}
+
+				// Can this state contain this property?
+				if (RepIndex >= StateDescriptor->RepIndexCount)
+				{
+					continue;
+				}
+
+				// Does this state contain this property?
+				const FReplicationStateMemberRepIndexToMemberIndexDescriptor& RepIndexToMemberIndexDescriptor = StateDescriptor->MemberRepIndexToMemberIndexDescriptors[RepIndex];
+				if (RepIndexToMemberIndexDescriptor.MemberIndex == FReplicationStateMemberRepIndexToMemberIndexDescriptor::InvalidEntry)
+				{
+					continue;
+				}
+
+				// We found the relevant state. Modify the external state changemasks.
+				const FReplicationInstanceProtocol::FFragmentData& Fragment = InstanceProtocol->FragmentData[StateIndex];
+
+				const FReplicationStateMemberChangeMaskDescriptor& ChangeMaskDescriptor = StateDescriptor->MemberChangeMaskDescriptors[RepIndexToMemberIndexDescriptor.MemberIndex];
+				checkSlow(ChangeMaskDescriptor.BitCount == 1);
+				FNetBitArrayView ConditionalChangeMask = GetMemberConditionalChangeMask(Fragment.ExternalSrcBuffer, StateDescriptor);
+
+				if (bIsActive)
+				{
+					ConditionalChangeMask.SetBits(ChangeMaskDescriptor.BitOffset, ChangeMaskDescriptor.BitCount);
+
+					// If a condition is enabled we also mark the corresponding regular changemask as dirty.
+					FNetBitArrayView MemberChangeMask = GetMemberChangeMask(Fragment.ExternalSrcBuffer, StateDescriptor);
+					const FReplicationStateHeader& Header = GetReplicationStateHeader(Fragment.ExternalSrcBuffer, StateDescriptor);
+					MarkDirty(Header, MemberChangeMask, ChangeMaskDescriptor);
+
+					// Enabled conditions causes new properties to be replicated which most likely have incorrect values at the receiving end.
+					BaselineInvalidationTracker->InvalidateBaselines(ObjectIndex, BaselineInvalidationTracker->InvalidateBaselineForAllConnections);
+				}
+				else
+				{
+					ConditionalChangeMask.ClearBits(ChangeMaskDescriptor.BitOffset, ChangeMaskDescriptor.BitCount);
+				}
+
+				return true;
+			}
+		}
+
+		UE_LOG(LogIris, Warning, TEXT("Trying to change non-existing custom conditional for RepIndex %u in protocol %s"), RepIndex, ToCStr(Protocol->DebugName));
+	}
+
+	return false;
+}
+
+void FReplicationConditionals::Update()
+{
+	UpdateObjectsInScope();
+}
+
+void FReplicationConditionals::GetChildSubObjectsToReplicate(uint32 ReplicatingConnectionId, const FConditionalsMask& LifetimeConditionals,  const FInternalNetHandle ParentObjectIndex, FSubObjectsToReplicateArray& OutSubObjectsToReplicate)
+{
+	// To mimic old system we use a weird replication order based on hierarchy, SubSubObjects are replicated before the parent
+	FChildSubObjectsInfo SubObjectsInfo;
+	if (NetHandleManager->GetChildSubObjects(ParentObjectIndex, SubObjectsInfo))
+	{
+		if (SubObjectsInfo.SubObjectLifeTimeConditions == nullptr)
+		{
+			for (uint32 ArrayIndex = 0; ArrayIndex < SubObjectsInfo.NumSubObjects; ++ArrayIndex)
+			{
+				FInternalNetHandle SubObjectIndex = SubObjectsInfo.ChildSubObjects[ArrayIndex];
+				GetChildSubObjectsToReplicate(ReplicatingConnectionId, LifetimeConditionals, SubObjectIndex, OutSubObjectsToReplicate);
+				OutSubObjectsToReplicate.Add(SubObjectIndex);
+			}
+		}
+		else
+		{
+			// Append child subobjects that fulfill the condition
+			for (uint32 ArrayIndex = 0; ArrayIndex < SubObjectsInfo.NumSubObjects; ++ArrayIndex)
+			{
+				const ELifetimeCondition LifeTimeCondition = (ELifetimeCondition)SubObjectsInfo.SubObjectLifeTimeConditions[ArrayIndex];
+				if (LifeTimeCondition == COND_NetGroup)
+				{
+					const FInternalNetHandle SubObjectIndex = SubObjectsInfo.ChildSubObjects[ArrayIndex];
+
+					uint32 GroupCount = 0U;
+					if (const FNetObjectGroupHandle* GroupMemberships = NetObjectGroups->GetGroupMemberships(SubObjectIndex, GroupCount))
+					{
+						bool bShouldReplicateSubObject = false;
+						for (uint32 GroupIt = 0U; GroupIt < GroupCount; ++GroupIt)
+						{
+							const FNetObjectGroupHandle NetGroup = GroupMemberships[GroupIt];
+							if (NetGroup == NetGroupOwnerNetObjectGroupHandle)
+							{
+								bShouldReplicateSubObject = LifetimeConditionals.IsConditionEnabled(COND_OwnerOnly);
+							}
+							else if (NetGroup == NetGroupReplayNetObjectGroupHandle)
+							{
+								bShouldReplicateSubObject = LifetimeConditionals.IsConditionEnabled(COND_ReplayOnly);
+							}
+							else
+							{
+								ENetFilterStatus ReplicationStatus = ENetFilterStatus::Disallow;
+								ensureAlwaysMsgf(ReplicationFiltering->GetSubObjectFilterStatus(NetGroup, ReplicatingConnectionId, ReplicationStatus), TEXT("FReplicationConditionals::GetChildSubObjectsToReplicat Trying to filter with group %u that is not a SubObjectFilterGroup"), NetGroup);
+								bShouldReplicateSubObject = ReplicationStatus != ENetFilterStatus::Disallow;
+							}
+						
+							if (bShouldReplicateSubObject)
+							{
+								GetChildSubObjectsToReplicate(ReplicatingConnectionId, LifetimeConditionals, SubObjectIndex, OutSubObjectsToReplicate);
+								OutSubObjectsToReplicate.Add(SubObjectIndex);
+								break;
+							}
+						}
+					}
+				}
+				else if (LifetimeConditionals.IsConditionEnabled(LifeTimeCondition))
+				{
+					const FInternalNetHandle SubObjectIndex = SubObjectsInfo.ChildSubObjects[ArrayIndex];
+					GetChildSubObjectsToReplicate(ReplicatingConnectionId, LifetimeConditionals, SubObjectIndex, OutSubObjectsToReplicate);
+					OutSubObjectsToReplicate.Add(SubObjectIndex);
+				}
+			}
+		}
+	}
+}
+
+void FReplicationConditionals::GetSubObjectsToReplicate(uint32 ReplicationConnectionId, FInternalNetHandle RootObjectIndex, FSubObjectsToReplicateArray& OutSubObjectsToReplicate)
+{
+	IRIS_PROFILER_SCOPE(FReplicationConditionals_GetSubObjectsToReplicate);
+
+	// For now, we do nothing to detect if a conditional has changed on the RootParent, we simply defer this until the next
+	// time the subobjects are marked as dirty. We might want to consider to explicitly mark object and subobjects as dirty when 
+	// the owning connections or conditionals such as bRepPhysics or Role is changed.
+	const FConditionalsMask LifetimeConditionals = GetLifetimeConditionals(ReplicationConnectionId, RootObjectIndex);
+	GetChildSubObjectsToReplicate(ReplicationConnectionId, LifetimeConditionals, RootObjectIndex, OutSubObjectsToReplicate);
+}
+
+bool FReplicationConditionals::ApplyConditionalsToChangeMask(uint32 ReplicatingConnectionId, FInternalNetHandle ParentObjectIndex, FInternalNetHandle ObjectIndex, uint32* ChangeMaskData, const uint32* ConditionalChangeMaskData, const FReplicationProtocol* Protocol)
+{
+	IRIS_PROFILER_SCOPE(FReplicationConditionals_ApplyConditionalsToChangeMask);
+
+	bool bMaskWasModified = false;
+
+	// Assume we need all information regarding connection filtering and replication conditionals.
+	FNetBitArrayView ChangeMask = MakeNetBitArrayView(ChangeMaskData, Protocol->ChangeMaskBitCount);
+
+	// Legacy lifetime conditionals support.
+	if (EnumHasAnyFlags(Protocol->ProtocolTraits, EReplicationProtocolTraits::HasLifetimeConditionals))
+	{
+		const FConditionalsMask LifetimeConditionals = GetLifetimeConditionals(ReplicatingConnectionId, ParentObjectIndex);
+		FConditionalsMask PrevLifeTimeConditions = ConnectionInfos[ReplicatingConnectionId].ObjectConditionals[ObjectIndex];
+		if (PrevLifeTimeConditions.IsUninitialized())
+		{
+			PrevLifeTimeConditions = LifetimeConditionals;
+		}
+		ConnectionInfos[ReplicatingConnectionId].ObjectConditionals[ObjectIndex] = LifetimeConditionals;
+
+		// Optimized path for single lifetime conditional state
+		if (Protocol->LifetimeConditionalsStateCount == 1U)
+		{
+			const FReplicationStateDescriptor* StateDescriptor = Protocol->ReplicationStateDescriptors[Protocol->FirstLifetimeConditionalsStateIndex];
+			const uint32 ChangeMaskBitOffset = Protocol->FirstLifetimeConditionalsChangeMaskOffset;
+
+			const FReplicationStateMemberChangeMaskDescriptor* ChangeMaskDescriptors = StateDescriptor->MemberChangeMaskDescriptors;
+			const FReplicationStateMemberLifetimeConditionDescriptor* LifetimeConditionDescriptors = StateDescriptor->MemberLifetimeConditionDescriptors;
+			for (uint32 MemberIt = 0U, MemberEndIt = StateDescriptor->MemberCount; MemberIt != MemberEndIt; ++MemberIt)
+			{
+				const FReplicationStateMemberLifetimeConditionDescriptor& LifetimeConditionDescriptor = LifetimeConditionDescriptors[MemberIt];
+				const ELifetimeCondition Condition = static_cast<ELifetimeCondition>(LifetimeConditionDescriptor.Condition);
+				// If condition was enabled we need to dirty changemask of relevant members. If it was disabled we clear the changemask of relevant members.
+				if (LifetimeConditionals.IsConditionEnabled(Condition))
+				{
+					if (!PrevLifeTimeConditions.IsConditionEnabled(Condition))
+					{
+						const FReplicationStateMemberChangeMaskDescriptor& ChangeMaskDescriptor = ChangeMaskDescriptors[MemberIt];
+						for (uint32 BitIt = ChangeMaskBitOffset + ChangeMaskDescriptor.BitOffset, BitEndIt = BitIt + ChangeMaskDescriptor.BitCount; BitIt != BitEndIt; ++BitIt)
+						{
+							bMaskWasModified = bMaskWasModified | !ChangeMask.GetBit(BitIt);
+							ChangeMask.SetBit(BitIt);
+						}
+					}
+				}
+				else
+				{
+					const FReplicationStateMemberChangeMaskDescriptor& ChangeMaskDescriptor = ChangeMaskDescriptors[MemberIt];
+					for (uint32 BitIt = ChangeMaskBitOffset + ChangeMaskDescriptor.BitOffset, BitEndIt = BitIt + ChangeMaskDescriptor.BitCount; BitIt != BitEndIt; ++BitIt)
+					{
+						bMaskWasModified = bMaskWasModified | ChangeMask.GetBit(BitIt);
+						ChangeMask.ClearBit(BitIt);
+					}
+				}
+			}
+		}
+		else
+		{
+			uint32 CurrentChangeMaskBitOffset = 0U;
+			uint32 LifetimeConditionalsStateIt = 0U;
+			const uint32 LifetimeConditionalsStateEndIt = Protocol->LifetimeConditionalsStateCount;
+			for (const FReplicationStateDescriptor* StateDescriptor : MakeArrayView(Protocol->ReplicationStateDescriptors, Protocol->ReplicationStateCount))
+			{
+				if (EnumHasAnyFlags(StateDescriptor->Traits, EReplicationStateTraits::HasLifetimeConditionals))
+				{
+					const FReplicationStateMemberChangeMaskDescriptor* ChangeMaskDescriptors = StateDescriptor->MemberChangeMaskDescriptors;
+					const FReplicationStateMemberLifetimeConditionDescriptor* LifetimeConditionDescriptors = StateDescriptor->MemberLifetimeConditionDescriptors;
+					for (uint32 MemberIt = 0U, MemberEndIt = StateDescriptor->MemberCount; MemberIt != MemberEndIt; ++MemberIt)
+					{
+						const FReplicationStateMemberLifetimeConditionDescriptor& LifetimeConditionDescriptor = LifetimeConditionDescriptors[MemberIt];
+						const ELifetimeCondition Condition = static_cast<ELifetimeCondition>(LifetimeConditionDescriptor.Condition);
+						// If the condition is fulfilled the changemask will remain intact so we can continue to the next member.
+						if (LifetimeConditionals.IsConditionEnabled(Condition))
+						{
+							if (!PrevLifeTimeConditions.IsConditionEnabled(Condition))
+							{
+								const FReplicationStateMemberChangeMaskDescriptor& ChangeMaskDescriptor = ChangeMaskDescriptors[MemberIt];
+								for (uint32 BitIt = CurrentChangeMaskBitOffset + ChangeMaskDescriptor.BitOffset, BitEndIt = BitIt + ChangeMaskDescriptor.BitCount; BitIt != BitEndIt; ++BitIt)
+								{
+									bMaskWasModified = bMaskWasModified | !ChangeMask.GetBit(BitIt);
+									ChangeMask.SetBit(BitIt);
+								}
+							}
+						}
+						else
+						{
+							const FReplicationStateMemberChangeMaskDescriptor& ChangeMaskDescriptor = ChangeMaskDescriptors[MemberIt];
+							for (uint32 BitIt = CurrentChangeMaskBitOffset + ChangeMaskDescriptor.BitOffset, BitEndIt = BitIt + ChangeMaskDescriptor.BitCount; BitIt != BitEndIt; ++BitIt)
+							{
+								bMaskWasModified = bMaskWasModified | ChangeMask.GetBit(BitIt);
+								ChangeMask.ClearBit(BitIt);
+							}
+						}
+					}
+
+					// Done processing all states with lifetime conditionals?
+					++LifetimeConditionalsStateIt;
+					if (LifetimeConditionalsStateIt == LifetimeConditionalsStateEndIt)
+					{
+						break;
+					}
+				}
+
+				CurrentChangeMaskBitOffset += StateDescriptor->ChangeMaskBitCount;
+			}
+		}
+	}
+
+	// Apply custom conditionals by word operations.
+	if (ConditionalChangeMaskData != nullptr)
+	{
+		const uint32 WordCount = ChangeMask.GetNumWords();
+		uint32 ChangedBits = 0;
+		for (uint32 WordIt = 0; WordIt != WordCount; ++WordIt)
+		{
+			const uint32 OldMask = ChangeMaskData[WordIt];
+			const uint32 ConditionalMask = ConditionalChangeMaskData[WordIt];
+			const uint32 NewMask = OldMask & ConditionalMask;
+			ChangeMaskData[WordIt] = NewMask;
+
+			ChangedBits |= OldMask & ~ConditionalMask;
+		}
+
+		bMaskWasModified = bMaskWasModified | (ChangedBits != 0U);
+	}
+
+	return bMaskWasModified;
+}
+
+void FReplicationConditionals::UpdateObjectsInScope()
+{
+	IRIS_PROFILER_SCOPE(FReplicationConditionals_UpdateObjectsInScope);
+
+	const FNetBitArray& ObjectsInScope = NetHandleManager->GetScopableInternalIndices();
+	const FNetBitArray& PrevObjectsInScope = NetHandleManager->GetPrevFrameScopableInternalIndices();
+
+	const uint32 WordCountForModifiedWords = Align(FPlatformMath::Max(MaxObjectCount, 1U), 32U)/32U;
+	TArray<uint32> ModifiedWords;
+	ModifiedWords.SetNumUninitialized(WordCountForModifiedWords);
+	uint32* ModifiedWordsStorage = ModifiedWords.GetData();
+	uint32 ModifiedWordIndex = 0U;
+	const uint32* ObjectsInScopeStorage = ObjectsInScope.GetData();
+	const uint32* PrevObjectsInScopeStorage = PrevObjectsInScope.GetData();
+	{
+		for (uint32 WordIt = 0U, WordEndIt = WordCountForModifiedWords; WordIt != WordEndIt; ++WordIt)
+		{
+			const bool bWordDiffers = ObjectsInScopeStorage[WordIt] != PrevObjectsInScopeStorage[WordIt];
+			ModifiedWordsStorage[ModifiedWordIndex] = WordIt;
+			ModifiedWordIndex += bWordDiffers;
+		}
+	}
+
+	// If the scope didn't change there's nothing for us to do.
+	if (ModifiedWordIndex == 0U)
+	{
+		return;
+	}
+
+	{
+		const FNetBitArray& ValidConnections = ReplicationConnections->GetValidConnections();
+		const uint32 MaxValidConnectionId = ValidConnections.FindLastOne();
+
+		for (uint32 WordIt = 0U, WordEndIt = ModifiedWordIndex; WordIt != WordEndIt; ++WordIt)
+		{
+			const uint32 WordIndex = ModifiedWordsStorage[WordIt];
+			const uint32 PrevExistingObjects = PrevObjectsInScopeStorage[WordIndex];
+			const uint32 ExistingObjects = ObjectsInScopeStorage[WordIndex];
+
+			const uint32 BitOffset = WordIndex*32U;
+
+			// Clear info for deleted objects.
+			{
+				uint32 DeletedObjects = PrevExistingObjects & ~ExistingObjects;
+				for (; DeletedObjects; )
+				{
+					const uint32 LeastSignificantBit = DeletedObjects & uint32(-int32(DeletedObjects));
+					DeletedObjects ^= LeastSignificantBit;
+
+					const uint32 ObjectIndex = BitOffset + FPlatformMath::CountTrailingZeros(LeastSignificantBit);
+					ClearPerObjectInfo(ObjectIndex);
+					if (MaxValidConnectionId != ~0U)
+					{
+						ClearConnectionInfosForObject(ValidConnections, MaxValidConnectionId, ObjectIndex);
+					}
+				}
+			}
+		}
+	}
+}
+
+FReplicationConditionals::FConditionalsMask FReplicationConditionals::GetLifetimeConditionals(uint32 ReplicatingConnectionId, FInternalNetHandle ParentObjectIndex) const
+{
+	FConditionalsMask ConditionalsMask{0};
+
+	const uint32 ObjectOwnerConnectionId = ReplicationFiltering->GetOwningConnection(ParentObjectIndex);
+	const bool bIsReplicatingToOwner = (ReplicatingConnectionId == ObjectOwnerConnectionId);
+
+	const FReplicationConditionals::FPerObjectInfo* ObjectInfo = GetPerObjectInfo(ParentObjectIndex);
+	const bool bRoleSimulated = ReplicatingConnectionId != ObjectInfo->AutonomousConnectionId;
+	const bool bRoleAutonomous = ReplicatingConnectionId == ObjectInfo->AutonomousConnectionId;
+	const bool bRepPhysics = ObjectInfo->bRepPhysics;
+
+	ConditionalsMask.SetConditionEnabled(COND_None, true);
+	ConditionalsMask.SetConditionEnabled(COND_Custom, true);
+	ConditionalsMask.SetConditionEnabled(COND_OwnerOnly, bIsReplicatingToOwner);
+	ConditionalsMask.SetConditionEnabled(COND_SkipOwner, !bIsReplicatingToOwner);
+	ConditionalsMask.SetConditionEnabled(COND_SimulatedOnly, bRoleSimulated);
+	ConditionalsMask.SetConditionEnabled(COND_AutonomousOnly, bRoleAutonomous);
+	ConditionalsMask.SetConditionEnabled(COND_SimulatedOrPhysics, bRoleSimulated | bRepPhysics);
+	ConditionalsMask.SetConditionEnabled(COND_InitialOrOwner, bIsReplicatingToOwner);
+	ConditionalsMask.SetConditionEnabled(COND_ReplayOrOwner, bIsReplicatingToOwner);
+	ConditionalsMask.SetConditionEnabled(COND_SimulatedOnlyNoReplay, bRoleSimulated);
+	ConditionalsMask.SetConditionEnabled(COND_SimulatedOrPhysicsNoReplay, bRoleSimulated | bRepPhysics);
+
+	return ConditionalsMask;
+}
+
+void FReplicationConditionals::ClearPerObjectInfo(FInternalNetHandle ObjectIndex)
+{
+	FPerObjectInfo& PerObjectInfo = PerObjectInfos.GetData()[ObjectIndex];
+	PerObjectInfo = {};
+}
+
+void FReplicationConditionals::ClearConnectionInfosForObject(const FNetBitArray& ValidConnections, uint32 MaxConnectionId, FInternalNetHandle ObjectIndex)
+{
+	IRIS_PROFILER_SCOPE(FReplicationConditionals_ClearConnectionInfosForObject);
+
+	// 0 is not a valid connection ID so we adjust it to one-based in the loop.
+	for (uint32 ConnIt = 0, ConnEndIt = MaxConnectionId; ConnIt != ConnEndIt; ++ConnIt)
+	{
+		const uint32 ConnectionId = ConnIt + 1U;
+		if (!ValidConnections.GetBit(ConnectionId))
+		{
+			continue;
+		}
+
+		FPerConnectionInfo& ConnectionInfo = ConnectionInfos[ConnectionId];
+		ConnectionInfo.ObjectConditionals[ObjectIndex] = FConditionalsMask{};
+	}
+}
+
+}
diff --git a/Engine/Source/Runtime/Experimental/Iris/Core/Private/Iris/ReplicationSystem/Conditionals/ReplicationConditionals.h b/Engine/Source/Runtime/Experimental/Iris/Core/Private/Iris/ReplicationSystem/Conditionals/ReplicationConditionals.h
new file mode 100644
--- /dev/null
+++ b/Engine/Source/Runtime/Experimental/Iris/Core/Private/Iris/ReplicationSystem/Conditionals/ReplicationConditionals.h
@@ -0,0 +1,135 @@
+// Copyright Epic Games, Inc. All Rights Reserved.
+
+#pragma once
+
+#include "CoreTypes.h"
+#include "Containers/Array.h"
+#include "Net/Core/NetBitArray.h"
+
+namespace UE::Net
+{
+	enum class EReplicationCondition : uint32;
+	struct FReplicationProtocol;
+	namespace Private
+	{
+		class FDeltaCompressionBaselineInvalidationTracker;
+		typedef uint32 FInternalNetHandle;
+		class FNetHandleManager;
+		class FReplicationConnections;
+		class FReplicationFiltering;
+		class FNetObjectGroups;
+	}
+}
+
+namespace UE::Net::Private
+{
+
+struct FReplicationConditionalsInitParams
+{
+	const FNetHandleManager* NetHandleManager = nullptr;
+	const FReplicationFiltering* ReplicationFiltering = nullptr;
+	const FReplicationConnections* ReplicationConnections = nullptr;
+	const FNetObjectGroups* NetObjectGroups = nullptr;
+	FDeltaCompressionBaselineInvalidationTracker* BaselineInvalidationTracker = nullptr;
+	uint32 MaxObjectCount = 0;
+	uint32 MaxConnectionCount = 0;
+};
+
+/**
+ * NOTE: Almost everything in this class is for backward compatibility mode only. As such some functions are
+ * fairly slow. For declared replication states this overhead will be eliminated because the state
+ * itself can do all the relevant tracking locally. The one thing that will possibly remain is the unmasking
+ * of changemasks. The unmasking could then either be done on the sending side, CPU vs bandwidth trade-off,
+ * or on the receiving side, not so CPU costly as there's typically one connection and maybe tens of objects.
+ * The unmasking is then likely to move to protocol operations.
+ */
+class FReplicationConditionals
+{
+public:
+	FReplicationConditionals();
+
+	void Init(FReplicationConditionalsInitParams& Params);
+
+	void AddConnection(uint32 ConnectionId);
+	void RemoveConnection(uint32 ConnectionId);
+
+	bool SetConditionConnectionFilter(FInternalNetHandle ObjectIndex, EReplicationCondition Condition, uint32 ConnectionId, bool bEnable);
+	bool SetCondition(FInternalNetHandle ObjectIndex, EReplicationCondition Condition, bool bEnable);
+
+	// For property custom conditions only
+	bool SetPropertyCustomCondition(FInternalNetHandle ObjectIndex, const void* Owner, uint16 RepIndex, bool bIsActive);
+
+	void Update();
+
+	bool ApplyConditionalsToChangeMask(uint32 ReplicatingConnectionId, FInternalNetHandle ParentObjectIndex, FInternalNetHandle ObjectIndex, uint32* ChangeMaskData, const uint32* ConditionalChangeMaskData, const FReplicationProtocol* Protocol);
+
+	using FSubObjectsToReplicateArray = TArray<FInternalNetHandle, TInlineAllocator<32>>;
+	void GetSubObjectsToReplicate(uint32 ReplicationConnectionId, FInternalNetHandle ParentObjectIndex, FSubObjectsToReplicateArray& OutSubObjectsToReplicate);
+
+private:
+	struct FPerObjectInfo
+	{
+		// Assume there can only be one connection which has the role autonomous connection and all else are simulated
+		uint16 AutonomousConnectionId : 15;
+		uint16 bRepPhysics : 1;
+	};
+
+	struct FConditionalsMask
+	{
+		bool IsUninitialized() const { return ConditionalsMask == 0; }
+		bool IsConditionEnabled(int Condition) const { return ConditionalsMask & (uint16(1) << unsigned(Condition)); }
+		bool SetConditionEnabled(int Condition, bool bEnabled) { return ConditionalsMask |= (uint16(bEnabled ? 1 : 0) << unsigned(Condition)); }
+
+		// Each LifetimeCondition is represented in this member via (1U << ELifetimeCondition)
+		uint16 ConditionalsMask;
+	};
+
+	struct FPerConnectionInfo
+	{
+		TArray<FConditionalsMask> ObjectConditionals;
+	};
+
+	struct FSubObjectConditionInfo
+	{
+		// This is a more compact storage form of ELifetimeCondition
+		int8 Condition;
+	};
+
+private:
+	void UpdateObjectsInScope();
+
+	FConditionalsMask GetLifetimeConditionals(uint32 ReplicatingConnectionId, FInternalNetHandle ParentObjectIndex) const;
+	bool ObjectHasLifetimeConditionals(FInternalNetHandle ObjectIndex) const;	
+
+	FPerObjectInfo* GetPerObjectInfo(FInternalNetHandle ObjectIndex);
+	const FPerObjectInfo* GetPerObjectInfo(FInternalNetHandle ObjectIndex) const;
+	void ClearPerObjectInfo(FInternalNetHandle ObjectIndex);
+	void ClearConnectionInfosForObject(const FNetBitArray& ValidConnections, uint32 MaxConnectionId, FInternalNetHandle ObjectIndex);
+
+	void GetChildSubObjectsToReplicate(uint32 ReplicatingConnectionId, const FConditionalsMask& LifetimeConditionals,  const FInternalNetHandle ParentObjectIndex, FSubObjectsToReplicateArray& OutSubObjectsToReplicate);
+
+private:
+	const FNetHandleManager* NetHandleManager = nullptr;
+	const FReplicationFiltering* ReplicationFiltering = nullptr;
+	const FReplicationConnections* ReplicationConnections = nullptr;
+	FDeltaCompressionBaselineInvalidationTracker* BaselineInvalidationTracker = nullptr;
+	const FNetObjectGroups* NetObjectGroups = nullptr;
+
+	uint32 MaxObjectCount = 0;
+	uint32 MaxConnectionCount = 0;
+
+	TArray<FPerObjectInfo> PerObjectInfos;
+	TArray<FPerConnectionInfo> ConnectionInfos;
+};
+
+inline FReplicationConditionals::FPerObjectInfo* FReplicationConditionals::GetPerObjectInfo(FInternalNetHandle ObjectIndex)
+{
+	return PerObjectInfos.GetData() + ObjectIndex;
+}
+
+inline const FReplicationConditionals::FPerObjectInfo* FReplicationConditionals::GetPerObjectInfo(FInternalNetHandle ObjectIndex) const
+{
+	return PerObjectInfos.GetData() + ObjectIndex;
+}
+
+}
diff --git a/Engine/Source/Runtime/Experimental/Iris/Core/Private/Iris/ReplicationSystem/DirtyNetObjectTracker.cpp b/Engine/Source/Runtime/Experimental/Iris/Core/Private/Iris/ReplicationSystem/DirtyNetObjectTracker.cpp
new file mode 100644
--- /dev/null
+++ b/Engine/Source/Runtime/Experimental/Iris/Core/Private/Iris/ReplicationSystem/DirtyNetObjectTracker.cpp
@@ -0,0 +1,124 @@
+// Copyright Epic Games, Inc. All Rights Reserved.
+
+#include "DirtyNetObjectTracker.h"
+#include "HAL/PlatformAtomics.h"
+#include "Iris/IrisConstants.h"
+#include "Iris/Core/IrisLog.h"
+#include "Iris/ReplicationSystem/ReplicationSystem.h"
+#include "Traits/IntType.h"
+#include <atomic>
+
+#if (UE_BUILD_SHIPPING || UE_BUILD_TEST)
+#	define UE_NET_ENABLE_DIRTYOBJECTTRACKER_LOG 0
+#else
+#	define UE_NET_ENABLE_DIRTYOBJECTTRACKER_LOG 0
+#endif 
+
+#if UE_NET_ENABLE_DIRTYOBJECTTRACKER_LOG
+#	define UE_LOG_DIRTYOBJECTTRACKER(Format, ...)  UE_LOG(LogIris, Log, Format, ##__VA_ARGS__)
+#else
+#	define UE_LOG_DIRTYOBJECTTRACKER(...)
+#endif
+
+#define UE_LOG_DIRTYOBJECTTRACKER_WARNING(Format, ...)  UE_LOG(LogIris, Warning, Format, ##__VA_ARGS__)
+
+namespace UE::Net::Private
+{
+
+#if UE_NET_ALLOW_MULTIPLE_REPLICATION_SYSTEMS
+constexpr uint32 DirtyNetObjectTrackerCount = FReplicationSystemFactory::MaxReplicationSystemCount;
+#else
+constexpr uint32 DirtyNetObjectTrackerCount = 1;
+#endif
+static FDirtyNetObjectTracker* DirtyNetObjectTrackers[DirtyNetObjectTrackerCount + 1]; // nullptr as last tracker
+
+FDirtyNetObjectTracker::FDirtyNetObjectTracker()
+: ReplicationSystemId(InvalidReplicationSystemId)
+, DirtyNetObjectContainer(nullptr)
+, DirtyNetObjectWordCount(0)
+, NetObjectIdRangeStart(0U)
+, NetObjectIdRangeEnd(0U)
+, NetObjectIdCount(0)
+{
+}
+
+FDirtyNetObjectTracker::~FDirtyNetObjectTracker()
+{
+	if (ReplicationSystemId < DirtyNetObjectTrackerCount && DirtyNetObjectTrackers[ReplicationSystemId] == this)
+	{
+		DirtyNetObjectTrackers[ReplicationSystemId] = nullptr;
+	}
+
+	Shutdown();
+}
+
+void FDirtyNetObjectTracker::Init(const FDirtyNetObjectTrackerInitParams& Params)
+{
+	check(Params.ReplicationSystemId < DirtyNetObjectTrackerCount);
+	check(Params.NetObjectIndexRangeEnd >= Params.NetObjectIndexRangeStart);
+	check(DirtyNetObjectContainer == nullptr);
+
+	ReplicationSystemId = Params.ReplicationSystemId;
+	NetObjectIdRangeStart = Params.NetObjectIndexRangeStart;
+	NetObjectIdRangeEnd = Params.NetObjectIndexRangeEnd;
+	/* 
+	 * For now we support all IDs up to RangeEnd. This could be expensive if we partition things more in some way or other.
+	 * In the latter case we would have to add functionality to FNetBitArrayView to handle an offset or add a "FNetSparseBitArray".
+	 */
+	NetObjectIdCount = NetObjectIdRangeEnd + 1;
+
+	DirtyNetObjectWordCount = (NetObjectIdCount + StorageTypeBitCount - 1)/StorageTypeBitCount;
+	DirtyNetObjectContainer = new StorageType[DirtyNetObjectWordCount];
+	ClearDirtyNetObjects();
+
+	if (DirtyNetObjectTrackers[ReplicationSystemId] != nullptr)
+	{
+		LowLevelFatalError(TEXT("DirtyNetObjectTrackerAlready initialized for ReplicationSystemId %u"), ReplicationSystemId);
+	}
+	
+	DirtyNetObjectTrackers[ReplicationSystemId] = this;
+
+	UE_LOG_DIRTYOBJECTTRACKER(TEXT("FDirtyNetObjectTracker::Init %u Id, Start:%u, End: %u"), ReplicationSystemId, NetObjectIdRangeStart, NetObjectIdRangeEnd);
+}
+
+void FDirtyNetObjectTracker::Shutdown()
+{
+	delete[] DirtyNetObjectContainer;
+	DirtyNetObjectContainer = nullptr;
+}
+
+void FDirtyNetObjectTracker::MarkNetObjectDirty(uint32 NetObjectIndex)
+{
+	if ((NetObjectIndex >= NetObjectIdRangeStart) & (NetObjectIndex <= NetObjectIdRangeEnd))
+	{
+		const uint32 BitOffset = NetObjectIndex;
+		const StorageType BitMask = StorageType(1) << (BitOffset & (StorageTypeBitCount - 1));
+
+		// ideally we'd have c++20 std::atomic_ref for this
+		FPlatformAtomics::InterlockedOr((TSignedIntType<sizeof(StorageType)>::Type*)(&DirtyNetObjectContainer[BitOffset/StorageTypeBitCount]), BitMask);
+
+		UE_LOG_DIRTYOBJECTTRACKER(TEXT("FDirtyNetObjectTracker::MarkNetObjectDirty %u ( InternalIndex: %u )"), ReplicationSystemId, NetObjectIndex);
+	}
+}
+
+FNetBitArrayView FDirtyNetObjectTracker::GetDirtyNetObjects() const
+{
+	return FNetBitArrayView(DirtyNetObjectContainer, NetObjectIdCount);
+}
+
+void FDirtyNetObjectTracker::ClearDirtyNetObjects()
+{
+	FMemory::Memset(DirtyNetObjectContainer, 0, DirtyNetObjectWordCount*sizeof(StorageType));
+
+	std::atomic_thread_fence(std::memory_order_seq_cst);
+}
+
+void MarkNetObjectStateDirty(uint32 ReplicationSystemId, uint32 NetObjectIndex)
+{
+	if (FDirtyNetObjectTracker* DirtyNetObjectTracker = DirtyNetObjectTrackers[FMath::Min(ReplicationSystemId, DirtyNetObjectTrackerCount)])
+	{
+		DirtyNetObjectTracker->MarkNetObjectDirty(NetObjectIndex);
+	}
+}
+
+}
diff --git a/Engine/Source/Runtime/Experimental/Iris/Core/Private/Iris/ReplicationSystem/DirtyNetObjectTracker.h b/Engine/Source/Runtime/Experimental/Iris/Core/Private/Iris/ReplicationSystem/DirtyNetObjectTracker.h
new file mode 100644
--- /dev/null
+++ b/Engine/Source/Runtime/Experimental/Iris/Core/Private/Iris/ReplicationSystem/DirtyNetObjectTracker.h
@@ -0,0 +1,49 @@
+// Copyright Epic Games, Inc. All Rights Reserved.
+
+#pragma once
+
+#include "Net/Core/NetBitArray.h"
+
+namespace UE::Net::Private
+{
+
+IRISCORE_API void MarkNetObjectStateDirty(uint32 ReplicationSystemId, uint32 NetObjectIndex);
+
+struct FDirtyNetObjectTrackerInitParams
+{
+	uint32 ReplicationSystemId;
+	uint32 NetObjectIndexRangeStart;
+	uint32 NetObjectIndexRangeEnd;
+};
+
+class FDirtyNetObjectTracker
+{
+public:
+	FDirtyNetObjectTracker();
+	~FDirtyNetObjectTracker();
+
+	void Init(const FDirtyNetObjectTrackerInitParams& Params);
+
+	// N.B. These methods do not use atomics.
+	// It's up to the user to determine whether it's safe to call the functions or not.
+	FNetBitArrayView GetDirtyNetObjects() const;
+	void ClearDirtyNetObjects();
+
+private:
+	friend IRISCORE_API void MarkNetObjectStateDirty(uint32 ReplicationSystemId, uint32 NetObjectIndex);
+
+	using StorageType = FNetBitArrayView::StorageWordType;
+	static constexpr uint32 StorageTypeBitCount = FNetBitArrayView::WordBitCount;
+
+	void Shutdown();
+	void MarkNetObjectDirty(uint32 NetObjectIndex);
+
+	uint32 ReplicationSystemId;
+	StorageType* DirtyNetObjectContainer;
+	uint32 DirtyNetObjectWordCount;
+	uint32 NetObjectIdRangeStart;
+	uint32 NetObjectIdRangeEnd;
+	uint32 NetObjectIdCount;
+};
+
+}
diff --git a/Engine/Source/Runtime/Experimental/Iris/Core/Private/Iris/ReplicationSystem/Filtering/FilterOutNetObjectFilter.cpp b/Engine/Source/Runtime/Experimental/Iris/Core/Private/Iris/ReplicationSystem/Filtering/FilterOutNetObjectFilter.cpp
new file mode 100644
--- /dev/null
+++ b/Engine/Source/Runtime/Experimental/Iris/Core/Private/Iris/ReplicationSystem/Filtering/FilterOutNetObjectFilter.cpp
@@ -0,0 +1,26 @@
+// Copyright Epic Games, Inc. All Rights Reserved.
+
+#include "Iris/ReplicationSystem/Filtering/FilterOutNetObjectFilter.h"
+
+void UFilterOutNetObjectFilter::Init(FNetObjectFilterInitParams& Params)
+{
+}
+
+bool UFilterOutNetObjectFilter::AddObject(uint32 ObjectIndex, FNetObjectFilterAddObjectParams& Params)
+{
+	return true;
+}
+
+void UFilterOutNetObjectFilter::RemoveObject(uint32 ObjectIndex, const FNetObjectFilteringInfo& Info)
+{
+}
+
+void UFilterOutNetObjectFilter::UpdateObjects(FNetObjectFilterUpdateParams& Params)
+{
+}
+
+void UFilterOutNetObjectFilter::Filter(FNetObjectFilteringParams& Params)
+{
+	// Filter out everything
+	Params.OutAllowedObjects.Reset();
+}
diff --git a/Engine/Source/Runtime/Experimental/Iris/Core/Private/Iris/ReplicationSystem/Filtering/NetObjectFilter.cpp b/Engine/Source/Runtime/Experimental/Iris/Core/Private/Iris/ReplicationSystem/Filtering/NetObjectFilter.cpp
new file mode 100644
--- /dev/null
+++ b/Engine/Source/Runtime/Experimental/Iris/Core/Private/Iris/ReplicationSystem/Filtering/NetObjectFilter.cpp
@@ -0,0 +1,35 @@
+// Copyright Epic Games, Inc. All Rights Reserved.
+
+#include "Iris/ReplicationSystem/Filtering/NetObjectFilter.h"
+
+FNetObjectFilteringParams::FNetObjectFilteringParams(const UE::Net::FNetBitArrayView InFilteredObjects)
+: FilteredObjects(InFilteredObjects)
+, FilteringInfos(nullptr)
+, StateBuffers(nullptr)
+, ConnectionId(0)
+{
+}
+
+UNetObjectFilter::UNetObjectFilter()
+{
+}
+
+void UNetObjectFilter::AddConnection(uint32 ConnectionId)
+{
+}
+
+void UNetObjectFilter::RemoveConnection(uint32 ConnectionId)
+{
+}
+
+void UNetObjectFilter::PreFilter(FNetObjectPreFilteringParams&)
+{
+}
+
+void UNetObjectFilter::Filter(FNetObjectFilteringParams&)
+{
+}
+
+void UNetObjectFilter::PostFilter(FNetObjectPostFilteringParams&)
+{
+}
diff --git a/Engine/Source/Runtime/Experimental/Iris/Core/Private/Iris/ReplicationSystem/Filtering/NetObjectGridFilter.cpp b/Engine/Source/Runtime/Experimental/Iris/Core/Private/Iris/ReplicationSystem/Filtering/NetObjectGridFilter.cpp
new file mode 100644
--- /dev/null
+++ b/Engine/Source/Runtime/Experimental/Iris/Core/Private/Iris/ReplicationSystem/Filtering/NetObjectGridFilter.cpp
@@ -0,0 +1,436 @@
+// Copyright Epic Games, Inc. All Rights Reserved.
+
+#include "Iris/ReplicationSystem/Filtering/NetObjectGridFilter.h"
+#include "Iris/ReplicationSystem/RepTag.h"
+#include "Iris/ReplicationSystem/ReplicationProtocol.h"
+#include "Iris/ReplicationSystem/ReplicationSystem.h"
+#include "Iris/ReplicationSystem/WorldLocations.h"
+#include "Iris/Core/IrisLog.h"
+#include "Iris/Core/IrisProfiler.h"
+
+void UNetObjectGridFilter::Init(FNetObjectFilterInitParams& Params)
+{
+	Config = TStrongObjectPtr<UNetObjectGridFilterConfig>(CastChecked<UNetObjectGridFilterConfig>(Params.Config));
+	checkf(Config.IsValid(), TEXT("Need config to operate."));
+
+	AssignedObjectInfoIndices.Init(Params.MaxObjectCount);
+
+	PerConnectionInfos.SetNum(Params.MaxConnectionCount + 1);
+
+	WorldLocations = &Params.ReplicationSystem->GetWorldLocations();
+}
+
+void UNetObjectGridFilter::AddConnection(uint32 ConnectionId)
+{
+	PerConnectionInfos[ConnectionId] = {};
+}
+
+void UNetObjectGridFilter::RemoveConnection(uint32 ConnectionId)
+{
+	PerConnectionInfos[ConnectionId] = {};
+}
+
+bool UNetObjectGridFilter::AddObject(uint32 ObjectIndex, FNetObjectFilterAddObjectParams& Params)
+{
+	// We support either a world location in the state, tagged with RepTag_WorldLocation, or via the WorldLocations instance.
+	UE::Net::FRepTagFindInfo WorldLocationTagInfo;
+	bool bHasWorldLocation = false;
+	if (WorldLocations->HasWorldLocation(ObjectIndex))
+	{
+		bHasWorldLocation = true;
+		// Craft tag info that will let us know we need to retrieve the location from WorldLocations
+		WorldLocationTagInfo.StateIndex = InvalidStateIndex;
+		WorldLocationTagInfo.ExternalStateOffset = InvalidStateOffset;
+	}
+	else if (!UE::Net::FindRepTag(Params.Protocol, UE::Net::RepTag_WorldLocation, WorldLocationTagInfo))
+	{
+		return false;
+	}
+
+	if (!bHasWorldLocation && ((WorldLocationTagInfo.ExternalStateOffset >= MAX_uint16) || (WorldLocationTagInfo.StateIndex >= MAX_uint16)))
+	{
+		return false;
+	}
+
+	FObjectLocationInfo& ObjectLocationInfo = static_cast<FObjectLocationInfo&>(Params.OutInfo);
+	ObjectLocationInfo.SetLocationStateOffset(WorldLocationTagInfo.ExternalStateOffset);
+	ObjectLocationInfo.SetLocationStateIndex(WorldLocationTagInfo.StateIndex);
+
+	const uint32 InfoIndex = AllocObjectInfo();
+	ObjectLocationInfo.SetInfoIndex(InfoIndex);
+
+	FPerObjectInfo& PerObjectInfo = ObjectInfos[InfoIndex];
+	PerObjectInfo.ObjectIndex = ObjectIndex;
+
+	// NetCullDistanceSqr is optional. 
+	{
+		UE::Net::FRepTagFindInfo NetCullDistanceSqrTagInfo;
+		if (UE::Net::FindRepTag(Params.Protocol, UE::Net::RepTag_CullDistanceSqr, NetCullDistanceSqrTagInfo))
+		{
+			if ((NetCullDistanceSqrTagInfo.ExternalStateOffset < MAX_uint16) && (NetCullDistanceSqrTagInfo.StateIndex < MAX_uint16))
+			{
+				PerObjectInfo.CullDistanceSqrStateIndex = NetCullDistanceSqrTagInfo.StateIndex;
+				PerObjectInfo.CullDistanceSqrStateOffset = NetCullDistanceSqrTagInfo.ExternalStateOffset;
+			}
+		}
+	}
+
+	AddCellInfoForObject(ObjectLocationInfo, Params.InstanceProtocol);
+	
+	if (PerObjectInfo.CullDistance > Config->MaxCullDistance)
+	{
+		// Too big an object. We expect it to be costly to move it across cells.
+		RemoveObject(ObjectIndex, ObjectLocationInfo);
+		return false;
+	}
+
+	return true;
+}
+
+void UNetObjectGridFilter::RemoveObject(uint32 ObjectIndex, const FNetObjectFilteringInfo& Info)
+{
+	const FObjectLocationInfo& ObjectLocationInfo = static_cast<const FObjectLocationInfo&>(Info);
+	RemoveCellInfoForObject(ObjectLocationInfo);
+
+	const uint32 InfoIndex = ObjectLocationInfo.GetInfoIndex();
+	FreeObjectInfo(InfoIndex);
+}
+
+void UNetObjectGridFilter::UpdateObjects(FNetObjectFilterUpdateParams& Params)
+{
+	IRIS_PROFILER_SCOPE(UNetObjectGridFilter_UpdateObjects);
+
+	for (SIZE_T ObjectIt = 0, ObjectEndIt = Params.ObjectCount; ObjectIt != ObjectEndIt; ++ObjectIt)
+	{
+		const uint32 ObjectIndex = Params.ObjectIndices[ObjectIt];
+
+		const FObjectLocationInfo& ObjectLocationInfo = static_cast<const FObjectLocationInfo&>(Params.FilteringInfos[ObjectIndex]);
+		const UE::Net::FReplicationInstanceProtocol* InstanceProtocol = Params.InstanceProtocols[ObjectIt];
+		UpdateCellInfoForObject(ObjectLocationInfo, InstanceProtocol);
+	}
+}
+
+void UNetObjectGridFilter::PreFilter(FNetObjectPreFilteringParams&)
+{
+	++FrameIndex;
+}
+
+void UNetObjectGridFilter::Filter(FNetObjectFilteringParams& Params)
+{
+	using namespace UE::Net;
+
+	IRIS_PROFILER_SCOPE(UNetObjectGridFilter_Filter);
+
+	// Update active cells for the connection
+	FPerConnectionInfo& ConnectionInfo = PerConnectionInfos[Params.ConnectionId];
+	TArray<FCellAndTimestamp, TInlineAllocator<32>> PrevCells = ConnectionInfo.RecentCells;
+	ConnectionInfo.RecentCells.Reset();
+
+	// Insert current cells for the views and remove any stale data from the previous cells.
+	TArray<FCellAndTimestamp, TInlineAllocator<32>> NewCells;
+	for (FReplicationView::FView& View : Params.View.Views)
+	{
+		FCellAndTimestamp CellAndTimestamp;
+		CellAndTimestamp.Timestamp = FrameIndex;
+		CalculateCellCoord(CellAndTimestamp.Cell, View.Pos);
+		NewCells.AddUnique(CellAndTimestamp);
+
+		for (const FCellAndTimestamp& PrevCell : PrevCells)
+		{
+			if ((PrevCell.Cell.X == CellAndTimestamp.Cell.X) & (PrevCell.Cell.Y == CellAndTimestamp.Cell.Y))
+			{
+				PrevCells.RemoveAtSwap(&PrevCell - PrevCells.GetData());
+				break;
+			}
+		}
+	}
+
+	// Prune old cells and insert still relevant in the new list.
+	{
+		const uint32 MaxFrameCount = Config->ViewPosRelevancyFrameCount;
+		for (const FCellAndTimestamp& PrevCell : PrevCells)
+		{
+			if ((FrameIndex - PrevCell.Timestamp) > MaxFrameCount)
+			{
+				continue;
+			}
+
+			NewCells.Add(PrevCell);
+		}
+	}
+
+	// Store new cells
+	ConnectionInfo.RecentCells = NewCells;
+
+	// Only allow objects in any of the relevant cells to replicate.
+	FNetBitArrayView AllowedObjects = Params.OutAllowedObjects;
+	AllowedObjects.Reset();
+
+	/**
+	 * The algorithm will simply iterate over all relevant cells and set the bits
+	 * in the AllowedObjects bitarray for all the objects present in the cell. An object can span multiple cells,
+	 * but We assume it's cheaper to modify the bitarray than to create a new set with the unique
+	 * object indices and then iterate over that. In both cases one do need to iterate over all
+	 * objects in all relevant cells anyway and setting a bit should be faster than inserting into a set.
+	 */
+	for (const FCellAndTimestamp& CellAndTimestamp : NewCells)
+	{
+		if (FCellObjects* Objects = Cells.Find(CellAndTimestamp.Cell))
+		{
+			for (const uint32 ObjectIndex : Objects->ObjectIndices)
+			{
+				AllowedObjects.SetBit(ObjectIndex);
+			}
+		}
+	}
+}
+
+uint32 UNetObjectGridFilter::AllocObjectInfo()
+{
+	uint32 Index = AssignedObjectInfoIndices.FindFirstZero();
+	if (Index >= uint32(ObjectInfos.Num()))
+	{
+		constexpr int32 NumElementsPerChunk = ObjectInfosChunkSize/sizeof(FPerObjectInfo);
+		ObjectInfos.Add(NumElementsPerChunk);
+	}
+
+	AssignedObjectInfoIndices.SetBit(Index);
+	return Index;
+}
+
+void UNetObjectGridFilter::FreeObjectInfo(uint32 Index)
+{
+	AssignedObjectInfoIndices.ClearBit(Index);
+	ObjectInfos[Index] = FPerObjectInfo();
+}
+
+void UNetObjectGridFilter::AddCellInfoForObject(const FObjectLocationInfo& ObjectLocationInfo, const UE::Net::FReplicationInstanceProtocol* InstanceProtocol)
+{
+	// Called for completely new objects
+	FPerObjectInfo& PerObjectInfo = ObjectInfos[ObjectLocationInfo.GetInfoIndex()];
+	PerObjectInfo.CullDistance = Config->DefaultCullDistance;
+	UpdatePositionAndCullDistance(ObjectLocationInfo, PerObjectInfo, InstanceProtocol);
+
+	FCellBox NewCellBox;
+	CalculateCellBox(PerObjectInfo, NewCellBox);
+	PerObjectInfo.CellBox = NewCellBox;
+
+	const uint32 ObjectIndex = PerObjectInfo.ObjectIndex;
+
+	// Add object to new cells.
+	{
+		FCellCoord Coord;
+		for (int32 X = NewCellBox.MinX, EndX = NewCellBox.MaxX + 1; X < EndX; ++X)
+		{
+			Coord.X = X;
+			for (int32 Y = NewCellBox.MinY, EndY = NewCellBox.MaxY + 1; Y < EndY; ++Y)
+			{
+				Coord.Y = Y;
+				FCellObjects& Cell = Cells.FindOrAdd(Coord);
+				Cell.ObjectIndices.Add(ObjectIndex);
+			}
+		}
+	}
+}
+
+void UNetObjectGridFilter::RemoveCellInfoForObject(const FObjectLocationInfo& ObjectLocationInfo)
+{
+	FPerObjectInfo& PerObjectInfo = ObjectInfos[ObjectLocationInfo.GetInfoIndex()];
+
+	const FCellBox CellBox = PerObjectInfo.CellBox;
+	const uint32 ObjectIndex = PerObjectInfo.ObjectIndex;
+
+	// Remove object from cells.
+	{
+		FCellCoord Coord;
+		for (int32 X = CellBox.MinX, EndX = CellBox.MaxX + 1; X < EndX; ++X)
+		{
+			Coord.X = X;
+			for (int32 Y = CellBox.MinY, EndY = CellBox.MaxY + 1; Y < EndY; ++Y)
+			{
+				Coord.Y = Y;
+				FCellObjects* Cell = Cells.Find(Coord);
+				checkSlow(Cell != nullptr);
+				Cell->ObjectIndices.Remove(ObjectIndex);
+			}
+		}
+	}
+}
+
+void UNetObjectGridFilter::UpdateCellInfoForObject(const FObjectLocationInfo& ObjectLocationInfo, const UE::Net::FReplicationInstanceProtocol* InstanceProtocol)
+{
+	FPerObjectInfo& PerObjectInfo = ObjectInfos[ObjectLocationInfo.GetInfoIndex()];
+
+	UpdatePositionAndCullDistance(ObjectLocationInfo, PerObjectInfo, InstanceProtocol);
+
+	const FCellBox PrevCellBox = PerObjectInfo.CellBox;
+	FCellBox NewCellBox;
+	CalculateCellBox(PerObjectInfo, NewCellBox);
+
+	if (NewCellBox != PrevCellBox)
+	{
+		const uint32 ObjectIndex = PerObjectInfo.ObjectIndex;
+		PerObjectInfo.CellBox = NewCellBox;
+
+		// Special case for disjoint cells
+		if (AreCellsDisjoint(NewCellBox, PrevCellBox))
+		{
+			// Remove object from previous cells.
+			{
+				FCellCoord Coord;
+				for (int32 X = PrevCellBox.MinX, EndX = PrevCellBox.MaxX + 1; X < EndX; ++X)
+				{
+					Coord.X = X;
+					for (int32 Y = PrevCellBox.MinY, EndY = PrevCellBox.MaxY + 1; Y < EndY; ++Y)
+					{
+						Coord.Y = Y;
+						FCellObjects* Cell = Cells.Find(Coord);
+						checkSlow(Cell != nullptr);
+						Cell->ObjectIndices.Remove(ObjectIndex);
+					}
+				}
+			}
+
+			// Add object to new cells.
+			{
+				FCellCoord Coord;
+				for (int32 X = NewCellBox.MinX, EndX = NewCellBox.MaxX + 1; X < EndX; ++X)
+				{
+					Coord.X = X;
+					for (int32 Y = NewCellBox.MinY, EndY = NewCellBox.MaxY + 1; Y < EndY; ++Y)
+					{
+						Coord.Y = Y;
+						FCellObjects& Cell = Cells.FindOrAdd(Coord);
+						Cell.ObjectIndices.Add(ObjectIndex);
+					}
+				}
+			}
+		}
+		// There's some overlap
+		else
+		{
+			/**
+			 * Naive algorithm to find out which cells to update.
+			 * 
+			 * 1. Calculate the cell overlap.
+			 * 2. For each coord exclusive to the previous cell we will remove the object.
+			 * 3. For each coord exclusive to the new cell we will add the object.
+			 */
+			FCellBox Overlap;
+			Overlap.MinX = FPlatformMath::Max(NewCellBox.MinX, PrevCellBox.MinX);
+			Overlap.MaxX = FPlatformMath::Min(NewCellBox.MaxX, PrevCellBox.MaxX);
+			Overlap.MinY = FPlatformMath::Max(NewCellBox.MinY, PrevCellBox.MinY);
+			Overlap.MaxY = FPlatformMath::Min(NewCellBox.MaxY, PrevCellBox.MaxY);
+			
+			// Remove object from previous cells.
+			{
+				FCellCoord Coord;
+				for (int32 X = PrevCellBox.MinX, EndX = PrevCellBox.MaxX + 1; X < EndX; ++X)
+				{
+					Coord.X = X;
+					for (int32 Y = PrevCellBox.MinY, EndY = PrevCellBox.MaxY + 1; Y < EndY; ++Y)
+					{
+						Coord.Y = Y;
+
+						if (DoesCellContainCoord(Overlap, Coord))
+						{
+							continue;
+						}
+
+						FCellObjects* Cell = Cells.Find(Coord);
+						checkSlow(Cell != nullptr);
+						Cell->ObjectIndices.Remove(ObjectIndex);
+					}
+				}
+			}
+
+			// Add object to new cells.
+			{
+				FCellCoord Coord;
+				for (int32 X = NewCellBox.MinX, EndX = NewCellBox.MaxX + 1; X < EndX; ++X)
+				{
+					Coord.X = X;
+					for (int32 Y = NewCellBox.MinY, EndY = NewCellBox.MaxY + 1; Y < EndY; ++Y)
+					{
+						Coord.Y = Y;
+
+						if (DoesCellContainCoord(Overlap, Coord))
+						{
+							continue;
+						}
+
+						FCellObjects& Cell = Cells.FindOrAdd(Coord);
+						Cell.ObjectIndices.Add(ObjectIndex);
+					}
+				}
+			}
+		}
+	}
+}
+
+void UNetObjectGridFilter::UpdatePositionAndCullDistance(const UNetObjectGridFilter::FObjectLocationInfo& ObjectLocationInfo, UNetObjectGridFilter::FPerObjectInfo& PerObjectInfo, const UE::Net::FReplicationInstanceProtocol* InstanceProtocol)
+{
+	TArrayView<const UE::Net::FReplicationInstanceProtocol::FFragmentData> FragmentDatas = MakeArrayView(InstanceProtocol->FragmentData, InstanceProtocol->FragmentCount);
+
+	// Update position
+	{
+		if (ObjectLocationInfo.IsUsingWorldLocations())
+		{
+			PerObjectInfo.Position = WorldLocations->GetWorldLocation(PerObjectInfo.ObjectIndex);
+		}
+		else
+		{
+			const UE::Net::FReplicationInstanceProtocol::FFragmentData& FragmentData = FragmentDatas[ObjectLocationInfo.GetLocationStateIndex()];
+			const uint8* LocationAddress = FragmentData.ExternalSrcBuffer + ObjectLocationInfo.GetLocationStateOffset();
+			const float* LocationComponents = reinterpret_cast<const float*>(LocationAddress);
+			PerObjectInfo.Position = FVector(LocationComponents[0], LocationComponents[1], LocationComponents[2]);
+		}
+	}
+
+	// Optionally update cull distance
+	if (PerObjectInfo.CullDistanceSqrStateOffset != InvalidStateOffset)
+	{
+		const UE::Net::FReplicationInstanceProtocol::FFragmentData& FragmentData = FragmentDatas[PerObjectInfo.CullDistanceSqrStateIndex];
+		const uint8* CullDistanceSqrAddress = FragmentData.ExternalSrcBuffer + PerObjectInfo.CullDistanceSqrStateOffset;
+		const float CullDistanceSqr = *reinterpret_cast<const float*>(CullDistanceSqrAddress);
+		PerObjectInfo.CullDistance = FPlatformMath::Sqrt(CullDistanceSqr);
+	}
+}
+
+void UNetObjectGridFilter::CalculateCellBox(const UNetObjectGridFilter::FPerObjectInfo& PerObjectInfo, UNetObjectGridFilter::FCellBox& OutCellBox)
+{
+	const float CullDistance = PerObjectInfo.CullDistance;
+	const FVector Position = PerObjectInfo.Position;
+	FVector MinPosition = Position - CullDistance;
+	FVector MaxPosition = Position + CullDistance;
+
+	MinPosition = MinPosition.ComponentMax(Config->MinPos);
+	MaxPosition = MaxPosition.ComponentMin(Config->MaxPos);
+
+	FCellBox CellBox;
+	CellBox.MinX = FPlatformMath::FloorToInt(MinPosition.X/Config->CellSizeX);
+	CellBox.MinY = FPlatformMath::FloorToInt(MinPosition.Y/Config->CellSizeY);
+	CellBox.MaxX = FPlatformMath::FloorToInt(MaxPosition.X/Config->CellSizeX);
+	CellBox.MaxY = FPlatformMath::FloorToInt(MaxPosition.Y/Config->CellSizeY);
+
+	OutCellBox = CellBox;
+}
+
+void UNetObjectGridFilter::CalculateCellCoord(UNetObjectGridFilter::FCellCoord& OutCoord, const FVector& Pos)
+{
+	FCellCoord Coord;
+	Coord.X = FPlatformMath::FloorToInt(Pos.X/Config->CellSizeX);
+	Coord.Y = FPlatformMath::FloorToInt(Pos.Y/Config->CellSizeY);
+
+	OutCoord = Coord;
+}
+
+bool UNetObjectGridFilter::AreCellsDisjoint(const UNetObjectGridFilter::FCellBox& A, const UNetObjectGridFilter::FCellBox& B)
+{
+	return (A.MinX > B.MaxX) | (A.MaxX < B.MinX) | (A.MinY > B.MaxY) | (A.MaxY < B.MinY);
+}
+
+bool UNetObjectGridFilter::DoesCellContainCoord(const UNetObjectGridFilter::FCellBox& Cell, const UNetObjectGridFilter::FCellCoord& Coord)
+{
+	return (Coord.X >= Cell.MinX) & (Coord.X <= Cell.MaxX) & (Coord.Y >= Cell.MinY) & (Coord.Y <= Cell.MaxY);
+}
diff --git a/Engine/Source/Runtime/Experimental/Iris/Core/Private/Iris/ReplicationSystem/Filtering/NetObjectGroups.cpp b/Engine/Source/Runtime/Experimental/Iris/Core/Private/Iris/ReplicationSystem/Filtering/NetObjectGroups.cpp
new file mode 100644
--- /dev/null
+++ b/Engine/Source/Runtime/Experimental/Iris/Core/Private/Iris/ReplicationSystem/Filtering/NetObjectGroups.cpp
@@ -0,0 +1,251 @@
+#include "NetObjectGroups.h"
+#include "Iris/ReplicationSystem/Filtering/NetObjectFilter.h"
+#include "Iris/ReplicationSystem/NetHandleManager.h" // for InvalidInternalIndex
+#include "Containers/ArrayView.h"
+
+namespace UE::Net::Private
+{
+
+FNetObjectGroups::FNetObjectGroups()
+: MaxGroupCount(0U)
+{
+}
+
+FNetObjectGroups::~FNetObjectGroups()
+{
+}
+
+bool FNetObjectGroups::IsMemberOf(const FNetObjectGroupMembership& Target, FNetObjectGroupHandle Group)
+{
+	for (uint32 It = 0; It < FNetObjectGroupMembership::MaxAssignedGroupCount; ++It)
+	{
+		if (Target.Groups[It] == Group)
+		{
+			return true;
+		}
+	}
+
+	return false;
+}
+
+bool FNetObjectGroups::AddGroupMembership(FNetObjectGroupMembership& Target, FNetObjectGroupHandle Group)
+{
+	for (uint32 It = 0; It < FNetObjectGroupMembership::MaxAssignedGroupCount; ++It)
+	{
+		if (Target.Groups[It] == Group)
+		{
+			return false;
+		}
+		else if (Target.Groups[It] == InvalidNetObjectGroupHandle)
+		{
+			Target.Groups[It] = Group;
+			return true;
+		}
+	}
+
+	return false;
+}
+
+void FNetObjectGroups::RemoveGroupMembership(FNetObjectGroupMembership& Target, FNetObjectGroupHandle Group)
+{
+	const uint32 LastIndex = FNetObjectGroupMembership::MaxAssignedGroupCount - 1U;
+	for (uint32 It = 0U; It < FNetObjectGroupMembership::MaxAssignedGroupCount; ++It)
+	{
+		if (Target.Groups[It] == Group)
+		{
+			FPlatformMemory::Memmove(&Target.Groups[It], &Target.Groups[It + 1], (LastIndex - It) * sizeof(FNetObjectGroupHandle));
+			Target.Groups[LastIndex] = InvalidNetObjectGroupHandle;
+
+			return;
+		}
+	}
+}
+
+void FNetObjectGroups::ResetGroupMembership(FNetObjectGroupMembership& Target)
+{
+	FMemory::Memzero(Target);
+}
+
+void FNetObjectGroups::Init(const FNetObjectGroupInitParams& Params)
+{
+	MaxGroupCount = Params.MaxGroupCount;
+
+	// Reserve first as invalid group
+	Groups.Add(FNetObjectGroup());
+
+	GroupMemberships.SetNumZeroed(Params.MaxObjectCount);
+}
+
+FNetObjectGroupHandle FNetObjectGroups::CreateGroup()
+{
+	if (ensure((uint32)Groups.Num() < MaxGroupCount))
+	{
+		const int Index = Groups.Add(FNetObjectGroup());
+		return FNetObjectGroupHandle(Index);
+	}
+	else
+	{
+		return InvalidNetObjectGroupHandle;
+	}
+}
+
+void FNetObjectGroups::DestroyGroup(FNetObjectGroupHandle GroupHandle)
+{
+	if (GroupHandle != InvalidNetObjectGroupHandle && Groups.IsValidIndex(GroupHandle))
+	{
+		ClearGroup(GroupHandle);
+
+		FNetObjectGroup& Group = Groups[GroupHandle];
+
+		if (EnumHasAnyFlags(Group.Traits, ENetObjectGroupTraits::IsFindableByName))
+		{
+			NamedGroups.FindAndRemoveChecked(Group.GroupName);
+		}
+
+		Groups.RemoveAt(GroupHandle);
+	}
+}
+
+void FNetObjectGroups::SetGroupName(FNetObjectGroupHandle GroupHandle, FName GroupName)
+{
+	if (FNetObjectGroup* Group = GetGroup(GroupHandle))
+	{
+		if (ensureAlwaysMsgf(EnumHasAnyFlags(Group->Traits, ENetObjectGroupTraits::IsFindableByName), TEXT("FNetObjectGroups::SetGroupName Cannot SetGroupName for grouphandle %u as it is FindableByName"), GroupHandle))
+		{
+			Group->GroupName = GroupName;
+		}
+	}
+}
+
+FNetObjectGroupHandle FNetObjectGroups::CreateNamedGroup(FName GroupName)
+{
+	if (NamedGroups.Contains(GroupName))
+	{
+		ensureAlwaysMsgf(false, TEXT("FNetObjectGroups, trying to create named group %s that already exists"), *GroupName.ToString());
+		return InvalidNetObjectGroupHandle;
+	}
+
+	FNetObjectGroupHandle GroupHandle = CreateGroup();
+	if (FNetObjectGroup* Group = GetGroup(GroupHandle))
+	{
+		Group->GroupName = GroupName;
+		Group->Traits |= ENetObjectGroupTraits::IsFindableByName;
+		NamedGroups.Add(GroupName, GroupHandle);
+	}
+
+	return GroupHandle;
+}
+
+FNetObjectGroupHandle FNetObjectGroups::GetNamedGroupHandle(FName GroupName)
+{
+	if (FNetObjectGroupHandle* Group = NamedGroups.Find(GroupName))
+	{
+		return *Group;
+	}
+
+	return InvalidNetObjectGroupHandle;
+}
+
+void FNetObjectGroups::DestroyNamedGroup(FName GroupName)
+{
+	DestroyGroup(GetNamedGroupHandle(GroupName));
+}
+
+void FNetObjectGroups::ClearGroup(FNetObjectGroupHandle GroupHandle)
+{
+	if (GroupHandle != InvalidNetObjectGroupHandle && Groups.IsValidIndex(GroupHandle))
+	{
+		FNetObjectGroup& Group = Groups[GroupHandle];
+
+		for (uint32 InternalIndex : MakeArrayView(Group.Members))
+		{
+			checkSlow(IsMemberOf(GroupMemberships[InternalIndex], GroupHandle));
+			RemoveGroupMembership(GroupMemberships[InternalIndex], GroupHandle);
+		}
+
+		Group.Members.Empty();
+	}
+}
+
+const FNetObjectGroup* FNetObjectGroups::GetGroup(FNetObjectGroupHandle GroupHandle) const
+{
+	return Groups.IsValidIndex(GroupHandle) ? &Groups[GroupHandle] : nullptr;
+}
+
+FNetObjectGroup* FNetObjectGroups::GetGroup(FNetObjectGroupHandle GroupHandle)
+{
+	return Groups.IsValidIndex(GroupHandle) ? &Groups[GroupHandle] : nullptr;
+}
+
+bool FNetObjectGroups::Contains(FNetObjectGroupHandle GroupHandle, FInternalNetHandle InternalIndex) const
+{
+	// If the group does not exist we cannot be in it..
+	const FNetObjectGroup* Group = GetGroup(GroupHandle);
+	if (!Group)
+	{
+		return false;
+	}
+
+	checkSlow(InternalIndex == 0U || (Group->Members.Contains(InternalIndex) == IsMemberOf(GroupMemberships[InternalIndex], GroupHandle)));
+
+	return InternalIndex ? IsMemberOf(GroupMemberships[InternalIndex], GroupHandle) : false;
+}
+
+void FNetObjectGroups::AddToGroup(FNetObjectGroupHandle GroupHandle, FInternalNetHandle InternalIndex)
+{
+	FNetObjectGroup* Group = GetGroup(GroupHandle);
+	if (InternalIndex != FNetHandleManager::InvalidInternalIndex && Group)
+	{
+		if (AddGroupMembership(GroupMemberships[InternalIndex], GroupHandle))
+		{
+			Group->Members.AddUnique(InternalIndex);
+		}
+		else
+		{
+			checkf(false, TEXT("FNetObjectGroups::AddToGroup, Failed to add ( InternalIndex: %u ) to (GroupIndex: %u) A NetObject can only be a member of %u groups."), InternalIndex, GroupHandle, FNetObjectGroupMembership::MaxAssignedGroupCount);
+		}
+	}
+}
+
+void FNetObjectGroups::RemoveFromGroup(FNetObjectGroupHandle GroupHandle, FInternalNetHandle InternalIndex)
+{
+	FNetObjectGroup* Group = GetGroup(GroupHandle);
+	if (InternalIndex != FNetHandleManager::InvalidInternalIndex && Group)
+	{
+		checkSlow(IsMemberOf(GroupMemberships[InternalIndex], GroupHandle));
+
+		RemoveGroupMembership(GroupMemberships[InternalIndex], GroupHandle);
+		Group->Members.RemoveSingle(InternalIndex);
+	}
+}
+
+uint32 FNetObjectGroups::GetNumGroupMemberships(FInternalNetHandle InternalIndex) const
+{
+	if (InternalIndex >= (uint32)GroupMemberships.Num())
+	{
+		return 0U;
+	}
+
+	uint32 Count = 0U;
+	for (uint32 It=0; It < FNetObjectGroupMembership::MaxAssignedGroupCount; ++It)
+	{
+		Count += GroupMemberships[InternalIndex].Groups[It] != 0U;
+	}
+	
+	return Count;
+}
+
+const FNetObjectGroupHandle* FNetObjectGroups::GetGroupMemberships(FInternalNetHandle InternalIndex, uint32& GroupCount) const
+{
+	GroupCount = GetNumGroupMemberships(InternalIndex);
+	if (GroupCount)
+	{
+		return GroupMemberships[InternalIndex].Groups;
+	}
+	else
+	{
+		return nullptr;
+	}
+}
+
+}
diff --git a/Engine/Source/Runtime/Experimental/Iris/Core/Private/Iris/ReplicationSystem/Filtering/NetObjectGroups.h b/Engine/Source/Runtime/Experimental/Iris/Core/Private/Iris/ReplicationSystem/Filtering/NetObjectGroups.h
new file mode 100644
--- /dev/null
+++ b/Engine/Source/Runtime/Experimental/Iris/Core/Private/Iris/ReplicationSystem/Filtering/NetObjectGroups.h
@@ -0,0 +1,101 @@
+// Copyright Epic Games, Inc. All Rights Reserved.
+
+#pragma once
+
+#include "CoreTypes.h"
+#include "Containers/Array.h"
+#include "Containers/SparseArray.h"
+
+namespace UE::Net
+{
+	typedef uint16 FNetObjectGroupHandle;
+	namespace Private
+	{
+		typedef uint32 FInternalNetHandle;
+	}
+}
+
+namespace UE::Net::Private
+{
+
+enum class ENetObjectGroupTraits : uint32
+{
+	None = 0U,
+	IsFindableByName = 1U << 0U,
+};
+ENUM_CLASS_FLAGS(ENetObjectGroupTraits);
+
+struct FNetObjectGroup
+{
+	FNetObjectGroup() : Traits(ENetObjectGroupTraits::None) {}
+
+	// Group members can only be replicated objects that have internal indices
+	TArray<uint32> Members;
+	FName GroupName;
+	ENetObjectGroupTraits Traits;
+};
+
+struct FNetObjectGroupInitParams
+{
+	uint32 MaxObjectCount;
+	uint32 MaxGroupCount;
+};
+
+class FNetObjectGroups
+{
+	UE_NONCOPYABLE(FNetObjectGroups)
+
+public:
+	FNetObjectGroups();
+	~FNetObjectGroups();
+
+	void Init(const FNetObjectGroupInitParams& Params);
+
+	FNetObjectGroupHandle CreateGroup();
+	void DestroyGroup(FNetObjectGroupHandle GroupHandle);
+	void ClearGroup(FNetObjectGroupHandle GroupHandle);
+	
+	const FNetObjectGroup* GetGroup(FNetObjectGroupHandle GroupHandle) const;
+	FNetObjectGroup* GetGroup(FNetObjectGroupHandle GroupHandle);
+	void SetGroupName(FNetObjectGroupHandle GroupHandle, FName GroupName);
+
+	inline bool IsValidGroup(FNetObjectGroupHandle GroupHandle) const { return GroupHandle && Groups.IsValidIndex(GroupHandle); }
+
+	bool Contains(FNetObjectGroupHandle GroupHandle, FInternalNetHandle InternalIndex) const;
+	void AddToGroup(FNetObjectGroupHandle GroupHandle, FInternalNetHandle InternalIndex);
+	void RemoveFromGroup(FNetObjectGroupHandle GroupHandle, FInternalNetHandle InternalIndex);
+
+	// Returns how many groups the given handle is a member of
+	uint32 GetNumGroupMemberships(FInternalNetHandle InternalIndex) const;
+	const FNetObjectGroupHandle* GetGroupMemberships(FInternalNetHandle InternalIndex, uint32& GroupCount) const;
+
+	// Create and manage named groups, only groups created as a named group will be findable by name
+	FNetObjectGroupHandle CreateNamedGroup(FName GroupName);
+	// Lookup NetObjectGroupHandle for a named group
+	FNetObjectGroupHandle GetNamedGroupHandle(FName GroupName);
+	// Destroy Named group
+	void DestroyNamedGroup(FName GroupName);
+
+private:
+	struct FNetObjectGroupMembership
+	{
+		enum { MaxAssignedGroupCount = 4 };
+		FNetObjectGroupHandle Groups[MaxAssignedGroupCount];		
+	};
+
+	static bool AddGroupMembership(FNetObjectGroupMembership& Target, FNetObjectGroupHandle Group);
+	static void RemoveGroupMembership(FNetObjectGroupMembership& Target, FNetObjectGroupHandle Group);
+	static void ResetGroupMembership(FNetObjectGroupMembership& Target);
+	static bool IsMemberOf(const FNetObjectGroupMembership& Target, FNetObjectGroupHandle Group);
+
+	// Group usage pattern should not be high frequency so memory layout should not be a major concern
+	TSparseArray<FNetObjectGroup> Groups;
+
+	// Track what groups each internal handle is a member of, we can tighten this up a bit if needed
+	TArray<FNetObjectGroupMembership> GroupMemberships;
+	uint32 MaxGroupCount;
+
+	TMap<FName, FNetObjectGroupHandle> NamedGroups;
+};
+
+}
diff --git a/Engine/Source/Runtime/Experimental/Iris/Core/Private/Iris/ReplicationSystem/Filtering/NopNetObjectFilter.cpp b/Engine/Source/Runtime/Experimental/Iris/Core/Private/Iris/ReplicationSystem/Filtering/NopNetObjectFilter.cpp
new file mode 100644
--- /dev/null
+++ b/Engine/Source/Runtime/Experimental/Iris/Core/Private/Iris/ReplicationSystem/Filtering/NopNetObjectFilter.cpp
@@ -0,0 +1,26 @@
+// Copyright Epic Games, Inc. All Rights Reserved.
+
+#include "Iris/ReplicationSystem/Filtering/NopNetObjectFilter.h"
+
+void UNopNetObjectFilter::Init(FNetObjectFilterInitParams& Params)
+{
+}
+
+bool UNopNetObjectFilter::AddObject(uint32 ObjectIndex, FNetObjectFilterAddObjectParams& Params)
+{
+	return true;
+}
+
+void UNopNetObjectFilter::RemoveObject(uint32 ObjectIndex, const FNetObjectFilteringInfo& Info)
+{
+}
+
+void UNopNetObjectFilter::UpdateObjects(FNetObjectFilterUpdateParams& Params)
+{
+}
+
+void UNopNetObjectFilter::Filter(FNetObjectFilteringParams& Params)
+{
+	// Allow all objects to be replicated, just as if no dynamic filtering was applied.
+	Params.OutAllowedObjects.SetAllBits();
+}
