Iris experimental replication system.

diff --git a/Engine/Source/Runtime/Experimental/Iris/Core/Private/Iris/ReplicationSystem/Filtering/ReplicationFiltering.cpp b/Engine/Source/Runtime/Experimental/Iris/Core/Private/Iris/ReplicationSystem/Filtering/ReplicationFiltering.cpp
new file mode 100644
--- /dev/null
+++ b/Engine/Source/Runtime/Experimental/Iris/Core/Private/Iris/ReplicationSystem/Filtering/ReplicationFiltering.cpp
@@ -0,0 +1,1940 @@
+// Copyright Epic Games, Inc. All Rights Reserved.
+
+#include "ReplicationFiltering.h"
+#include "HAL/PlatformMath.h"
+#include "Iris/IrisConfigInternal.h"
+#include "Iris/IrisConstants.h"
+#include "Iris/Core/IrisLog.h"
+#include "Iris/Core/IrisProfiler.h"
+#include "Iris/ReplicationSystem/NetHandleManager.h"
+#include "Iris/ReplicationSystem/ReplicationConnections.h"
+#include "Iris/ReplicationSystem/ReplicationProtocol.h"
+#include "Iris/ReplicationSystem/ReplicationSystem.h"
+#include "Iris/ReplicationSystem/DeltaCompression/DeltaCompressionBaselineInvalidationTracker.h"
+#include "Iris/ReplicationSystem/Filtering/NetObjectGroups.h"
+#include "Iris/ReplicationSystem/Filtering/NetObjectFilterDefinitions.h"
+
+#include <limits>
+
+namespace UE::Net::Private
+{
+
+inline static ENetFilterStatus GetDependentObjectFilterStatus(const FNetHandleManager* NetHandleManager, const FNetBitArray& ObjectsInScope, FInternalNetHandle ObjectIndex)
+{
+	for (const FInternalNetHandle ParentObjectIndex : NetHandleManager->GetDependentObjectParents(ObjectIndex))
+	{
+		if (GetDependentObjectFilterStatus(NetHandleManager, ObjectsInScope, ParentObjectIndex) == ENetFilterStatus::Allow)
+		{
+			return ENetFilterStatus::Allow;
+		}
+	}
+	return ObjectsInScope.GetBit(ObjectIndex) ? ENetFilterStatus::Allow : ENetFilterStatus::Disallow;
+}
+
+class FNetObjectFilterHandleUtil
+{
+public:
+	static bool IsInvalidHandle(FNetObjectFilterHandle Handle);
+	static bool IsDynamicFilter(FNetObjectFilterHandle Handle);
+	static bool IsStaticFilter(FNetObjectFilterHandle Handle);
+
+	static FNetObjectFilterHandle MakeDynamicFilterHandle(uint32 FilterIndex);
+	static uint32 GetDynamicFilterIndex(FNetObjectFilterHandle Handle);
+
+private:
+	// Most significant bit in the filter handle acts as a dynamic/static filter classifier.
+	static constexpr FNetObjectFilterHandle DynamicNetObjectFilterHandleFlag = 1U << (sizeof(FNetObjectFilterHandle)*8U - 1U);
+};
+
+class FReplicationFiltering::FUpdateDirtyObjectsBatchHelper
+{
+public:
+	enum Constants : uint32
+	{
+		MaxObjectCountPerBatch = 512U,
+	};
+
+	struct FPerFilterInfo
+	{
+		uint32* ObjectIndices;
+		uint32 ObjectCount;
+
+		FReplicationInstanceProtocol const** InstanceProtocols;
+	};
+
+	FUpdateDirtyObjectsBatchHelper(const FNetHandleManager* InNetHandleManager, uint32 FilterCount)
+		: NetHandleManager(InNetHandleManager)
+	{
+		PerFilterInfos.SetNumUninitialized(FilterCount);
+		ObjectIndicesStorage.SetNumUninitialized(FilterCount * MaxObjectCountPerBatch);
+		InstanceProtocolsStorage.SetNumUninitialized(FilterCount * MaxObjectCountPerBatch);
+
+		uint32 FilterIndex = 0;
+		for (FPerFilterInfo& PerFilterInfo : PerFilterInfos)
+		{
+			PerFilterInfo.ObjectIndices = ObjectIndicesStorage.GetData() + FilterIndex * MaxObjectCountPerBatch;
+			PerFilterInfo.InstanceProtocols = InstanceProtocolsStorage.GetData() + FilterIndex * MaxObjectCountPerBatch;
+			++FilterIndex;
+		}
+	}
+
+	void PrepareBatch(const uint32* ObjectIndices, uint32 ObjectCount, const uint8* FilterIndices)
+	{
+		ResetBatch();
+
+		FPerFilterInfo* PerFilterInfosData = PerFilterInfos.GetData();
+		for (const uint32 ObjectIndex : MakeArrayView(ObjectIndices, ObjectCount))
+		{
+			const uint8 FilterIndex = FilterIndices[ObjectIndex];
+			if (FilterIndex == InvalidDynamicFilterIndex)
+			{
+				continue;
+			}
+
+			if (const FReplicationInstanceProtocol* InstanceProtocol = NetHandleManager->GetReplicatedObjectDataNoCheck(ObjectIndex).InstanceProtocol)
+			{
+				FPerFilterInfo& PerFilterInfo = PerFilterInfosData[FilterIndex];
+				PerFilterInfo.ObjectIndices[PerFilterInfo.ObjectCount] = ObjectIndex;
+				PerFilterInfo.InstanceProtocols[PerFilterInfo.ObjectCount] = InstanceProtocol;
+				++PerFilterInfo.ObjectCount;
+			}
+		}
+	}
+
+	TArray<FPerFilterInfo, TInlineAllocator<16>> PerFilterInfos;
+
+private:
+	void ResetBatch()
+	{
+		for (FPerFilterInfo& PerFilterInfo : PerFilterInfos)
+		{
+			PerFilterInfo.ObjectCount = 0U;
+		}
+	}
+
+	TArray<uint32> ObjectIndicesStorage;
+	TArray<const FReplicationInstanceProtocol*> InstanceProtocolsStorage;
+	const FNetHandleManager* NetHandleManager;
+};
+
+DEFINE_LOG_CATEGORY_STATIC(LogIrisFiltering, Log, All);
+
+FReplicationFiltering::FReplicationFiltering()
+: bHasNewConnection(0)
+, bHasRemovedConnection(0)
+, bHasDirtyConnectionFilter(0)
+, bHasDirtyOwner(0)
+{
+	StaticChecks();
+}
+
+void FReplicationFiltering::StaticChecks()
+{
+	static_assert(TNumericLimits<PerObjectInfoIndexType>::Max() % (UsedPerObjectInfoStorageGrowSize*32U) == UsedPerObjectInfoStorageGrowSize*32U - 1, "Bit array grow code expects not to be able to return an out of bound index.");
+	static_assert(sizeof(FNetBitArrayBase::StorageWordType) == sizeof(uint32), "Expected FNetBitArrayBase::StorageWordType to be four bytes in size."); 
+}
+
+void FReplicationFiltering::Init(FReplicationFilteringInitParams& Params)
+{
+	check(Params.Connections != nullptr);
+	check(Params.Connections->GetMaxConnectionCount() <= std::numeric_limits<decltype(ObjectIndexToOwningConnection)::ElementType>::max());
+
+	ReplicationSystem = Params.ReplicationSystem;
+
+	Connections = Params.Connections;
+	NetHandleManager = Params.NetHandleManager;
+	BaselineInvalidationTracker = Params.BaselineInvalidationTracker;
+	Groups = Params.Groups;
+
+	MaxObjectCount = Params.MaxObjectCount;
+	WordCountForObjectBitArrays = Align(FPlatformMath::Max(MaxObjectCount, 1U), sizeof(FNetBitArrayBase::StorageWordType)*8U)/(sizeof(FNetBitArrayBase::StorageWordType)*8U);
+
+	// Connection specifics
+	ConnectionInfos.SetNum(Params.Connections->GetMaxConnectionCount() + 1U);
+	ValidConnections.Init(ConnectionInfos.Num());
+	NewConnections.Init(ConnectionInfos.Num());
+
+	// Filter specifics
+	ObjectsWithDirtyConnectionFilter.Init(MaxObjectCount);
+	ObjectsWithDirtyOwner.Init(MaxObjectCount);
+	ObjectsWithOwnerFilter.Init(MaxObjectCount);
+	ObjectsWithPerObjectInfo.Init(MaxObjectCount);
+
+	// Group filtering
+	MaxGroupCount = Params.MaxGroupCount;
+	GroupInfos.SetNumZeroed(MaxGroupCount);
+	FilterGroups.Init(GroupInfos.Num());
+	DirtyFilterGroups.Init(GroupInfos.Num());
+	// SubObjectFilter groups
+	SubObjectFilterGroups.Init(GroupInfos.Num());
+	DirtySubObjectFilterGroups.Init(GroupInfos.Num());
+
+	ObjectIndexToPerObjectInfoIndex.SetNumZeroed(MaxObjectCount);
+
+	PerObjectInfoStorageCountForConnections = Align(FPlatformMath::Max(uint32(ConnectionInfos.Num()), 1U), 32U)/32U;
+	PerObjectInfoStorageCountPerItem = sizeof(FPerObjectInfo) + PerObjectInfoStorageCountForConnections - 1U;
+
+	// Owning connections
+	ObjectIndexToOwningConnection.SetNumZeroed(MaxObjectCount);
+
+	// Dynamic filters
+	{
+		NetObjectFilteringInfos.SetNumUninitialized(Params.MaxObjectCount);
+
+		ObjectIndexToDynamicFilterIndex.SetNumUninitialized(Params.MaxObjectCount);
+		FMemory::Memset(ObjectIndexToDynamicFilterIndex.GetData(), InvalidDynamicFilterIndex, Params.MaxObjectCount*sizeof(decltype(ObjectIndexToDynamicFilterIndex)::ElementType));
+
+		ObjectsRequiringDynamicFilterUpdate.Init(MaxObjectCount);
+	}
+
+	InitFilters();
+}
+
+void FReplicationFiltering::Filter(const FNetBitArrayView& DirtyObjects)
+{
+	IRIS_PROFILER_SCOPE(FReplicationFiltering_Filter);
+
+	ResetRemovedConnections();
+
+	InitNewConnections();
+
+	UpdateObjectsInScope();
+
+	// Update group filters if needed
+	UpdateGroupFiltering();
+
+	UpdateOwnerAndConnectionFiltering();
+
+	UpdateSubObjectFilters();
+
+	/**
+	 * This last filtering allows users to filter out objects based on arbitrary criteria.
+	 */
+	if (HasDynamicFilters())
+	{
+		NotifyFiltersOfDirtyObjects(DirtyObjects);
+
+		PreUpdateDynamicFiltering();
+		UpdateDynamicFiltering();
+		PostUpdateDynamicFiltering();
+	}
+}
+
+FNetBitArrayView FReplicationFiltering::GetObjectsInScope(uint32 ConnectionId) const
+{
+	const FPerConnectionInfo& ConnectionInfo = ConnectionInfos[ConnectionId];
+	if (HasDynamicFilters())
+	{
+		return MakeNetBitArrayView(ConnectionInfo.ObjectsInScope);
+	}
+	else
+	{
+		return MakeNetBitArrayView(ConnectionInfo.ObjectsInScopeBeforeDynamicFiltering);
+	}
+}
+
+FNetBitArrayView FReplicationFiltering::GetGroupFilteredOutObjects(uint32 ConnectionId) const
+{
+	const FPerConnectionInfo& ConnectionInfo = ConnectionInfos[ConnectionId];
+	return MakeNetBitArrayView(ConnectionInfo.GroupFilteredOutObjects);
+}
+
+void FReplicationFiltering::SetOwningConnection(uint32 ObjectIndex, uint32 ConnectionId)
+{
+	// We allow valid connections as well as 0, which would prevent the object from being replicated to anyone.
+	if ((ConnectionId != 0U) && !Connections->IsValidConnection(ConnectionId))
+	{
+		UE_LOG(LogIrisFiltering, Warning, TEXT("Trying to set unknown owning connection on object %u. Connection: %u"), ObjectIndex, ConnectionId);
+		return;
+	}
+
+	const uint16 OldConnectionId = ObjectIndexToOwningConnection[ObjectIndex];
+	ObjectIndexToOwningConnection[ObjectIndex] = ConnectionId;
+	if (ConnectionId != OldConnectionId)
+	{
+		bHasDirtyOwner = 1;
+		ObjectsWithDirtyOwner.SetBit(ObjectIndex);
+		if (HasOwnerFilter(ObjectIndex))
+		{
+			bHasDirtyConnectionFilter = 1;
+			ObjectsWithDirtyConnectionFilter.SetBit(ObjectIndex);
+		}
+
+		InvalidateBaselinesForObject(ObjectIndex, ConnectionId, OldConnectionId);
+	}
+}
+
+bool FReplicationFiltering::SetFilter(uint32 ObjectIndex, FNetObjectFilterHandle Filter)
+{
+	if (Filter == ConnectionFilterHandle)
+	{
+		ensureAlwaysMsgf(false, TEXT("Use SetConnectionFilter to enable connection filtering of objects. Cause of ensure must be fixed!"));
+		return false;
+	}
+
+	const bool bWantsToUseDynamicFilter = FNetObjectFilterHandleUtil::IsDynamicFilter(Filter);
+	const uint8 OldDynamicFilterIndex = ObjectIndexToDynamicFilterIndex[ObjectIndex];
+	const uint32 NewDynamicFilterIndex = bWantsToUseDynamicFilter ? FNetObjectFilterHandleUtil::GetDynamicFilterIndex(Filter) : InvalidDynamicFilterIndex;
+	const bool bWasUsingDynamicFilter = OldDynamicFilterIndex != InvalidDynamicFilterIndex;
+
+	// Validate the filter
+	if (bWantsToUseDynamicFilter && (NewDynamicFilterIndex >= (uint32)DynamicFilterInfos.Num()))
+	{
+		ensureAlwaysMsgf(false, TEXT("Invalid dynamic filter 0x%08X. Filter is not being changed. Cause of ensure must be fixed!"), NewDynamicFilterIndex);
+		return false;
+	}
+	else if (!bWantsToUseDynamicFilter && (Filter != InvalidNetObjectFilterHandle) && (Filter != ToOwnerFilterHandle))
+	{
+		ensureAlwaysMsgf(false, TEXT("Invalid static filter 0x%08X. Filter is not being changed. Cause of ensure must be fixed!"), Filter);
+		return false;
+	}
+
+	const FNetHandleManager::FReplicatedObjectData& ObjectData = NetHandleManager->GetReplicatedObjectDataNoCheck(ObjectIndex);
+	const TArray<uint8*>& ReplicatedObjectsStateBuffers = NetHandleManager->GetReplicatedObjectStateBuffers();
+	// Let subobjects be filtered like their owners.
+	if (bWantsToUseDynamicFilter && (ObjectData.SubObjectRootIndex != FNetHandleManager::InvalidInternalIndex))
+	{
+		UE_LOG(LogIrisFiltering, Warning, TEXT("Ignoring request to use dynamic filter on object %s due to it being a subobject. Any filtering already successfully applied remains as is."), ToCStr(ObjectData.Handle.ToString()));
+		return false;
+	}
+
+	const FNetBitArray& ObjectsInScope = NetHandleManager->GetScopableInternalIndices();
+	auto&& ClearAndDirtyStaticFilter = [this](uint32 ObjIndex)
+	{
+		this->bHasDirtyConnectionFilter = 1;
+
+		this->ObjectsWithOwnerFilter.ClearBit(ObjIndex);
+		this->ObjectsWithDirtyConnectionFilter.SetBit(ObjIndex);
+		this->FreePerObjectInfoForObject(ObjIndex);
+	};
+
+	auto&& TrySetDynamicFilter = [this, &ObjectData, &ReplicatedObjectsStateBuffers](uint32 ObjIndex, uint32 FilterIndex)
+	{
+		FNetObjectFilteringInfo& NetObjectFilteringInfo = this->NetObjectFilteringInfos[ObjIndex];
+		NetObjectFilteringInfo = {};
+		FNetObjectFilterAddObjectParams AddParams = { NetObjectFilteringInfo, ObjectData.InstanceProtocol, ObjectData.Protocol, ReplicatedObjectsStateBuffers.GetData()[ObjIndex] };
+		FFilterInfo& FilterInfo = this->DynamicFilterInfos[FilterIndex];
+		if (FilterInfo.Filter->AddObject(ObjIndex, AddParams))
+		{
+			++FilterInfo.ObjectCount;
+			FilterInfo.FilteredObjects.SetBit(ObjIndex);
+			this->ObjectIndexToDynamicFilterIndex[ObjIndex] = FilterIndex;
+			this->DynamicFilterEnabledObjects.SetBit(ObjIndex);
+			return true;
+		}
+		
+		return false;
+	};
+
+	// Clear dynamic filter info if needed
+	if (bWasUsingDynamicFilter)
+	{
+		RemoveFromDynamicFilter(ObjectIndex, OldDynamicFilterIndex);
+	}
+	// Clear static filter info if needed
+	else
+	{
+		ClearAndDirtyStaticFilter(ObjectIndex);
+	}
+
+	// Set dynamic filter
+	if (bWantsToUseDynamicFilter)
+	{
+		const bool bSuccess = TrySetDynamicFilter(ObjectIndex, NewDynamicFilterIndex);
+		if (bSuccess)
+		{
+			return true;
+		}
+		else
+		{
+			UE_LOG(LogIrisFiltering, Verbose, TEXT("Filter '%s' does not support object %u."), ToCStr(DynamicFilterInfos[NewDynamicFilterIndex].Filter->GetFName().GetPlainNameString()), ObjectIndex);
+			return false;
+		}
+	}
+	// Set static filter
+	else
+	{
+		if (Filter == InvalidNetObjectFilterHandle)
+		{
+			return true;
+		}
+		else if (Filter == ToOwnerFilterHandle)
+		{
+			ObjectsWithOwnerFilter.SetBit(ObjectIndex);
+			return true;
+		}
+
+		// The validity of the filter has already been verified. How did we end up here?
+		check(false);
+	}
+
+	// Unknown filter
+	return false;
+}
+
+bool FReplicationFiltering::SetConnectionFilter(uint32 ObjectIndex, const FNetBitArrayView& ConnectionIndices, ENetFilterStatus ReplicationStatus)
+{
+	// If we have some sort of other filtering set then remove that. From now on we want to use connection filtering.
+	if (HasOwnerFilter(ObjectIndex))
+	{
+		ObjectsWithOwnerFilter.ClearBit(ObjectIndex);
+	}
+
+	const uint8 DynamicFilterIndex = ObjectIndexToDynamicFilterIndex[ObjectIndex];
+	if (DynamicFilterIndex != InvalidDynamicFilterIndex)
+	{
+		RemoveFromDynamicFilter(ObjectIndex, DynamicFilterIndex);
+	}
+
+	bHasDirtyConnectionFilter = 1;
+	ObjectsWithDirtyConnectionFilter.SetBit(ObjectIndex);
+
+	PerObjectInfoIndexType ObjectInfoIndex = ObjectIndexToPerObjectInfoIndex.GetData()[ObjectIndex];
+	FPerObjectInfo* ObjectInfo = nullptr;
+	if (ObjectInfoIndex == 0)
+	{
+		ObjectInfo = AllocPerObjectInfoForObject(ObjectIndex);
+	}
+	else
+	{
+		ObjectInfo = GetPerObjectInfo(ObjectInfoIndex);
+	}
+
+	{
+		// The ConnectionIndices bit array does not necessarily have to have storage for the max amount of connections
+		const uint32 InWordCount = ConnectionIndices.GetNumWords();
+		const uint32 MaxWordCount = PerObjectInfoStorageCountForConnections;
+		const uint32* InConnectionIndicesData = ConnectionIndices.GetData();
+		const uint32 WordMask = (ReplicationStatus == ENetFilterStatus::Allow ? 0U : ~0U);
+		for (uint32 WordIt = 0, WordEndIt = FPlatformMath::Min(InWordCount, MaxWordCount); WordIt < WordEndIt; ++WordIt)
+		{
+			ObjectInfo->ConnectionIds[WordIt] = InConnectionIndicesData[WordIt] ^ WordMask;
+		}
+
+		// Fill in connections that weren't passed in ConnectionIndices
+		for (uint32 WordIt = InWordCount; WordIt < MaxWordCount; ++WordIt)
+		{
+			ObjectInfo->ConnectionIds[WordIt] = WordMask;
+		}
+	}
+
+	return true;
+}
+
+FNetObjectFilterHandle FReplicationFiltering::GetFilterHandle(const FName FilterName) const
+{
+	for (const FFilterInfo& Info : DynamicFilterInfos)
+	{
+		if (Info.Name == FilterName)
+		{
+			return FNetObjectFilterHandleUtil::MakeDynamicFilterHandle(&Info - DynamicFilterInfos.GetData());
+		}
+	}
+
+	return InvalidNetObjectFilterHandle;
+}
+
+UNetObjectFilter* FReplicationFiltering::GetFilter(const FName FilterName) const
+{
+	for (const FFilterInfo& Info : DynamicFilterInfos)
+	{
+		if (Info.Name == FilterName)
+		{
+			return Info.Filter.Get();
+		}
+	}
+
+	return nullptr;
+}
+
+// Connection handling
+void FReplicationFiltering::AddConnection(uint32 ConnectionId)
+{
+	bHasNewConnection = 1;
+	ValidConnections.SetBit(ConnectionId);
+	NewConnections.SetBit(ConnectionId);
+
+	UE_LOG(LogIrisFiltering, Verbose, TEXT("ReplicationFiltering::AddConnection ConnectionId: %u"), ConnectionId);
+
+	for (FFilterInfo& Info : DynamicFilterInfos)
+	{
+		Info.Filter->AddConnection(ConnectionId);
+	}
+
+	// We defer needed processing to InitNewConnections().
+}
+
+void FReplicationFiltering::RemoveConnection(uint32 ConnectionId)
+{
+	bHasRemovedConnection = 1U;
+	ValidConnections.ClearBit(ConnectionId);
+	// If for whatever reason this connection is removed before we've done any new connection processing.
+	NewConnections.ClearBit(ConnectionId);
+
+	UE_LOG(LogIrisFiltering, Verbose, TEXT("ReplicationFiltering::RemoveConnection ConnectionId: %u"), ConnectionId);
+	
+	// Reset connection info
+	FPerConnectionInfo& ConnectionInfo = ConnectionInfos[ConnectionId];
+	ConnectionInfo.ConnectionFilteredObjects.Empty();
+	ConnectionInfo.GroupFilteredOutObjects.Empty();
+	ConnectionInfo.ObjectsInScopeBeforeDynamicFiltering.Empty();
+	ConnectionInfo.ObjectsInScope.Empty();
+
+	for (FFilterInfo& Info : DynamicFilterInfos)
+	{
+		Info.Filter->RemoveConnection(ConnectionId);
+	}
+
+	// Reset SubObject filter for removed connection
+	SubObjectFilterGroups.ForAllSetBits([this, ConnectionId](uint32 GroupIndex)
+	{
+		FNetObjectGroupHandle GroupHandle(GroupIndex);
+		if (!IsReservedNetObjectGroupHandle(GroupHandle))
+		{
+			SetConnectionFilterStatus(*GetPerObjectInfo(GroupInfos[GroupHandle].ConnectionStateIndex), ConnectionId, ENetFilterStatus::Disallow);
+			DirtySubObjectFilterGroups.SetBit(GroupIndex);
+		}
+	});
+}
+
+void FReplicationFiltering::InitNewConnections()
+{
+	if (!bHasNewConnection)
+	{
+		return;
+	}
+
+	IRIS_PROFILER_SCOPE(FReplicationFiltering_InitNewConnections);
+
+	bHasNewConnection = 0;
+	
+	auto&& InitNewConnection = [this](uint32 ConnectionId)
+	{
+		// Copy default scope
+		const FNetBitArray& ScopableInternalIndices = NetHandleManager->GetScopableInternalIndices();
+		FPerConnectionInfo& ConnectionInfo = this->ConnectionInfos[ConnectionId];
+
+		ConnectionInfo.ConnectionFilteredObjects.Init(MaxObjectCount);
+		ConnectionInfo.ConnectionFilteredObjects.Copy(ScopableInternalIndices);
+		ConnectionInfo.ConnectionFilteredObjects.ClearBit(FNetHandleManager::InvalidInternalIndex);
+
+		// Everything enabled by default in groups, i.e. not filtered out.
+		ConnectionInfo.GroupFilteredOutObjects.Init(MaxObjectCount);
+
+		// The combined result of scoped objects, connection filteirng and group filtering.
+		ConnectionInfo.ObjectsInScopeBeforeDynamicFiltering.Init(MaxObjectCount);
+
+		// The final result of all filtering.
+		ConnectionInfo.ObjectsInScope.Init(MaxObjectCount);
+
+		if (HasDynamicFilters())
+		{
+			ConnectionInfo.DynamicFilteredOutObjects.Init(MaxObjectCount);
+		}
+
+		// Update group filter by iterating over all registered groups for the connection
+		{
+			auto&& InitGroupFilterForConnection = [this, &ConnectionInfo, ConnectionId](uint32 GroupIndex)
+			{
+				const FNetObjectGroupHandle GroupHandle = GroupIndex;
+				const FPerObjectInfo* ConnectionState = GetPerObjectInfo(this->GroupInfos[GroupIndex].ConnectionStateIndex);
+
+				// Setup filter for connection
+				if (GetConnectionFilterStatus(*ConnectionState, ConnectionId) == ENetFilterStatus::Disallow)
+				{
+					// Apply filter
+					const FNetObjectGroup* Group = this->Groups->GetGroup(GroupHandle);
+					FNetBitArray& GroupFilteredOutObjects = ConnectionInfo.GroupFilteredOutObjects;
+					for (uint32 ObjectIndex : MakeArrayView(Group->Members.GetData(), Group->Members.Num()))
+					{
+						GroupFilteredOutObjects.SetBit(ObjectIndex);
+						// Filter subobjects
+						for (const FInternalNetHandle SubObjectIndex : NetHandleManager->GetSubObjects(ObjectIndex))
+						{
+							GroupFilteredOutObjects.SetBit(SubObjectIndex);
+						}
+					}
+				}
+			};
+
+			FilterGroups.ForAllSetBits(InitGroupFilterForConnection);
+		}
+
+		// Update connection scope with owner filtering
+		{
+			auto&& MaskObjectToOwner = [this, ConnectionId, &ConnectionInfo](uint32 ObjectIndex)
+			{
+				const bool bIsOwner = (ConnectionId == ObjectIndexToOwningConnection.GetData()[ObjectIndex]);
+				ConnectionInfo.ConnectionFilteredObjects.SetBitValue(ObjectIndex, bIsOwner);
+			};
+
+			ObjectsWithOwnerFilter.ForAllSetBits(MaskObjectToOwner);
+		}
+
+		// Update connection scope with connection filtering
+		{
+			auto&& MaskObjectToConnection = [this, ConnectionId, &ConnectionInfo](uint32 ObjectIndex)
+			{
+				const PerObjectInfoIndexType ObjectInfoIndex = ObjectIndexToPerObjectInfoIndex.GetData()[ObjectIndex];
+				const FPerObjectInfo* ObjectInfo = GetPerObjectInfo(ObjectInfoIndex);
+				const ENetFilterStatus ReplicationStatus = GetConnectionFilterStatus(*ObjectInfo, ConnectionId);
+				ConnectionInfo.ConnectionFilteredObjects.SetBitValue(ObjectIndex, ReplicationStatus == ENetFilterStatus::Allow);
+			};
+
+			ObjectsWithPerObjectInfo.ForAllSetBits(MaskObjectToConnection);
+		}
+
+		// Combine connection and group filtering
+		ConnectionInfo.ObjectsInScopeBeforeDynamicFiltering.Copy(ConnectionInfo.ConnectionFilteredObjects);
+		ConnectionInfo.ObjectsInScopeBeforeDynamicFiltering.Combine(ConnectionInfo.GroupFilteredOutObjects, FNetBitArray::AndNotOp);
+	};
+
+	NewConnections.ForAllSetBits(InitNewConnection);
+	NewConnections.Reset();
+}
+
+void FReplicationFiltering::ResetRemovedConnections()
+{
+	if (!bHasRemovedConnection)
+	{
+		return;
+	}
+
+	bHasRemovedConnection = 0;
+
+	// Reset group filter status
+	// We might want to introduce a way to specify default state for group filters
+	// Currently we just clear the filter
+	auto&& ResetGroupFilterStatus = [this](uint32 GroupIndex)
+	{
+		FPerObjectInfo* ConnectionStateInfo = GetPerObjectInfo(this->GroupInfos[GroupIndex].ConnectionStateIndex);
+
+		// Special case we want to mask with the valid connections
+		const uint32* ValidConnectionsData = this->ValidConnections.GetData();
+		const uint32 NumWords = this->ValidConnections.GetNumWords();
+
+		// We do not need to do anything more than to restore the state to the default as we reset the effects of the filter
+		for (uint32 WordIt = 0; WordIt < NumWords; ++WordIt)
+		{
+			ConnectionStateInfo->ConnectionIds[WordIt] &= ValidConnectionsData[WordIt];
+		}
+	};
+
+	FilterGroups.ForAllSetBits(ResetGroupFilterStatus);
+}
+
+void FReplicationFiltering::UpdateObjectsInScope()
+{
+	const FNetBitArray& ObjectsInScope = NetHandleManager->GetScopableInternalIndices();
+	const FNetBitArray& PrevObjectsInScope = NetHandleManager->GetPrevFrameScopableInternalIndices();
+
+	/**
+	 * It's possible for an object to be created, have some filtering applied and then be removed later the same frame.
+	 * We can detect it using our various dirty bit arrays and force deletion of filtering data associated with the object.
+	 */
+	FNetBitArray FakePrevObjectsInScope(ObjectsInScope.GetNumBits());
+
+	TArray<FNetBitArrayBase::StorageWordType> ModifiedWords;
+	ModifiedWords.SetNumUninitialized(WordCountForObjectBitArrays);
+	uint32* ModifiedWordsStorage = ModifiedWords.GetData();
+	uint32 ModifiedWordIndex = 0;
+	uint32* ObjectsWithDirtyConnectionFilterStorage = ObjectsWithDirtyConnectionFilter.GetData();
+	uint32* ObjectsWithDirtyOwnerStorage = ObjectsWithDirtyOwner.GetData();
+	const uint32* ObjectsInScopeStorage = ObjectsInScope.GetData();
+	uint32* FakePrevObjectsInScopeStorage = FakePrevObjectsInScope.GetData();
+	{
+		const uint32* PrevObjectsInScopeStorage = PrevObjectsInScope.GetData();
+
+		for (uint32 WordIt = 0, WordEndIt = WordCountForObjectBitArrays; WordIt != WordEndIt; ++WordIt)
+		{
+			const uint32 ObjectsInScopeWord = ObjectsInScopeStorage[WordIt];
+			const uint32 PrevObjectsInScopeWord = PrevObjectsInScopeStorage[WordIt];
+			const uint32 ObjectsWithDirtyConnectionFilterWord = ObjectsWithDirtyConnectionFilterStorage[WordIt];
+			const uint32 ObjectsWithDirtyOwnerWord = ObjectsWithDirtyOwnerStorage[WordIt];
+			const uint32 SameFrameRemovedWord = ~(ObjectsInScopeWord | PrevObjectsInScopeWord) & (ObjectsWithDirtyConnectionFilterWord | ObjectsWithDirtyOwnerWord);
+
+			// Pretend that same frame removed objects existed in the previous frame.
+			FakePrevObjectsInScopeStorage[WordIt] = PrevObjectsInScopeWord | SameFrameRemovedWord;
+
+			// Store modified word index to later avoid iterating over unchanged words.
+			const bool bWordDiffers = ((ObjectsInScopeWord ^ PrevObjectsInScopeWord) | SameFrameRemovedWord) != 0U;
+			ModifiedWordsStorage[ModifiedWordIndex] = WordIt;
+			ModifiedWordIndex += bWordDiffers;
+		}
+	}
+
+	// If the scope didn't change there's nothing for us to do.
+	if (ModifiedWordIndex == 0)
+	{
+		return;
+	}
+
+	// Clear info for deleted objects
+	{
+		uint32* ObjectsWithOwnerFilterStorage = ObjectsWithOwnerFilter.GetData();
+		uint16* ObjectIndexToOwningConnectionStorage = ObjectIndexToOwningConnection.GetData();
+		const uint32* SubObjectInternalIndicesStorage = NetHandleManager->GetSubObjectInternalIndices().GetData();
+		uint32 PrevParentIndex = FNetHandleManager::InvalidInternalIndex;
+		for (uint32 WordIt = 0, WordEndIt = ModifiedWordIndex; WordIt != WordEndIt; ++WordIt)
+		{
+			const uint32 WordIndex = ModifiedWordsStorage[WordIt];
+			const uint32 PrevExistingObjects = FakePrevObjectsInScopeStorage[WordIndex];
+			const uint32 ExistingObjects = ObjectsInScopeStorage[WordIndex];
+
+			// Deleted objects can't be dirty and can't have filtering.
+			ObjectsWithDirtyConnectionFilterStorage[WordIndex] &= ExistingObjects;
+			ObjectsWithOwnerFilterStorage[WordIndex] &= ExistingObjects;
+			ObjectsWithDirtyOwnerStorage[WordIndex] &= ExistingObjects;
+
+			// Clear dynamic filters and owner info from deleted objects.
+			const uint32 BitOffset = WordIndex*32U;
+			uint32 DeletedObjects = (PrevExistingObjects & ~ExistingObjects);
+			for ( ; DeletedObjects; )
+			{
+				const uint32 LeastSignificantBit = DeletedObjects & uint32(-int32(DeletedObjects));
+				DeletedObjects ^= LeastSignificantBit;
+
+				const uint32 ObjectIndex = BitOffset + FPlatformMath::CountTrailingZeros(LeastSignificantBit);
+				ObjectIndexToOwningConnection[ObjectIndex] = 0U;
+
+				// Free per object info if present
+				FreePerObjectInfoForObject(ObjectIndex);
+
+				// Remove dynamic filtering
+				const uint32 DynamicFilterIndex = ObjectIndexToDynamicFilterIndex[ObjectIndex];
+				if (DynamicFilterIndex != InvalidDynamicFilterIndex)
+				{
+					RemoveFromDynamicFilter(ObjectIndex, DynamicFilterIndex);
+				}
+			}
+
+			// Make sure subobjects that are added after the parent gets properly updated.
+			// Dirtying the parent will cause the subobjects to be updated too.
+			uint32 AddedObjects = ExistingObjects & ~PrevExistingObjects;
+			if (uint32 AddedSubObjects = AddedObjects & SubObjectInternalIndicesStorage[WordIndex])
+			{
+				for ( ; AddedSubObjects; )
+				{
+					const uint32 LeastSignificantBit = AddedSubObjects & uint32(-int32(AddedSubObjects));
+					AddedSubObjects ^= LeastSignificantBit;
+
+					const uint32 ObjectIndex = BitOffset + FPlatformMath::CountTrailingZeros(LeastSignificantBit);
+					const uint32 ParentIndex = NetHandleManager->GetReplicatedObjectDataNoCheck(ObjectIndex).SubObjectRootIndex;
+					
+					if (ParentIndex == PrevParentIndex || ParentIndex == FNetHandleManager::InvalidInternalIndex)
+					{
+						continue;
+					}
+
+					// We need to make sure the subobject gets the same filter status as the parent.
+					ObjectsRequiringDynamicFilterUpdate.SetBit(ParentIndex);
+
+					PrevParentIndex = ParentIndex;
+
+					// If parent is a member of a group filter we need to refresh group filtering to include subobject
+					uint32 GroupMembershipCount = 0U;
+					const FNetObjectGroupHandle* GroupHandles = Groups->GetGroupMemberships(ParentIndex, GroupMembershipCount);
+					for (FNetObjectGroupHandle GroupHandle : MakeArrayView(GroupHandles, GroupMembershipCount))
+					{
+						if (FilterGroups.GetBit(GroupHandle))
+						{
+							DirtyFilterGroups.SetBit(GroupHandle);
+						}
+					}
+
+					// If the parent is new as well we don't have to do anything, but if it already existed
+					// we may need to do some work if it has owner or filtering info.
+					if (!PrevObjectsInScope.GetBit(ParentIndex))
+					{
+						continue;
+					}
+
+					if (ObjectIndexToOwningConnectionStorage[ParentIndex])
+					{
+						bHasDirtyOwner = 1;
+						ObjectsWithDirtyOwner.SetBit(ParentIndex);
+					}
+
+					if (HasOwnerFilter(ParentIndex) || HasConnectionFilter(ParentIndex))
+					{
+						bHasDirtyConnectionFilter = 1;
+						// Updating the parent will update all subobjects.
+						ObjectsWithDirtyConnectionFilter.SetBit(ParentIndex);
+					}
+				}
+			}
+		}
+	}
+
+	// Update the scope for all valid connections.
+	for (FPerConnectionInfo& ConnectionInfo : ConnectionInfos)
+	{
+		const SIZE_T Index = &ConnectionInfo - ConnectionInfos.GetData();
+		if (!ValidConnections.GetBit(Index))
+		{
+			continue;
+		}
+
+		// Or in brand new objects and mask off deleted objects by anding with now existing objects.
+		uint32* FilteredObjectsStorage = ConnectionInfo.ConnectionFilteredObjects.GetData();
+		uint32* GroupFilteredOutObjectsStorage = ConnectionInfo.GroupFilteredOutObjects.GetData();
+		uint32* ObjectsInScopeBeforeDynamicFiltering = ConnectionInfo.ObjectsInScopeBeforeDynamicFiltering.GetData();
+		uint32* DynamicFilteredOutObjects = ConnectionInfo.DynamicFilteredOutObjects.GetData();
+
+		for (uint32 WordIt = 0, WordEndIt = ModifiedWordIndex; WordIt != WordEndIt; ++WordIt)
+		{
+			const uint32 WordIndex = ModifiedWordsStorage[WordIt];
+
+			const uint32 PrevExistingObjects = FakePrevObjectsInScopeStorage[WordIndex];
+			const uint32 ExistingObjects = ObjectsInScopeStorage[WordIndex];
+			const uint32 NewObjects = ExistingObjects & ~PrevExistingObjects;
+			const uint32 FilteredObjectsWord = (FilteredObjectsStorage[WordIndex] | NewObjects) & ExistingObjects;
+			const uint32 GroupFilteredOutWord = GroupFilteredOutObjectsStorage[WordIndex] & ExistingObjects;
+			FilteredObjectsStorage[WordIndex] = FilteredObjectsWord;
+			GroupFilteredOutObjectsStorage[WordIndex] = GroupFilteredOutWord;
+			ObjectsInScopeBeforeDynamicFiltering[WordIndex] = FilteredObjectsWord & ~GroupFilteredOutWord;
+
+			// Make sure the we restore dynamic filtering for new objects
+			if (HasDynamicFilters())
+			{
+				DynamicFilteredOutObjects[WordIndex] &= ~NewObjects;
+			}
+		}
+	}
+}
+
+uint32 FReplicationFiltering::GetOwningConnectionIfDirty(uint32 ObjectIndex) const
+{
+	// If this is a subobject we must check if our parent is updated or not
+	if (NetHandleManager->GetSubObjectInternalIndices().GetBit(ObjectIndex))
+	{
+		const FNetHandleManager::FReplicatedObjectData& ObjectData = NetHandleManager->GetReplicatedObjectDataNoCheck(ObjectIndex);
+		if (ObjectData.IsSubObject())
+		{
+			const uint32 ParentIndex = NetHandleManager->GetReplicatedObjectDataNoCheck(ObjectIndex).SubObjectRootIndex;
+			return ObjectIndexToOwningConnection[ParentIndex];
+		}
+	}
+	
+	return ObjectIndexToOwningConnection[ObjectIndex];
+}
+
+void FReplicationFiltering::UpdateOwnerAndConnectionFiltering()
+{
+	if (!(bHasDirtyOwner || bHasDirtyConnectionFilter))
+	{
+		return;
+	}
+
+	// Update owners
+	if (bHasDirtyOwner)
+	{
+		uint16* const ObjectIndexToOwningConnectionStorage = ObjectIndexToOwningConnection.GetData();
+		auto&& UpdateOwners = [this, ObjectIndexToOwningConnectionStorage](uint32 ObjectIndex)
+		{
+			const uint32 OwningConnectionId = ObjectIndexToOwningConnectionStorage[ObjectIndex];
+
+			for (const FInternalNetHandle SubObjectIndex : NetHandleManager->GetSubObjects(ObjectIndex))
+			{
+				const FNetHandleManager::FReplicatedObjectData& ObjectData = NetHandleManager->GetReplicatedObjectDataNoCheck(SubObjectIndex);
+				if (ObjectData.IsSubObject())
+				{
+					ObjectIndexToOwningConnectionStorage[SubObjectIndex] = OwningConnectionId;
+				}
+			}
+		};
+
+		ObjectsWithDirtyOwner.ForAllSetBits(UpdateOwners);
+	}
+
+	const FNetBitArray& GlobalObjectsInScope = NetHandleManager->GetScopableInternalIndices();
+
+	// Update filtering
+	if (bHasDirtyConnectionFilter)
+	{
+		auto&& UpdateConnectionScope = [this, &GlobalObjectsInScope](uint32 ConnectionId)
+		{
+			FPerConnectionInfo& ConnectionInfo = this->ConnectionInfos[ConnectionId];
+			FNetBitArrayView ConnectionScope = MakeNetBitArrayView(ConnectionInfo.ConnectionFilteredObjects);
+			FNetBitArrayView GroupFilteredOutObjects = MakeNetBitArrayView(ConnectionInfo.GroupFilteredOutObjects);
+			FNetBitArrayView ObjectsInScopeBeforeDynamicFiltering = MakeNetBitArrayView(ConnectionInfo.ObjectsInScopeBeforeDynamicFiltering);
+
+			// Update filter info
+			{
+				auto&& MaskObject = [this, &ConnectionScope, &GroupFilteredOutObjects, &ObjectsInScopeBeforeDynamicFiltering, ConnectionId, &GlobalObjectsInScope](uint32 ObjectIndex)
+				{
+					bool bObjectIsInScope = true;
+					if (HasOwnerFilter(ObjectIndex))
+					{
+						const uint32 OwningConnection = ObjectIndexToOwningConnection.GetData()[ObjectIndex];
+						const bool bIsOwner = (ConnectionId == OwningConnection);
+						bObjectIsInScope = bIsOwner;
+					}
+					else if (HasConnectionFilter(ObjectIndex))
+					{
+						const PerObjectInfoIndexType ObjectInfoIndex = ObjectIndexToPerObjectInfoIndex.GetData()[ObjectIndex];
+						const FPerObjectInfo* ObjectInfo = GetPerObjectInfo(ObjectInfoIndex);
+						const ENetFilterStatus ReplicationStatus = GetConnectionFilterStatus(*ObjectInfo, ConnectionId);
+						bObjectIsInScope = ReplicationStatus == ENetFilterStatus::Allow;
+					}
+
+					// Update scope for parent object.
+					{
+						const bool bIsGroupEnabled = !GroupFilteredOutObjects.GetBit(ObjectIndex);
+						ConnectionScope.SetBitValue(ObjectIndex, bObjectIsInScope);
+						ObjectsInScopeBeforeDynamicFiltering.SetBitValue(ObjectIndex, bObjectIsInScope & bIsGroupEnabled);
+					}
+
+					// Subobjects follow suit.
+					for (const FInternalNetHandle SubObjectIndex : NetHandleManager->GetSubObjects(ObjectIndex))
+					{
+						const bool bEnableObject = bObjectIsInScope && GlobalObjectsInScope.GetBit(SubObjectIndex);
+						const bool bIsGroupEnabled = !GroupFilteredOutObjects.GetBit(SubObjectIndex);
+
+						ConnectionScope.SetBitValue(SubObjectIndex, bEnableObject);
+						ObjectsInScopeBeforeDynamicFiltering.SetBitValue(SubObjectIndex, bEnableObject & bIsGroupEnabled);
+					}
+				};
+
+				// Update the connection scope for dirty objects
+				ObjectsWithDirtyConnectionFilter.ForAllSetBits(MaskObject);
+			}
+		};
+
+		ValidConnections.ForAllSetBits(UpdateConnectionScope);
+	}
+	
+	// Clear out dirtiness
+	bHasDirtyConnectionFilter = 0;
+	bHasDirtyOwner = 0;
+	ObjectsWithDirtyConnectionFilter.Reset();
+	ObjectsWithDirtyOwner.Reset();
+}
+
+void FReplicationFiltering::UpdateGroupFiltering()
+{
+	// Adding objects to an active group filter is deferred in order to avoid triggering constant filter updates.
+	auto&& UpdateGroupFilter = [this](uint32 GroupIndex)
+	{
+		if (this->FilterGroups.GetBit(GroupIndex))
+		{
+			const FPerObjectInfo* ConnectionStateInfo = GetPerObjectInfo(this->GroupInfos[GroupIndex].ConnectionStateIndex);
+			const FNetObjectGroup* Group = this->Groups->GetGroup(GroupIndex);
+			FPerConnectionInfo* LocalConnectionInfos = this->ConnectionInfos.GetData();
+
+			auto&& UpdateGroupFilterForConnection = [this, ConnectionStateInfo, Group, LocalConnectionInfos](uint32 ConnectionId)
+			{
+				if (this->GetConnectionFilterStatus(*ConnectionStateInfo, ConnectionId) == ENetFilterStatus::Disallow)
+				{
+					FNetBitArray& GroupFilteredOutObjects = LocalConnectionInfos[ConnectionId].GroupFilteredOutObjects;
+					FNetBitArray& ObjectsInScopeBeforeDynamicFiltering = LocalConnectionInfos[ConnectionId].ObjectsInScopeBeforeDynamicFiltering;
+					const FNetBitArray& ScopableObjects = NetHandleManager->GetScopableInternalIndices();
+					for (uint32 ObjectIndex : MakeArrayView(Group->Members.GetData(), Group->Members.Num()))
+					{
+						GroupFilteredOutObjects.SetBit(ObjectIndex);
+						ObjectsInScopeBeforeDynamicFiltering.ClearBit(ObjectIndex);
+
+						// Filter subobjects
+						for (const FInternalNetHandle SubObjectIndex : NetHandleManager->GetSubObjects(ObjectIndex))
+						{
+							const bool bIsScopable = ScopableObjects.GetBit(SubObjectIndex);
+							GroupFilteredOutObjects.SetBitValue(SubObjectIndex, bIsScopable);
+							ObjectsInScopeBeforeDynamicFiltering.ClearBit(SubObjectIndex);
+						}
+					}
+				}
+			};
+
+			this->ValidConnections.ForAllSetBits(UpdateGroupFilterForConnection);
+		}
+	};
+
+	DirtyFilterGroups.ForAllSetBits(UpdateGroupFilter);
+
+	// Clear out dirtiness
+	DirtyFilterGroups.Reset();
+}
+
+void FReplicationFiltering::PreUpdateDynamicFiltering()
+{
+	IRIS_PROFILER_SCOPE(FReplicationFiltering_PreUpdateDynamicFiltering);
+
+	// Give filters a chance to prepare for filtering. It's only called if any object has the filter set.
+	{
+		FNetObjectPreFilteringParams PreFilteringParams;
+		for (FFilterInfo& Info : DynamicFilterInfos)
+		{
+			if (Info.ObjectCount == 0U)
+			{
+				continue;
+			}
+			
+			Info.Filter->PreFilter(PreFilteringParams);
+		}
+	}
+}
+
+/**
+ * Per connection we will combine the effects of connection/owner and group filtering before passing
+ * the objects with a specific dynamic filter set. The filter can then decide which objects should
+ * be filtered out or not. The result is ANDed with the objects passed so that it's trivial to
+ * implement on or off filters and as a safety mechanism to prevent filters from changing the status
+ * of objects it should not be concerned about.
+ */
+
+void FReplicationFiltering::UpdateDynamicFiltering()
+{
+	IRIS_PROFILER_SCOPE(FReplicationFiltering_UpdateDynamicFiltering);
+
+	// Working set
+	TArray<uint32, TInlineAllocator<256>> DisabledDependentObjects;
+
+	uint32* AllowedObjectsData = static_cast<uint32*>(FMemory_Alloca(WordCountForObjectBitArrays * sizeof(uint32)));
+	uint32* NewDynamicFilteredOutObjectsData = static_cast<uint32*>(FMemory_Alloca(WordCountForObjectBitArrays * sizeof(uint32)));
+
+	FNetBitArrayView AllowedObjects(AllowedObjectsData, MaxObjectCount, FNetBitArrayView::NoResetNoValidate);
+	FNetBitArrayView NewDynamicFilteredOutObjects(NewDynamicFilteredOutObjectsData, MaxObjectCount, FNetBitArrayView::NoResetNoValidate);
+
+	// Subobjects will never be added to a dynamic filter, but objects can become dependent at any time.
+	// We need to make sure they are not filtered out.
+	const uint32* SubObjectsData = NetHandleManager->GetSubObjectInternalIndices().GetData();
+	const uint32* DependentObjectsData = NetHandleManager->GetDependentObjectInternalIndices().GetData();
+	const uint32* ObjectsRequiringDynamicFilterUpdateData = ObjectsRequiringDynamicFilterUpdate.GetData();
+
+	uint32* ConnectionIds = static_cast<uint32*>(FMemory_Alloca(ValidConnections.GetNumBits() * sizeof(uint32)));
+	uint32 ConnectionCount = 0;
+	ValidConnections.ForAllSetBits([ConnectionIds, &ConnectionCount](uint32 Bit) { ConnectionIds[ConnectionCount++] = Bit; });
+
+	for (const uint32 ConnId : MakeArrayView(ConnectionIds, ConnectionCount))
+	{
+		FPerConnectionInfo& ConnectionInfo = this->ConnectionInfos[ConnId];
+
+		NewDynamicFilteredOutObjects.Reset();
+
+		/*
+		 * Apply dynamic filters.
+		 * 
+		 * The algorithm will loop over all connections, call each filter and update
+		 * the DynamicFilteredOutObjects bit array. Once all filters have been applied
+		 * the result is compared with that of the previous frame.
+		 * Special handling is required for objects who are marked as requiring update,
+		 * as they may have updated subobject information, as well as dependent objects.
+		 */
+		{
+			for (FFilterInfo& Info : DynamicFilterInfos)
+			{
+				if (Info.ObjectCount == 0U)
+				{
+					continue;
+				}
+
+				FNetObjectFilteringParams FilteringParams(MakeNetBitArrayView(Info.FilteredObjects));
+				FilteringParams.OutAllowedObjects = AllowedObjects;
+				FilteringParams.FilteringInfos = NetObjectFilteringInfos.GetData();
+				FilteringParams.StateBuffers = NetHandleManager->GetReplicatedObjectStateBuffers().GetData();
+				FilteringParams.ConnectionId = ConnId;
+				FilteringParams.View = Connections->GetReplicationView(ConnId);
+				Info.Filter->Filter(FilteringParams);
+
+				const uint32* FilteredObjectsData = Info.FilteredObjects.GetData();
+				for (SIZE_T WordIt = 0, WordEndIt = WordCountForObjectBitArrays; WordIt != WordEndIt; ++WordIt)
+				{
+					const uint32 FilteredObjects = FilteredObjectsData[WordIt];
+					const uint32 FilterAllowedObjects = AllowedObjectsData[WordIt];
+					// Keep status of objects not affected by this filter. Inverse the effect of the filter and add which objects are filtered out.
+					NewDynamicFilteredOutObjectsData[WordIt] = (NewDynamicFilteredOutObjectsData[WordIt] & ~FilteredObjects) | (~FilterAllowedObjects & FilteredObjects);
+				}
+			}
+		}
+
+		/*
+		 * Now that all dynamic filters have been called we can do some post-processing
+		 * and try only to do the more expensive operations, such as subobject management,
+		 * for objects that have changed filter status since the previous frame.
+		 */
+		{
+			DisabledDependentObjects.Reset();
+
+			const uint32* DynamicFilterEnabledObjectsData = DynamicFilterEnabledObjects.GetData();
+			FNetBitArrayView DynamicFilteredOutObjects = MakeNetBitArrayView(ConnectionInfo.DynamicFilteredOutObjects);
+			uint32* DynamicFilteredOutObjectsData = ConnectionInfo.DynamicFilteredOutObjects.GetData();
+			for (SIZE_T WordIt = 0, WordEndIt = WordCountForObjectBitArrays; WordIt != WordEndIt; ++WordIt)
+			{
+				const uint32 SubObjects = SubObjectsData[WordIt];
+				const uint32 DependentObjects = DependentObjectsData[WordIt];
+				const uint32 ObjectsRequiringUpdate = ObjectsRequiringDynamicFilterUpdateData[WordIt];
+				const uint32 PrevFilteredOutObjects = DynamicFilteredOutObjectsData[WordIt];
+				const uint32 NewFilteredOutObjects = NewDynamicFilteredOutObjectsData[WordIt];
+				const uint32 FilterEnabledObjects = DynamicFilterEnabledObjectsData[WordIt];
+
+				const uint32 ModifiedScopeObjects = (PrevFilteredOutObjects ^ NewFilteredOutObjects);
+				// ObjectsRequiringUpdate may contain objects that had a dynamic filter set last frame. We need to update
+				// the filtered out objects.
+				const uint32 ObjectsToProcess = ((ModifiedScopeObjects | DependentObjects) & FilterEnabledObjects) | ObjectsRequiringUpdate;
+
+				if (!ObjectsToProcess)
+				{
+					continue;
+				}
+
+				// Update dynamic filter, but preserve subobject information.
+				DynamicFilteredOutObjectsData[WordIt] = (PrevFilteredOutObjects & SubObjects) | NewFilteredOutObjects;
+
+				const uint32 BitOffset = WordIt*32U;
+
+				// Calculate which objects need to be updated due to being filtered out.
+				// We want to process as few as possible since modifying subobject status is expensive.
+				for (uint32 DisabledObjects = NewFilteredOutObjects & ObjectsToProcess; DisabledObjects; )
+				{
+					// Extract one set bit from DisabledObjects. The code below will extract the least significant bit set and then clear it.
+					const uint32 LeastSignificantBit = DisabledObjects & uint32(-int32(DisabledObjects));
+					DisabledObjects ^= LeastSignificantBit;
+
+					const uint32 ObjectIndex = BitOffset + FPlatformMath::CountTrailingZeros(LeastSignificantBit);
+					for (const FInternalNetHandle SubObjectIndex : NetHandleManager->GetSubObjects(ObjectIndex))
+					{						
+						DynamicFilteredOutObjects.SetBit(SubObjectIndex);
+					}
+
+					// Store dependent objects for later processing.
+					if (DependentObjects & LeastSignificantBit)
+					{
+						DisabledDependentObjects.Add(ObjectIndex);
+					}
+				}
+
+				for (uint32 EnabledObjects = ~NewFilteredOutObjects & ObjectsToProcess; EnabledObjects; )
+				{
+					const uint32 LeastSignificantBit = EnabledObjects & uint32(-int32(EnabledObjects));
+					EnabledObjects ^= LeastSignificantBit;
+
+					const uint32 ObjectIndex = BitOffset + FPlatformMath::CountTrailingZeros(LeastSignificantBit);
+					for (const FInternalNetHandle SubObjectIndex : NetHandleManager->GetSubObjects(ObjectIndex))
+					{
+						DynamicFilteredOutObjects.ClearBit(SubObjectIndex);
+					}
+				}
+			}
+		}
+
+		// Update the entire scope for the connection
+		{
+			uint32* ObjectsInScopeData = ConnectionInfo.ObjectsInScope.GetData();
+			const uint32* ObjectsInScopeBeforeDynamicFilteringData = ConnectionInfo.ObjectsInScopeBeforeDynamicFiltering.GetData();
+			const uint32* DynamicFilteredOutObjectsData = ConnectionInfo.DynamicFilteredOutObjects.GetData();
+			for (SIZE_T WordIt = 0, WordEndIt = WordCountForObjectBitArrays; WordIt != WordEndIt; ++WordIt)
+			{
+				ObjectsInScopeData[WordIt] = ObjectsInScopeBeforeDynamicFilteringData[WordIt] & ~DynamicFilteredOutObjectsData[WordIt];
+			}
+		}
+
+		// The scope for the connection is now fully updated, apart from disabled dependent objects.
+		// If any object that has a dependency isn't filtered out we must re-enable the dependent object.
+		{
+			for (const uint32 DependentObjectIndex : DisabledDependentObjects)
+			{
+				if (GetDependentObjectFilterStatus(NetHandleManager, ConnectionInfo.ObjectsInScope, DependentObjectIndex) == ENetFilterStatus::Allow)
+				{
+					ConnectionInfo.ObjectsInScope.SetBit(DependentObjectIndex);
+					for (const FInternalNetHandle SubObjectIndex : NetHandleManager->GetSubObjects(DependentObjectIndex))
+					{
+						ConnectionInfo.ObjectsInScope.SetBit(SubObjectIndex);
+					}
+				}
+			}
+		}
+	}
+
+	ObjectsRequiringDynamicFilterUpdate.Reset();
+}
+
+void FReplicationFiltering::PostUpdateDynamicFiltering()
+{
+	IRIS_PROFILER_SCOPE(FReplicationFiltering_PostUpdateDynamicFiltering);
+
+	// Give filters a chance to prepare for filtering. It's only called if any object has the filter set.
+	{
+		FNetObjectPostFilteringParams PostFilteringParams;
+		for (FFilterInfo& Info : DynamicFilterInfos)
+		{
+			if (Info.ObjectCount == 0U)
+			{
+				continue;
+			}
+			
+			Info.Filter->PostFilter(PostFilteringParams);
+		}
+	}
+}
+
+void FReplicationFiltering::NotifyFiltersOfDirtyObjects(const FNetBitArrayView& DirtyObjects)
+{
+	FUpdateDirtyObjectsBatchHelper BatchHelper(NetHandleManager, DynamicFilterInfos.Num());
+
+	constexpr SIZE_T MaxBatchObjectCount = FUpdateDirtyObjectsBatchHelper::Constants::MaxObjectCountPerBatch;
+	uint32 ObjectIndices[MaxBatchObjectCount];
+
+	const uint32 BitCount = ~0U;
+	for (uint32 ObjectCount, StartIndex = 0; (ObjectCount = DirtyObjects.GetSetBitIndices(StartIndex, BitCount, ObjectIndices, MaxBatchObjectCount)) > 0; )
+	{
+		BatchNotifyFiltersOfDirtyObjects(BatchHelper, ObjectIndices, ObjectCount);
+
+		StartIndex = ObjectIndices[ObjectCount - 1] + 1U;
+		if ((StartIndex == DirtyObjects.GetNumBits()) | (ObjectCount < MaxBatchObjectCount))
+		{
+			break;
+		}
+	}
+}
+
+void FReplicationFiltering::BatchNotifyFiltersOfDirtyObjects(FUpdateDirtyObjectsBatchHelper& BatchHelper, const uint32* ObjectIndices, uint32 ObjectCount)
+{
+	BatchHelper.PrepareBatch(ObjectIndices, ObjectCount, ObjectIndexToDynamicFilterIndex.GetData());
+
+	FNetObjectFilterUpdateParams UpdateParameters;
+	UpdateParameters.StateBuffers = NetHandleManager->GetReplicatedObjectStateBuffers().GetData();
+	UpdateParameters.FilteringInfos = NetObjectFilteringInfos.GetData();
+
+	for (const FUpdateDirtyObjectsBatchHelper::FPerFilterInfo& PerFilterInfo : BatchHelper.PerFilterInfos)
+	{
+		if (PerFilterInfo.ObjectCount == 0)
+		{
+			continue;
+		}
+
+		UpdateParameters.ObjectIndices = PerFilterInfo.ObjectIndices;
+		UpdateParameters.ObjectCount = PerFilterInfo.ObjectCount;
+		UpdateParameters.InstanceProtocols = PerFilterInfo.InstanceProtocols;
+
+		const SIZE_T FilterIndex = &PerFilterInfo - BatchHelper.PerFilterInfos.GetData();
+		UNetObjectFilter* Filter = DynamicFilterInfos[FilterIndex].Filter.Get();
+		Filter->UpdateObjects(UpdateParameters);
+	}
+}
+
+bool FReplicationFiltering::HasOwnerFilter(uint32 ObjectIndex) const
+{
+	const bool bHasOwnerFilter = ObjectsWithOwnerFilter.GetBit(ObjectIndex);
+	return bHasOwnerFilter;
+}
+
+bool FReplicationFiltering::HasConnectionFilter(uint32 ObjectIndex) const
+{
+	const bool bHasConnectionFilter = ObjectsWithPerObjectInfo.GetBit(ObjectIndex);
+	return bHasConnectionFilter;
+}
+
+FReplicationFiltering::PerObjectInfoIndexType FReplicationFiltering::AllocPerObjectInfo()
+{
+	const uint32 UsedPerObjectInfoBitsPerWord = sizeof(decltype(UsedPerObjectInfoStorage)::ElementType)*8;
+
+	FNetBitArrayView UsedPerObjectInfos = MakeNetBitArrayView(UsedPerObjectInfoStorage.GetData(), UsedPerObjectInfoStorage.Num()*UsedPerObjectInfoBitsPerWord);
+	uint32 FreeIndex = UsedPerObjectInfos.FindFirstZero();
+
+	// Grow used indices bit array if needed
+	if (FreeIndex == FNetBitArrayView::InvalidIndex)
+	{
+		checkf(UsedPerObjectInfos.GetNumBits() < TNumericLimits<PerObjectInfoIndexType>::Max(), TEXT("Filtering per object info storage exhausted. Contact the UE Network team."));
+		FreeIndex = UsedPerObjectInfos.GetNumBits();
+		UsedPerObjectInfoStorage.AddZeroed(UsedPerObjectInfoStorageGrowSize/UsedPerObjectInfoBitsPerWord);
+		FNetBitArrayView NewUsedPerObjectInfos = MakeNetBitArrayView(UsedPerObjectInfoStorage.GetData(), UsedPerObjectInfoStorage.Num()*UsedPerObjectInfoBitsPerWord);
+		NewUsedPerObjectInfos.SetBit(FreeIndex);
+		// Mark index 0 as used so we can use it as an invalid index
+		if (FreeIndex == 0U)
+		{
+			FreeIndex = 1U;
+			NewUsedPerObjectInfos.SetBit(1);
+		}
+		PerObjectInfoStorage.AddUninitialized(PerObjectInfoStorageCountPerItem*UsedPerObjectInfoStorageGrowSize);
+	}
+	else
+	{
+		UsedPerObjectInfos.SetBit(FreeIndex);
+	}
+	
+	return FreeIndex;
+}
+
+void FReplicationFiltering::FreePerObjectInfo(PerObjectInfoIndexType Index)
+{
+	const uint32 UsedPerObjectInfoBitsPerWord = sizeof(decltype(UsedPerObjectInfoStorage)::ElementType)*8U;
+
+	FNetBitArrayView UsedPerObjectInfos = MakeNetBitArrayView(UsedPerObjectInfoStorage.GetData(), UsedPerObjectInfoStorage.Num()*UsedPerObjectInfoBitsPerWord);
+	UsedPerObjectInfos.ClearBit(Index);
+}
+
+FReplicationFiltering::FPerObjectInfo* FReplicationFiltering::AllocPerObjectInfoForObject(uint32 ObjectIndex)
+{
+	ObjectsWithPerObjectInfo.SetBit(ObjectIndex);
+	const PerObjectInfoIndexType ObjectInfoIndex = AllocPerObjectInfo();
+	ObjectIndexToPerObjectInfoIndex.GetData()[ObjectIndex] = ObjectInfoIndex;
+
+	FPerObjectInfo* ObjectInfo = GetPerObjectInfo(ObjectInfoIndex);
+	SetPerObjectInfoFilterStatus(*ObjectInfo, ENetFilterStatus::Allow);
+
+	return ObjectInfo;
+}
+
+void FReplicationFiltering::FreePerObjectInfoForObject(uint32 ObjectIndex)
+{
+	PerObjectInfoIndexType& ObjectInfoIndex = ObjectIndexToPerObjectInfoIndex.GetData()[ObjectIndex];
+	if (ObjectInfoIndex == 0)
+	{
+		return;
+	}
+
+	ObjectsWithPerObjectInfo.ClearBit(ObjectIndex);
+
+	FreePerObjectInfo(ObjectInfoIndex);
+	ObjectInfoIndex = 0;
+}
+
+FReplicationFiltering::FPerObjectInfo* FReplicationFiltering::GetPerObjectInfo(PerObjectInfoIndexType Index)
+{
+	const uint32 ObjectInfoIndex = Index*PerObjectInfoStorageCountPerItem;
+
+	checkSlow(ObjectInfoIndex < (uint32)PerObjectInfoStorage.Num());
+
+	uint32* StoragePointer = PerObjectInfoStorage.GetData() + ObjectInfoIndex;
+	return reinterpret_cast<FPerObjectInfo*>(StoragePointer);
+}
+
+const FReplicationFiltering::FPerObjectInfo* FReplicationFiltering::GetPerObjectInfo(PerObjectInfoIndexType Index) const
+{
+	const uint32 ObjectInfoIndex = Index*PerObjectInfoStorageCountPerItem;
+
+	checkSlow(ObjectInfoIndex < (uint32)PerObjectInfoStorage.Num());
+
+	const uint32* StoragePointer = PerObjectInfoStorage.GetData() + ObjectInfoIndex;
+	return reinterpret_cast<const FPerObjectInfo*>(StoragePointer);
+}
+
+void FReplicationFiltering::AddSubObjectFilter(FNetObjectGroupHandle GroupHandle)
+{
+	check(!FilterGroups.GetBit(GroupHandle));
+
+	if (ensure(GroupHandle != InvalidNetObjectGroupHandle && !SubObjectFilterGroups.GetBit(GroupHandle)))
+	{
+		const bool bGroupExists = SubObjectFilterGroups.GetBit(GroupHandle);
+		if (ensure(bGroupExists == false))
+		{
+			SubObjectFilterGroups.SetBit(GroupHandle);
+			GroupInfos[GroupHandle].ConnectionStateIndex = AllocPerObjectInfo();
+
+			// By default filter is active, i.e. we filter out all connections
+			SetPerObjectInfoFilterStatus(*GetPerObjectInfo(GroupInfos[GroupHandle].ConnectionStateIndex), ENetFilterStatus::Disallow);
+
+			UE_LOG(LogIrisFiltering, Verbose, TEXT("ReplicationFiltering::AddSubObjectFilter (%s) GroupIndex: %u, FilterStatus: DisallowReplication"), *(Groups->GetGroup(GroupHandle)->GroupName.ToString()), GroupHandle);
+		}
+	}
+}
+
+void FReplicationFiltering::RemoveSubObjectFilter(FNetObjectGroupHandle GroupHandle)
+{
+	if (GroupHandle != InvalidNetObjectGroupHandle  && SubObjectFilterGroups.GetBit(GroupHandle))
+	{
+		// Mark group as no longer a SubObjectFilter group
+		SubObjectFilterGroups.ClearBit(GroupHandle);
+		const PerObjectInfoIndexType ConnectionStateIndex = GroupInfos[GroupHandle].ConnectionStateIndex;
+		GroupInfos[GroupHandle].ConnectionStateIndex = 0U;
+		FreePerObjectInfo(ConnectionStateIndex);
+
+		UE_LOG(LogIrisFiltering, Verbose, TEXT("ReplicationFiltering::RemoveSubObjectFilter GroupIndex: %u"), GroupHandle);
+	}
+}
+
+void FReplicationFiltering::UpdateSubObjectFilters()
+{
+	// We want to remove all groups that have no members and no enabled connections
+	auto&& UpdateSubObjectFilterGroup = [this](uint32 GroupIndex)
+	{
+		FNetObjectGroupHandle GroupHandle(GroupIndex);
+		if (const FNetObjectGroup* Group = Groups->GetGroup(FNetObjectGroupHandle(GroupIndex)))
+		{
+			if (Group->Members.Num() == 0U && !IsConnectionFilterStatusAllowedForAnyConnection(*GetPerObjectInfo(GroupInfos[GroupHandle].ConnectionStateIndex)))
+			{
+				ReplicationSystem->DestroyGroup(GroupHandle);
+			}
+		}
+	};
+
+	FNetBitArray::ForAllSetBits(DirtySubObjectFilterGroups, SubObjectFilterGroups, FNetBitArray::AndOp, UpdateSubObjectFilterGroup);
+	DirtySubObjectFilterGroups.Reset();
+}
+
+void FReplicationFiltering::SetSubObjectFilterStatus(FNetObjectGroupHandle GroupHandle, ENetFilterStatus ReplicationStatus)
+{
+	IRIS_PROFILER_SCOPE(SetSubObjectFilterStatus)
+
+	if (IsReservedNetObjectGroupHandle(GroupHandle))
+	{
+		UE_LOG(LogIrisFiltering, Warning, TEXT("FReplicationFiltering::SetSubObjectFilterStatus - Trying to set filter for reserved GroupHandle: %u which is not allowed."), GroupHandle);
+		return;
+	}
+
+	if (!SubObjectFilterGroups.GetBit(GroupHandle))
+	{
+		UE_LOG(LogIrisFiltering, Warning, TEXT("FReplicationFiltering::SetSubObjectFilterStatus - Trying to ReplicationStatus for GroupIndex: %u that is not a SubObjectFilterGroup"), GroupHandle);
+		return;
+	}
+
+	FPerObjectInfo* FilterInfo = GetPerObjectInfo(GroupInfos[GroupHandle].ConnectionStateIndex);
+	SetPerObjectInfoFilterStatus(*FilterInfo, ReplicationStatus);
+	if (!IsConnectionFilterStatusAllowedForAnyConnection(*FilterInfo))
+	{
+		DirtySubObjectFilterGroups.SetBit(GroupHandle);
+	}
+}
+
+void FReplicationFiltering::SetSubObjectFilterStatus(FNetObjectGroupHandle GroupHandle, const FNetBitArrayView& ConnectionsBitArray, ENetFilterStatus ReplicationStatus)
+{
+	if (ConnectionsBitArray.GetNumBits() > ValidConnections.GetNumBits())
+	{
+		UE_LOG(LogIrisFiltering, Warning, TEXT("FReplicationFiltering::SetSubObjectFilterStatus - Trying to set filter for %u, with invalid Connections parameters."), GroupHandle);
+		return;
+	}
+
+	if (IsReservedNetObjectGroupHandle(GroupHandle))
+	{
+		UE_LOG(LogIrisFiltering, Warning, TEXT("FReplicationFiltering::SetSubObjectFilterStatus - Trying to set filter for reserved GroupHandle: %u which is not allowed."), GroupHandle);
+		return;
+	}
+
+	if (!SubObjectFilterGroups.GetBit(GroupHandle))
+	{
+		UE_LOG(LogIrisFiltering, Warning, TEXT("FReplicationFiltering::SetSubObjectFilterStatus - Trying to ReplicationStatus for GroupIndex: %u that is not a SubObjectFilterGroup"), GroupHandle);
+		return;
+	}
+
+	const ENetFilterStatus InvertedReplicationStatus = ReplicationStatus == ENetFilterStatus::Allow ? ENetFilterStatus::Disallow : ENetFilterStatus::Allow;
+	const ENetFilterStatus DefaultReplicationStatus = ENetFilterStatus::Disallow;
+
+	// It is intentional that we iterate over all connections
+	for (uint32 ConnectionId = 1U; ConnectionId < ValidConnections.GetNumBits(); ++ConnectionId)
+	{
+		if (ConnectionId < ConnectionsBitArray.GetNumBits())
+		{			
+			SetSubObjectFilterStatus(GroupHandle, ConnectionId, ConnectionsBitArray.GetBit(ConnectionId) ? ReplicationStatus : InvertedReplicationStatus);
+		}
+		else
+		{
+			// To avoid issues with unspecified connections we set the status to the default
+			SetSubObjectFilterStatus(GroupHandle, ConnectionId, DefaultReplicationStatus);
+		}
+	}
+}
+
+void FReplicationFiltering::SetSubObjectFilterStatus(FNetObjectGroupHandle GroupHandle, uint32 ConnectionId, ENetFilterStatus ReplicationStatus)
+{
+	if (IsReservedNetObjectGroupHandle(GroupHandle))
+	{
+		UE_LOG(LogIrisFiltering, Warning, TEXT("FReplicationFiltering::SetSubObjectFilterStatus - Trying to set filter for reserved GroupHandle: %u which is not allowed."), GroupHandle);
+		return;
+	}
+
+	if (ensureAlways(ValidConnections.GetBit(ConnectionId) && SubObjectFilterGroups.GetBit(GroupHandle)))
+	{
+		UE_LOG(LogIrisFiltering, Verbose, TEXT("ReplicationFiltering::SetSubObjectFilterStatus GroupIndex: %u, ConnectionId: %u, FilterStatus: %u"), GroupHandle, ConnectionId, ReplicationStatus == ENetFilterStatus::Allow ? 1U : 0U);
+		FPerObjectInfo* FilterInfo = GetPerObjectInfo(GroupInfos[GroupHandle].ConnectionStateIndex);
+		SetConnectionFilterStatus(*FilterInfo, ConnectionId, ReplicationStatus);
+		if (!IsConnectionFilterStatusAllowedForAnyConnection(*FilterInfo))
+		{
+			DirtySubObjectFilterGroups.SetBit(GroupHandle);
+		}
+	}
+}
+
+bool FReplicationFiltering::GetSubObjectFilterStatus(FNetObjectGroupHandle GroupHandle, uint32 ConnectionId, ENetFilterStatus& OutReplicationStatus) const
+{
+	if (!(ValidConnections.GetBit(ConnectionId) && SubObjectFilterGroups.GetBit(GroupHandle)))
+	{
+		return false;
+	}
+
+	const FPerObjectInfo* ConnectionState = GetPerObjectInfo(GroupInfos[GroupHandle].ConnectionStateIndex);
+	OutReplicationStatus = GetConnectionFilterStatus(*ConnectionState, ConnectionId);
+
+	return true;
+}
+
+void FReplicationFiltering::AddGroupFilter(FNetObjectGroupHandle GroupHandle)
+{
+	check(!SubObjectFilterGroups.GetBit(GroupHandle));
+
+	if (ensure(GroupHandle && !FilterGroups.GetBit(GroupHandle)))
+	{
+		const bool bGroupExists = FilterGroups.GetBit(GroupHandle);
+		if (ensure(bGroupExists == false))
+		{
+			FilterGroups.SetBit(GroupHandle);
+			DirtyFilterGroups.SetBit(GroupHandle);
+			GroupInfos[GroupHandle].ConnectionStateIndex = AllocPerObjectInfo();
+
+			// By default filter is active, i.e. we filter out all connections
+			SetPerObjectInfoFilterStatus(*GetPerObjectInfo(GroupInfos[GroupHandle].ConnectionStateIndex), ENetFilterStatus::Disallow);
+
+			UE_LOG(LogIrisFiltering, Verbose, TEXT("ReplicationFiltering::AddGroupFilter GroupIndex: %u, FilterStatus: DisallowReplication"), GroupHandle);
+		}
+	}
+}
+
+// Remove group from filtered groups
+void FReplicationFiltering::RemoveGroupFilter(FNetObjectGroupHandle GroupHandle)
+{
+	if (GroupHandle && FilterGroups.GetBit(GroupHandle))
+	{
+		ValidConnections.ForAllSetBits( [GroupHandle, this](uint32 ConnectionIndex) { this->SetGroupFilterStatus(GroupHandle, ConnectionIndex, ENetFilterStatus::Allow); } );
+
+		// Mark group as no longer a group filter group
+		FilterGroups.ClearBit(GroupHandle);
+		const PerObjectInfoIndexType ConnectionStateIndex = GroupInfos[GroupHandle].ConnectionStateIndex;
+		GroupInfos[GroupHandle].ConnectionStateIndex = 0U;
+		FreePerObjectInfo(ConnectionStateIndex);
+
+		UE_LOG(LogIrisFiltering, Verbose, TEXT("ReplicationFiltering::RemoveGroupFilter GroupIndex: %u"), GroupHandle);
+	}
+}
+
+void FReplicationFiltering::SetGroupFilterStatus(FNetObjectGroupHandle GroupHandle, ENetFilterStatus ReplicationStatus)
+{
+	IRIS_PROFILER_SCOPE(SetGroupFilterStatus)
+
+	if (IsReservedNetObjectGroupHandle(GroupHandle))
+	{
+		UE_LOG(LogIrisFiltering, Warning, TEXT("FReplicationFiltering::SetGroupFilterStatus - Trying to set filter for reserer NetObjectGroupHandle: %u which is not allowed."), GroupHandle);
+		return;
+	}
+
+	if (!FilterGroups.GetBit(GroupHandle))
+	{
+		UE_LOG(LogIrisFiltering, Warning, TEXT("FReplicationFiltering::SetGroupFilterStatus - Trying to set filter for invalid GroupIndex: %u, Make sure group is added to filtering"), GroupHandle);
+		return;
+	}
+	
+	// It is intentional that we iterate overall connections, not only the valid ones as we want the group filter status to be initialized correctly for new connections as well.
+	for (uint32 ConnectionId = 1U; ConnectionId < ValidConnections.GetNumBits(); ++ConnectionId)
+	{
+		InternalSetGroupFilterStatus(GroupHandle, ConnectionId, ReplicationStatus);
+	}
+}
+
+void FReplicationFiltering::SetGroupFilterStatus(FNetObjectGroupHandle GroupHandle, const FNetBitArrayView& ConnectionsBitArray, ENetFilterStatus ReplicationStatus)
+{
+	IRIS_PROFILER_SCOPE(SetGroupFilterStatus)
+		
+	if (IsReservedNetObjectGroupHandle(GroupHandle))
+	{
+		UE_LOG(LogIrisFiltering, Warning, TEXT("FReplicationFiltering::SetGroupFilterStatus - Trying to set filter for reserer NetObjectGroupHandle: %u which is not allowed."), GroupHandle);		
+		return;
+	}
+
+	if (!FilterGroups.GetBit(GroupHandle))
+	{
+		UE_LOG(LogIrisFiltering, Warning, TEXT("FReplicationFiltering::SetGroupFilterStatus - Trying to set filter for invalid GroupIndex: %u, Make sure group is added to filtering"), GroupHandle);
+		return;
+	}
+
+	if (ConnectionsBitArray.GetNumBits() > ValidConnections.GetNumBits())
+	{
+		UE_LOG(LogIrisFiltering, Warning, TEXT("FReplicationFiltering::SetGroupFilterStatus - Trying to set filter for %u, with invalid Connections parameters."), GroupHandle);
+		return;
+	}
+
+	// It is intentional that we iterate overall connections, not only the valid ones as we want the group filter status to be initialized correctly for new connections as well.
+	for (uint32 ConnectionId = 1U; ConnectionId < ValidConnections.GetNumBits(); ++ConnectionId)
+	{
+		ENetFilterStatus ReplicationStatusToSet = ReplicationStatus;
+		if (ConnectionId >= ConnectionsBitArray.GetNumBits() || !ConnectionsBitArray.GetBit(ConnectionId))
+		{
+			ReplicationStatusToSet = ReplicationStatus == ENetFilterStatus::Allow ? ENetFilterStatus::Disallow : ENetFilterStatus::Allow;
+		}
+		InternalSetGroupFilterStatus(GroupHandle, ConnectionId, ReplicationStatusToSet);
+	}
+}
+
+void FReplicationFiltering::SetGroupFilterStatus(FNetObjectGroupHandle GroupHandle, uint32 ConnectionId, ENetFilterStatus ReplicationStatus)
+{
+	IRIS_PROFILER_SCOPE(SetGroupFilterStatus)
+
+	if (IsReservedNetObjectGroupHandle(GroupHandle))
+	{
+		UE_LOG(LogIrisFiltering, Warning, TEXT("FReplicationFiltering::SetGroupFilterStatus - Trying to set filter for reserer NetObjectGroupHandle: %u which is not allowed."), GroupHandle);
+		return;
+	}
+
+	if (!FilterGroups.GetBit(GroupHandle))
+	{
+		UE_LOG(LogIrisFiltering, Warning, TEXT("FReplicationFiltering::SetGroupFilterStatus - Trying to set invalid filter ConnectionId: %u, GroupIndex: %u"), ConnectionId, GroupHandle);
+		return;
+	}
+	
+	InternalSetGroupFilterStatus(GroupHandle, ConnectionId, ReplicationStatus);
+}
+
+bool FReplicationFiltering::GetIsFilteredOutByAnyGroup(uint32 ObjectInternalIndex, uint32 ConnectionId) const
+{
+	uint32 GroupMembershipCount = 0U;
+	const FNetObjectGroupHandle* GroupHandles = Groups->GetGroupMemberships(ObjectInternalIndex, GroupMembershipCount);
+
+	if (!GroupHandles)
+	{
+		return false;
+	}
+
+	for (FNetObjectGroupHandle GroupHandle : MakeArrayView(GroupHandles, GroupMembershipCount))
+	{
+		if (FilterGroups.GetBit(GroupHandle))
+		{
+			const FPerObjectInfo* ConnectionState = GetPerObjectInfo(GroupInfos[GroupHandle].ConnectionStateIndex);
+			if (GetConnectionFilterStatus(*ConnectionState, ConnectionId) == ENetFilterStatus::Disallow)
+			{
+				return true;
+			}
+		}
+	}
+
+	return false;
+}
+
+bool FReplicationFiltering::UpdateGroupFilterEffectsForObject(uint32 ObjectIndex, uint32 ConnectionId, const FNetBitArray& ScopableObjects)
+{
+	FPerConnectionInfo& ConnectionInfo = ConnectionInfos[ConnectionId];
+	FNetBitArray& GroupFilteredOutObjects = ConnectionInfo.GroupFilteredOutObjects;
+	FNetBitArray& ConnectionFilteredObjects = ConnectionInfo.ConnectionFilteredObjects;
+	FNetBitArray& ObjectsInScopeBeforeDynamicFiltering = ConnectionInfo.ObjectsInScopeBeforeDynamicFiltering;
+
+	auto&& ClearGroupFilterForObject = [&GroupFilteredOutObjects, &ConnectionFilteredObjects, &ObjectsInScopeBeforeDynamicFiltering](uint32 InternalObjectIndex)
+	{
+		GroupFilteredOutObjects.ClearBit(InternalObjectIndex);
+		ObjectsInScopeBeforeDynamicFiltering.SetBitValue(InternalObjectIndex, ConnectionFilteredObjects.GetBit(InternalObjectIndex));
+	};
+
+	if (!GetIsFilteredOutByAnyGroup(ObjectIndex, ConnectionId))
+	{
+		ClearGroupFilterForObject(ObjectIndex);
+	
+		// Filter subobjects
+		for (const FInternalNetHandle SubObjectIndex : NetHandleManager->GetSubObjects(ObjectIndex))
+		{
+			if (!GetIsFilteredOutByAnyGroup(SubObjectIndex, ConnectionId))
+			{
+				ClearGroupFilterForObject(SubObjectIndex);
+			}
+		}
+		
+		return true;
+	}
+
+	return false;
+}
+
+void FReplicationFiltering::InternalSetGroupFilterStatus(FNetObjectGroupHandle GroupHandle, uint32 ConnectionId, ENetFilterStatus ReplicationStatus)
+{
+	if (GroupHandle == NotReplicatedNetObjectGroupHandle)
+	{
+		UE_LOG(LogIrisFiltering, Warning, TEXT("FReplicationFiltering::SetGroupFilterStatus - Trying to set filter NotReplicatedNetObjectGroupHandle GroupIndex: %u which is not allowed."), GroupHandle);		
+		return;
+	}
+
+	if (!FilterGroups.GetBit(GroupHandle))
+	{
+		UE_LOG(LogIrisFiltering, Warning, TEXT("FReplicationFiltering::SetGroupFilterStatus - Trying to set invalid filter ConnectionId: %u, GroupIndex: %u"), ConnectionId, GroupHandle);
+		return;
+	}
+
+	if (ReplicationStatus == ENetFilterStatus::Disallow)
+	{
+		// Filter out objects in group
+		FPerObjectInfo* ConnectionState = GetPerObjectInfo(GroupInfos[GroupHandle].ConnectionStateIndex);
+		if (GetConnectionFilterStatus(*ConnectionState, ConnectionId) == ENetFilterStatus::Allow)
+		{
+			// Mark filter as active for the connection
+			SetConnectionFilterStatus(*ConnectionState, ConnectionId, ENetFilterStatus::Disallow);
+
+			UE_LOG(LogIrisFiltering, Verbose, TEXT("ReplicationFiltering::SetGroupFilterStatus GroupIndex: %u, ConnectionId: %u, FilterStatus: DisallowReplication"), GroupHandle, ConnectionId);
+
+			if (ValidConnections.GetBit(ConnectionId))
+			{
+				// New Connections will be initialized the next filter update
+				if (!NewConnections.GetBit(ConnectionId))
+				{
+					const FNetObjectGroup* Group = Groups->GetGroup(GroupHandle);
+			
+					// Update group filter mask for the connection
+					FPerConnectionInfo& ConnectionInfo = ConnectionInfos[ConnectionId];
+					FNetBitArray& GroupFilteredOutObjects = ConnectionInfo.GroupFilteredOutObjects;
+					FNetBitArray& ObjectsInScopeBeforeDynamicFiltering = ConnectionInfo.ObjectsInScopeBeforeDynamicFiltering;
+					const FNetBitArray& ScopableObjects = NetHandleManager->GetScopableInternalIndices();
+					for (uint32 ObjectIndex : MakeArrayView(Group->Members.GetData(), Group->Members.Num()))
+					{
+						GroupFilteredOutObjects.SetBit(ObjectIndex);
+						ObjectsInScopeBeforeDynamicFiltering.ClearBit(ObjectIndex);
+						// Filter subobjects
+						for (const FInternalNetHandle SubObjectIndex : NetHandleManager->GetSubObjects(ObjectIndex))
+						{
+							GroupFilteredOutObjects.SetBitValue(SubObjectIndex, ScopableObjects.GetBit(SubObjectIndex));
+							ObjectsInScopeBeforeDynamicFiltering.ClearBit(SubObjectIndex);
+						}
+					}
+				}
+			}
+		}
+	}
+	else
+	{
+		// Clear filter effects
+		FPerObjectInfo* ConnectionState = GetPerObjectInfo(GroupInfos[GroupHandle].ConnectionStateIndex);
+		if (GetConnectionFilterStatus(*ConnectionState, ConnectionId) == ENetFilterStatus::Disallow)
+		{
+			// Mark filter as no longer being active for the connection
+			SetConnectionFilterStatus(*ConnectionState, ConnectionId, ENetFilterStatus::Allow);
+
+			UE_LOG(LogIrisFiltering, Verbose, TEXT("ReplicationFiltering::SetGroupFilterStatus GroupIndex: %u, ConnectionId: %u, FilterStatus: AllowReplication"), GroupHandle, ConnectionId);
+
+			if (ValidConnections.GetBit(ConnectionId))
+			{
+				// New Connections will be initialized the next filter update
+				if (!NewConnections.GetBit(ConnectionId))
+				{
+					const FNetObjectGroup* Group = Groups->GetGroup(GroupHandle);
+
+					// Update group filter mask for the connection
+					const FNetBitArray& ScopedObjects = NetHandleManager->GetScopableInternalIndices();
+			
+					for (uint32 ObjectIndex : MakeArrayView(Group->Members.GetData(), Group->Members.Num()))
+					{
+						UpdateGroupFilterEffectsForObject(ObjectIndex, ConnectionId, ScopedObjects);
+					}
+				}
+			}
+		}
+	}
+}
+
+bool FReplicationFiltering::GetGroupFilterStatus(FNetObjectGroupHandle GroupHandle, uint32 ConnectionId, ENetFilterStatus& OutReplicationStatus) const
+{
+	if (!(ValidConnections.GetBit(ConnectionId) && FilterGroups.GetBit(GroupHandle)))
+	{
+		return false;
+	}
+
+	const FPerObjectInfo* ConnectionState = GetPerObjectInfo(GroupInfos[GroupHandle].ConnectionStateIndex);
+	OutReplicationStatus = GetConnectionFilterStatus(*ConnectionState, ConnectionId);
+
+	return true;
+}
+
+void FReplicationFiltering::NotifyObjectAddedToGroup(FNetObjectGroupHandle GroupHandle, uint32 ObjectIndex)
+{
+	// As we want to avoid the fundamentally bad case of doing tons of single object adds to a group and updating group filter for every call we deffer the update until next call to filter
+	if (FilterGroups.GetBit(GroupHandle))
+	{
+		DirtyFilterGroups.SetBit(GroupHandle);
+
+		UE_LOG(LogIrisFiltering, Verbose, TEXT("ReplicationFiltering::NotifyObjectAddedToGroup Added %s to Groupfilter with GroupIndex: %u"), *(NetHandleManager->GetReplicatedObjectDataNoCheck(ObjectIndex).Handle.ToString()), GroupHandle);
+	}
+	else if (SubObjectFilterGroups.GetBit(GroupHandle))
+	{
+		UE_LOG(LogIrisFiltering, Verbose, TEXT("ReplicationFiltering::NotifyObjectAddedToGroup Added %s to SubObjectFilter with GroupIndex: %u"), *(NetHandleManager->GetReplicatedObjectDataNoCheck(ObjectIndex).Handle.ToString()), GroupHandle);
+	}
+}
+
+void FReplicationFiltering::NotifyObjectRemovedFromGroup(FNetObjectGroupHandle GroupHandle, uint32 ObjectIndex)
+{
+	UE_LOG(LogIrisFiltering, Verbose, TEXT("ReplicationFiltering::NotifyObjectRemovedFromGroup GroupIndex: %u, %s"), GroupHandle, *(NetHandleManager->GetReplicatedObjectDataNoCheck(ObjectIndex).Handle.ToString()));
+
+	if (FilterGroups.GetBit(GroupHandle))
+	{
+		FPerObjectInfo* ConnectionState = GetPerObjectInfo(GroupInfos[GroupHandle].ConnectionStateIndex);
+		const FNetBitArray& ScopedObjects = NetHandleManager->GetScopableInternalIndices();
+
+		auto&& MarkGroupsDirty = [this, ObjectIndex, GroupHandle, ConnectionState, &ScopedObjects](uint32 ConnectionId)
+		{
+			if (GetConnectionFilterStatus(*ConnectionState, ConnectionId) == ENetFilterStatus::Disallow)
+			{
+				FNetBitArray& GroupFilteredOutObjects = this->ConnectionInfos[ConnectionId].GroupFilteredOutObjects;
+				UpdateGroupFilterEffectsForObject(ObjectIndex, ConnectionId, ScopedObjects);
+			}
+		};
+
+		// We need to update groupfilter status for all initialized connections
+		FNetBitArray::ForAllSetBits(ValidConnections, NewConnections, FNetBitArrayBase::AndNotOp, MarkGroupsDirty);
+	}
+	else if (SubObjectFilterGroups.GetBit(GroupHandle))
+	{		
+		DirtySubObjectFilterGroups.SetBit(GroupHandle);
+	}
+}
+
+void FReplicationFiltering::NotifyAddedDependentObject(uint32 ObjectIndex)
+{
+	ObjectsRequiringDynamicFilterUpdate.SetBit(ObjectIndex);
+}
+
+void FReplicationFiltering::NotifyRemovedDependentObject(uint32 ObjectIndex)
+{
+	ObjectsRequiringDynamicFilterUpdate.SetBit(ObjectIndex);
+}
+
+ENetFilterStatus FReplicationFiltering::GetConnectionFilterStatus(const FPerObjectInfo& ObjectInfo, uint32 ConnectionId) const
+{
+	return (ObjectInfo.ConnectionIds[ConnectionId >> 5U] & (1U << (ConnectionId & 31U))) != 0U ? ENetFilterStatus::Allow : ENetFilterStatus::Disallow;
+}
+
+bool FReplicationFiltering::IsConnectionFilterStatusAllowedForAnyConnection(const FPerObjectInfo& ObjectInfo) const
+{
+	for (uint32 WordIt = 0, WordEndIt = PerObjectInfoStorageCountPerItem; WordIt != WordEndIt; ++WordIt)
+	{
+		if (ObjectInfo.ConnectionIds[WordIt] != 0U)
+		{
+			return true;
+		}
+	}
+
+	return false;
+}
+
+void FReplicationFiltering::SetConnectionFilterStatus(FPerObjectInfo& ObjectInfo, uint32 ConnectionId, ENetFilterStatus ReplicationStatus)
+{
+	const uint32 WordMask = 1 << (ConnectionId & 31U);
+	const uint32 ValueMask = ReplicationStatus == ENetFilterStatus::Allow ? WordMask : 0U;
+
+	ObjectInfo.ConnectionIds[ConnectionId >> 5U] = (ObjectInfo.ConnectionIds[ConnectionId >> 5U] & ~WordMask) | ValueMask;
+}
+
+void FReplicationFiltering::SetPerObjectInfoFilterStatus(FPerObjectInfo& ObjectInfo, ENetFilterStatus ReplicationStatus)
+{
+	const uint32 InitialValue = ReplicationStatus == ENetFilterStatus::Allow ? ~0U : 0U;
+	for (uint32 WordIt = 0, WordEndIt = PerObjectInfoStorageCountPerItem; WordIt != WordEndIt; ++WordIt)
+	{
+		ObjectInfo.ConnectionIds[WordIt] = InitialValue;
+	}
+}
+
+void FReplicationFiltering::InitFilters()
+{
+	const UNetObjectFilterDefinitions* FilterDefinitions = GetDefault<UNetObjectFilterDefinitions>();
+	for (const FNetObjectFilterDefinition& FilterDefinition : FilterDefinitions->GetFilterDefinitions())
+	{
+		constexpr UObject* ClassOuter = nullptr;
+		constexpr bool bMatchExactClass = true;
+		UClass* NetObjectFilterClass = Cast<UClass>(StaticFindObject(UClass::StaticClass(), ClassOuter, ToCStr(FilterDefinition.ClassName.ToString()), bMatchExactClass));
+		if (NetObjectFilterClass == nullptr || !NetObjectFilterClass->IsChildOf(UNetObjectFilter::StaticClass()))
+		{
+			ensureAlwaysMsgf(NetObjectFilterClass != nullptr, TEXT("NetObjectFilter class is not a NetObjectFilter or could not be found: %s"), ToCStr(FilterDefinition.ClassName.ToString()));
+			continue;
+		}
+
+		UClass* NetObjectFilterConfigClass = nullptr;
+		if (!FilterDefinition.ConfigClassName.IsNone())
+		{
+			NetObjectFilterConfigClass = Cast<UClass>(StaticFindObject(UClass::StaticClass(), ClassOuter, ToCStr(FilterDefinition.ConfigClassName.ToString()), bMatchExactClass));
+			if (NetObjectFilterConfigClass == nullptr || !NetObjectFilterConfigClass->IsChildOf(UNetObjectFilterConfig::StaticClass()))
+			{
+				ensureAlwaysMsgf(NetObjectFilterConfigClass != nullptr, TEXT("NetObjectFilterConfig class is not a NetObjectFilterConfig or could not be found: %s"), ToCStr(FilterDefinition.ConfigClassName.ToString()));
+				continue;
+			}
+		}
+
+		TStrongObjectPtr<UNetObjectFilter> Filter(NewObject<UNetObjectFilter>((UObject*)GetTransientPackage(), NetObjectFilterClass, MakeUniqueObjectName(nullptr, NetObjectFilterClass, FilterDefinition.FilterName)));
+		check(Filter.IsValid());
+
+		FNetObjectFilterInitParams InitParams;
+		InitParams.ReplicationSystem = ReplicationSystem;
+		InitParams.Config = (NetObjectFilterConfigClass ? NewObject<UNetObjectFilterConfig>((UObject*)GetTransientPackage(), NetObjectFilterConfigClass) : nullptr);
+		InitParams.MaxObjectCount = MaxObjectCount;
+		InitParams.MaxConnectionCount = Connections->GetMaxConnectionCount();
+
+		Filter->Init(InitParams);
+
+		FFilterInfo& Info = DynamicFilterInfos.Emplace_GetRef();
+		Info.Filter = MoveTemp(Filter);
+		Info.Name = FilterDefinition.FilterName;
+		Info.ObjectCount = 0;
+		Info.FilteredObjects.Init(MaxObjectCount);
+	}
+
+	if (HasDynamicFilters())
+	{
+		DynamicFilterEnabledObjects.Init(MaxObjectCount);
+	}
+}
+
+void FReplicationFiltering::RemoveFromDynamicFilter(uint32 ObjectIndex, uint32 FilterIndex)
+{
+	ObjectIndexToDynamicFilterIndex[ObjectIndex] = InvalidDynamicFilterIndex;
+	FNetObjectFilteringInfo& NetObjectFilteringInfo = NetObjectFilteringInfos[ObjectIndex];
+	FFilterInfo& FilterInfo = DynamicFilterInfos[FilterIndex];
+	--FilterInfo.ObjectCount;
+	FilterInfo.FilteredObjects.ClearBit(ObjectIndex);
+	FilterInfo.Filter->RemoveObject(ObjectIndex, NetObjectFilteringInfo);
+
+	DynamicFilterEnabledObjects.ClearBit(ObjectIndex);
+	ObjectsRequiringDynamicFilterUpdate.SetBit(ObjectIndex);
+}
+
+void FReplicationFiltering::InvalidateBaselinesForObject(uint32 ObjectIndex, uint32 NewOwningConnectionId, uint32 PrevOwningConnectionId)
+{
+	{
+		const FNetHandleManager::FReplicatedObjectData& ObjectData = NetHandleManager->GetReplicatedObjectDataNoCheck(ObjectIndex);
+		// $IRIS TODO Only invalidate baselines if the object has conditions based on owner/non-owner.
+		if (EnumHasAnyFlags(ObjectData.Protocol->ProtocolTraits, EReplicationProtocolTraits::HasLifetimeConditionals))
+		{
+			for (const uint32 ConnId : {NewOwningConnectionId, PrevOwningConnectionId})
+			{
+				if (ConnId)
+				{
+					BaselineInvalidationTracker->InvalidateBaselines(ObjectIndex, ConnId);
+				}
+			}
+		}
+	}
+
+	// Invalidate baselines for subobjects
+	for (const FInternalNetHandle SubObjectIndex : NetHandleManager->GetSubObjects(ObjectIndex))
+	{
+		const FNetHandleManager::FReplicatedObjectData& SubObjectData = NetHandleManager->GetReplicatedObjectDataNoCheck(SubObjectIndex);
+		if (SubObjectData.IsSubObject())
+		{
+			if (EnumHasAnyFlags(SubObjectData.Protocol->ProtocolTraits, EReplicationProtocolTraits::HasLifetimeConditionals))
+			{
+				for (const uint32 ConnId : {NewOwningConnectionId, PrevOwningConnectionId})
+				{
+					if (ConnId)
+					{
+						BaselineInvalidationTracker->InvalidateBaselines(ObjectIndex, ConnId);
+					}
+				}
+			}
+		}
+	}
+}
+
+// FNetObjectFilterHandleUtil
+inline bool FNetObjectFilterHandleUtil::IsInvalidHandle(FNetObjectFilterHandle Handle)
+{
+	return Handle == InvalidNetObjectFilterHandle;
+}
+
+inline bool FNetObjectFilterHandleUtil::IsDynamicFilter(FNetObjectFilterHandle Handle)
+{
+	return (Handle & DynamicNetObjectFilterHandleFlag) != 0;
+}
+
+inline bool FNetObjectFilterHandleUtil::IsStaticFilter(FNetObjectFilterHandle Handle)
+{
+	return (Handle & DynamicNetObjectFilterHandleFlag) == 0;
+}
+
+inline FNetObjectFilterHandle FNetObjectFilterHandleUtil::MakeDynamicFilterHandle(uint32 FilterIndex)
+{
+	const FNetObjectFilterHandle Handle = DynamicNetObjectFilterHandleFlag | FNetObjectFilterHandle(FilterIndex);
+	return Handle;
+}
+
+inline uint32 FNetObjectFilterHandleUtil::GetDynamicFilterIndex(FNetObjectFilterHandle Handle)
+{
+	return (Handle & DynamicNetObjectFilterHandleFlag) ? (Handle & ~DynamicNetObjectFilterHandleFlag) : ~0U;
+}
+
+}
diff --git a/Engine/Source/Runtime/Experimental/Iris/Core/Private/Iris/ReplicationSystem/Filtering/ReplicationFiltering.h b/Engine/Source/Runtime/Experimental/Iris/Core/Private/Iris/ReplicationSystem/Filtering/ReplicationFiltering.h
new file mode 100644
--- /dev/null
+++ b/Engine/Source/Runtime/Experimental/Iris/Core/Private/Iris/ReplicationSystem/Filtering/ReplicationFiltering.h
@@ -0,0 +1,258 @@
+// Copyright Epic Games, Inc. All Rights Reserved.
+
+#pragma once
+
+#include "CoreTypes.h"
+#include "Net/Core/NetBitArray.h"
+#include "Iris/ReplicationSystem/Filtering/NetObjectFilter.h"
+#include "Containers/Array.h"
+#include "UObject/StrongObjectPtr.h"
+
+class UReplicationSystem;
+namespace UE::Net
+{
+	typedef uint32 FNetObjectFilterHandle;
+	typedef uint16 FNetObjectGroupHandle;
+	namespace Private
+	{
+		class FDeltaCompressionBaselineInvalidationTracker;
+		class FNetHandleManager;
+		class FNetObjectGroups;
+		class FReplicationConnections;
+	}
+}
+
+namespace UE::Net::Private
+{
+
+struct FReplicationFilteringInitParams
+{
+	TObjectPtr<UReplicationSystem> ReplicationSystem;
+	const FNetHandleManager* NetHandleManager = nullptr;
+	FNetObjectGroups* Groups = nullptr;
+	FDeltaCompressionBaselineInvalidationTracker* BaselineInvalidationTracker = nullptr;
+	FReplicationConnections* Connections = nullptr;
+	uint32 MaxGroupCount = 0;
+	uint32 MaxObjectCount = 0;
+};
+
+class FReplicationFiltering
+{
+public:
+	FReplicationFiltering();
+
+	void Init(FReplicationFilteringInitParams& Params);
+
+	/**
+	 * Determine which objects are allowed to be replicated to which connections.
+	 * This is done for all connections every tick.
+	 */
+	void Filter(const FNetBitArrayView& DirtyObjects);
+
+	FNetBitArrayView GetObjectsInScope(uint32 ConnectionId) const;
+
+	FNetBitArrayView GetGroupFilteredOutObjects(uint32 ConnectionId) const;
+
+	// Who owns what?
+	void SetOwningConnection(uint32 ObjectIndex, uint32 ConnectionId);
+	uint32 GetOwningConnection(uint32 ObjectIndex) const { return !bHasDirtyOwner ? ObjectIndexToOwningConnection[ObjectIndex] : GetOwningConnectionIfDirty(ObjectIndex); }
+
+	// Various filters
+	bool SetFilter(uint32 ObjectIndex, FNetObjectFilterHandle Filter);
+
+	// Set whether an object is allowed to be replicated to certain connections or not.
+	bool SetConnectionFilter(uint32 ObjectIndex, const FNetBitArrayView& ConnectionIndices, ENetFilterStatus ReplicationStatus);
+
+	FNetObjectFilterHandle GetFilterHandle(const FName FilterName) const;
+	UNetObjectFilter* GetFilter(const FName FilterName) const;
+
+	// Connection handling
+	void AddConnection(uint32 ConnectionId);
+	void RemoveConnection(uint32 ConnectionId);
+
+	// Group based filtering
+	void AddGroupFilter(FNetObjectGroupHandle GroupHandle);
+	void RemoveGroupFilter(FNetObjectGroupHandle GroupHandle);
+	bool IsGroupFilterGroup(FNetObjectGroupHandle GroupHandle) const { return GroupHandle && FilterGroups.GetBit(GroupHandle); }
+
+	void SetGroupFilterStatus(FNetObjectGroupHandle GroupHandle, ENetFilterStatus ReplicationStatus);
+	void SetGroupFilterStatus(FNetObjectGroupHandle GroupHandle, const FNetBitArrayView& ConnectionsBitArray, ENetFilterStatus);
+	void SetGroupFilterStatus(FNetObjectGroupHandle GroupHandle, uint32 ConnectionId, ENetFilterStatus ReplicationStatus);
+	bool GetGroupFilterStatus(FNetObjectGroupHandle GroupHandle, uint32 ConnectionId, ENetFilterStatus& OutReplicationStatus) const;
+
+	void NotifyObjectAddedToGroup(FNetObjectGroupHandle GroupHandle, uint32 ObjectIndex);
+	void NotifyObjectRemovedFromGroup(FNetObjectGroupHandle GroupHandle, uint32 ObjectIndex);
+
+	void NotifyAddedDependentObject(uint32 ObjectIndex);
+	void NotifyRemovedDependentObject(uint32 ObjectIndex);
+	
+	// SubOjectFilter status
+	void AddSubObjectFilter(FNetObjectGroupHandle GroupHandle);
+	void RemoveSubObjectFilter(FNetObjectGroupHandle GroupHandle);
+	bool IsSubObjectFilterGroup(FNetObjectGroupHandle GroupHandle) const { return GroupHandle && SubObjectFilterGroups.GetBit(GroupHandle); }
+
+	void SetSubObjectFilterStatus(FNetObjectGroupHandle GroupHandle, ENetFilterStatus ReplicationStatus);
+	void SetSubObjectFilterStatus(FNetObjectGroupHandle GroupHandle, const FNetBitArrayView& ConnectionsBitArray, ENetFilterStatus);
+	void SetSubObjectFilterStatus(FNetObjectGroupHandle GroupHandle, uint32 ConnectionId, ENetFilterStatus ReplicationStatus);
+	bool GetSubObjectFilterStatus(FNetObjectGroupHandle GroupHandle, uint32 ConnectionId, ENetFilterStatus& OutReplicationStatus) const;
+
+private:
+	struct FPerConnectionInfo
+	{
+		// Objects filtered depending on owning connection or user set connection filtering
+		FNetBitArray ConnectionFilteredObjects;
+		// Objects filtered out due to one or more groups it belongs to is filtered out
+		FNetBitArray GroupFilteredOutObjects;
+		// Connection and group filtering is assumed to happen seldom. Avoid recalculating from scratch every frame.
+		FNetBitArray ObjectsInScopeBeforeDynamicFiltering;
+		// Objects in scope after all kinds of filtering, including dynamic filtering, has been applied
+		FNetBitArray ObjectsInScope;
+
+		// Which objects are filtered out by dynamic filters.
+		FNetBitArray DynamicFilteredOutObjects;
+	};
+
+	struct FPerObjectInfo
+	{
+		// Note: Array is likely larger than one element.
+		uint32 ConnectionIds[1];
+	};
+
+	static constexpr uint32 UsedPerObjectInfoStorageGrowSize = 32; // 256 bytes, 1024 indices
+	typedef uint16 PerObjectInfoIndexType;
+
+	struct FPerGroupInfo
+	{
+		PerObjectInfoIndexType ConnectionStateIndex;
+	};
+
+	struct FFilterInfo
+	{
+		TStrongObjectPtr<UNetObjectFilter> Filter;
+		FName Name;
+		uint32 ObjectCount = 0;
+		// Objects with this filter set.
+		FNetBitArray FilteredObjects;
+	};
+
+private:
+	class FUpdateDirtyObjectsBatchHelper;
+	
+	static void StaticChecks();
+
+	void InitFilters();
+
+	void InitNewConnections();
+	void ResetRemovedConnections();
+	void UpdateObjectsInScope();
+	void UpdateOwnerAndConnectionFiltering();
+	void UpdateGroupFiltering();
+	void PreUpdateDynamicFiltering();
+	void UpdateDynamicFiltering();
+	void PostUpdateDynamicFiltering();
+	void UpdateSubObjectFilters();
+
+	bool HasDynamicFilters() const;
+
+	// Helper to update and reset group filter effects if objects are removed from a filter or after a filter status change, returns true of the group filter was changed
+	bool UpdateGroupFilterEffectsForObject(uint32 ObjectIndex, uint32 ConnectionId, const FNetBitArray& ScopableObjects);
+	bool HasOwnerFilter(uint32 ObjectIndex) const;
+	bool HasConnectionFilter(uint32 ObjectIndex) const;
+
+	PerObjectInfoIndexType AllocPerObjectInfo();
+	void FreePerObjectInfo(PerObjectInfoIndexType Index);
+
+	FPerObjectInfo* AllocPerObjectInfoForObject(uint32 ObjectIndex);
+	void FreePerObjectInfoForObject(uint32 ObjectIndex);
+	
+	FPerObjectInfo* GetPerObjectInfo(PerObjectInfoIndexType Index);
+	const FPerObjectInfo* GetPerObjectInfo(PerObjectInfoIndexType Index) const;
+
+	void SetPerObjectInfoFilterStatus(FPerObjectInfo& ObjectInfo, ENetFilterStatus ReplicationStatus);
+
+	ENetFilterStatus GetConnectionFilterStatus(const FPerObjectInfo& ObjectInfo, uint32 ConnectionId) const;
+	bool IsConnectionFilterStatusAllowedForAnyConnection(const FPerObjectInfo& ObjectInfo) const;
+	void SetConnectionFilterStatus(FPerObjectInfo& ObjectInfo, uint32 ConnectionId, ENetFilterStatus ReplicationStatus);
+	bool GetIsFilteredOutByAnyGroup(uint32 ObjectInternalIndex, uint32 ConnectionId) const;
+	void InternalSetGroupFilterStatus(FNetObjectGroupHandle GroupHandle, uint32 ConnectionId, ENetFilterStatus ReplicationStatus);
+	uint32 GetOwningConnectionIfDirty(uint32 ObjectIndex) const;
+
+	void RemoveFromDynamicFilter(uint32 ObjectIndex, uint32 FilterIndex);
+
+	void NotifyFiltersOfDirtyObjects(const FNetBitArrayView& DirtyObjects);
+	void BatchNotifyFiltersOfDirtyObjects(FUpdateDirtyObjectsBatchHelper& BatchHelper, const uint32* ObjectIndices, uint32 ObjectCount);
+
+	void InvalidateBaselinesForObject(uint32 ObjectIndex, uint32 NewOwningConnectionId, uint32 PrevOwningConnectionId);
+
+private:
+	// Used for ObjectIndexToDynamicFilterIndex lookup
+	static constexpr uint8 InvalidDynamicFilterIndex = 255U;
+
+	// General
+	TObjectPtr<UReplicationSystem> ReplicationSystem;
+	const FNetHandleManager* NetHandleManager = nullptr;
+
+	// Groups
+	FNetObjectGroups* Groups = nullptr;
+
+	// Baseline invalidation tracker
+	FDeltaCompressionBaselineInvalidationTracker* BaselineInvalidationTracker = nullptr;
+
+	// Connection specifics
+	FReplicationConnections* Connections = nullptr;
+	TArray<FPerConnectionInfo> ConnectionInfos;
+	FNetBitArray ValidConnections;
+	FNetBitArray NewConnections;
+
+	// Object specifics
+	uint32 MaxObjectCount = 0;
+	uint32 WordCountForObjectBitArrays = 0;
+
+	// Filter specifics
+	FNetBitArray ObjectsWithDirtyConnectionFilter;
+	FNetBitArray ObjectsWithDirtyOwner;
+
+	FNetBitArray ObjectsWithOwnerFilter;
+	TArray<uint16> ObjectIndexToOwningConnection;
+
+	// For non-owner filtered objects
+	// The storage for PerObjectInfo.
+	TArray<uint32> PerObjectInfoStorage;
+	// Storage for bit array indicating used/free status of PerObjectInfos
+	TArray<uint32> UsedPerObjectInfoStorage;
+
+	FNetBitArray ObjectsWithPerObjectInfo;
+
+	// Groups
+	TArray<FPerGroupInfo> GroupInfos;
+	uint32 MaxGroupCount = 0;
+	FNetBitArray FilterGroups;
+	FNetBitArray DirtyFilterGroups;
+	FNetBitArray SubObjectFilterGroups;
+	FNetBitArray DirtySubObjectFilterGroups;
+
+	TArray<PerObjectInfoIndexType> ObjectIndexToPerObjectInfoIndex;
+	uint32 PerObjectInfoStorageCountForConnections = 0;
+	// How many elements from UsedPerObjectInfoStorage is needed to hold a FPerObjectInfo
+	uint32 PerObjectInfoStorageCountPerItem = 0;
+
+	// Dynamic filters 
+	TArray<FNetObjectFilteringInfo> NetObjectFilteringInfos;
+	TArray<uint8> ObjectIndexToDynamicFilterIndex;
+	TArray<FFilterInfo> DynamicFilterInfos;
+	FNetBitArray DynamicFilterEnabledObjects;
+	FNetBitArray ObjectsRequiringDynamicFilterUpdate;
+
+	//
+	uint32 bHasNewConnection : 1;
+	uint32 bHasRemovedConnection : 1;
+	uint32 bHasDirtyConnectionFilter: 1;
+	uint32 bHasDirtyOwner : 1;
+};
+
+inline bool FReplicationFiltering::HasDynamicFilters() const
+{
+	return DynamicFilterInfos.Num() > 0;
+}
+
+}
diff --git a/Engine/Source/Runtime/Experimental/Iris/Core/Private/Iris/ReplicationSystem/LegacyPushModel.h b/Engine/Source/Runtime/Experimental/Iris/Core/Private/Iris/ReplicationSystem/LegacyPushModel.h
new file mode 100644
--- /dev/null
+++ b/Engine/Source/Runtime/Experimental/Iris/Core/Private/Iris/ReplicationSystem/LegacyPushModel.h
@@ -0,0 +1,96 @@
+// Copyright Epic Games, Inc. All Rights Reserved.
+
+#pragma once
+
+#if WITH_PUSH_MODEL
+
+#include "CoreTypes.h"
+#include "Iris/ReplicationSystem/NetHandle.h"
+#include "Net/Core/PushModel/PushModel.h"
+
+namespace UE::Net
+{
+	struct FReplicationProtocol;
+	struct FReplicationInstanceProtocol;
+}
+
+namespace UE::Net::Private
+{
+
+extern IRISCORE_API bool bIsIrisPushModelForceEnabled;
+extern IRISCORE_API int IrisPushModelMode;
+inline bool IsIrisPushModelEnabled(bool bIsPushModelEnabled = IS_PUSH_MODEL_ENABLED()) { return (bIsIrisPushModelForceEnabled | (IrisPushModelMode > 0)) & bIsPushModelEnabled; }
+
+struct FNetPushObjectHandle
+{
+public:
+	FNetPushObjectHandle() : Value(Invalid) {}
+	FNetPushObjectHandle(uint32 InIndex, uint32 InReplicationSystemId) : Index(InIndex), ReplicationSystemId(InReplicationSystemId + 1U), ZeroPadding(0) {}
+
+	uint32 GetInternalIndex() const { return Index; }
+	uint32 GetReplicationSystemId() const { check(ReplicationSystemId != 0U); return ReplicationSystemId - 1U; }
+	bool IsValid() const { return Value != Invalid; }
+
+private:
+	friend struct FNetHandleLegacyPushModelHelper;
+
+	enum : uint64 { Invalid = 0 };
+
+	explicit FNetPushObjectHandle(UEPushModelPrivate::FNetIrisPushObjectId PushId) : Value(PushId) {}
+	uint64 GetValue() const { return Value; }
+
+	union 
+	{
+		struct
+		{
+			uint64 Index : FNetHandle::IdBits;								// Ever increasing index, we use different indices for static and dynamic object
+			uint64 ReplicationSystemId : FNetHandle::ReplicationSystemIdBits;	// ReplicationSystemId, when running in pie, we track the owning instance
+			uint64 ZeroPadding : (64 - FNetHandle::IdBits - FNetHandle::ReplicationSystemIdBits);
+		};
+		uint64 Value : 64;
+	};
+};
+
+struct FNetHandleLegacyPushModelHelper
+{
+	static void InitPushModel();
+	static void ShutdownPushModel();
+
+	IRISCORE_API static void SetNetPushID(UObject* Object, FNetPushObjectHandle PushHandle);
+	IRISCORE_API static void ClearNetPushID(UObject* Object);
+
+private:
+	static void MarkPropertyOwnerDirty(const UObject* Object, UEPushModelPrivate::FNetIrisPushObjectId PushId, const int32 RepIndex);
+	static void MarkPropertiesOwnerDirty(const UObject* Object, UEPushModelPrivate::FNetIrisPushObjectId PushId, const int32 StartRepIndex, const int32 EndRepIndex);
+
+	static void MarkPropertyDirty(const UObject* Object, UEPushModelPrivate::FNetIrisPushObjectId PushId, const int32 RepIndex);
+	static void MarkPropertiesDirty(const UObject* Object, UEPushModelPrivate::FNetIrisPushObjectId PushId, const int32 StartRepIndex, const int32 EndRepIndex);
+
+	static void OptionallyMarkPropertyOwnerDirty(const UObject* Object, UEPushModelPrivate::FNetIrisPushObjectId PushId, const int32 RepIndex);
+	static void OptionallyMarkPropertiesOwnerDirty(const UObject* Object, UEPushModelPrivate::FNetIrisPushObjectId PushId, const int32 StartRepIndex, const int32 EndRepIndex);
+};
+
+}
+
+#define UE_NET_IRIS_SET_PUSH_ID(...) UE::Net::Private::FNetHandleLegacyPushModelHelper::SetNetPushID(__VA_ARGS__)
+#define UE_NET_IRIS_CLEAR_PUSH_ID(...) UE::Net::Private::FNetHandleLegacyPushModelHelper::ClearNetPushID(__VA_ARGS__)
+
+#define UE_NET_IRIS_INIT_LEGACY_PUSH_MODEL() UE::Net::Private::FNetHandleLegacyPushModelHelper::InitPushModel()
+#define UE_NET_IRIS_SHUTDOWN_LEGACY_PUSH_MODEL()  UE::Net::Private::FNetHandleLegacyPushModelHelper::ShutdownPushModel()
+
+#else
+
+namespace UE::Net::Private
+{
+
+inline constexpr bool IsIrisPushModelEnabled(bool bIsPushModelEnabled = false) { return false; }
+
+}
+
+#define UE_NET_IRIS_SET_PUSH_ID(...)
+#define UE_NET_IRIS_CLEAR_PUSH_ID(...)
+
+#define UE_NET_IRIS_INIT_LEGACY_PUSH_MODEL() 
+#define UE_NET_IRIS_SHUTDOWN_LEGACY_PUSH_MODEL() 
+
+#endif
diff --git a/Engine/Source/Runtime/Experimental/Iris/Core/Private/Iris/ReplicationSystem/NetBlob/NetBlob.cpp b/Engine/Source/Runtime/Experimental/Iris/Core/Private/Iris/ReplicationSystem/NetBlob/NetBlob.cpp
new file mode 100644
--- /dev/null
+++ b/Engine/Source/Runtime/Experimental/Iris/Core/Private/Iris/ReplicationSystem/NetBlob/NetBlob.cpp
@@ -0,0 +1,163 @@
+// Copyright Epic Games, Inc. All Rights Reserved.
+
+#include "Iris/ReplicationSystem/NetBlob/NetBlob.h"
+#include "Iris/ReplicationSystem/ObjectReferenceCache.h"
+#include "Iris/ReplicationSystem/ReplicationOperations.h"
+#include "Iris/ReplicationSystem/ReplicationOperationsInternal.h"
+#include "Iris/ReplicationSystem/ReplicationSystem.h"
+#include "Iris/ReplicationSystem/ReplicationSystemInternal.h"
+#include "Iris/Serialization/NetBitStreamReader.h"
+#include "Iris/Serialization/NetBitStreamWriter.h"
+#include "Iris/Serialization/InternalNetSerializationContext.h"
+#include "Iris/Serialization/NetSerializationContext.h"
+#include "Iris/Serialization/ObjectNetSerializer.h"
+
+namespace UE::Net
+{
+
+// NetBlob
+FNetBlob::FNetBlob(const FNetBlobCreationInfo& InCreationInfo)
+: CreationInfo(InCreationInfo)
+, RefCount(0)
+{
+}
+
+FNetBlob::~FNetBlob()
+{
+	if (const FReplicationStateDescriptor* Descriptor = BlobDescriptor.GetReference())
+	{
+		if (EnumHasAnyFlags(Descriptor->Traits, EReplicationStateTraits::HasDynamicState))
+		{
+			if (uint8* StateBuffer = QuantizedBlobState.Get())
+			{
+				Private::FReplicationStateOperationsInternal::FreeDynamicState(StateBuffer, Descriptor);
+			}
+		}
+	}
+}
+
+void FNetBlob::SetState(const TRefCountPtr<const FReplicationStateDescriptor>& InBlobDescriptor, TUniquePtr<uint8> InQuantizedBlobState)
+{
+	BlobDescriptor = InBlobDescriptor;
+	QuantizedBlobState = MoveTemp(InQuantizedBlobState);
+}
+
+void FNetBlob::SerializeCreationInfo(FNetSerializationContext& Context, const FNetBlobCreationInfo& CreationInfo)
+{
+	FNetBitStreamWriter* Writer = Context.GetBitStreamWriter();
+
+	// Currently we only need to serialize the blob type as reliability is expected to be handled externally.
+	Writer->WriteBits((CreationInfo.Type == 0 ? 0U : 1U), 1U);
+	if (CreationInfo.Type != 0)
+	{
+		Writer->WriteBits(CreationInfo.Type, 7U);
+	}
+}
+
+void FNetBlob::DeserializeCreationInfo(FNetSerializationContext& Context, FNetBlobCreationInfo& OutCreationInfo)
+{
+	FNetBitStreamReader* Reader = Context.GetBitStreamReader();
+	FNetBlobType Type = Reader->ReadBits(1U);
+	if (Type != 0)
+	{
+		Type = Reader->ReadBits(7U);
+	}
+
+	OutCreationInfo.Type = Type;
+	OutCreationInfo.Flags = ENetBlobFlags::None;
+}
+
+void FNetBlob::SerializeWithObject(FNetSerializationContext& Context, FNetHandle NetHandle) const
+{
+	SerializeBlob(Context);
+}
+
+void FNetBlob::DeserializeWithObject(FNetSerializationContext& Context, FNetHandle NetHandle)
+{
+	check(BlobDescriptor.IsValid());
+	DeserializeBlob(Context);
+}
+
+void FNetBlob::Serialize(FNetSerializationContext& Context) const
+{
+	SerializeBlob(Context);
+}
+
+void FNetBlob::Deserialize(FNetSerializationContext& Context)
+{
+	DeserializeBlob(Context);
+}
+
+ENetObjectReferenceResolveResult FNetBlob::ResolveObjectReferences(FNetSerializationContext& Context) const
+{
+	if (BlobDescriptor.IsValid() && QuantizedBlobState.IsValid())
+	{
+		return Private::FReplicationStateOperationsInternal::TryToResolveObjectReferences(Context, QuantizedBlobState.Get(), BlobDescriptor);
+	}
+	
+	return ENetObjectReferenceResolveResult::None;
+}
+
+void FNetBlob::SerializeBlob(FNetSerializationContext& Context) const
+{
+	if (BlobDescriptor.IsValid() && QuantizedBlobState.IsValid())
+	{
+		FReplicationStateOperations::Serialize(Context, QuantizedBlobState.Get(), BlobDescriptor);
+	}
+}
+
+void FNetBlob::DeserializeBlob(FNetSerializationContext& Context)
+{
+	if (BlobDescriptor.IsValid() && QuantizedBlobState.IsValid())
+	{
+		FReplicationStateOperations::Deserialize(Context, QuantizedBlobState.Get(), BlobDescriptor);
+	}
+}
+
+void FNetBlob::Release() const
+{
+	if (--RefCount == 0)
+	{
+		delete this;
+	}
+}
+
+TArrayView<const FNetObjectReference> FNetBlob::GetExports() const
+{
+	return MakeArrayView<const FNetObjectReference>(nullptr, 0);
+};
+
+// NetObjectAttachment
+FNetObjectAttachment::FNetObjectAttachment(const FNetBlobCreationInfo& CreationInfo)
+: FNetBlob(CreationInfo)
+{
+}
+
+FNetObjectAttachment::~FNetObjectAttachment()
+{
+}
+
+void FNetObjectAttachment::SerializeObjectReference(FNetSerializationContext& Context) const
+{
+	WriteFullNetObjectReference(Context, NetObjectReference);
+	WriteFullNetObjectReference(Context, TargetObjectReference);
+}
+
+void FNetObjectAttachment::DeserializeObjectReference(FNetSerializationContext& Context)
+{
+	ReadFullNetObjectReference(Context, NetObjectReference);
+	ReadFullNetObjectReference(Context, TargetObjectReference);
+}
+
+void FNetObjectAttachment::SerializeSubObjectReference(FNetSerializationContext& Context, FNetHandle NetHandle) const
+{
+	WriteFullNetObjectReference(Context, TargetObjectReference);
+}
+
+void FNetObjectAttachment::DeserializeSubObjectReference(FNetSerializationContext& Context, FNetHandle NetHandle)
+{
+	NetObjectReference = Private::FObjectReferenceCache::MakeNetObjectReference(NetHandle);
+	ReadFullNetObjectReference(Context, TargetObjectReference);
+}
+
+}
