Iris experimental replication system.

diff --git a/Engine/Source/Runtime/Experimental/Iris/Core/Private/Iris/ReplicationSystem/NetBlob/NetBlobManager.cpp b/Engine/Source/Runtime/Experimental/Iris/Core/Private/Iris/ReplicationSystem/NetBlob/NetBlobManager.cpp
new file mode 100644
--- /dev/null
+++ b/Engine/Source/Runtime/Experimental/Iris/Core/Private/Iris/ReplicationSystem/NetBlob/NetBlobManager.cpp
@@ -0,0 +1,433 @@
+// Copyright Epic Games, Inc. All Rights Reserved.
+
+#include "Iris/ReplicationSystem/NetBlob/NetBlobManager.h"
+#include "HAL/IConsoleManager.h"
+#include "Iris/ReplicationSystem/NetBlob/NetObjectBlobHandler.h"
+#include "Iris/ReplicationSystem/NetBlob/NetRPCHandler.h"
+#include "Iris/ReplicationSystem/NetBlob/PartialNetObjectAttachmentHandler.h"
+#include "Iris/ReplicationSystem/ReplicationConnections.h"
+#include "Iris/ReplicationSystem/ReplicationSystem.h"
+#include "Iris/ReplicationSystem/ReplicationSystemInternal.h"
+#include "Iris/ReplicationSystem/ReplicationWriter.h"
+#include "Iris/Core/NetObjectReference.h"
+#include "Iris/Core/IrisLog.h"
+
+namespace UE::Net::Private
+{
+
+static TAutoConsoleVariable<int32> CVarEnableIrisRPCs(TEXT("net.Iris.EnableRPCs"), 1, TEXT( "If > 0 let Iris replicate and execute RPCs."));
+
+FNetBlobManager::FNetBlobManager()
+: ReplicationSystem(nullptr)
+, ObjectReferenceCache(nullptr)
+, Connections(nullptr)
+, PartialNetObjectAttachmentHandlerConfig(nullptr)
+, NetHandleManager(nullptr)
+, bIsServer(false)
+, bSendAttachmentsWithObject(false)
+{
+}
+	
+void FNetBlobManager::Init(FNetBlobManagerInitParams& InitParams)
+{
+	BlobHandlerManager.Init();
+
+	ReplicationSystem = InitParams.ReplicationSystem;
+	Connections = &InitParams.ReplicationSystem->GetReplicationSystemInternal()->GetConnections();
+	NetHandleManager = &InitParams.ReplicationSystem->GetReplicationSystemInternal()->GetNetHandleManager();
+	ObjectReferenceCache = &InitParams.ReplicationSystem->GetReplicationSystemInternal()->GetObjectReferenceCache();
+	bIsServer = ReplicationSystem->IsServer();
+	bSendAttachmentsWithObject = InitParams.bSendAttachmentsWithObject; 
+
+	RegisterDefaultHandlers();
+
+	AttachmentSendQueue.Init(this);
+}
+
+bool FNetBlobManager::RegisterNetBlobHandler(UNetBlobHandler* Handler)
+{
+	return BlobHandlerManager.RegisterHandler(Handler);
+}
+
+bool FNetBlobManager::QueueNetObjectAttachment(uint32 ConnectionId, const FNetObjectReference& TargetRef, const TRefCountPtr<FNetObjectAttachment>& Attachment)
+{
+	if (!Attachment.IsValid())
+	{
+		return false;
+	}
+
+	if (!Connections->IsValidConnection(ConnectionId))
+	{
+		UE_LOG(LogIris, Warning, TEXT("Dropping attachment to invalid connection %u."), ConnectionId);
+		return true;
+	}
+
+	FInternalNetHandle OwnerIndex = 0;
+	FInternalNetHandle SubObjectIndex = 0;
+	const FNetObjectReference& NetObjectReference = Attachment->GetNetObjectReference();
+
+	const FNetObjectReference OwnerOrSubObjectReference = TargetRef;	
+	FNetObjectReference OwnerReference = OwnerOrSubObjectReference;
+
+	bool bCanSendRpc = GetOwnerAndSubObjectIndicesFromHandle(OwnerOrSubObjectReference.GetRefHandle(), OwnerIndex, SubObjectIndex);
+	if (!bCanSendRpc)
+	{
+		// If the TargetRef is a valid reference but not replicated, then the owner must be used instead
+		if (OwnerOrSubObjectReference.IsValid())
+		{
+			// See if we can find a replicated handle
+			OwnerReference = ObjectReferenceCache->GetReplicatedOuter(TargetRef);
+			bCanSendRpc = GetOwnerAndSubObjectIndicesFromHandle(OwnerReference.GetRefHandle(), OwnerIndex, SubObjectIndex);
+		}
+
+		if (!bCanSendRpc)
+		{
+			UE_LOG(LogIris, Warning, TEXT("QueueNetObjectAttachment %s Failed. Unable to resolve object reference %s."), *NetObjectReference.ToString(), *OwnerReference.ToString());
+			return false;
+		}
+	}
+
+	Attachment->SetNetObjectReference(OwnerReference, OwnerOrSubObjectReference);
+	AttachmentSendQueue.Enqueue(ConnectionId, OwnerIndex, SubObjectIndex, Attachment);
+	return true;
+}
+
+bool FNetBlobManager::SendRPC(const UObject* Object, const UObject* SubObject, const UFunction* Function, const void* Parameters)
+{
+	if (CVarEnableIrisRPCs.GetValueOnGameThread() <= 0)
+	{
+		return false;
+	}
+
+	UNetRPCHandler* Handler = RPCHandler.Get();
+	if (Handler == nullptr)
+	{
+		return false;
+	}
+
+	// May the RPC be sent?
+	if ((Function->FunctionFlags & (bIsServer ? (FUNC_NetClient | FUNC_NetMulticast) : FUNC_NetServer)) == 0)
+	{
+		checkf(false, TEXT("Trying to call RPC %s in the wrong direction."), ToCStr(Function->GetName()));
+		return true;
+	}
+
+	// Check if there are any connections at all
+	const FNetBitArray& ValidConnections = Connections->GetValidConnections();
+	if (ValidConnections.FindFirstOne() == FNetBitArray::InvalidIndex)
+	{
+		return true;
+	}
+
+	FInternalNetHandle OwnerIndex = 0;
+	FInternalNetHandle SubObjectIndex = 0;
+
+	const FNetObjectReference OwnerOrSubObjectReference = ObjectReferenceCache->GetOrCreateObjectReference(SubObject ? SubObject : Object);	
+	FNetObjectReference OwnerReference = OwnerOrSubObjectReference;
+
+	bool bCanSendRpc = GetOwnerAndSubObjectIndicesFromHandle(OwnerOrSubObjectReference.GetRefHandle(), OwnerIndex, SubObjectIndex);
+	if (!bCanSendRpc)
+	{
+		// If the subobject is a valid reference but not replicated, then the owner must be used instead
+		if (OwnerOrSubObjectReference.IsValid() && SubObject)
+		{
+			OwnerReference = ObjectReferenceCache->GetOrCreateObjectReference(Object);
+			bCanSendRpc = GetOwnerAndSubObjectIndicesFromHandle(OwnerReference.GetRefHandle(), OwnerIndex, SubObjectIndex);
+		}
+
+		if (!bCanSendRpc)
+		{
+			UE_LOG(LogIris, Warning, TEXT("SendRPC %s for %s Failed. Unable to resolve object reference %s."), ToCStr(Function->GetName()), ToCStr(Object->GetName()), ToCStr(OwnerReference.ToString()));
+			return false;
+		}
+	}
+
+	const TRefCountPtr<FNetRPC>& RPC = Handler->CreateRPC(OwnerReference, Function, Parameters);
+	if (!RPC.IsValid())
+	{
+		return true;
+	}
+
+	RPC->SetNetObjectReference(OwnerReference, OwnerOrSubObjectReference);
+	AttachmentSendQueue.Enqueue(OwnerIndex, SubObjectIndex, reinterpret_cast<const TRefCountPtr<FNetObjectAttachment>&>(RPC));
+	return true;
+}
+
+bool FNetBlobManager::SendRPC(uint32 ConnectionId, const UObject* Object, const UObject* SubObject, const UFunction* Function, const void* Parameters)
+{
+	if (CVarEnableIrisRPCs.GetValueOnGameThread() <= 0)
+	{
+		return false;
+	}
+
+	UNetRPCHandler* Handler = RPCHandler.Get();
+	if (Handler == nullptr)
+	{
+		return false;
+	}
+
+	if ((Function->FunctionFlags & (bIsServer ? (FUNC_NetClient | FUNC_NetMulticast) : FUNC_NetServer)) == 0)
+	{
+		checkf(false, TEXT("Trying to call RPC %s in the wrong direction."), ToCStr(Function->GetName()));
+		return true;
+	}
+
+	if (!Connections->IsValidConnection(ConnectionId))
+	{
+		UE_LOG(LogIris, Warning, TEXT("Trying to call RPC on non-existing connection %u."), ConnectionId);
+		return true;
+	}
+
+	FInternalNetHandle OwnerIndex = 0;
+	FInternalNetHandle SubObjectIndex = 0;
+
+	const FNetObjectReference OwnerOrSubObjectReference = ObjectReferenceCache->GetOrCreateObjectReference(SubObject ? SubObject : Object);	
+	FNetObjectReference OwnerReference = OwnerOrSubObjectReference;
+
+	bool bCanSendRpc = GetOwnerAndSubObjectIndicesFromHandle(OwnerOrSubObjectReference.GetRefHandle(), OwnerIndex, SubObjectIndex);
+	if (!bCanSendRpc)
+	{
+		// If the subobject is a valid reference but  not replicated, then the owner must be used instead
+		if (OwnerOrSubObjectReference.IsValid() && SubObject)
+		{
+			OwnerReference = ObjectReferenceCache->GetOrCreateObjectReference(Object);
+			bCanSendRpc = GetOwnerAndSubObjectIndicesFromHandle(OwnerReference.GetRefHandle(), OwnerIndex, SubObjectIndex);
+		}
+
+		if (!bCanSendRpc)
+		{
+			UE_LOG(LogIris, Warning, TEXT("SendRPC %s for %s Failed. Unable to resolve object reference %s."), ToCStr(Function->GetName()), ToCStr(Object->GetName()), ToCStr(OwnerReference.ToString()));
+			return false;
+		}
+	}
+
+	const TRefCountPtr<FNetRPC>& RPC = Handler->CreateRPC(OwnerReference, Function, Parameters);
+	if (!RPC.IsValid())
+	{
+		UE_LOG(LogIris, Warning, TEXT("Unable to create RPC for function %s."), ToCStr(Function->GetName()));
+		return true;
+	}
+
+	RPC->SetNetObjectReference(OwnerReference, OwnerOrSubObjectReference);
+	AttachmentSendQueue.Enqueue(ConnectionId, OwnerIndex, SubObjectIndex, reinterpret_cast<const TRefCountPtr<FNetObjectAttachment>&>(RPC));
+	return true;
+}
+
+bool FNetBlobManager::GetOwnerAndSubObjectIndicesFromHandle(FNetHandle NetHandle, FInternalNetHandle& OutOwnerIndex, FInternalNetHandle& OutSubObjectIndex)
+{
+	if (!NetHandle.IsValid())
+	{
+		return false;
+	}
+
+	const FInternalNetHandle ObjectIndex = NetHandleManager->GetInternalIndex(NetHandle);
+	if (ObjectIndex == FNetHandleManager::InvalidInternalIndex)
+	{
+		return false;
+	}
+
+	const FNetHandleManager::FReplicatedObjectData& ObjectData = NetHandleManager->GetReplicatedObjectDataNoCheck(ObjectIndex);
+	if (ObjectData.SubObjectRootIndex != FNetHandleManager::InvalidInternalIndex)
+	{
+		OutOwnerIndex = ObjectData.SubObjectRootIndex;
+		OutSubObjectIndex = ObjectIndex;
+	}
+	else
+	{
+		OutOwnerIndex = ObjectIndex;
+		OutSubObjectIndex = FNetHandleManager::InvalidInternalIndex;
+	}
+
+	return true;
+}
+
+void FNetBlobManager::ProcessNetObjectAttachmentSendQueue()
+{
+	AttachmentSendQueue.ProcessQueue(Connections);
+}
+
+void FNetBlobManager::AddConnection(uint32 ConnectionId)
+{
+	BlobHandlerManager.AddConnection(ConnectionId);
+}
+
+void FNetBlobManager::RemoveConnection(uint32 ConnectionId)
+{
+	BlobHandlerManager.RemoveConnection(ConnectionId);
+}
+
+void FNetBlobManager::RegisterDefaultHandlers()
+{
+	// NetRPCHandler
+	{
+		RPCHandler = TStrongObjectPtr<UNetRPCHandler>(NewObject<UNetRPCHandler>());
+		RPCHandler->Init(*ReplicationSystem);
+		if (!RegisterNetBlobHandler(RPCHandler.Get()))
+		{
+#if !WITH_AUTOMATION_WORKER
+			checkf(false, TEXT("%s"), TEXT("Unable to register RPC handler. RPCs cannot be sent or received."));
+#endif
+			RPCHandler.Reset();
+		}
+	}
+
+	// PartialNetObjectAttachmentHandler
+	{
+		PartialNetObjectAttachmentHandlerConfig = GetDefault<UPartialNetObjectAttachmentHandlerConfig>();
+		FPartialNetObjectAttachmentHandlerInitParams InitParams = {};
+		InitParams.ReplicationSystem  = ReplicationSystem;
+		InitParams.Config = PartialNetObjectAttachmentHandlerConfig;
+		PartialNetObjectAttachmentHandler = TStrongObjectPtr<UPartialNetObjectAttachmentHandler>(NewObject<UPartialNetObjectAttachmentHandler>());
+		PartialNetObjectAttachmentHandler->Init(InitParams);
+		if (!RegisterNetBlobHandler(PartialNetObjectAttachmentHandler.Get()))
+		{
+#if !WITH_AUTOMATION_WORKER
+			checkf(false, TEXT("%s"), TEXT("Unable to register PartialNetObjectAttachment handler. Attachments cannot be split."));
+#endif
+			PartialNetObjectAttachmentHandler.Reset();
+		}
+	}
+
+	// NetObjectBlobHandler
+	{
+		NetObjectBlobHandler = TStrongObjectPtr<UNetObjectBlobHandler>(NewObject<UNetObjectBlobHandler>());
+		if (!RegisterNetBlobHandler(NetObjectBlobHandler.Get()))
+		{
+#if !WITH_AUTOMATION_WORKER
+			checkf(false, TEXT("%s"), TEXT("Unable to register NetObjectBlobHandler handler. Replicated objects cannot be split."));
+#endif
+			NetObjectBlobHandler.Reset();
+		}
+	}
+}
+
+// FNetObjectAttachmentQueue
+FNetBlobManager::FNetObjectAttachmentSendQueue::FNetObjectAttachmentSendQueue()
+: Manager(nullptr)
+, bHasMulticastAttachments(false)
+{
+}
+
+void FNetBlobManager::FNetObjectAttachmentSendQueue::Init(FNetBlobManager* InManager)
+{
+	Manager = InManager;
+}
+
+void FNetBlobManager::FNetObjectAttachmentSendQueue::Enqueue(uint32 ConnectionId, FInternalNetHandle OwnerIndex, FInternalNetHandle SubObjectIndex, const TRefCountPtr<FNetObjectAttachment>& Attachment)
+{
+	FNetObjectAttachmentQueueEntry& QueueEntry = AttachmentQueue.AddDefaulted_GetRef();
+	QueueEntry.ConnectionId = ConnectionId;
+	QueueEntry.OwnerIndex = OwnerIndex;
+	QueueEntry.SubObjectIndex = SubObjectIndex;
+	QueueEntry.Attachment = Attachment;
+}
+
+void FNetBlobManager::FNetObjectAttachmentSendQueue::Enqueue(FInternalNetHandle OwnerIndex, FInternalNetHandle SubObjectIndex, const TRefCountPtr<FNetObjectAttachment>& Attachment)
+{
+	FNetObjectAttachmentQueueEntry& QueueEntry = AttachmentQueue.AddDefaulted_GetRef();
+	QueueEntry.ConnectionId = 0;
+	QueueEntry.OwnerIndex = OwnerIndex;
+	QueueEntry.SubObjectIndex = SubObjectIndex;
+	QueueEntry.Attachment = Attachment;
+
+	bHasMulticastAttachments = true;
+}
+
+void FNetBlobManager::FNetObjectAttachmentSendQueue::ProcessQueue(FReplicationConnections* InConnections)
+{
+	if (AttachmentQueue.Num() <= 0)
+	{
+		return;
+	}
+
+	const FNetBitArray& ValidConnections = InConnections->GetValidConnections();
+	const FNetBitArrayView ReplicatingConnections = MakeNetBitArrayView(ValidConnections);
+
+	uint32* ConnectionIds = nullptr;
+	uint32 ConnectionCount = 0;
+
+	if (bHasMulticastAttachments)
+	{
+		ConnectionIds = static_cast<uint32*>(FMemory_Alloca(ReplicatingConnections.GetNumBits() * 4));
+		ReplicatingConnections.ForAllSetBits([ConnectionIds, &ConnectionCount](uint32 Bit) { ConnectionIds[ConnectionCount++] = Bit; });
+	}
+
+	TArray<TRefCountPtr<FNetBlob>> PartialNetBlobs;
+	for (const FNetObjectAttachmentQueueEntry& Entry : MakeArrayView(AttachmentQueue))
+	{
+		PartialNetBlobs.Reset();
+
+		const TRefCountPtr<FNetObjectAttachment>& Attachment = Entry.Attachment;
+		const FReplicationStateDescriptor* ReplicationStateDescriptor = Attachment->GetReplicationStateDescriptor();
+
+		const bool bMulticast = Entry.ConnectionId == 0;
+		const bool bHasConnectionSpecificSerialization = ReplicationStateDescriptor && EnumHasAnyFlags(ReplicationStateDescriptor->Traits, EReplicationStateTraits::HasConnectionSpecificSerialization);
+
+		if (!(bMulticast && bHasConnectionSpecificSerialization) && !PreSerializeAndSplitNetBlob(Entry.ConnectionId, Attachment, PartialNetBlobs))
+		{
+			checkf(false, TEXT("Unable to split %s NetObjectAttachment."), (EnumHasAnyFlags(Attachment->GetCreationInfo().Flags, ENetBlobFlags::Reliable) ? TEXT("reliable") : TEXT("unreliable")));
+			continue;
+		}
+
+		const bool bIsMultiPartAttachment = (PartialNetBlobs.Num() > 1);
+		TArrayView<const TRefCountPtr<FNetBlob>> AttachmentsView = MakeArrayView(PartialNetBlobs);
+		if (bMulticast)
+		{
+			if (bIsMultiPartAttachment)
+			{
+				UE_LOG(LogIris, Warning, TEXT("Splitting multicast net object attachment %s."), ReplicationStateDescriptor->DebugName ? ReplicationStateDescriptor->DebugName->Name : TEXT("Unknown"));
+			}
+
+			const FNetBlobCreationInfo& BlobCreationInfo = Attachment->GetCreationInfo();
+			const bool bIsReliableRPC = EnumHasAnyFlags(BlobCreationInfo.Flags, ENetBlobFlags::Reliable);
+			for (uint32 ConnectionId : MakeArrayView(ConnectionIds, ConnectionCount))
+			{
+				// Objects won't be prioritized until there's a view so let's avoid queuing multicast attachments.
+				if (!bIsReliableRPC && InConnections->GetReplicationView(ConnectionId).Views.Num() <= 0)
+				{
+					continue;
+				}
+
+				if (bHasConnectionSpecificSerialization)
+				{
+					PartialNetBlobs.Reset();
+					if (!PreSerializeAndSplitNetBlob(ConnectionId, Attachment, PartialNetBlobs))
+					{
+						checkf(false, TEXT("Unable to split %s NetObjectAttachment with connection specific serialization."), (EnumHasAnyFlags(Attachment->GetCreationInfo().Flags, ENetBlobFlags::Reliable) ? TEXT("reliable") : TEXT("unreliable")));
+						continue;
+					}
+
+					AttachmentsView = MakeArrayView(PartialNetBlobs);
+				}
+
+				FReplicationConnection* Connection = InConnections->GetConnection(ConnectionId);
+				// We're only iterating over valid connections so the Connection pointer must be valid.
+				Connection->ReplicationWriter->QueueNetObjectAttachments(Entry.OwnerIndex, Entry.SubObjectIndex, AttachmentsView);
+			}
+		}
+		else
+		{
+			if (FReplicationConnection* Connection = InConnections->GetConnection(Entry.ConnectionId))
+			{
+				Connection->ReplicationWriter->QueueNetObjectAttachments(Entry.OwnerIndex, Entry.SubObjectIndex, AttachmentsView);
+			}
+		}
+	}
+
+	// Clear queue
+	AttachmentQueue.Reset();
+	bHasMulticastAttachments = false;
+}
+
+bool FNetBlobManager::FNetObjectAttachmentSendQueue::PreSerializeAndSplitNetBlob(uint32 ConnectionId, const TRefCountPtr<FNetObjectAttachment>& Attachment, TArray<TRefCountPtr<FNetBlob>>& OutPartialNetBlobs) const
+{
+	if (!Manager->PartialNetObjectAttachmentHandler.IsValid())
+	{
+		OutPartialNetBlobs.Add(reinterpret_cast<const TRefCountPtr<FNetBlob>&>(Attachment));
+		return true;
+	}
+
+	return Manager->PartialNetObjectAttachmentHandler->PreSerializeAndSplitNetBlob(ConnectionId, Attachment, OutPartialNetBlobs, Manager->bSendAttachmentsWithObject);
+}
+
+}
diff --git a/Engine/Source/Runtime/Experimental/Iris/Core/Private/Iris/ReplicationSystem/NetBlob/NetBlobManager.h b/Engine/Source/Runtime/Experimental/Iris/Core/Private/Iris/ReplicationSystem/NetBlob/NetBlobManager.h
new file mode 100644
--- /dev/null
+++ b/Engine/Source/Runtime/Experimental/Iris/Core/Private/Iris/ReplicationSystem/NetBlob/NetBlobManager.h
@@ -0,0 +1,114 @@
+// Copyright Epic Games, Inc. All Rights Reserved.
+
+#pragma once
+#include "CoreTypes.h"
+#include "Iris/ReplicationSystem/NetBlob/NetBlob.h"
+#include "Iris/ReplicationSystem/NetBlob/NetBlobHandlerManager.h"
+#include "Iris/ReplicationSystem/NetHandleManager.h" // For FInternalNetHandle
+#include "UObject/StrongObjectPtr.h"
+
+class UNetBlobHandler;
+class UNetObjectBlobHandler;
+class UPartialNetObjectAttachmentHandler;
+class UPartialNetObjectAttachmentHandlerConfig;
+class UNetRPCHandler;
+namespace UE::Net
+{
+	class FNetObjectReference;
+	namespace Private
+	{
+		class FNetHandleManager;
+		class FObjectReferenceCache;
+		class FReplicationConnections;
+	}
+}
+
+namespace UE::Net::Private
+{
+
+struct FNetBlobManagerInitParams
+{
+	UReplicationSystem* ReplicationSystem = nullptr;
+	bool bSendAttachmentsWithObject = false;
+};
+
+
+class FNetBlobManager
+{
+public:
+	FNetBlobManager();
+	
+	void Init(FNetBlobManagerInitParams& InitParams);
+
+	bool RegisterNetBlobHandler(UNetBlobHandler* Handler);
+
+	bool QueueNetObjectAttachment(uint32 ConnectionId, const FNetObjectReference& TargetRef, const TRefCountPtr<FNetObjectAttachment>& Attachment);
+	bool SendRPC(const UObject* Object, const UObject* SubObject, const UFunction* Function, const void* Parameters);
+	bool SendRPC(uint32 ConnectionId, const UObject* Object, const UObject* SubObject, const UFunction* Function, const void* Parameters);
+
+	void ProcessNetObjectAttachmentSendQueue();
+
+	FNetBlobHandlerManager& GetNetBlobHandlerManager() { return BlobHandlerManager; }
+	const FNetBlobHandlerManager& GetNetBlobHandlerManager() const { return BlobHandlerManager; }
+
+	const UPartialNetObjectAttachmentHandler* GetPartialNetObjectAttachmentHandler() const { return PartialNetObjectAttachmentHandler.Get(); }
+	const UNetObjectBlobHandler* GetNetObjectBlobHandler() const { return NetObjectBlobHandler.Get(); }
+
+	// Connection handling
+	void AddConnection(uint32 ConnectionId);
+	void RemoveConnection(uint32 ConnectionId);
+
+private:
+	void RegisterDefaultHandlers();
+
+	bool GetOwnerAndSubObjectIndicesFromHandle(FNetHandle NetHandle, FInternalNetHandle& OutOwnerIndex, FInternalNetHandle& OutSubObjectIndex);
+
+	class FNetObjectAttachmentSendQueue
+	{
+	public:
+		FNetObjectAttachmentSendQueue();
+
+		void Init(FNetBlobManager* Manager);
+
+		// Unicast
+		void Enqueue(uint32 ConnectionId, FInternalNetHandle OwnerIndex, FInternalNetHandle SubObjectIndex, const TRefCountPtr<FNetObjectAttachment>& Attachment);
+
+		// Multicast
+		void Enqueue(FInternalNetHandle OwnerIndex, FInternalNetHandle SubObjectIndex, const TRefCountPtr<FNetObjectAttachment>& Attachment);
+
+		void ProcessQueue(FReplicationConnections* Connections);
+
+	private:
+		struct FNetObjectAttachmentQueueEntry
+		{
+			uint32 ConnectionId;
+			FInternalNetHandle OwnerIndex;
+			FInternalNetHandle SubObjectIndex;
+			TRefCountPtr<FNetObjectAttachment> Attachment;
+		};
+		typedef TArray<FNetObjectAttachmentQueueEntry> FQueue;
+
+		bool PreSerializeAndSplitNetBlob(uint32 ConnectionId, const TRefCountPtr<FNetObjectAttachment>& Attachment, TArray<TRefCountPtr<FNetBlob>>& OutPartialNetBlobs) const;
+
+		FNetBlobManager* Manager;
+		FQueue AttachmentQueue;
+		bool bHasMulticastAttachments;
+	};
+
+
+	FNetBlobHandlerManager BlobHandlerManager;
+	FNetObjectAttachmentSendQueue AttachmentSendQueue;
+	TStrongObjectPtr<UNetRPCHandler> RPCHandler;
+	TStrongObjectPtr<UPartialNetObjectAttachmentHandler> PartialNetObjectAttachmentHandler;
+	TStrongObjectPtr<UNetObjectBlobHandler> NetObjectBlobHandler;
+
+	UReplicationSystem* ReplicationSystem;
+	FObjectReferenceCache* ObjectReferenceCache;
+	FReplicationConnections* Connections;
+	const UPartialNetObjectAttachmentHandlerConfig* PartialNetObjectAttachmentHandlerConfig;
+	const FNetHandleManager* NetHandleManager;
+	bool bIsServer;
+	bool bSendAttachmentsWithObject;
+};
+
+}
diff --git a/Engine/Source/Runtime/Experimental/Iris/Core/Private/Iris/ReplicationSystem/NetBlob/ReliableNetBlobQueue.cpp b/Engine/Source/Runtime/Experimental/Iris/Core/Private/Iris/ReplicationSystem/NetBlob/ReliableNetBlobQueue.cpp
new file mode 100644
--- /dev/null
+++ b/Engine/Source/Runtime/Experimental/Iris/Core/Private/Iris/ReplicationSystem/NetBlob/ReliableNetBlobQueue.cpp
@@ -0,0 +1,381 @@
+// Copyright Epic Games, Inc. All Rights Reserved.
+
+#include "Iris/ReplicationSystem/NetBlob/ReliableNetBlobQueue.h"
+#include "Iris/ReplicationSystem/NetBlob/NetBlobHandlerManager.h"
+#include "Iris/ReplicationSystem/ObjectReferenceCache.h"
+#include "Iris/Core/IrisLog.h"
+#include "Iris/Serialization/NetBitStreamReader.h"
+#include "Iris/Serialization/NetBitStreamUtil.h"
+#include "Iris/Serialization/NetBitStreamWriter.h"
+#include "Iris/Serialization/NetSerializationContext.h"
+#include "Iris/Serialization/InternalNetSerializationContext.h"
+#include "Iris/Serialization/NetExportContext.h"
+
+namespace UE::Net::Private
+{
+
+static const FName NetError_ReliableQueueFull("Reliable attachment queue full");
+static const FName NetError_InvalidSequence("Invalid sequence number");
+
+// ReliableNetBlobQueue
+FReliableNetBlobQueue::FReliableNetBlobQueue()
+: Sent{}
+, Acked{}
+, FirstSeq(0)
+, LastSeq(0)
+, UnsentBlobCount(0)
+{
+}
+
+FReliableNetBlobQueue::~FReliableNetBlobQueue()
+{
+}
+
+// If everything is sent and acked and the first sequence index is zero then it's safe to destroy this instance.
+bool FReliableNetBlobQueue::IsSafeToDestroy() const
+{
+	if (FirstSeq != LastSeq)
+	{
+		return false;
+	}
+
+	if (SequenceToIndex(FirstSeq) != 0)
+	{
+		return false;
+	}
+
+	return true;
+}
+
+uint32 FReliableNetBlobQueue::Serialize(FNetSerializationContext& Context, FReliableNetBlobQueue::ReplicationRecord& OutRecord)
+{
+	FNetHandle InvalidNetHandle;
+	return SerializeInternal(Context, InvalidNetHandle, OutRecord, false);
+}
+
+uint32 FReliableNetBlobQueue::SerializeWithObject(FNetSerializationContext& Context, FNetHandle NetHandle, FReliableNetBlobQueue::ReplicationRecord& OutRecord)
+{
+	return SerializeInternal(Context, NetHandle, OutRecord, true);
+}
+
+uint32 FReliableNetBlobQueue::SerializeInternal(FNetSerializationContext& Context, FNetHandle NetHandle, FReliableNetBlobQueue::ReplicationRecord& OutRecord, const bool bSerializeWithObject)
+{
+	FNetBitStreamWriter* Writer = Context.GetBitStreamWriter();
+	const FObjectReferenceCache* ObjectReferenceCache = Context.GetInternalContext()->ObjectReferenceCache;
+
+	uint32 PrevHasMoreBlobsWritePos = 0;
+	uint32 SerializedCount = 0;
+	uint32 PrevWrittenSeq = ~0U;
+	uint32 WrittenSeq[2] = {~0U, ~0U};
+	uint32 WrittenCount[2] = {0U, 0U};
+	uint32 WrittenIndex = 0;
+	for (uint32 Seq = FirstSeq, EndSeq = LastSeq; Seq < EndSeq; ++Seq)
+	{
+		const uint32 Index = SequenceToIndex(Seq);
+		if (IsIndexSent(Index))
+		{
+			continue;
+		}
+
+		// Disjoint sequence handling
+		// First in sequence?
+		if (WrittenSeq[WrittenIndex] == ~0U)
+		{
+			WrittenSeq[WrittenIndex] = Seq;
+		}
+		// If broken sequence start the next one.
+		else if (PrevWrittenSeq + 1U != Seq)
+		{
+			++WrittenIndex;
+			// Our ReplicationRecord can't hold a lot of info.
+			if (WrittenIndex >= UE_ARRAY_COUNT(WrittenCount))
+			{
+				break;
+			}
+			WrittenSeq[WrittenIndex] = Seq;
+		}
+
+		FNetBitStreamRollbackScope RollbackScope(*Writer);
+		FNetExportRollbackScope ExportRollbackScope(Context);
+
+		const TRefCountPtr<FNetBlob>& Attachment = NetBlobs[Index];
+
+		// If this sequence is disjoint from the previous sequence we need to serialize the full index.
+		// It's important that the sequence number is sent first so we can validate it before receiving exports and payload.
+		if (Writer->WriteBool(Seq != PrevWrittenSeq + 1U))
+		{
+			Writer->WriteBits(Index, IndexBitCount);
+		}
+
+		// If we have exports, write them now, if attachment is rolled back we will roll back any written exports as well.
+		ObjectReferenceCache->WriteExports(Context, Attachment->CallGetExports());
+
+		Attachment->SerializeCreationInfo(Context, Attachment->GetCreationInfo());
+		if (bSerializeWithObject)
+		{
+			Attachment->SerializeWithObject(Context, NetHandle);
+		}
+		else
+		{
+			Attachment->Serialize(Context);
+		}
+
+		const uint32 HasMoreBlobsWritePos = Writer->GetPosBits();
+		Writer->WriteBool(false); // Don't know yet if there's more to come
+
+		if (Writer->IsOverflown())
+		{
+			ExportRollbackScope.Rollback();
+			break;
+		}
+		// We did manage to serialize yet another blob
+		else
+		{
+			if (SerializedCount > 0)
+			{
+				FNetBitStreamWriteScope WriteScope(*Writer, PrevHasMoreBlobsWritePos);
+				Writer->WriteBool(true);
+			}
+
+			++SerializedCount;
+
+			// Our ReplicationRecord currently only holds two sequences with number and count, one byte each.
+			++WrittenCount[WrittenIndex];
+			if (WrittenCount[WrittenIndex] == 255)
+			{
+				++WrittenIndex;
+				if (WrittenIndex >= UE_ARRAY_COUNT(WrittenCount))
+				{
+					break;
+				}
+			}
+
+			PrevWrittenSeq = Seq;
+			PrevHasMoreBlobsWritePos = HasMoreBlobsWritePos;
+		}
+	}
+
+	//
+	const ReplicationRecord Record = ((WrittenSeq[0] & 255) << 24U) | ((WrittenCount[0] & 255) << 16U) | ((WrittenSeq[1] & 255) << 8U) | ((WrittenCount[1] & 255) << 0U);
+	OutRecord = (SerializedCount ? Record : FReliableNetBlobQueue::InvalidReplicationRecord);
+
+	return SerializedCount;
+}
+
+void FReliableNetBlobQueue::CommitReplicationRecord(FReliableNetBlobQueue::ReplicationRecord Record)
+{
+	uint32 RecordSeq[2] = {(Record >> 24U) & 255, (Record >> 8U) & 255U};
+	uint32 RecordCount[2] = {(Record >> 16U) & 255, (Record >> 0U) & 255U};
+
+	UnsentBlobCount -= RecordCount[0] + RecordCount[1];
+	for (const uint32 Index : {0U, 1U})
+	{
+		for (uint32 Seq = RecordSeq[Index], EndSeq = Seq + RecordCount[Index]; Seq != EndSeq; ++Seq)
+		{
+			SetSequenceIsSent(Seq);
+		}
+	}
+}
+
+uint32 FReliableNetBlobQueue::Deserialize(FNetSerializationContext& Context)
+{
+	FNetHandle InvalidNetHandle;
+	return DeserializeInternal(Context, InvalidNetHandle, false);
+}
+
+uint32 FReliableNetBlobQueue::DeserializeWithObject(FNetSerializationContext& Context, FNetHandle NetHandle)
+{
+	return DeserializeInternal(Context, NetHandle, true);
+}
+
+uint32 FReliableNetBlobQueue::DeserializeInternal(FNetSerializationContext& Context, FNetHandle NetHandle, const bool bSerializeWithObject)
+{
+	INetBlobReceiver* BlobReceiver = Context.GetNetBlobReceiver();
+	checkSlow(BlobReceiver != nullptr);
+
+	FNetBitStreamReader* Reader = Context.GetBitStreamReader();
+	FObjectReferenceCache* ObjectReferenceCache = Context.GetInternalContext()->ObjectReferenceCache;
+
+	uint32 Index = ~0U;
+	uint32 DeserializedCount = 0;
+	bool bHasMoreBlobs = false;
+	do
+	{
+		if (const bool bIsNewSequence = Reader->ReadBool())
+		{
+			Index = Reader->ReadBits(IndexBitCount);
+		}
+		else
+		{
+			Index = (Index + 1) % MaxUnackedBlobCount;
+		}
+
+		uint32 ReceivedSeq = FirstSeq - (FirstSeq % MaxUnackedBlobCount) + Index;
+		if (ReceivedSeq < FirstSeq)
+		{
+			ReceivedSeq += MaxUnackedBlobCount;
+		}
+
+		if (!IsValidReceiveSequence(ReceivedSeq) || IsIndexAcked(Index))
+		{
+			UE_LOG(LogIris, Warning, TEXT("Invalid reliable sequence number. ReceivedSeq: %u FirstSeq: %u LastSeq: %u IsAcked: %u"), ReceivedSeq, FirstSeq, LastSeq, unsigned(IsIndexAcked(Index)));
+			Context.SetError(NetError_InvalidSequence);
+
+			return DeserializedCount;
+		}
+
+		// If we have exports, import them now.
+		ObjectReferenceCache->ReadExports(Context);
+
+		FNetBlobCreationInfo CreationInfo;
+		FNetBlob::DeserializeCreationInfo(Context, CreationInfo);
+		CreationInfo.Flags = CreationInfo.Flags | ENetBlobFlags::Reliable;
+		const TRefCountPtr<FNetBlob>& Blob = BlobReceiver->CreateNetBlob(CreationInfo);
+		if (!Blob.IsValid())
+		{
+			UE_LOG(LogIris, Warning, TEXT("%s"), TEXT("Unable to create blob."));
+			Context.SetError(GNetError_UnsupportedNetBlob);
+
+			return DeserializedCount;
+		}
+
+		if (bSerializeWithObject)
+		{
+			Blob->DeserializeWithObject(Context, NetHandle);
+		}
+		else
+		{
+			Blob->Deserialize(Context);
+		}
+		
+		bHasMoreBlobs = Reader->ReadBool();
+
+		// Break if something went wrong with deserialization.
+		if (Context.HasErrorOrOverflow())
+		{
+			UE_LOG(LogIris, Warning, TEXT("Failed to deserialize reliable attachments for %s"), *NetHandle.ToString());
+			break;
+		}
+
+		NetBlobs[Index] = Blob;
+		LastSeq = FMath::Max(LastSeq, ReceivedSeq + 1U);
+		SetIndexIsAcked(Index);
+
+		++DeserializedCount;
+	} while (bHasMoreBlobs);
+
+	return DeserializedCount;
+}
+
+bool FReliableNetBlobQueue::Enqueue(const TRefCountPtr<FNetBlob>& Blob)
+{
+	if (IsFull())
+	{
+		return false;
+	}
+
+	++UnsentBlobCount;
+
+	const uint32 Index = SequenceToIndex(LastSeq);
+	++LastSeq;
+
+	NetBlobs[Index] = Blob;
+	return true;
+}
+
+const TRefCountPtr<FNetBlob>* FReliableNetBlobQueue::Peek() const
+{
+	const uint32 Index = SequenceToIndex(FirstSeq);
+	if (IsIndexAcked(Index))
+	{
+		return &NetBlobs[Index];
+	}
+
+	return nullptr;
+}
+
+void FReliableNetBlobQueue::Pop()
+{
+	const uint32 Index = SequenceToIndex(FirstSeq);
+	checkSlow(IsIndexAcked(Index) && (FirstSeq < LastSeq));
+	NetBlobs[Index].SafeRelease();
+	ClearIndexIsAcked(Index);
+	++FirstSeq;
+
+	// $TODO. For a memory optimization one can change the storage implementation and free the memory here if everything is acked.
+}
+
+void FReliableNetBlobQueue::ProcessPacketDeliveryStatus(EPacketDeliveryStatus Status, FReliableNetBlobQueue::ReplicationRecord Record)
+{
+	switch (Status)
+	{
+	case EPacketDeliveryStatus::Delivered:
+	{
+		OnPacketDelivered(Record);
+		break;
+	}
+	case EPacketDeliveryStatus::Lost:
+	{
+		OnPacketDropped(Record);
+		break;
+	}
+	default:
+	{
+		break;
+	}
+	}
+}
+
+void FReliableNetBlobQueue::OnPacketDelivered(FReliableNetBlobQueue::ReplicationRecord Record)
+{
+	uint32 RecordSeq[2] = {(Record >> 24U) & 255, (Record >> 8U) & 255U};
+	uint32 RecordCount[2] = {(Record >> 16U) & 255, (Record >> 0U) & 255U};
+
+	// Mark blobs as acked
+	for (const uint32 Index : {0U, 1U})
+	{
+		for (uint32 Seq = RecordSeq[Index], EndSeq = Seq + RecordCount[Index]; Seq != EndSeq; ++Seq)
+		{
+			SetSequenceIsAcked(Seq);
+		}
+	}
+
+	// Remove acked blobs, allowing for more blobs to be added to the queue.
+	PopInOrderAckedBlobs();
+}
+
+void FReliableNetBlobQueue::OnPacketDropped(FReliableNetBlobQueue::ReplicationRecord Record)
+{
+	uint32 RecordSeq[2] = {(Record >> 24U) & 255, (Record >> 8U) & 255U};
+	uint32 RecordCount[2] = {(Record >> 16U) & 255, (Record >> 0U) & 255U};
+
+	// Mark blobs as unsent
+	UnsentBlobCount += RecordCount[0] + RecordCount[1];
+	for (const uint32 Index : {0U, 1U})
+	{
+		for (uint32 Seq = RecordSeq[Index], EndSeq = Seq + RecordCount[Index]; Seq != EndSeq; ++Seq)
+		{
+			ClearSequenceIsSent(Seq);
+		}
+	}
+}
+
+void FReliableNetBlobQueue::PopInOrderAckedBlobs()
+{
+	for (; FirstSeq != LastSeq; ++FirstSeq)
+	{
+		const uint32 Index = SequenceToIndex(FirstSeq);
+		if (!IsIndexAcked(Index))
+		{
+			return;
+		}
+
+		NetBlobs[Index].SafeRelease();
+		ClearIndexIsAcked(Index);
+		ClearIndexIsSent(Index);
+	}
+
+	// $TODO. For a memory optimization one can change the storage implementation and free the memory here if everything is acked.
+}
+
+}
diff --git a/Engine/Source/Runtime/Experimental/Iris/Core/Private/Iris/ReplicationSystem/NetDependencyData.cpp b/Engine/Source/Runtime/Experimental/Iris/Core/Private/Iris/ReplicationSystem/NetDependencyData.cpp
new file mode 100644
--- /dev/null
+++ b/Engine/Source/Runtime/Experimental/Iris/Core/Private/Iris/ReplicationSystem/NetDependencyData.cpp
@@ -0,0 +1,112 @@
+// Copyright Epic Games, Inc. All Rights Reserved.
+
+#include "NetDependencyData.h"
+#include "Iris/Core/IrisLog.h"
+#include "ReplicationOperationsInternal.h"
+
+namespace UE::Net::Private
+{
+
+FNetDependencyData::FNetDependencyData()
+{
+}
+
+void FNetDependencyData::FreeStoredDependencyDataForObject(FInternalNetHandle InternalHandle)
+{
+	if (FDependencyInfo* Entry = DependencyInfos.Find(InternalHandle))
+	{
+		for (const uint32 ArrayIndex : Entry->ArrayIndices)
+		{
+			if (ArrayIndex != FDependencyInfo::InvalidCacheIndex)
+			{
+				checkSlow(DependentObjectsStorage[ArrayIndex].Num() == 0);
+				DependentObjectsStorage.RemoveAt(ArrayIndex);
+			}
+		}
+
+		if (Entry->SubObjectConditionalArrayIndex != FDependencyInfo::InvalidCacheIndex)
+		{
+			SubObjectConditionalsStorage.RemoveAt(Entry->SubObjectConditionalArrayIndex);
+		}
+
+		DependencyInfos.Remove(InternalHandle);
+	}
+}
+
+FNetDependencyData::FDependencyInfo& FNetDependencyData::GetOrCreateCacheEntry(FInternalNetHandle InternalHandle)
+{
+	FDependencyInfo* Entry = DependencyInfos.Find(InternalHandle);
+
+	if (!Entry)
+	{
+		Entry = &DependencyInfos.Add(InternalHandle);
+		*Entry = FDependencyInfo();
+	}
+
+	return *Entry;
+}
+
+FNetDependencyData::FSubObjectConditionalsArray& FNetDependencyData::GetOrCreateSubObjectConditionalsArray(FInternalNetHandle OwnerHandle)
+{
+	FDependencyInfo& Entry = GetOrCreateCacheEntry(OwnerHandle);
+	check(Entry.ArrayIndices[EArrayType::ChildSubObjects] != FDependencyInfo::InvalidCacheIndex);
+	
+	if (Entry.SubObjectConditionalArrayIndex == FDependencyInfo::InvalidCacheIndex)
+	{
+		FSparseArrayAllocationInfo AllocInfo = SubObjectConditionalsStorage.AddUninitialized();
+		Entry.SubObjectConditionalArrayIndex = AllocInfo.Index;
+
+		FSubObjectConditionalsArray* SubObjectConditionalsArray = new (AllocInfo.Pointer) FSubObjectConditionalsArray();
+		
+		// Make sure that we initialize the conditionals to match the number of SubObjects
+		const int32 NumChildSubObjects = DependentObjectsStorage[Entry.ArrayIndices[EArrayType::ChildSubObjects]].Num();
+		SubObjectConditionalsArray->SetNumZeroed(NumChildSubObjects);
+
+		return *SubObjectConditionalsArray;
+	}
+	else
+	{
+		return SubObjectConditionalsStorage[Entry.SubObjectConditionalArrayIndex];
+	}
+}
+
+FNetDependencyData::FInternalNetHandleArray& FNetDependencyData::GetOrCreateInternalChildSubObjectsArray(FInternalNetHandle OwnerHandle, FSubObjectConditionalsArray*& OutSubObjectConditionals)
+{
+	FDependencyInfo& Entry = GetOrCreateCacheEntry(OwnerHandle);
+	if (Entry.ArrayIndices[ChildSubObjects] == FDependencyInfo::InvalidCacheIndex)
+	{
+		// Allocate storage for subobjects / atomic dependencies
+		FSparseArrayAllocationInfo AllocInfo = DependentObjectsStorage.AddUninitialized();
+		Entry.ArrayIndices[ChildSubObjects] = AllocInfo.Index;
+
+		FInternalNetHandleArray* InternalHandleArray = new (AllocInfo.Pointer) FInternalNetHandleArray();
+		OutSubObjectConditionals = nullptr;
+		return *InternalHandleArray;
+	}
+	else
+	{
+		OutSubObjectConditionals = Entry.SubObjectConditionalArrayIndex != FDependencyInfo::InvalidCacheIndex ? &SubObjectConditionalsStorage[Entry.SubObjectConditionalArrayIndex] : nullptr;
+		return DependentObjectsStorage[Entry.ArrayIndices[ChildSubObjects]];
+	}
+}
+
+FNetDependencyData::FInternalNetHandleArray& FNetDependencyData::GetOrCreateInternalHandleArray(FInternalNetHandle OwnerHandle, EArrayType ArrayType)
+{
+	FDependencyInfo& Entry = GetOrCreateCacheEntry(OwnerHandle);
+	if (Entry.ArrayIndices[ArrayType] == FDependencyInfo::InvalidCacheIndex)
+	{
+		// Allocate storage for subobjects / atomic dependencies
+		FSparseArrayAllocationInfo AllocInfo = DependentObjectsStorage.AddUninitialized();
+		Entry.ArrayIndices[ArrayType] = AllocInfo.Index;
+
+		FInternalNetHandleArray* InternalHandleArray = new (AllocInfo.Pointer) FInternalNetHandleArray();
+
+		return *InternalHandleArray;
+	}
+	else
+	{
+		return DependentObjectsStorage[Entry.ArrayIndices[ArrayType]];
+	}
+}
+
+}
\ No newline at end of file
diff --git a/Engine/Source/Runtime/Experimental/Iris/Core/Private/Iris/ReplicationSystem/NetDependencyData.h b/Engine/Source/Runtime/Experimental/Iris/Core/Private/Iris/ReplicationSystem/NetDependencyData.h
new file mode 100644
--- /dev/null
+++ b/Engine/Source/Runtime/Experimental/Iris/Core/Private/Iris/ReplicationSystem/NetDependencyData.h
@@ -0,0 +1,136 @@
+// Copyright Epic Games, Inc. All Rights Reserved.
+
+#pragma once
+
+#include "CoreTypes.h"
+#include "Containers/ArrayView.h"
+#include "Containers/Map.h"
+#include "Containers/SparseArray.h"
+
+namespace UE::Net::Private
+{
+	typedef uint32 FInternalNetHandle;
+	typedef int8 FLifeTimeConditionStorage;
+}
+
+namespace UE::Net::Private
+{
+
+struct FChildSubObjectsInfo
+{
+	const FInternalNetHandle* ChildSubObjects = nullptr;
+	const FLifeTimeConditionStorage* SubObjectLifeTimeConditions = nullptr;
+	uint32 NumSubObjects = 0U;
+};
+
+class FNetDependencyData
+{
+public:
+	FNetDependencyData();
+
+	typedef TArray<FInternalNetHandle, TInlineAllocator<8>> FInternalNetHandleArray;
+
+	typedef TArray<FLifeTimeConditionStorage, TInlineAllocator<8>> FSubObjectConditionalsArray;
+
+	enum EArrayType
+	{
+		SubObjects = 0U,
+		ChildSubObjects,
+		DependentObjects,
+		ParentObjects,
+		Count
+	};
+
+	template<EArrayType TypeIndex>
+	FInternalNetHandleArray& GetOrCreateInternalHandleArray(FInternalNetHandle InternalHandle)
+	{
+		static_assert(TypeIndex != EArrayType::Count, "Invalid array type index");
+		return GetOrCreateInternalHandleArray(InternalHandle, TypeIndex);
+	};
+
+	FSubObjectConditionalsArray& GetOrCreateSubObjectConditionalsArray(FInternalNetHandle InternalHandle);
+
+	FInternalNetHandleArray& GetOrCreateInternalChildSubObjectsArray(FInternalNetHandle InternalHandle, FSubObjectConditionalsArray*& OutSubObjectConditionals);
+
+	bool GetInternalChildSubObjectAndConditionalArrays(FInternalNetHandle InternalHandle, FInternalNetHandleArray*& OutChildSubObjectsArray, FSubObjectConditionalsArray*& OutSubObjectConditionals)
+	{
+		const FDependencyInfo* Entry = DependencyInfos.Find(InternalHandle);
+		const uint32 ArrayIndex = Entry ? Entry->ArrayIndices[ChildSubObjects] : FDependencyInfo::InvalidCacheIndex;
+
+		if (ArrayIndex == FDependencyInfo::InvalidCacheIndex)
+		{
+			return false;
+		}
+
+		OutChildSubObjectsArray = &DependentObjectsStorage[ArrayIndex];
+		OutSubObjectConditionals = Entry->SubObjectConditionalArrayIndex != FDependencyInfo::InvalidCacheIndex ? &SubObjectConditionalsStorage[Entry->SubObjectConditionalArrayIndex] : nullptr;
+		return true;
+	}
+
+	bool GetChildSubObjects(FInternalNetHandle InternalHandle, FChildSubObjectsInfo& OutInfo) const
+	{
+		const FDependencyInfo* Entry = DependencyInfos.Find(InternalHandle);
+		const uint32 ArrayIndex = Entry ? Entry->ArrayIndices[ChildSubObjects] : FDependencyInfo::InvalidCacheIndex;
+		if (ArrayIndex == FDependencyInfo::InvalidCacheIndex)
+		{
+			return false;
+		}
+
+		OutInfo.ChildSubObjects = DependentObjectsStorage[ArrayIndex].GetData();			
+		OutInfo.SubObjectLifeTimeConditions = Entry->SubObjectConditionalArrayIndex != FDependencyInfo::InvalidCacheIndex ? SubObjectConditionalsStorage[Entry->SubObjectConditionalArrayIndex].GetData() : nullptr;
+		OutInfo.NumSubObjects = DependentObjectsStorage[ArrayIndex].Num();
+		return true;
+	}
+
+
+	template<EArrayType TypeIndex>
+	FInternalNetHandleArray* GetInternalHandleArray(FInternalNetHandle InternalHandle)
+	{
+		static_assert(TypeIndex != EArrayType::Count, "Invalid array type index");
+		const FDependencyInfo* Entry = DependencyInfos.Find(InternalHandle);
+		const uint32 ArrayIndex = Entry ? Entry->ArrayIndices[TypeIndex] : FDependencyInfo::InvalidCacheIndex;
+		if (ArrayIndex != FDependencyInfo::InvalidCacheIndex)
+		{
+			return &DependentObjectsStorage[ArrayIndex];
+		}
+		return nullptr;
+	}
+
+	template<EArrayType TypeIndex>
+	TArrayView<const FInternalNetHandle> GetInternalHandleArray(FInternalNetHandle InternalHandle) const
+	{
+		static_assert(TypeIndex != EArrayType::Count, "Invalid array type index");
+
+		const FDependencyInfo* Entry = DependencyInfos.Find(InternalHandle);
+		const uint32 ArrayIndex = Entry ? Entry->ArrayIndices[TypeIndex] : FDependencyInfo::InvalidCacheIndex;
+		if (ArrayIndex != FDependencyInfo::InvalidCacheIndex)
+		{
+			return MakeArrayView(DependentObjectsStorage[ArrayIndex]);
+		}
+		return MakeArrayView<const FInternalNetHandle>(nullptr, 0);
+	}
+
+	void FreeStoredDependencyDataForObject(FInternalNetHandle InternalHandle);
+	
+private:
+	
+	struct FDependencyInfo
+	{
+		constexpr static uint32 InvalidCacheIndex = ~(0U);
+		uint32 ArrayIndices[EArrayType::Count] = { InvalidCacheIndex, InvalidCacheIndex, InvalidCacheIndex, InvalidCacheIndex };
+		uint32 SubObjectConditionalArrayIndex = InvalidCacheIndex;
+	};
+
+private:
+	FInternalNetHandleArray& GetOrCreateInternalHandleArray(FInternalNetHandle InternalHandle, EArrayType Type);
+	FDependencyInfo& GetOrCreateCacheEntry(FInternalNetHandle InternalHandle);
+	
+	// Map to track the replicated objects with subObjects or dependencies
+	TMap<FInternalNetHandle, FDependencyInfo> DependencyInfos;
+	// Storage for DependentObjects and SubObjects
+	TSparseArray<FInternalNetHandleArray> DependentObjectsStorage;
+	// Storage for SubObject conditionals
+	TSparseArray<FSubObjectConditionalsArray> SubObjectConditionalsStorage;
+};
+
+}
diff --git a/Engine/Source/Runtime/Experimental/Iris/Core/Private/Iris/ReplicationSystem/NetHandleManager.cpp b/Engine/Source/Runtime/Experimental/Iris/Core/Private/Iris/ReplicationSystem/NetHandleManager.cpp
new file mode 100644
--- /dev/null
+++ b/Engine/Source/Runtime/Experimental/Iris/Core/Private/Iris/ReplicationSystem/NetHandleManager.cpp
@@ -0,0 +1,749 @@
+// Copyright Epic Games, Inc. All Rights Reserved.
+
+#include "NetHandleManager.h"
+#include "ReplicationOperationsInternal.h"
+#include "ReplicationProtocol.h"
+#include "ReplicationProtocolManager.h"
+#include "Iris/Core/IrisLog.h"
+#include "Iris/Serialization/NetSerializationContext.h"
+#include "Iris/Serialization/InternalNetSerializationContext.h"
+#include "Net/Core/Trace/NetTrace.h"
+#include "UObject/Object.h"
+#include "UObject/UObjectGlobals.h"
+
+namespace UE::Net::Private
+{
+
+FNetHandleManager::FNetHandleManager(FReplicationProtocolManager& InReplicationProtocolManager, uint32 InReplicationSystemId, uint32 InMaxActiveObjectCount)
+: ActiveObjectCount(0)
+, MaxActiveObjectCount(InMaxActiveObjectCount)
+, ReplicationSystemId(InReplicationSystemId)
+, ScopableInternalIndices(MaxActiveObjectCount)
+, PrevFrameScopableInternalIndices(MaxActiveObjectCount)
+, AssignedInternalIndices(MaxActiveObjectCount)
+, SubObjectInternalIndices(MaxActiveObjectCount)
+, DependentObjectInternalIndices(MaxActiveObjectCount)
+, ObjectsWithDependentObjectsInternalIndices(MaxActiveObjectCount)
+, DestroyedStartupObjectInternalIndices(MaxActiveObjectCount)
+, WantToBeDormantInternalIndices(MaxActiveObjectCount)
+, NextStaticHandleIndex(1U) // Index 0 is always reserved, for both static and dynamic handles
+, NextDynamicHandleIndex(1U)
+, ReplicationProtocolManager(InReplicationProtocolManager)
+{
+	// first is reserved for invalid handle
+	AssignedInternalIndices.SetBit(0);
+
+	ReplicatedObjectData.SetNumUninitialized(MaxActiveObjectCount);
+	ReplicatedObjectRefCount.SetNumZeroed(MaxActiveObjectCount);
+	ReplicatedObjectStateBuffers.SetNumZeroed(MaxActiveObjectCount);
+	
+	ReplicatedInstances.SetNumZeroed(MaxActiveObjectCount);
+
+	// For convenience we initialize ReplicatedObjectData for InvalidInternalIndex so that GetReplicatedObjectDataNoCheck returns something useful.
+	ReplicatedObjectData[InvalidInternalIndex] = FReplicatedObjectData();
+}
+
+uint32 FNetHandleManager::GetNextHandleId(uint32 HandleId) const
+{
+	// Since we use the lowest bit in the index to indicate if the handle is static or dynamic we can not use all bits as the index
+	constexpr uint32 NetHandleIdIndexBitMask = (1U << (FNetHandle::IdBits - 1U)) - 1U;
+
+	uint32 NextHandleId = (HandleId + 1) & NetHandleIdIndexBitMask;
+	if (NextHandleId == 0)
+	{
+		++NextHandleId;
+	}
+	return NextHandleId;
+}
+
+FInternalNetHandle FNetHandleManager::GetNextFreeInternalIndex() const
+{
+	const uint32 NextFreeIndex = AssignedInternalIndices.FindFirstZero();
+	return NextFreeIndex != FNetBitArray::InvalidIndex ? NextFreeIndex : InvalidInternalIndex;
+}
+
+uint32 FNetHandleManager::InternalCreateNetObject(const FNetHandle NetHandle, const FReplicationProtocol* ReplicationProtocol)
+{
+	if (ActiveObjectCount >= MaxActiveObjectCount)
+	{
+		return InvalidInternalIndex;
+	}
+
+	// verify that the handle is free
+	if (HandleMap.Contains(NetHandle))
+	{
+		ensureAlwaysMsgf(false, TEXT("NetHandleManager::InternalCreateNetHandleHandle %s already exists"), *NetHandle.ToString());
+		return InvalidInternalIndex;
+	}
+
+	const uint32 InternalIndex = GetNextFreeInternalIndex();
+	if (InternalIndex != InvalidInternalIndex)
+	{
+		// Store data;
+		FReplicatedObjectData& Data = ReplicatedObjectData[InternalIndex];
+
+		Data = FReplicatedObjectData();
+
+		Data.Handle = NetHandle;
+		Data.Protocol = ReplicationProtocol;
+		Data.InstanceProtocol = nullptr;
+		ReplicatedObjectStateBuffers.GetData()[InternalIndex] = nullptr;
+		Data.ReceiveStateBuffer = nullptr;
+		Data.bShouldPropagateChangedStates = 1U;
+
+		++ActiveObjectCount;
+
+		// Add map entry so that we can map from NetHandle to InternalIndex
+		HandleMap.Add(Data.Handle, InternalIndex);
+
+		// Mark Handle index as assigned and scopable for now
+		AssignedInternalIndices.SetBit(InternalIndex);
+		ScopableInternalIndices.SetBit(InternalIndex);
+
+		// When a handle is first created, it is not set to be a subobject
+		SubObjectInternalIndices.ClearBit(InternalIndex);
+
+		ReplicatedObjectRefCount[InternalIndex] = 0;
+
+		return InternalIndex;
+	}
+		
+	return InvalidInternalIndex;
+}
+
+void FNetHandleManager::AttachInstanceProtocol(FInternalNetHandle InternalIndex, const FReplicationInstanceProtocol* InstanceProtocol, UObject* Instance)
+{
+	if (ensureAlways((InternalIndex != InvalidInternalIndex) && InstanceProtocol))
+	{
+		FReplicatedObjectData& Data = ReplicatedObjectData[InternalIndex];
+		Data.InstanceProtocol = InstanceProtocol;
+
+		check(ReplicatedInstances[InternalIndex] == nullptr);
+		ReplicatedInstances[InternalIndex] = Instance;
+	}
+}
+
+const FReplicationInstanceProtocol* FNetHandleManager::DetachInstanceProtocol(FInternalNetHandle InternalIndex)
+{
+	if (ensure(InternalIndex != InvalidInternalIndex))
+	{
+		FReplicatedObjectData& Data = ReplicatedObjectData[InternalIndex];
+		const FReplicationInstanceProtocol* InstanceProtocol = Data.InstanceProtocol;
+		Data.InstanceProtocol = nullptr;
+		ReplicatedInstances[InternalIndex] = nullptr;
+		
+		return InstanceProtocol;
+	}
+
+	return nullptr;
+}
+
+FNetHandle FNetHandleManager::AllocateNetHandle(bool bIsStatic)
+{
+	uint32& NextHandleId = bIsStatic ? NextStaticHandleIndex : NextDynamicHandleIndex;
+
+	const uint32 NewHandleId = MakeNetHandleId(NextHandleId, bIsStatic);
+	FNetHandle NewHandle = MakeNetHandle(NewHandleId, ReplicationSystemId);
+
+	// Verify that the handle is free
+	if (HandleMap.Contains(NewHandle))
+	{
+		checkf(false, TEXT("FNetHandleManager::AllocateNetHandle - Handle %s already exists!"), *NewHandle.ToString());
+
+		return FNetHandle();
+	}
+
+	// Bump NextHandleId
+	NextHandleId = GetNextHandleId(NextHandleId);
+
+	return NewHandle;
+}
+
+FNetHandle FNetHandleManager::CreateNetObject(FNetHandle WantedHandle, const FReplicationProtocol* ReplicationProtocol)
+{
+	FNetHandle NetHandle = WantedHandle;
+
+	const FInternalNetHandle InternalIndex = InternalCreateNetObject(NetHandle, ReplicationProtocol);
+	if (InternalIndex != InvalidInternalIndex)
+	{
+		FReplicatedObjectData& Data = ReplicatedObjectData[InternalIndex];
+		
+		// Allocate storage for outgoing data. We need a valid pointer even if the size is zero.
+		uint8* StateBuffer = (uint8*)FMemory::Malloc(FPlatformMath::Max(ReplicationProtocol->InternalTotalSize, 1U), ReplicationProtocol->InternalTotalAlignment);
+		// Clear state buffer.
+		FMemory::Memzero(StateBuffer, ReplicationProtocol->InternalTotalSize);
+		if (EnumHasAnyFlags(ReplicationProtocol->ProtocolTraits, EReplicationProtocolTraits::HasConditionalChangeMask))
+		{
+			// Enable all conditions by default. This is to be compatible with the implementation of FRepChangedPropertyTracker where we have hooks into ReplicationSystem.
+			FNetBitArrayView ConditionalChangeMask(reinterpret_cast<uint32*>(StateBuffer + ReplicationProtocol->GetConditionalChangeMaskOffset()), ReplicationProtocol->ChangeMaskBitCount, FNetBitArrayView::NoResetNoValidate);
+			ConditionalChangeMask.SetAllBits();
+		}
+
+		ReplicatedObjectStateBuffers.GetData()[InternalIndex] = StateBuffer;
+
+		// Bump protocol refcount
+		ReplicationProtocol->AddRef();
+
+		return NetHandle;
+	}
+	else
+	{
+		return FNetHandle();
+	}
+}
+
+// Create NetHandle not owned by us
+FNetHandle FNetHandleManager::CreateNetObjectFromRemote(FNetHandle WantedHandle, const FReplicationProtocol* ReplicationProtocol)
+{
+	if (!ensureAlwaysMsgf(WantedHandle.IsValid() && !WantedHandle.IsCompleteHandle(), TEXT("FNetHandleManager::CreateNetHandleFromRemote Expected WantedHandle %s to be valid and incomplete"), *WantedHandle.ToString()))
+	{
+		return FNetHandle();
+	}
+
+	FNetHandle NetHandle = MakeNetHandle(WantedHandle.GetId(), ReplicationSystemId);
+
+	const FInternalNetHandle InternalIndex = InternalCreateNetObject(NetHandle, ReplicationProtocol);
+	if (InternalIndex != InvalidInternalIndex)
+	{
+		FReplicatedObjectData& Data = ReplicatedObjectData[InternalIndex];
+
+		// Allocate storage for incoming data
+		Data.ReceiveStateBuffer = (uint8*)FMemory::Malloc(ReplicationProtocol->InternalTotalSize, ReplicationProtocol->InternalTotalAlignment);
+
+		// Since we currently do not have any default values, just initialize to zero
+		// N.B. If this Memzero is optimized away for some reason it still needs to be done for protocols with dynamic state.
+		FMemory::Memzero(Data.ReceiveStateBuffer, ReplicationProtocol->InternalTotalSize);
+
+		// Note: We could initialize this from default but at the moment it is part of the contract for all serializers to write the value when we serialize and we will only apply dirty states
+
+		// Don't bother initializing the conditional changemask if present since it's currently unused on the receiving end.
+
+		// Bump protocol refcount
+		ReplicationProtocol->AddRef();
+
+		return NetHandle;
+	}
+	else
+	{
+		return FNetHandle();
+	}
+}
+
+void FNetHandleManager::InternalDestroyNetObject(FInternalNetHandle InternalIndex)
+{
+	FReplicatedObjectData& Data = ReplicatedObjectData[InternalIndex];
+
+	UE_LOG(LogIris, Verbose, TEXT("FNetHandleManager::InternalDestroyNetObject ( InternalIndex: %u ) %s"), InternalIndex, *Data.Handle.ToString());
+
+	uint8* StateBuffer = ReplicatedObjectStateBuffers.GetData()[InternalIndex];
+	// Free any allocated resources
+	if (EnumHasAnyFlags(Data.Protocol->ProtocolTraits, EReplicationProtocolTraits::HasDynamicState))
+	{
+		FNetSerializationContext FreeContext;
+		FInternalNetSerializationContext InternalContext;
+		FreeContext.SetInternalContext(&InternalContext);
+		if (StateBuffer != nullptr)
+		{
+			FReplicationProtocolOperationsInternal::FreeDynamicState(FreeContext, StateBuffer, Data.Protocol);
+		}
+		if (Data.ReceiveStateBuffer != nullptr)
+		{
+			FReplicationProtocolOperationsInternal::FreeDynamicState(FreeContext, Data.ReceiveStateBuffer, Data.Protocol);
+		}
+	}
+	
+	// If we have subobjects remove them from our list
+	if (FNetDependencyData::FInternalNetHandleArray* SubObjectArray = SubObjects.GetInternalHandleArray<FNetDependencyData::EArrayType::SubObjects>(InternalIndex))
+	{
+		for (FInternalNetHandle SubObjectInternalHandle : *SubObjectArray)
+		{
+			InternalRemoveSubObject(InternalIndex, SubObjectInternalHandle, false);
+		}
+		SubObjectArray->Reset(0U);
+	}
+	// Clear ChildSubObjectArray
+	if (FNetDependencyData::FInternalNetHandleArray* ChildSubObjectArray = SubObjects.GetInternalHandleArray<FNetDependencyData::EArrayType::ChildSubObjects>(InternalIndex))
+	{
+		ChildSubObjectArray->Reset(0U);
+	}
+
+	// If we are a subobject remove from owner and hierarchical parents
+	if (Data.SubObjectRootIndex != InvalidInternalIndex)
+	{
+		InternalRemoveSubObject(Data.SubObjectRootIndex, InternalIndex);
+	}
+
+	// Remove from all dependent object relationships and clear data
+	InternalRemoveDependentObject(InternalIndex);
+
+	// Free all stored data for object
+	SubObjects.FreeStoredDependencyDataForObject(InternalIndex);
+
+	// Decrease protocol refcount
+	Data.Protocol->Release();
+	if (Data.Protocol->GetRefCount() == 0)
+	{
+		ReplicationProtocolManager.DestroyReplicationProtocol(Data.Protocol);
+	}
+
+	FMemory::Free(StateBuffer);
+	FMemory::Free(Data.ReceiveStateBuffer);
+
+	// Clear pointer to state buffer
+	ReplicatedObjectStateBuffers.GetData()[InternalIndex] = nullptr;
+
+	UE_NET_TRACE_NETHANDLE_DESTROYED(Data.Handle);
+
+	Data = FReplicatedObjectData();
+
+	// tracking
+	AssignedInternalIndices.ClearBit(InternalIndex);
+
+	// Restore internal state
+	SubObjectInternalIndices.ClearBit(InternalIndex);
+	WantToBeDormantInternalIndices.ClearBit(InternalIndex);
+	ObjectsWithDependentObjectsInternalIndices.ClearBit(InternalIndex);
+
+	// Cleanup cross reference to destruction info
+	if (DestroyedStartupObjectInternalIndices.GetBit(InternalIndex))
+	{
+		DestroyedStartupObjectInternalIndices.ClearBit(InternalIndex);
+		uint32 OtherInternalIndex = 0U;
+		if (DestroyedStartupObject.RemoveAndCopyValue(InternalIndex, OtherInternalIndex))
+		{
+			DestroyedStartupObject.Remove(OtherInternalIndex);
+		}
+	}
+
+	--ActiveObjectCount;
+}
+
+FNetHandle FNetHandleManager::CreateHandleForDestructionInfo(FNetHandle Handle, const FReplicationProtocol* DestroydObjectProtocol)
+{
+	// Create destruction info handle carrying destruction info
+	FNetHandle AllocatedHandle = AllocateNetHandle(false);
+	FNetHandle DestructionInfoHandle = CreateNetObject(AllocatedHandle, DestroydObjectProtocol);
+
+	if (DestructionInfoHandle.IsValid())
+	{
+		const uint32 InternalIndex = GetInternalIndex(DestructionInfoHandle);
+		const uint32 DestroyedInternalIndex = GetInternalIndex(Handle);
+		
+		// mark the internal index
+		DestroyedStartupObjectInternalIndices.SetBit(InternalIndex);
+				
+		// Is the object replicated, if so we must make sure that we do not add it to scope by accident
+		if (DestroyedInternalIndex)
+		{			
+			DestroyedStartupObject.FindOrAdd(InternalIndex, DestroyedInternalIndex);
+		
+			// Mark the replicated index as destroyed
+			DestroyedStartupObjectInternalIndices.SetBit(DestroyedInternalIndex);
+			DestroyedStartupObject.FindOrAdd(DestroyedInternalIndex, InternalIndex);
+		}
+	}
+
+	return DestructionInfoHandle;
+}
+
+void FNetHandleManager::RemoveFromScope(FInternalNetHandle InternalIndex)
+{
+	// Can only remove an object from scope if it is assignable
+	if (ensure(AssignedInternalIndices.GetBit(InternalIndex)))
+	{
+		ScopableInternalIndices.ClearBit(InternalIndex);
+	}
+}
+
+void FNetHandleManager::DestroyNetObject(FNetHandle Handle)
+{
+	FInternalNetHandle InternalIndex = HandleMap.FindAndRemoveChecked(Handle);
+	if (ensure(AssignedInternalIndices.GetBit(InternalIndex)))
+	{
+		FReplicatedObjectData& Data = ReplicatedObjectData[InternalIndex];
+		check(Data.Handle == Handle);
+
+		// Remove from scopable objects if not already done
+		ScopableInternalIndices.ClearBit(InternalIndex);
+
+		// We always defer the actual destroy
+		PendingDestroyInternalIndices.Add(InternalIndex);
+	}
+}
+
+void FNetHandleManager::DestroyObjectsPendingDestroy()
+{
+	// Destroy Objects pending destroy
+	for (int32 It = PendingDestroyInternalIndices.Num() -1; It >= 0; --It)
+	{
+		const FInternalNetHandle InternalIndex = PendingDestroyInternalIndices[It];
+		if (ReplicatedObjectRefCount[InternalIndex] == 0)
+		{
+			InternalDestroyNetObject(InternalIndex);
+			PendingDestroyInternalIndices.RemoveAtSwap(It);
+		}
+	}
+}
+
+bool FNetHandleManager::AddSubObject(FNetHandle OwnerHandle, FNetHandle SubObjectHandle, FNetHandle RelativeOtherSubObjectHandle, EAddSubObjectFlags Flags)
+{
+	check(OwnerHandle != SubObjectHandle);
+
+	// validate objects
+	const FInternalNetHandle OwnerInternalIndex = GetInternalIndex(OwnerHandle);
+	const FInternalNetHandle SubObjectInternalIndex = GetInternalIndex(SubObjectHandle);
+
+	const bool bIsValidOwner = ensure(OwnerInternalIndex != InvalidInternalIndex);
+	const bool bIsValidSubObejct = ensure(SubObjectInternalIndex != InvalidInternalIndex);
+
+	if (!(bIsValidOwner && bIsValidSubObejct))
+	{
+		return false;
+	}
+
+	FReplicatedObjectData& SubObjectData = GetReplicatedObjectDataNoCheck(SubObjectInternalIndex);
+	if (!ensureAlwaysMsgf(SubObjectData.SubObjectRootIndex == InvalidInternalIndex, TEXT("FNetHandleManager::AddSubObject %s is already marked as a subobject"), *SubObjectHandle.ToString()))
+	{
+		return false;
+	}
+
+	FNetDependencyData::FInternalNetHandleArray& SubObjectArray = SubObjects.GetOrCreateInternalHandleArray<FNetDependencyData::SubObjects>(OwnerInternalIndex);
+	SubObjectArray.Add(SubObjectInternalIndex);
+	SubObjectData.SubObjectRootIndex = OwnerInternalIndex;
+	SubObjectData.bDestroySubObjectWithOwner = EnumHasAnyFlags(Flags, EAddSubObjectFlags::DestroyWithOwner);
+	// Mark the object as a subobject
+	SetIsSubObject(SubObjectInternalIndex, true);
+
+	const FInternalNetHandle RelativeOtherSubObjectInternalIndex = EnumHasAnyFlags(Flags, EAddSubObjectFlags::ReplicateWithSubObject) ? GetInternalIndex(RelativeOtherSubObjectHandle) : InvalidInternalIndex;
+	if (RelativeOtherSubObjectInternalIndex != InvalidInternalIndex && ensureAlwaysMsgf(SubObjectArray.Find(RelativeOtherSubObjectInternalIndex) != INDEX_NONE, TEXT("RelativeOtherSubObjectHandle %s Must be a Subobject of %s"), *RelativeOtherSubObjectHandle.ToString(), *OwnerHandle.ToString()))
+	{
+		// Add to child array of RelativeOtherSubObjectRelativeIndex for hierarchical replication order
+		FNetDependencyData::FSubObjectConditionalsArray* SubObjectConditionalsArray = nullptr;
+		FNetDependencyData::FInternalNetHandleArray& ChildSubObjectArray = SubObjects.GetOrCreateInternalChildSubObjectsArray(RelativeOtherSubObjectInternalIndex, SubObjectConditionalsArray);
+		ChildSubObjectArray.Add(SubObjectInternalIndex);
+		if (SubObjectConditionalsArray)
+		{
+			SubObjectConditionalsArray->Add(ELifetimeCondition::COND_None);
+		}
+		SubObjectData.SubObjectParentIndex = RelativeOtherSubObjectInternalIndex;
+	}
+	else
+	{
+		// Add to child array of root
+		FNetDependencyData::FSubObjectConditionalsArray* SubObjectConditionalsArray = nullptr;
+		FNetDependencyData::FInternalNetHandleArray& ChildSubObjectArray = SubObjects.GetOrCreateInternalChildSubObjectsArray(OwnerInternalIndex, SubObjectConditionalsArray);
+		ChildSubObjectArray.Add(SubObjectInternalIndex);
+		if (SubObjectConditionalsArray)
+		{
+			SubObjectConditionalsArray->Add(ELifetimeCondition::COND_None);
+		}
+		SubObjectData.SubObjectParentIndex = OwnerInternalIndex;
+	}
+
+	return true;
+}
+
+void FNetHandleManager::InternalRemoveSubObject(FInternalNetHandle OwnerInternalIndex, FInternalNetHandle SubObjectInternalIndex, bool bRemoveFromSubObjectArray)
+{
+	// both must be valid
+	if (OwnerInternalIndex != InvalidInternalIndex && SubObjectInternalIndex != InvalidInternalIndex)
+	{
+		FReplicatedObjectData& SubObjectData = GetReplicatedObjectDataNoCheck(SubObjectInternalIndex);
+		check(SubObjectData.SubObjectRootIndex == OwnerInternalIndex);
+		
+		if (bRemoveFromSubObjectArray)
+		{
+			// Remove from root parent
+			if (FNetDependencyData::FInternalNetHandleArray* SubObjectArray = SubObjects.GetInternalHandleArray<FNetDependencyData::EArrayType::SubObjects>(OwnerInternalIndex))
+			{
+				SubObjectArray->Remove(SubObjectInternalIndex);
+			}
+			// Remove from parents ChildSubObjectArray
+			if (SubObjectData.SubObjectParentIndex != InvalidInternalIndex)
+			{
+				FNetDependencyData::FInternalNetHandleArray* ChildSubObjectArray;
+				FNetDependencyData::FSubObjectConditionalsArray* SubObjectConditionsArray;
+
+				if (SubObjects.GetInternalChildSubObjectAndConditionalArrays(SubObjectData.SubObjectParentIndex, ChildSubObjectArray, SubObjectConditionsArray))
+				{
+					const int32 ArrayIndex = ChildSubObjectArray->Find(SubObjectInternalIndex);
+					if (ensureAlways(ArrayIndex != INDEX_NONE))
+					{
+						ChildSubObjectArray->RemoveAt(ArrayIndex);
+						if (SubObjectConditionsArray)
+						{
+							SubObjectConditionsArray->RemoveAt(ArrayIndex);
+							check(SubObjectConditionsArray->Num() == ChildSubObjectArray->Num());
+						}
+					}
+				}
+			}
+		}
+
+		SubObjectData.SubObjectRootIndex = InvalidInternalIndex;
+		SubObjectData.bDestroySubObjectWithOwner = false;
+
+		SetIsSubObject(SubObjectInternalIndex, false);
+	}
+}
+
+void FNetHandleManager::RemoveSubObject(FNetHandle Handle)
+{
+	FInternalNetHandle SubObjectInternalIndex = GetInternalIndex(Handle);
+	checkSlow(SubObjectInternalIndex);
+
+	if (SubObjectInternalIndex != InvalidInternalIndex)
+	{
+		const uint32 OwnerInternalIndex = ReplicatedObjectData[SubObjectInternalIndex].SubObjectRootIndex;
+		if (OwnerInternalIndex != InvalidInternalIndex)
+		{
+			InternalRemoveSubObject(OwnerInternalIndex, SubObjectInternalIndex);
+		}
+	}
+}
+
+bool FNetHandleManager::SetSubObjectNetCondition(FInternalNetHandle SubObjectInternalIndex, FLifeTimeConditionStorage SubObjectCondition)
+{
+	if (ensure(SubObjectInternalIndex != InvalidInternalIndex))
+	{
+		const FInternalNetHandle SubObjectParentIndex = ReplicatedObjectData[SubObjectInternalIndex].SubObjectParentIndex;
+		if (ensure(SubObjectParentIndex != InvalidInternalIndex))
+		{
+			FNetDependencyData::FInternalNetHandleArray* SubObjectsArray;
+			FNetDependencyData::FSubObjectConditionalsArray* SubObjectConditionals;
+			
+			if (SubObjects.GetInternalChildSubObjectAndConditionalArrays(SubObjectParentIndex, SubObjectsArray, SubObjectConditionals))
+			{
+				const int32 SubObjectArrayIndex = SubObjectsArray->Find(SubObjectInternalIndex);
+				if (ensure(SubObjectArrayIndex != INDEX_NONE))
+				{
+					if (!SubObjectConditionals)
+					{
+						SubObjectConditionals = &SubObjects.GetOrCreateSubObjectConditionalsArray(SubObjectParentIndex);
+					}
+
+					check(SubObjectConditionals->Num() == SubObjectsArray->Num());
+
+					SubObjectConditionals->GetData()[SubObjectArrayIndex] = SubObjectCondition;
+					return true;
+				}
+			}
+
+		}
+	}
+
+	return false;
+}
+
+FNetHandle FNetHandleManager::GetSubObjectOwner(FNetHandle SubObjectHandle) const
+{
+	const FInternalNetHandle SubObjectInternalIndex = GetInternalIndex(SubObjectHandle);
+	const FInternalNetHandle OwnerInternalIndex = SubObjectInternalIndex != InvalidInternalIndex ? ReplicatedObjectData[SubObjectInternalIndex].SubObjectRootIndex : InvalidInternalIndex;
+
+	return OwnerInternalIndex != InvalidInternalIndex ? ReplicatedObjectData[OwnerInternalIndex].Handle : FNetHandle();
+}
+
+bool FNetHandleManager::AddDependentObject(FNetHandle ParentHandle, FNetHandle DependentObjectHandle, EAddDependentObjectFlags Flags)
+{
+	check(ParentHandle != DependentObjectHandle);
+
+	// validate objects
+	FInternalNetHandle ParentInternalIndex = GetInternalIndex(ParentHandle);
+	FInternalNetHandle DependentObjectInternalIndex = GetInternalIndex(DependentObjectHandle);
+
+	const bool bIsValidOwner = ensure(ParentInternalIndex != InvalidInternalIndex);
+	const bool bIsValidDependentObject = ensure(DependentObjectInternalIndex != InvalidInternalIndex);
+
+	if (!(bIsValidOwner && bIsValidDependentObject))
+	{
+		return false;
+	}
+
+	FReplicatedObjectData& DependentObjectData = GetReplicatedObjectDataNoCheck(DependentObjectInternalIndex);
+	FReplicatedObjectData& ParentObjectData = GetReplicatedObjectDataNoCheck(ParentInternalIndex);
+
+	// SubObjects cannot have dependent objects or be a dependent object (for now)
+	check(!DependentObjectData.IsSubObject() && !ParentObjectData.IsSubObject());
+	check(!SubObjectInternalIndices.GetBit(DependentObjectInternalIndex));
+	check(!SubObjectInternalIndices.GetBit(ParentInternalIndex));
+
+	// Add dependent to parents dependent object list
+	FNetDependencyData::FInternalNetHandleArray& ParentDependentObjectArray = SubObjects.GetOrCreateInternalHandleArray<FNetDependencyData::DependentObjects>(ParentInternalIndex);
+	const bool bDependentIsAlreadyDependant = ParentDependentObjectArray.Find(DependentObjectInternalIndex) != INDEX_NONE;
+	if (!bDependentIsAlreadyDependant)
+	{
+		ParentDependentObjectArray.Add(DependentObjectInternalIndex);
+	}
+
+	FNetDependencyData::FInternalNetHandleArray& DependentParentObjectArray = SubObjects.GetOrCreateInternalHandleArray<FNetDependencyData::ParentObjects>(DependentObjectInternalIndex);
+	const bool bDependentHadParentAlready = DependentParentObjectArray.Find(ParentInternalIndex) != INDEX_NONE;
+	if (!bDependentHadParentAlready)
+	{
+		DependentParentObjectArray.Add(ParentInternalIndex);
+	}
+
+	// Update cached info to avoid to do map lookups to find out if we have are a dependent object or have dependent objects
+	DependentObjectData.bIsDependentObject = true;
+	ParentObjectData.bHasDependentObjects = true;
+	ObjectsWithDependentObjectsInternalIndices.SetBit(ParentInternalIndex);
+	DependentObjectInternalIndices.SetBit(DependentObjectInternalIndex);
+
+	if (EnumHasAnyFlags(Flags, EAddDependentObjectFlags::WarnIfAlreadyDependentObject) && (bDependentHadParentAlready || bDependentIsAlreadyDependant))
+	{
+		// If this gets out of sync something is messed up
+		check(bDependentHadParentAlready == bDependentIsAlreadyDependant);
+		ensureAlwaysMsgf(false, TEXT("FNetHandleManager::AddDependentObject %s already is a child of %s parent %s"), *DependentObjectHandle.ToString(), *GetReplicatedObjectDataNoCheck(DependentObjectData.SubObjectRootIndex).Handle.ToString(), *ParentHandle.ToString());
+	}
+
+	return true;
+}
+
+void FNetHandleManager::RemoveDependentObject(FNetHandle DependentHandle)
+{
+	FInternalNetHandle DependentInternalIndex = GetInternalIndex(DependentHandle);
+
+	if (DependentInternalIndex != InvalidInternalIndex)
+	{
+		InternalRemoveDependentObject(DependentInternalIndex);
+	}
+}
+
+void FNetHandleManager::InternalRemoveDependentObject(FInternalNetHandle ParentInternalIndex, FInternalNetHandle DependentInternalIndex, ERemoveDependentObjectFlags Flags)
+{
+	if (EnumHasAnyFlags(Flags, ERemoveDependentObjectFlags::RemoveFromDependentParentObjects))
+	{
+		if (FNetDependencyData::FInternalNetHandleArray* ParentObjectArray = SubObjects.GetInternalHandleArray<FNetDependencyData::ParentObjects>(DependentInternalIndex))
+		{
+			ParentObjectArray->Remove(ParentInternalIndex);
+			if (ParentObjectArray->Num() == 0)
+			{
+				FReplicatedObjectData& DependentObjectData = GetReplicatedObjectDataNoCheck(DependentInternalIndex);
+				DependentObjectData.bIsDependentObject = false;
+				DependentObjectInternalIndices.ClearBit(DependentInternalIndex);
+			}
+		}
+	}
+
+	if (EnumHasAnyFlags(Flags, ERemoveDependentObjectFlags::RemoveFromParentDependentObjects))
+	{
+		if (FNetDependencyData::FInternalNetHandleArray* ParentDependentObjectArray = SubObjects.GetInternalHandleArray<FNetDependencyData::DependentObjects>(ParentInternalIndex))
+		{
+			ParentDependentObjectArray->Remove(DependentInternalIndex);
+			if (ParentDependentObjectArray->Num() == 0)
+			{
+				FReplicatedObjectData& ParentObjectData = GetReplicatedObjectDataNoCheck(ParentInternalIndex);
+				ParentObjectData.bHasDependentObjects = false;
+				ObjectsWithDependentObjectsInternalIndices.ClearBit(ParentInternalIndex);
+			}
+		}
+	}
+}
+
+void FNetHandleManager::InternalRemoveDependentObject(FInternalNetHandle DependentInternalIndex)
+{
+	// Remove from all parents
+	if (FNetDependencyData::FInternalNetHandleArray* ParentObjectArray = SubObjects.GetInternalHandleArray<FNetDependencyData::ParentObjects>(DependentInternalIndex))
+	{
+		for (FInternalNetHandle ParentInternalIndex : *ParentObjectArray)
+		{
+			// Flag is set to only update data on the parent to avoid modifying the array we iterate over
+			InternalRemoveDependentObject(ParentInternalIndex, DependentInternalIndex, ERemoveDependentObjectFlags::RemoveFromParentDependentObjects);
+		}
+		ParentObjectArray->Reset();
+	}
+
+	// Remove from our dependents
+	if (FNetDependencyData::FInternalNetHandleArray* DependentObjectArray = SubObjects.GetInternalHandleArray<FNetDependencyData::DependentObjects>(DependentInternalIndex))
+	{
+		for (FInternalNetHandle ChildDepedentInternalIndex : *DependentObjectArray)
+		{
+			// Flag is set to only update data on the childDependentObject to avoid modifying the array we iterate over
+			InternalRemoveDependentObject(DependentInternalIndex, ChildDepedentInternalIndex, ERemoveDependentObjectFlags::RemoveFromDependentParentObjects);
+		}
+		DependentObjectArray->Reset();		
+	}
+
+	FReplicatedObjectData& DependentObjectData = GetReplicatedObjectDataNoCheck(DependentInternalIndex);
+
+	// Clear out flags on this object	
+	DependentObjectData.bIsDependentObject = false;
+	DependentObjectData.bHasDependentObjects = false;
+	ObjectsWithDependentObjectsInternalIndices.ClearBit(DependentInternalIndex);
+	DependentObjectInternalIndices.ClearBit(DependentInternalIndex);
+}
+
+void FNetHandleManager::RemoveDependentObject(FNetHandle ParentHandle, FNetHandle DependentHandle)
+{
+	// Validate objects
+	FInternalNetHandle ParentInternalIndex = GetInternalIndex(ParentHandle);
+	FInternalNetHandle DependentInternalIndex = GetInternalIndex(DependentHandle);
+
+	if (!(ParentInternalIndex != InvalidInternalIndex) && (DependentInternalIndex != InvalidInternalIndex))
+	{
+		return;
+	}
+
+	InternalRemoveDependentObject(ParentInternalIndex, DependentInternalIndex);
+}
+
+void FNetHandleManager::SetShouldPropagateChangedStates(FInternalNetHandle ObjectInternalIndex, bool bShouldPropagateChangedStates)
+{
+	if (ObjectInternalIndex != InvalidInternalIndex)
+	{
+		if (bShouldPropagateChangedStates)
+		{
+			// Currently we do not support re-enabling state propagation
+			// $IRIS: $TODO: Implement method to force dirty all changes 
+			// https://jira.it.epicgames.com/browse/UE-127368
+
+			checkf(false, TEXT("Re-enabling state change propagation is currently Not implemented."));			
+			return;
+		}
+
+		ReplicatedObjectData[ObjectInternalIndex].bShouldPropagateChangedStates = bShouldPropagateChangedStates ? 1U : 0U;
+	}
+}
+
+void FNetHandleManager::SetShouldPropagateChangedStates(FNetHandle Handle, bool bShouldPropagateChangedStates)
+{
+	FInternalNetHandle ObjectInternalIndex = GetInternalIndex(Handle);
+	return SetShouldPropagateChangedStates(ObjectInternalIndex, bShouldPropagateChangedStates);
+}
+
+void FNetHandleManager::AddReferencedObjects(FReferenceCollector& Collector)
+{
+	Collector.AddReferencedObjects(ReplicatedInstances);
+}
+
+uint32 FNetHandleManager::MakeNetHandleId(uint32 Id, bool bIsStatic)
+{
+	return (Id << 1U) | (bIsStatic ? 1U : 0U);
+}
+
+FNetHandle FNetHandleManager::MakeNetHandle(uint32 Id, uint32 ReplicationSystemId)
+{
+	check((Id & FNetHandle::IdMask) == Id);
+	check(ReplicationSystemId < FNetHandle::MaxReplicationSystemId);
+
+	FNetHandle Handle;
+
+	Handle.Id = Id;
+	Handle.ReplicationSystemId = ReplicationSystemId + 1U;
+
+	return Handle;
+}
+
+FNetHandle FNetHandleManager::MakeNetHandleFromId(uint32 Id)
+{
+	check((Id & FNetHandle::IdMask) == Id);
+
+	FNetHandle Handle;
+
+	Handle.Id = Id;
+	Handle.ReplicationSystemId = 0U;
+
+	return Handle;
+}
+
+
+}
diff --git a/Engine/Source/Runtime/Experimental/Iris/Core/Private/Iris/ReplicationSystem/NetHandleManager.h b/Engine/Source/Runtime/Experimental/Iris/Core/Private/Iris/ReplicationSystem/NetHandleManager.h
new file mode 100644
--- /dev/null
+++ b/Engine/Source/Runtime/Experimental/Iris/Core/Private/Iris/ReplicationSystem/NetHandleManager.h
@@ -0,0 +1,435 @@
+// Copyright Epic Games, Inc. All Rights Reserved.
+
+#pragma once
+
+#include "CoreTypes.h"
+#include "Containers/Array.h"
+#include "Containers/Map.h"
+#include "Containers/ArrayView.h"
+#include "Net/Core/NetBitArray.h"
+#include "Iris/ReplicationState/ReplicationStateDescriptor.h"
+#include "Iris/ReplicationSystem/NetHandle.h"
+#include "Iris/ReplicationSystem/NetDependencyData.h"
+
+class FReferenceCollector;
+namespace UE::Net
+{
+	struct FReplicationProtocol;
+	struct FReplicationInstanceProtocol;
+
+	namespace Private
+	{
+		class FReplicationProtocolManager;
+	}
+}
+
+namespace UE::Net::Private
+{
+
+typedef uint32 FInternalNetHandle;
+
+enum class EAddSubObjectFlags : uint32
+{
+	None = 0U,
+	WarnIfAlreadySubObject = 1U,
+	SkipIfAlreadySubObject = WarnIfAlreadySubObject << 1U,
+	DestroyWithOwner = SkipIfAlreadySubObject << 1U,
+	ReplicateWithSubObject = DestroyWithOwner << 1U,
+	Default = WarnIfAlreadySubObject | DestroyWithOwner,
+};
+ENUM_CLASS_FLAGS(EAddSubObjectFlags);
+
+enum class EAddDependentObjectFlags : uint32
+{
+	None = 0U,
+	WarnIfAlreadyDependentObject = 1U,
+};
+ENUM_CLASS_FLAGS(EAddDependentObjectFlags);
+
+enum class ERemoveDependentObjectFlags : uint32
+{
+	None = 0U,
+	RemoveFromDependentParentObjects = 1U,
+	RemoveFromParentDependentObjects = RemoveFromDependentParentObjects << 1U,
+	All = RemoveFromDependentParentObjects | RemoveFromParentDependentObjects,
+};
+ENUM_CLASS_FLAGS(ERemoveDependentObjectFlags);
+
+// Internal class to manage NetHandles and their internal data
+class FNetHandleManager
+{
+public:
+	// We will never assign index 0
+	enum { InvalidInternalIndex = 0U };
+
+	// We need to store some internal data for Replicated objects
+	// $TODO: This should be split up into separate array according to usage patterns, it is getting a bit too large
+	struct FReplicatedObjectData
+	{
+		FReplicatedObjectData()
+		: Protocol(nullptr)
+		, InstanceProtocol(nullptr)
+		, ReceiveStateBuffer(nullptr)
+		, SubObjectRootIndex(InvalidInternalIndex)
+		, SubObjectParentIndex(InvalidInternalIndex)
+		, Flags(0U)
+		{
+		}
+
+		FNetHandle Handle;
+		const FReplicationProtocol* Protocol;
+		const FReplicationInstanceProtocol* InstanceProtocol;
+		uint8* ReceiveStateBuffer;
+		uint32 SubObjectRootIndex;
+		uint32 SubObjectParentIndex;
+
+		union
+		{
+			uint32 Flags : 32U;
+			struct
+			{
+				uint32 bShouldPropagateChangedStates : 1U;
+				uint32 bTearOff : 1U;
+				uint32 bDestroySubObjectWithOwner : 1U;
+				uint32 bIsDependentObject : 1U;
+				uint32 bHasDependentObjects : 1U;
+				uint32 Padding : 27U;
+			};
+		};
+	
+		// Returns true if this is a SubObject
+		bool IsSubObject() const { return SubObjectRootIndex && !!bDestroySubObjectWithOwner; }
+
+		// Returns true if this is a DepedentObject
+		bool IsDependentObject() const { return bIsDependentObject; }
+	};
+
+	// We want to handle both actual Networked and "network addressable" handles in the same way
+	// As we expect the number of FNetHandles to be much greater than the actual number of Replicated Objects we
+	// use a map to map from Handle to InternalIndex, only Replicated Objects will ever be assigned an InternalNetHandle
+	typedef TMap<FNetHandle, FInternalNetHandle> FHandleMap;
+
+public:
+	FNetHandleManager(FReplicationProtocolManager& InReplicationProtocolManager, uint32 InReplicationSystemId, uint32 MaxActiveObjects);
+
+	// Return true if this is a scopable index
+	bool IsScopableIndex(FInternalNetHandle InternalIndex) const { return ScopableInternalIndices.GetBit(InternalIndex); }
+
+	// Return true if this index is assigned
+	bool IsAssignedIndex(FInternalNetHandle InternalIndex) const { return InternalIndex && ScopableInternalIndices.GetBit(InternalIndex); }
+
+	static FNetHandle MakeNetHandle(uint32 Id, uint32 ReplicationSystemId);
+	static FNetHandle MakeNetHandleFromId(uint32 Id);
+
+	// Returns a valid handle if the wanted handle can be allocated
+	FNetHandle AllocateNetHandle(bool bIsStatic);
+
+	// Create local Net Object
+	FNetHandle CreateNetObject(FNetHandle WantedHandle, const FReplicationProtocol* ReplicationProtocol);
+
+	// Create NetObject on request from remote
+	FNetHandle CreateNetObjectFromRemote(FNetHandle WantedHandle, const FReplicationProtocol* ReplicationProtocol);
+
+	// Attach Instance protocol to handle
+	// Instance can be null, we only track the Instance for legacy support
+	void AttachInstanceProtocol(FInternalNetHandle InternalIndex, const FReplicationInstanceProtocol* InstanceProtocol, UObject* Instance);
+
+	// Used when we want to detach the instance protocol due to async destroy
+	const FReplicationInstanceProtocol* DetachInstanceProtocol(FInternalNetHandle InternalIndex);
+
+	// Creates a DestructionInfo Handle which is used to replicate persistently destroyed static objects when late joining or streaming in levels with static objects that already has been destroyed
+	// Returns a new handle used to replicate the destruction info, and adds a cross-reference to the original handle being destroyed if that still exists 
+	// Handle is a reference to an already replicated Handle that is about to be destroyed
+	FNetHandle CreateHandleForDestructionInfo(FNetHandle Handle, const FReplicationProtocol* DestroydObjectProtocol);
+	
+	void DestroyNetObject(FNetHandle Handle);
+
+	// Mark object as no longer scopable, object will be removed from scope for all connections
+	void RemoveFromScope(FInternalNetHandle InternalIndex);
+
+	TArrayView<const FInternalNetHandle> GetObjectsPendingDestroy() const { return MakeArrayView(PendingDestroyInternalIndices); }
+	void DestroyObjectsPendingDestroy();
+
+	const FReplicatedObjectData& GetReplicatedObjectDataNoCheck(FInternalNetHandle InternalIndex) const { return ReplicatedObjectData[InternalIndex]; }
+	FReplicatedObjectData& GetReplicatedObjectDataNoCheck(FInternalNetHandle InternalIndex) { return ReplicatedObjectData[InternalIndex]; }
+
+	inline const FReplicatedObjectData& GetReplicatedObjectData(FInternalNetHandle InternalIndex) const;
+
+	inline const uint8* GetReplicatedObjectStateBufferNoCheck(FInternalNetHandle InternalObjectIndex) const { return ReplicatedObjectStateBuffers.GetData()[InternalObjectIndex]; }
+	inline uint8* GetReplicatedObjectStateBufferNoCheck(FInternalNetHandle InternalObjectIndex) { return ReplicatedObjectStateBuffers.GetData()[InternalObjectIndex]; }
+	inline const TArray<uint8*>& GetReplicatedObjectStateBuffers() const { return ReplicatedObjectStateBuffers; }
+
+	// Verify a handle against internal handle, A handle is valid if it matches internal storage
+	inline bool IsValidNetHandle(FNetHandle Handle) const;
+
+	// Returns true of the handle is for local replicated object
+	inline bool IsLocalNetHandle(FNetHandle Handle) const;
+
+	// Returns true of the handle is for a remote replicated object
+	inline bool IsRemoteNetHandle(FNetHandle Handle) const;
+
+	// Extract an Full handle from an incomplete one, a incomplete handle is a handle with only an index
+	inline FNetHandle GetCompleteNetHandle(FNetHandle IncompleteHandle) const;
+	
+	// Get Handle from internal index
+	inline FNetHandle GetNetHandleFromInternalIndex(FInternalNetHandle InternalIndex) const;
+
+	// Get internal index from handle
+	inline FInternalNetHandle GetInternalIndex(FNetHandle Handle) const;
+
+	// Get bitarray for all currently scopable internal indices
+	const FNetBitArray& GetScopableInternalIndices() const { return ScopableInternalIndices; }
+
+	// Get bitarray for all internal indices that was scopable last update
+	const FNetBitArray& GetPrevFrameScopableInternalIndices() const { return PrevFrameScopableInternalIndices; }
+	void SetPrevFrameScopableInternalIndicesToCurrent() { PrevFrameScopableInternalIndices = ScopableInternalIndices; }
+
+	// Get bitarray for all internal indices that currently are assigned
+	const FNetBitArray& GetAssignedInternalIndices() const { return AssignedInternalIndices; }
+
+	// SubObjects
+	const FNetBitArray& GetSubObjectInternalIndices() const { return SubObjectInternalIndices; }	
+	bool AddSubObject(FNetHandle OwnerHandle, FNetHandle SubObjectHandle, FNetHandle RelativeOtherSubObjectHandle, EAddSubObjectFlags Flags = EAddSubObjectFlags::Default);
+	inline bool AddSubObject(FNetHandle OwnerHandle, FNetHandle SubObjectHandle, EAddSubObjectFlags Flags = EAddSubObjectFlags::Default) { return AddSubObject(OwnerHandle, SubObjectHandle, FNetHandle(), Flags); }
+	void RemoveSubObject(FNetHandle SubObjectHandle);
+	FNetHandle GetSubObjectOwner(FNetHandle SubObjectHandle) const;
+	bool SetSubObjectNetCondition(FInternalNetHandle SubObjectInternalIndex, FLifeTimeConditionStorage SubObjectCondition);
+
+	// DependentObjects
+	const FNetBitArray& GetDependentObjectInternalIndices() const { return DependentObjectInternalIndices; }
+	const FNetBitArray& GetObjectsWithDependentObjectsInternalIndices() const { return ObjectsWithDependentObjectsInternalIndices; }
+	bool AddDependentObject(FNetHandle ParentHandle, FNetHandle DependentHandle, EAddDependentObjectFlags Flags = EAddDependentObjectFlags::WarnIfAlreadyDependentObject);
+	void RemoveDependentObject(FNetHandle ParentHandle, FNetHandle DependentHandle);
+
+	// Remove DependentHandles from all dependent object tracking
+	// $TODO: This might need a pass ObjectReplicationBridge as well since most likely also must restore filter/polling logic
+	void RemoveDependentObject(FNetHandle DependentHandle);
+	
+	void SetShouldPropagateChangedStates(FNetHandle Handle, bool bShouldPropagateChangedStates);
+	void SetShouldPropagateChangedStates(FInternalNetHandle ObjectInternalIndex, bool bShouldPropagateChangedStates);
+
+	uint32 GetMaxActiveObjectCount() const { return MaxActiveObjectCount; }
+	uint32 GetActiveObjectCount() const { return ActiveObjectCount; }
+
+	// We do refcount objects tracked by each connection in order to know when it is safe to reuse an InternalIndex
+	void AddNetObjectRef(FInternalNetHandle InternalIndex) { ++ReplicatedObjectRefCount[InternalIndex]; }
+	void ReleaseNetObjectRef(FInternalNetHandle InternalIndex) { check(ReplicatedObjectRefCount[InternalIndex] > 0); --ReplicatedObjectRefCount[InternalIndex]; }
+	uint16 GetNetObjectRefCount(FInternalNetHandle ObjectInternalIndex) const { return ReplicatedObjectRefCount[ObjectInternalIndex]; }
+
+	// Get dependent objects for the given ParentIndex
+	inline TArrayView<const FInternalNetHandle> GetDependentObjects(FInternalNetHandle ParentIndex) const;
+
+	// Get all parents of the given DependentIndex
+	inline TArrayView<const FInternalNetHandle> GetDependentObjectParents(FInternalNetHandle DependentIndex) const;
+
+	// Get all subobjects for the given OwnerIndex, Note: only valid for the root
+	inline TArrayView<const FInternalNetHandle> GetSubObjects(FInternalNetHandle OwnerIndex) const;
+
+	// Get child subobjects for a object, used when we do hierarchical operations such as conditional serialization
+	inline TArrayView<const FInternalNetHandle> GetChildSubObjects(FInternalNetHandle ParentObjectIndex) const;
+
+	// Get child subobjects and condtionals array if one exists, if there are no child subobjects the method returns false
+	inline bool GetChildSubObjects(FInternalNetHandle OwnerIndex, FChildSubObjectsInfo& OutInfo) const;
+
+	const FHandleMap& GetReplicatedHandles() const { return HandleMap; }
+
+	const TArray<UObject*>& GetReplicatedInstances() const { return ReplicatedInstances; }
+
+	void AddReferencedObjects(FReferenceCollector& Collector);
+
+	bool GetIsDestroyedStartupObject(FInternalNetHandle InternalIndex) const { return DestroyedStartupObjectInternalIndices.GetBit(InternalIndex); }
+	inline uint32 GetOriginalDestroyedStartupObjectIndex(FInternalNetHandle InternalIndex) const;
+
+	const FNetBitArray& GetDestroyedStartupObjectInternalIndices() const { return DestroyedStartupObjectInternalIndices; }
+
+	const FNetBitArray& GetWantToBeDormantInternalIndices() const { return WantToBeDormantInternalIndices; }
+	FNetBitArray& GetWantToBeDormantInternalIndices() { return WantToBeDormantInternalIndices; }
+
+	// Iterate over all dependent objects and their dependent objects
+	template <typename T>
+	void ForAllDependentObjectsRecursive(FInternalNetHandle ObjectIndex, T&& Functor) const
+	{
+		if (ObjectsWithDependentObjectsInternalIndices.GetBit(ObjectIndex))
+		{
+			for (const FInternalNetHandle DependentObjectIndex : GetDependentObjects(ObjectIndex))
+			{
+				Functor(DependentObjectIndex);
+				ForAllDependentObjectsRecursive(DependentObjectIndex, Functor);
+			}
+			return;
+		}
+	};
+
+private:
+	FInternalNetHandle InternalCreateNetObject(const FNetHandle NetHandle, const FReplicationProtocol* ReplicationProtocol);
+	void InternalDestroyNetObject(FInternalNetHandle InternalIndex);
+
+	static uint32 MakeNetHandleId(uint32 Seed, bool bIsStatic);
+	uint32 GetNextHandleId(uint32 HandleIndex) const;
+
+	// Get the next free internal index, returns InvalidInternalIndex if a free one cannot be found
+	FInternalNetHandle GetNextFreeInternalIndex() const;
+
+	inline void SetIsSubObject(FInternalNetHandle InternalIndex, bool IsSubObject) { SubObjectInternalIndices.SetBitValue(InternalIndex, IsSubObject); }
+	void InternalRemoveSubObject(FInternalNetHandle OwnerInternalIndex, FInternalNetHandle SubObjectInternalIndex, bool bRemoveFromSubObjectArray = true);
+
+	void InternalRemoveDependentObject(FInternalNetHandle ParentInternalIndex, FInternalNetHandle DependentInternalIndex, ERemoveDependentObjectFlags Flags = ERemoveDependentObjectFlags::All);
+	void InternalRemoveDependentObject(FInternalNetHandle DependentInternalIndex);
+
+	// The current replicated object count
+	uint32 ActiveObjectCount;
+
+	// Max allowed replicated object count
+	uint32 MaxActiveObjectCount;
+
+	uint32 ReplicationSystemId;
+
+	FHandleMap HandleMap;
+
+	// Bitset used in order to track assigned internal which are scopable
+	FNetBitArray ScopableInternalIndices;
+
+	// Which internal indices were used last net frame. This can be used to find out which ones are new and deleted this frame. 
+	FNetBitArray PrevFrameScopableInternalIndices;
+
+	// Bitset containing all internal indices that are assigned
+	FNetBitArray AssignedInternalIndices;
+
+	// Bitset containing all internal indices that are SubObjects
+	FNetBitArray SubObjectInternalIndices;
+
+	// Bitset containing all internal indices that are dependentObjects
+	FNetBitArray DependentObjectInternalIndices;
+
+	// Bitset containing all internal indices that have dependentObjects
+	FNetBitArray ObjectsWithDependentObjectsInternalIndices;
+
+	// Bitset marking internal indices of static objects that are permanently destroyed
+	// We use this for special scoping
+	FNetBitArray DestroyedStartupObjectInternalIndices;
+
+	// Bitset marking internal indices that wants to be dormant
+	FNetBitArray WantToBeDormantInternalIndices;
+
+	// Both are marked in DestroyedStartupObjectsInternalIndicies
+	// When one is destroyed both mapping will be cleared
+	// Map OriginalIndex -> DestructionInfoIndex
+	// Map DestructionInfoIndex -> Orignal
+	TMap<uint32, uint32> DestroyedStartupObject;
+private:
+
+	// Array used in order to track objects pending destroy
+	TArray<FInternalNetHandle> PendingDestroyInternalIndices;
+
+	// Just an array containing data about our replicated objects
+	TArray<FReplicatedObjectData> ReplicatedObjectData;
+
+	// Pointers to state buffers for all replicated objects
+	TArray<uint8*> ReplicatedObjectStateBuffers;
+
+	// Refcounts for all tracked objects
+	TArray<uint16> ReplicatedObjectRefCount;
+
+	// Raw pointers to all bound instances
+	TArray<UObject*> ReplicatedInstances;
+
+	// Assign handles
+	uint32 NextStaticHandleIndex;
+	uint32 NextDynamicHandleIndex;
+
+	FNetDependencyData SubObjects;
+
+	FReplicationProtocolManager& ReplicationProtocolManager;
+};
+
+const FNetHandleManager::FReplicatedObjectData& FNetHandleManager::GetReplicatedObjectData(FInternalNetHandle InternalIndex) const
+{
+	check(AssignedInternalIndices.GetBit(InternalIndex));
+	return GetReplicatedObjectDataNoCheck(InternalIndex);
+}
+
+FInternalNetHandle FNetHandleManager::GetInternalIndex(FNetHandle Handle) const
+{
+	const FInternalNetHandle* InternalIndex = HandleMap.Find(Handle);
+	return InternalIndex ? *InternalIndex : FNetHandleManager::InvalidInternalIndex;
+}
+
+FNetHandle FNetHandleManager::GetNetHandleFromInternalIndex(FInternalNetHandle InternalIndex) const
+{
+	check(AssignedInternalIndices.GetBit(InternalIndex));
+	return GetReplicatedObjectDataNoCheck(InternalIndex).Handle;
+}
+
+FNetHandle FNetHandleManager::GetCompleteNetHandle(FNetHandle IncompleteHandle) const
+{
+	if (const FInternalNetHandle* InternalIndex = HandleMap.Find(IncompleteHandle))
+	{
+		return GetReplicatedObjectDataNoCheck(*InternalIndex).Handle;
+	}
+	else
+	{
+		return FNetHandle();
+	}
+}
+
+bool FNetHandleManager::IsValidNetHandle(FNetHandle Handle) const
+{
+	return HandleMap.Contains(Handle);
+}
+
+// Returns true of the handle is a local handle
+bool FNetHandleManager::IsLocalNetHandle(FNetHandle Handle) const
+{
+	if (const uint32 InternalIndex = GetInternalIndex(Handle))
+	{
+		// For the time being only replicated objects owned by this peer has a state buffer`1
+		return ReplicatedObjectStateBuffers[InternalIndex] != nullptr;
+	}
+
+	return false;
+}
+
+// Returns true of the handle is a remote handle
+bool FNetHandleManager::IsRemoteNetHandle(FNetHandle Handle) const
+{
+	if (const uint32 InternalIndex = GetInternalIndex(Handle))
+	{
+		// For the time being only replicated objects owned by this peer has a state buffer
+		return ReplicatedObjectStateBuffers[InternalIndex] == nullptr;
+	}
+
+	return false;
+}
+
+uint32 FNetHandleManager::GetOriginalDestroyedStartupObjectIndex(FInternalNetHandle InternalIndex) const
+{
+	const uint32* FoundOriginalInternalIndex = DestroyedStartupObject.Find(InternalIndex);
+	return FoundOriginalInternalIndex ? *FoundOriginalInternalIndex : 0U;
+}
+
+TArrayView<const FInternalNetHandle> FNetHandleManager::GetSubObjects(FInternalNetHandle OwnerIndex) const
+{
+	return SubObjects.GetInternalHandleArray<FNetDependencyData::EArrayType::SubObjects>(OwnerIndex);
+}
+
+bool FNetHandleManager::GetChildSubObjects(FInternalNetHandle OwnerIndex, FChildSubObjectsInfo& OutInfo) const
+{
+	return SubObjects.GetChildSubObjects(OwnerIndex, OutInfo);
+}
+
+TArrayView<const FInternalNetHandle> FNetHandleManager::GetChildSubObjects(FInternalNetHandle OwnerIndex) const
+{
+	return SubObjects.GetInternalHandleArray<FNetDependencyData::EArrayType::ChildSubObjects>(OwnerIndex);
+}
+
+TArrayView<const FInternalNetHandle> FNetHandleManager::GetDependentObjects(FInternalNetHandle ParentIndex) const
+{
+	return SubObjects.GetInternalHandleArray<FNetDependencyData::EArrayType::DependentObjects>(ParentIndex);
+}
+
+TArrayView<const FInternalNetHandle> FNetHandleManager::GetDependentObjectParents(FInternalNetHandle DependentIndex) const
+{
+	return SubObjects.GetInternalHandleArray<FNetDependencyData::EArrayType::ParentObjects>(DependentIndex);
+}
+	
+}
diff --git a/Engine/Source/Runtime/Experimental/Iris/Core/Private/Iris/ReplicationSystem/NetTokenStore.cpp b/Engine/Source/Runtime/Experimental/Iris/Core/Private/Iris/ReplicationSystem/NetTokenStore.cpp
new file mode 100644
--- /dev/null
+++ b/Engine/Source/Runtime/Experimental/Iris/Core/Private/Iris/ReplicationSystem/NetTokenStore.cpp
@@ -0,0 +1,175 @@
+// Copyright Epic Games, Inc. All Rights Reserved.
+
+#include "Iris/ReplicationSystem/NetTokenStore.h"
+#include "Iris/ReplicationSystem/NetTokenStoreState.h"
+#include "Iris/Serialization/NetBitStreamReader.h"
+#include "Iris/Serialization/NetBitStreamWriter.h"
+#include "Iris/Serialization/NetBitStreamUtil.h"
+#include "Iris/Serialization/NetSerializationContext.h"
+
+namespace UE::Net
+{
+
+FNetTokenDataStore::FNetTokenDataStore()
+: TypeId(FNetToken::InvalidTokenTypeId)
+{
+}
+
+FNetTokenDataStore::~FNetTokenDataStore()
+{
+}
+
+FNetTokenStoreKey FNetTokenDataStore::GetTokenKey(FNetToken Token, const FNetTokenStoreState& TokenStoreState) const
+{
+	if (Token.GetIndex() < (uint32)TokenStoreState.TokenInfos.Num())
+	{
+		const FNetTokenStoreKey& Key = TokenStoreState.TokenInfos[Token.GetIndex()];
+		if (ensureAlwaysMsgf(GetTypeId() == Key.GetTypeId(), TEXT("Cannot resolve NetToken %s with TypeId: %u in DataStore with TypeId: %u"), *Token.ToString(), Key.GetTypeId(), GetTypeId()))
+		{
+			return Key;
+		}
+	}
+
+	return FNetTokenStoreKey();
+}
+
+FNetToken FNetTokenDataStore::CreateToken(FNetTokenStoreKey Key, FNetTokenStoreState& TokenStoreState)
+{
+	const uint32 NextTokenIndex = TokenStoreState.TokenInfos.Num();
+
+	if (!ensure(NextTokenIndex < FNetToken::MaxNetTokenCount))
+	{
+		return FNetToken();
+	}
+
+	// Store token info
+	TokenStoreState.TokenInfos.Add(Key);
+
+	return FNetToken::MakeNetToken(NextTokenIndex);
+}
+
+FNetToken ReadNetToken(FNetBitStreamReader* Reader)
+{
+	uint32 Index = ReadPackedUint32(Reader);
+	
+	if (!Reader->IsOverflown())
+	{
+		return FNetToken(Index);
+	}
+	else
+	{
+		return FNetToken();
+	}
+}
+
+void WriteNetToken(FNetBitStreamWriter* Writer, FNetToken Token)
+{
+	WritePackedUint32(Writer, Token.GetIndex());
+}
+
+FNetTokenStore::FNetTokenStore()
+: LocalNetTokenStoreState(new FNetTokenStoreState)
+{
+}
+
+FNetTokenStore::~FNetTokenStore()
+{
+	delete LocalNetTokenStoreState;
+}
+
+bool FNetTokenStore::RegisterDataStore(FNetTokenDataStore* DataStore)
+{
+	if (TokenDataStores.Num() >= FNetToken::MaxTypeIdCount)
+	{
+		return false;
+	}
+
+	if (!DataStore)
+	{
+		return false;
+	}
+
+	if (!ensure(DataStore->GetTypeId() == FNetToken::InvalidTokenTypeId))
+	{
+		// Already registered
+		return false;
+	}
+
+	DataStore->TypeId = TokenDataStores.Num();
+	TokenDataStores.Add(DataStore);
+
+	return true;
+}
+
+void FNetTokenStore::WriteTokenData(FNetSerializationContext& Context, const FNetToken& NetToken) const
+{
+	if (NetToken.IsValid())
+	{
+		WriteTokenDataForIndex(Context, NetToken.GetIndex());
+	}
+}
+
+void FNetTokenStore::ReadTokenData(FNetSerializationContext& Context, const FNetToken& NetToken, FNetTokenStoreState& RemoteNetTokenStoreState)
+{
+	if (NetToken.IsValid())
+	{
+		// TODO: Guard this better, a map might be better after-all!
+		// Also need to protect this if we process inbound data in parallel.
+		if (!RemoteNetTokenStoreState.ReserveTokenCount(NetToken.GetIndex() + 1))
+		{
+			Context.GetBitStreamReader()->DoOverflow();
+			return;
+		}
+
+		ReadTokenDataForIndex(Context, NetToken.GetIndex(), RemoteNetTokenStoreState);
+	}
+}
+
+void FNetTokenStore::WriteTokenDataForIndex(FNetSerializationContext& Context, uint32 TokenIndex) const
+{
+	FNetBitStreamWriter* Writer = Context.GetBitStreamWriter();
+
+	// Write type
+	const FNetTokenStoreKey& TokenKey = LocalNetTokenStoreState->TokenInfos[TokenIndex];
+	Writer->WriteBits(TokenKey.TypeId, FNetToken::TokenTypeIdBits);
+
+	// Write token data
+	TokenDataStores[TokenKey.TypeId]->WriteTokenData(Context, TokenKey);
+}
+
+void FNetTokenStore::ReadTokenDataForIndex(FNetSerializationContext& Context, uint32 TokenIndex, FNetTokenStoreState& TokenStoreState)
+{
+	FNetBitStreamReader* Reader = Context.GetBitStreamReader();
+
+	// Read type
+	FNetToken::FTypeId TokenTypeId = (FNetToken::FTypeId)Reader->ReadBits(FNetToken::TokenTypeIdBits);
+
+	// Validate that we managed to read and verify the type
+	if (Reader->IsOverflown() || TokenTypeId >= (uint32)TokenDataStores.Num() )
+	{
+		Reader->DoOverflow();
+		return;
+	}
+
+	const FNetTokenStoreKey StoreKey = TokenDataStores[TokenTypeId]->ReadTokenData(Context);
+
+	// Validate that we managed to read the data for the key
+	if (Reader->IsOverflown() || !StoreKey.IsValid())
+	{
+		Reader->DoOverflow();
+		return;
+	}
+
+	// Since the same tokendata might be exported multiple times we better validate that it is the same data
+	const FNetTokenStoreKey& ExistingStoreKey = TokenStoreState.TokenInfos[TokenIndex];
+	if (!ensureAlways(!ExistingStoreKey.IsValid() || StoreKey == ExistingStoreKey))
+	{
+		Reader->DoOverflow();
+		return;
+	}
+	
+	// Store
+	TokenStoreState.TokenInfos[TokenIndex] = StoreKey;
+}
+
+}
diff --git a/Engine/Source/Runtime/Experimental/Iris/Core/Private/Iris/ReplicationSystem/NetTokenStoreState.cpp b/Engine/Source/Runtime/Experimental/Iris/Core/Private/Iris/ReplicationSystem/NetTokenStoreState.cpp
new file mode 100644
--- /dev/null
+++ b/Engine/Source/Runtime/Experimental/Iris/Core/Private/Iris/ReplicationSystem/NetTokenStoreState.cpp
@@ -0,0 +1,29 @@
+// Copyright Epic Games, Inc. All Rights Reserved.
+
+#include "Iris/ReplicationSystem/NetTokenStoreState.h"
+
+namespace UE::Net
+{
+
+FNetTokenStoreState::FNetTokenStoreState()
+{
+	// We reserve the first token as an invalid token
+	TokenInfos.Add(FNetTokenStoreKey());
+}
+
+bool FNetTokenStoreState::ReserveTokenCount(uint32 NewCount)
+{
+	if (NewCount <= FNetToken::MaxNetTokenCount)
+	{
+		const uint32 NewSize = FMath::Max<uint32>(TokenInfos.Num(), NewCount);
+		TokenInfos.SetNumZeroed(NewSize);
+
+		return true;
+	}
+	else
+	{
+		return false;
+	}
+}
+
+}
diff --git a/Engine/Source/Runtime/Experimental/Iris/Core/Private/Iris/ReplicationSystem/NetTokenStoreState.h b/Engine/Source/Runtime/Experimental/Iris/Core/Private/Iris/ReplicationSystem/NetTokenStoreState.h
new file mode 100644
--- /dev/null
+++ b/Engine/Source/Runtime/Experimental/Iris/Core/Private/Iris/ReplicationSystem/NetTokenStoreState.h
@@ -0,0 +1,25 @@
+// Copyright Epic Games, Inc. All Rights Reserved.
+
+#pragma once
+
+#include "Iris/ReplicationSystem/NetToken.h"
+
+namespace UE::Net
+{
+
+class FNetTokenStoreState
+{
+public:
+	// The size of the array is managed by FNetTokenDataStream
+	FNetTokenStoreState();
+
+	// Reserve space for tokens
+	bool ReserveTokenCount(uint32 NewCount);
+
+public:
+
+	// This basically maps from token to tokenkey
+	TArray<FNetTokenStoreKey> TokenInfos;
+};
+
+}
