Iris experimental replication system.

diff --git a/Engine/Source/Runtime/Experimental/Iris/Core/Private/Iris/ReplicationSystem/ObjectPollFrequencyLimiter.cpp b/Engine/Source/Runtime/Experimental/Iris/Core/Private/Iris/ReplicationSystem/ObjectPollFrequencyLimiter.cpp
new file mode 100644
--- /dev/null
+++ b/Engine/Source/Runtime/Experimental/Iris/Core/Private/Iris/ReplicationSystem/ObjectPollFrequencyLimiter.cpp
@@ -0,0 +1,178 @@
+// Copyright Epic Games, Inc. All Rights Reserved.
+
+#include "ObjectPollFrequencyLimiter.h"
+#include "Iris/Core/IrisProfiler.h"
+#include "HAL/IConsoleManager.h"
+#include "Math/VectorRegister.h"
+
+namespace UE
+{
+namespace Net
+{
+namespace Private
+{
+
+// Storing poll frame period as uint8
+const int MaxPollFramePeriod = 256;
+float PollFrequencyMultiplier = 1.0f;
+
+static FAutoConsoleVariableRef CVarPollFrequencyMultiplier(
+		TEXT("net.Iris.PollFrequencyMultiplier"),
+		PollFrequencyMultiplier,
+		TEXT("Multiplied with the NetUpdateFrequency to decide how often PreReplication is called on an Actor. Default factor is 1.0.")
+		);
+
+FObjectPollFrequencyLimiter::FObjectPollFrequencyLimiter()
+{
+}
+
+void FObjectPollFrequencyLimiter::Init(uint32 MaxActiveObjectCount)
+{
+	// We want to be able to easily update 32 objects at a time.
+	const uint32 StorageObjectCount = (MaxActiveObjectCount + 31U) & ~(MaxActiveObjectCount - 1U);
+
+	// Allocate max amount of memory required.
+	FramesBetweenUpdates.SetNumZeroed(StorageObjectCount);
+	FrameCounters.SetNumZeroed(StorageObjectCount);
+}
+
+void FObjectPollFrequencyLimiter::Deinit()
+{
+	FramesBetweenUpdates.Empty();
+	FrameCounters.Empty();
+}
+
+void FObjectPollFrequencyLimiter::Update(const FNetBitArrayView& ScopableObjects, const FNetBitArrayView& DirtyObjects, FNetBitArrayView& OutObjectsToPoll)
+{
+	IRIS_PROFILER_SCOPE(ObjectPollFrequencyLimiter_Update);
+
+	typedef FNetBitArrayView::StorageWordType WordType;
+	constexpr uint32 WordBitCount = FNetBitArrayView::WordBitCount;
+
+	if (!MaxInternalHandle)
+	{
+		return;
+	}
+
+	++FrameIndex;
+
+	uint8* CountersData = FrameCounters.GetData();
+	const uint8* FramesBetweenUpdatesData = FramesBetweenUpdates.GetData();
+
+	const WordType* ScopableObjectsData = ScopableObjects.GetData();
+	const WordType* DirtyObjectsData = DirtyObjects.GetData();
+	WordType* ObjectsToPollData = OutObjectsToPoll.GetData();
+
+	// Decrease counters by one and if they wrap around it's time to poll.
+	// Algorithm can easily be throttled to only allow at most N objects to be polled per frame.
+#if PLATFORM_ENABLE_VECTORINTRINSICS == 1 && !PLATFORM_ENABLE_VECTORINTRINSICS_NEON
+	{
+		static_assert(WordBitCount == 32U, "Code assumes a NetBitArray word size of 32 bits.");
+
+		const VectorRegisterInt AllBitsSet = _mm_cmpeq_epi32(_mm_setzero_si128(), _mm_setzero_si128());
+		const VectorRegisterInt Zero = _mm_setzero_si128();
+		const VectorRegisterInt One = _mm_set1_epi8(char(1));
+
+		for (uint32 ObjectIndex = 0, ObjectEndIndex = MaxInternalHandle;
+			ObjectIndex <= ObjectEndIndex; 
+			ObjectIndex += WordBitCount, CountersData += WordBitCount, FramesBetweenUpdatesData += WordBitCount, ++ScopableObjectsData, ++DirtyObjectsData, ++ObjectsToPollData)
+		{
+			// Skip ranges with no scopable objects.
+			const WordType ObjectsInScopeWord = *ScopableObjectsData;
+			if (!ObjectsInScopeWord)
+			{
+				continue;
+			}
+
+			VectorRegisterInt Values0 = _mm_loadu_si128(reinterpret_cast<const __m128i*>(CountersData + 0U));
+			VectorRegisterInt Values1 = _mm_loadu_si128(reinterpret_cast<const __m128i*>(CountersData + 16U));
+
+			const VectorRegisterInt FramesBetweenUpdates0 = _mm_loadu_si128(reinterpret_cast<const __m128i*>(FramesBetweenUpdatesData + 0U));
+			const VectorRegisterInt FramesBetweenUpdates1 = _mm_loadu_si128(reinterpret_cast<const __m128i*>(FramesBetweenUpdatesData + 16U));
+
+			Values0 = _mm_sub_epi8(Values0, One);
+			const VectorRegisterInt MaskToUpdate0 = _mm_cmpeq_epi8(Values0, AllBitsSet);
+			Values0 = VectorIntSelect(MaskToUpdate0, FramesBetweenUpdates0, Values0);
+
+			Values1 = _mm_sub_epi8(Values1, One);
+			const VectorRegisterInt MaskToUpdate1 = _mm_cmpeq_epi8(Values1, AllBitsSet);
+			Values1 = VectorIntSelect(MaskToUpdate1, FramesBetweenUpdates1, Values1);
+
+			const uint32 ObjectsToPoll0 = _mm_movemask_epi8(MaskToUpdate0);
+			const uint32 ObjectsToPoll1 = _mm_movemask_epi8(MaskToUpdate1);
+
+			_mm_storeu_si128(reinterpret_cast<__m128i*>(CountersData + 0U), Values0);
+			_mm_storeu_si128(reinterpret_cast<__m128i*>(CountersData + 16U), Values1);
+
+			// Poll objects that have been set to dirty or are due.
+			const WordType ObjectsToPoll = (ObjectsToPoll1 << 16U) | ObjectsToPoll0;
+			const WordType DirtyObjectWord = *DirtyObjectsData;
+			*ObjectsToPollData = (ObjectsToPoll | DirtyObjectWord) & ObjectsInScopeWord;
+		}
+	}
+#else
+	// Slower path using scalar integer instructions.
+	{
+		for (uint32 ObjectIndex = 0, ObjectEndIndex = MaxInternalHandle;
+			ObjectIndex <= ObjectEndIndex;
+			ObjectIndex += WordBitCount, CountersData += WordBitCount, FramesBetweenUpdatesData += WordBitCount, ++ScopableObjectsData, ++DirtyObjectsData, ++ObjectsToPollData)
+		{
+			// Skip ranges with no scopable objects.
+			WordType ObjectsInScopeWord = *ScopableObjectsData;
+			if (!ObjectsInScopeWord)
+			{
+				continue;
+			}
+
+			WordType ObjectsToPoll = 0;
+			for (SIZE_T It = 0, IndexOffset = 0; It < 8; ++It, IndexOffset += 4U, ObjectsInScopeWord >>= 4U)
+			{
+				if (!(ObjectsInScopeWord & 15U))
+				{
+					continue;
+				}
+
+				uint8 Counter0 = CountersData[IndexOffset + 0];
+				uint8 Counter1 = CountersData[IndexOffset + 1];
+				uint8 Counter2 = CountersData[IndexOffset + 2];
+				uint8 Counter3 = CountersData[IndexOffset + 3];
+
+				const uint8 FramesBetweenUpdates0 = FramesBetweenUpdatesData[IndexOffset + 0];
+				const uint8 FramesBetweenUpdates1 = FramesBetweenUpdatesData[IndexOffset + 1];
+				const uint8 FramesBetweenUpdates2 = FramesBetweenUpdatesData[IndexOffset + 2];
+				const uint8 FramesBetweenUpdates3 = FramesBetweenUpdatesData[IndexOffset + 3];
+
+				--Counter0;
+				--Counter1;
+				--Counter2;
+				--Counter3;
+
+				const uint32 Mask0 = Counter0 == 255 ? 1 : 0;
+				const uint32 Mask1 = Counter1 == 255 ? 1 : 0;
+				const uint32 Mask2 = Counter2 == 255 ? 1 : 0;
+				const uint32 Mask3 = Counter3 == 255 ? 1 : 0;
+
+				Counter0 = Mask0 ? FramesBetweenUpdates0 : Counter0;
+				Counter1 = Mask1 ? FramesBetweenUpdates1 : Counter1;
+				Counter2 = Mask2 ? FramesBetweenUpdates2 : Counter2;
+				Counter3 = Mask3 ? FramesBetweenUpdates3 : Counter3;
+
+				const uint32 Mask = (Mask3 << 3U) | (Mask2 << 2U) | (Mask1 << 1U) | Mask0;
+				ObjectsToPoll |= (Mask << IndexOffset);
+
+				CountersData[IndexOffset + 0] = Counter0;
+				CountersData[IndexOffset + 1] = Counter1;
+				CountersData[IndexOffset + 2] = Counter2;
+				CountersData[IndexOffset + 3] = Counter3;
+			}
+
+			const WordType DirtyObjectWord = *DirtyObjectsData;
+			*ObjectsToPollData = (ObjectsToPoll | DirtyObjectWord) & ObjectsInScopeWord;
+		}
+	}
+#endif
+}
+
+}
+}
+}
diff --git a/Engine/Source/Runtime/Experimental/Iris/Core/Private/Iris/ReplicationSystem/ObjectPollFrequencyLimiter.h b/Engine/Source/Runtime/Experimental/Iris/Core/Private/Iris/ReplicationSystem/ObjectPollFrequencyLimiter.h
new file mode 100644
--- /dev/null
+++ b/Engine/Source/Runtime/Experimental/Iris/Core/Private/Iris/ReplicationSystem/ObjectPollFrequencyLimiter.h
@@ -0,0 +1,66 @@
+// Copyright Epic Games, Inc. All Rights Reserved.
+
+#pragma once
+
+#include "CoreTypes.h"
+#include "Containers/Array.h"
+#include "Net/Core/NetBitArray.h"
+
+namespace UE::Net::Private
+{
+	typedef uint32 FInternalNetHandle;
+}
+
+namespace UE::Net::Private
+{
+
+extern IRISCORE_API float PollFrequencyMultiplier;
+extern IRISCORE_API const int MaxPollFramePeriod;
+
+class FObjectPollFrequencyLimiter
+{
+public:
+	FObjectPollFrequencyLimiter();
+
+	void Init(uint32 MaxActiveObjectCount);
+	void Deinit();
+
+	void SetPollFramePeriod(FInternalNetHandle InternalHandle, uint8 PollFramePeriod);
+
+	void SetPollWithObject(FInternalNetHandle ObjectToPollWithInternalHandle, FInternalNetHandle InternalHandle);
+
+	void Update(const FNetBitArrayView& ScopableObjects, const FNetBitArrayView& DirtyObjects, FNetBitArrayView& OutObjectsToPoll);
+
+private:
+	uint32 GetPollFramePeriodForFrequency(float PollFrequency) const;
+
+	uint32 MaxInternalHandle = 0;
+	uint32 FrameIndex = 0;
+	// We store the number of frames between updates as a byte to be able to process 16 objects per instruction.
+	// This limits polling to at least every 256th frame. At 30Hz this means every 8.5 seconds.
+	uint8 FrameIndexOffsets[256] = {};
+	TArray<uint8> FramesBetweenUpdates;
+	TArray<uint8> FrameCounters;
+};
+
+
+inline void FObjectPollFrequencyLimiter::SetPollFramePeriod(FInternalNetHandle InternalHandle, uint8 PollFramePeriod)
+{
+	MaxInternalHandle = FPlatformMath::Max(MaxInternalHandle, InternalHandle);
+
+	FramesBetweenUpdates[InternalHandle] = PollFramePeriod;
+	// Spread the polling of objects with the same frequency so that if you add lots of objects the same frame they won't be polled at the same time.
+	const uint32 FrameOffset = FrameIndexOffsets[PollFramePeriod]++;
+	FrameCounters[InternalHandle] = FrameOffset % (uint32(PollFramePeriod) + 1U);
+}
+
+inline void FObjectPollFrequencyLimiter::SetPollWithObject(FInternalNetHandle ObjectToPollWithInternalHandle, FInternalNetHandle InternalHandle)
+{
+	MaxInternalHandle = FPlatformMath::Max(MaxInternalHandle, InternalHandle);
+
+	// Copy state from object to poll with
+	FramesBetweenUpdates[InternalHandle] = FramesBetweenUpdates[ObjectToPollWithInternalHandle];
+	FrameCounters[InternalHandle] = FrameCounters[ObjectToPollWithInternalHandle];
+}
+
+}
diff --git a/Engine/Source/Runtime/Experimental/Iris/Core/Private/Iris/ReplicationSystem/ObjectReferenceCache.cpp b/Engine/Source/Runtime/Experimental/Iris/Core/Private/Iris/ReplicationSystem/ObjectReferenceCache.cpp
new file mode 100644
--- /dev/null
+++ b/Engine/Source/Runtime/Experimental/Iris/Core/Private/Iris/ReplicationSystem/ObjectReferenceCache.cpp
@@ -0,0 +1,1429 @@
+// Copyright Epic Games, Inc. All Rights Reserved.
+
+#include "ObjectReferenceCache.h"
+#include "Iris/Core/IrisLog.h"
+#include "Iris/ReplicationSystem/ObjectReplicationBridge.h"
+#include "Iris/ReplicationSystem/ReplicationSystem.h"
+#include "Iris/ReplicationSystem/ReplicationSystemInternal.h"
+#include "Iris/ReplicationSystem/StringTokenStore.h"
+#include "Iris/Serialization/NetBitStreamWriter.h"
+#include "Iris/Serialization/NetBitStreamReader.h"
+#include "Iris/Serialization/ObjectNetSerializer.h"
+#include "UObject/Package.h"
+#include "Net/Core/Trace/NetDebugName.h"
+#include "Net/Core/Trace/NetTrace.h"
+#include "Core/Public/Misc/StringBuilder.h"
+#include "Iris/Serialization/NetSerializationContext.h"
+#include "Iris/Serialization/InternalNetSerializationContext.h"
+#include "Iris/Serialization/NetExportContext.h"
+#include "Iris/Serialization/NetBitStreamUtil.h"
+
+DEFINE_LOG_CATEGORY_STATIC(LogIrisReferences, Log, All);
+
+#if (UE_BUILD_SHIPPING || UE_BUILD_TEST)
+#	define UE_NET_ENABLE_REFERENCECACHE_LOG 0
+#	define UE_NET_VALIDATE_REFERENCECACHE 0
+#else
+#	define UE_NET_ENABLE_REFERENCECACHE_LOG 1
+#	define UE_NET_VALIDATE_REFERENCECACHE 1
+#endif 
+
+#if UE_NET_ENABLE_REFERENCECACHE_LOG
+#	define UE_LOG_REFERENCECACHE(Verbosity, Format, ...)  UE_LOG(LogIrisReferences, Verbosity, Format, ##__VA_ARGS__)
+#else
+#	define UE_LOG_REFERENCECACHE(...)
+#endif
+
+#define UE_LOG_REFERENCECACHE_WARNING(Format, ...)  UE_LOG(LogIrisReferences, Warning, Format, ##__VA_ARGS__)
+
+// $TODO: If we add information about why a dynamic object is being destroyed (filtering or actual destroy) then we can 
+// try to clean up references to reduce cache pollution
+#define UE_NET_CLEANUP_REFERENCES_WITH_DYNAMIC_OUTER 0
+
+namespace UE::Net::Private
+{
+
+/**
+ * Don't allow infinite recursion of InternalLoadObject - an attacker could
+ * send malicious packets that cause a stack overflow on the server.
+ */
+static const int INTERNAL_READ_REF_RECURSION_LIMIT = 16;
+
+FObjectReferenceCache::FObjectReferenceCache()
+: ReplicationSystem(nullptr)
+, ReplicationBridge(nullptr)
+, NetTokenStore(nullptr)
+, StringTokenStore(nullptr)
+, NetHandleManager(nullptr)
+, bIsAuthority(false)
+{
+}
+
+void FObjectReferenceCache::Init(UReplicationSystem* InReplicationSystem)
+{
+	ReplicationSystem = InReplicationSystem;
+	ReplicationBridge = InReplicationSystem->GetReplicationBridgeAs<UObjectReplicationBridge>();
+	NetTokenStore =  &InReplicationSystem->GetReplicationSystemInternal()->GetNetTokenStore();
+	StringTokenStore = InReplicationSystem->GetStringTokenStore();
+	NetHandleManager = &InReplicationSystem->GetReplicationSystemInternal()->GetNetHandleManager();
+	bIsAuthority = InReplicationSystem->IsServer();
+}
+
+bool FObjectReferenceCache::SupportsObjectInternal(const UObject* Object) const
+{
+	if (Object->IsFullNameStableForNetworking())
+	{
+		// If object is fully net addressable, it's definitely supported
+		return true;
+	}
+
+	if (Object->IsSupportedForNetworking())
+	{
+		// This means the server will explicitly tell the client to spawn and assign the id for this object
+		return true;
+	}
+
+	return false;
+}
+
+// $IRIS: $TODO: refactor into utility methods as this is coped from GUIDCache
+bool FObjectReferenceCache::IsDynamicObject(const UObject* Object) const
+{
+	checkSlow(Object != NULL);
+	checkSlow(Object->IsSupportedForNetworking());
+
+	// Any non net addressable object is dynamic
+	return IsDynamicInternal(Object);
+}
+
+bool FObjectReferenceCache::IsDynamicInternal(const UObject* Object) const
+{
+	// Any non net addressable object is dynamic
+	return !Object->IsFullNameStableForNetworking();
+}
+
+bool FObjectReferenceCache::IsAuthority() const
+{
+	return bIsAuthority;
+}
+
+bool FObjectReferenceCache::CanClientLoadObjectInternal(const UObject* Object, bool bIsDynamic) const
+{
+	// PackageMapClient can't load maps, we must wait for the client to load the map when ready
+	// These guids are special guids, where the guid and all child guids resolve once the map has been loaded
+	if (bIsDynamic || Object->GetOutermost()->ContainsMap())
+	{
+		return false;
+	}
+	// We can load everything else
+	return true;
+}
+
+bool FObjectReferenceCache::ShouldIgnoreWhenMissing(FNetHandle RefHandle) const
+{
+	if (RefHandle.IsDynamic())
+	{
+		return true;		// Ignore missing dynamic guids (even on server because client may send RPC on/with object it doesn't know server destroyed)
+	}
+
+	if (IsAuthority())
+	{
+		return false;		// Server never ignores when missing, always warns
+	}
+
+	const FCachedNetObjectReference* CacheObject = ReferenceHandleToCachedReference.Find(RefHandle);
+	if (CacheObject == nullptr)
+	{
+		return false;		// If we haven't been told about this static guid before, we need to warn
+	}
+
+	const FCachedNetObjectReference* OutermostCacheObject = CacheObject;
+
+	while (OutermostCacheObject != nullptr && OutermostCacheObject->OuterNetHandle.IsValid())
+	{
+		OutermostCacheObject = ReferenceHandleToCachedReference.Find(OutermostCacheObject->OuterNetHandle);
+	}
+
+	if (OutermostCacheObject != nullptr)
+	{
+		// If our outer package is not fully loaded, then don't warn, assume it will eventually come in
+		if (OutermostCacheObject->bIsPending)
+		{
+			// Outer is pending, don't warn
+			return true;
+		}
+		// Sometimes, other systems async load packages, which we don't track, but still must be aware of
+		if (OutermostCacheObject->Object != nullptr && !OutermostCacheObject->Object->GetOutermost()->IsFullyLoaded())
+		{
+			return true;
+		}
+	}
+
+	// Ignore warnings when we explicitly are told to
+	return CacheObject->bIgnoreWhenMissing;
+}
+
+bool FObjectReferenceCache::RenamePathForPie(FString &Str, bool bReading)
+{
+	return GPlayInEditorID != -1 ? ReplicationBridge->RemapPathForPIE(Str, bReading) : false;
+}
+
+FNetHandle FObjectReferenceCache::CreateObjectReferenceHandle(const UObject* Object)
+{
+	FNetObjectReference Ref;
+	CreateObjectReferenceInternal(Object, Ref);
+
+	return Ref.GetRefHandle();
+}
+
+bool FObjectReferenceCache::CreateObjectReferenceInternal(const UObject* Object, FNetObjectReference& OutReference)
+{
+	if (!IsValid(Object))
+	{
+		return false;
+	}
+
+	// Check if we are already know about this object
+	if (FNetHandle* ReferenceHandle = ObjectToNetReferenceHandle.Find(Object))
+	{
+		// Verify that this is the same object and clean up if it is not
+		FCachedNetObjectReference& CachedObject = ReferenceHandleToCachedReference.FindChecked(*ReferenceHandle);
+
+		if (CachedObject.Object.Get() == Object)
+		{
+			UE_LOG_REFERENCECACHE(VeryVerbose, TEXT("ObjectReferenceCache::CreateObjectReferenceHandle Found existing %s for Object %s, OuterNetHandle: %s"), *CachedObject.NetHandle.ToString(), ToCStr(Object->GetPathName()), *CachedObject.OuterNetHandle.ToString());
+			OutReference = MakeNetObjectReference(CachedObject);			
+			return true;
+		}
+		else
+		{
+			UE_LOG_REFERENCECACHE(Verbose, TEXT("Detected stale cached Object removing %s from ObjectToNetReferenceHandle lookup"), *CachedObject.NetHandle.ToString());
+
+			ObjectToNetReferenceHandle.Remove(CachedObject.ObjectKey);
+			ObjectToNetReferenceHandle.Remove(Object);
+
+			if ((*ReferenceHandle).IsStatic())
+			{
+				// Note we only cleanse the object reference
+				// we still keep the cachedObject data around in order to be able to serialize static destruction infos
+				CachedObject.ObjectKey = nullptr;
+				CachedObject.Object = nullptr;
+			}
+			else
+			{
+				ReferenceHandleToCachedReference.Remove(*ReferenceHandle);
+			}
+		}
+	}
+
+	// Translated from GuidCache, but kept as separate variable to avoid calling the same functions multiple times
+	const bool bIsFullNameStableForNetworking = Object->IsFullNameStableForNetworking();
+	const bool bIsDynamic = !bIsFullNameStableForNetworking;
+	const bool bIsStatic = !bIsDynamic;
+
+	// Check if we should replicate this object?
+	if (!(bIsFullNameStableForNetworking || Object->IsSupportedForNetworking()))
+	{
+		return false;
+	}
+
+	// We currently do not support creating new references from client to server so we currently using a path based reference
+	if (!IsAuthority())
+	{
+		// Find known outer if there is one, and express reference as a relative path to the outer
+		for (const UObject* Outer = Object; Outer != nullptr; Outer = Outer->GetOuter())
+		{
+			if (FNetHandle* ReplicatedOuter = ObjectToNetReferenceHandle.Find(Outer))
+			{
+				FString RelativeObjectPath = Object->GetPathName(Outer);
+				RenamePathForPie(RelativeObjectPath, false);		
+				OutReference = MakeNetObjectReference(*ReplicatedOuter, StringTokenStore->GetOrCreateToken(RelativeObjectPath));
+
+				return true;
+			}
+		}
+
+		if (bIsDynamic)
+		{
+			// If the object is dynamic and we did not find an outer handle there is not much we can do
+			return false;
+		}
+		else
+		{
+			// For static objects we can always fall back on the full pathname
+			FString ObjectPath = Object->GetPathName();
+			RenamePathForPie(ObjectPath, false);		
+			OutReference = MakeNetObjectReference(FNetHandle(), StringTokenStore->GetOrCreateToken(ObjectPath));
+		
+			return true;
+		}
+	}
+
+	// Generate a new Handle
+	FNetHandle NetHandle = NetHandleManager->AllocateNetHandle(bIsStatic);
+	FNetToken PathToken;
+	
+	check(NetHandle.IsValid());
+
+	// If our name is stable, or at least relative to our outer
+	if (bIsFullNameStableForNetworking || Object->IsNameStableForNetworking())
+	{
+		FString ObjectPath = Object->GetName();
+		RenamePathForPie(ObjectPath, false);
+		PathToken = StringTokenStore->GetOrCreateToken(ObjectPath);
+	}
+
+	// Register cached ref
+	FCachedNetObjectReference CachedObject;
+
+	CachedObject.Object = MakeWeakObjectPtr(const_cast<UObject*>(Object));
+	CachedObject.ObjectKey = Object;
+	CachedObject.NetHandle = NetHandle;
+	CachedObject.RelativePath = PathToken;
+	CachedObject.OuterNetHandle = CreateObjectReferenceHandle(Object->GetOuter());
+	CachedObject.bIsPackage = bIsStatic && Object && !Object->GetOuter();
+	CachedObject.bNoLoad = !CanClientLoadObjectInternal(Object, bIsDynamic); // We are probably going to deal with this as a dependency
+	CachedObject.bIgnoreWhenMissing = !CachedObject.bNoLoad;
+	CachedObject.bIsPending = false;
+	CachedObject.bIsBroken = false;
+
+	check( CachedObject.bIsPackage == (Cast< UPackage >( Object ) != NULL));		// Make sure it really is a package
+
+	ReferenceHandleToCachedReference.Add(NetHandle, CachedObject);
+	ObjectToNetReferenceHandle.Add(Object, NetHandle);
+
+#if UE_NET_TRACE_ENABLED
+	if (NetHandle.IsValid())
+	{
+		TStringBuilder<512> Builder;
+		Builder.Appendf(TEXT("Ref_%s"), PathToken.IsValid() ? StringTokenStore->ResolveToken(PathToken) : ToCStr(Object->GetName()));
+		UE_NET_TRACE_NETHANDLE_CREATED(NetHandle, CreatePersistentNetDebugName(Builder.ToString()), 0, IsAuthority() ? 1U : 0U);
+	}
+#endif
+	
+	UE_LOG_REFERENCECACHE(Verbose, TEXT("ObjectReferenceCache::CreateObjectReferenceHandle %s for Object %s, OuterNetHandle: %s"), *CachedObject.NetHandle.ToString(), ToCStr(Object->GetPathName()), *CachedObject.OuterNetHandle.ToString());
+
+	// Create reference
+	OutReference = MakeNetObjectReference(CachedObject);
+
+#if UE_NET_CLEANUP_REFERENCES_WITH_DYNAMIC_OUTER
+	// We want to track references with a dynamic root so that we can purge them when the dynamic root is removed
+	if (bIsDynamic)
+	{
+		const FNetHandle ReplicatedOuterHandle = GetDynamicRoot(NetHandle);
+		if (ReplicatedOuterHandle.IsValid())
+		{
+			HandleToDynamicOuter.Add(ReplicatedOuterHandle, NetHandle);
+		}
+	}
+#endif
+
+	return true;
+}
+
+FNetHandle FObjectReferenceCache::GetObjectReferenceHandleFromObject(const UObject* Object) const
+{
+	if (Object == nullptr)
+	{
+		return FNetHandle();
+	}
+
+	// Check if we are already know about this object
+	if (const FNetHandle* Reference = ObjectToNetReferenceHandle.Find(Object))
+	{
+		// Verify that this is the same object
+		const FCachedNetObjectReference& CachedObject = ReferenceHandleToCachedReference.FindChecked(*Reference);
+
+		const UObject* ExistingObject = CachedObject.Object.Get();
+		if (ExistingObject == Object)
+		{
+			UE_LOG_REFERENCECACHE(VeryVerbose, TEXT("ObjectReferenceCache::GetObjectReferenceHandleFromObject Found existing %s for Object %s, OuterNetHandle: %s"), *CachedObject.NetHandle.ToString(), ToCStr(Object->GetPathName()), *CachedObject.OuterNetHandle.ToString());
+			
+			return *Reference;
+		}
+		else
+		{
+			UE_LOG_REFERENCECACHE(Verbose, TEXT("ObjectReferenceCache::Found %s cached RefHandle %s for object %s"), CachedObject.Object.IsStale() ? TEXT("Stale") : TEXT("UnResolved"), *CachedObject.NetHandle.ToString(), ToCStr(Object->GetName()));
+		}
+	}
+
+	return FNetHandle();
+}
+
+void FObjectReferenceCache::AddRemoteReference(FNetHandle RefHandle, const UObject* Object)
+{
+	// Only clients should get here
+	check(!IsAuthority()); 
+	check(IsValid(Object));
+	check(RefHandle.IsCompleteHandle());
+
+	if (ensure(RefHandle.IsDynamic()))
+	{
+		UE_LOG_REFERENCECACHE(Verbose, TEXT("AddRemoteReference: %s, Object: %s Pointer: %p" ), *RefHandle.ToString(), Object ? *Object->GetName() : TEXT( "NULL" ), Object);
+
+	#if UE_NET_VALIDATE_REFERENCECACHE
+		// Dynamic object might have been referenced before it gets replicated, then it will then already be registered (with a path relative to its outer)
+		if (FNetHandle* ExistingHandle = ObjectToNetReferenceHandle.Find(Object))
+		{
+			FCachedNetObjectReference& ExistingCachedObject = ReferenceHandleToCachedReference.FindChecked(*ExistingHandle);
+			if (ExistingCachedObject.NetHandle != RefHandle)
+			{
+				UE_LOG_REFERENCECACHE_WARNING(TEXT("AddRemoteReference: Detected conflicting Ref %s"), *ExistingCachedObject.NetHandle.ToString());
+			}
+		}
+	#endif
+
+		// The object could already be known as it has been exported as a reference prior to being replicated to this connection
+		if (FCachedNetObjectReference* CacheObjectPtr = ReferenceHandleToCachedReference.Find(RefHandle))
+		{
+			// Patch-up object pointers
+			CacheObjectPtr->Object = MakeWeakObjectPtr(const_cast<UObject*>(Object));
+			CacheObjectPtr->ObjectKey = Object;
+		}
+		else
+		{
+			LLM_SCOPE(ELLMTag::Iris);
+
+			FCachedNetObjectReference CachedObject;
+	
+			CachedObject.Object = MakeWeakObjectPtr(const_cast<UObject*>(Object));
+			CachedObject.ObjectKey = Object;
+			CachedObject.NetHandle = RefHandle;
+			CachedObject.RelativePath = FNetToken();
+			CachedObject.OuterNetHandle = FNetHandle();
+			CachedObject.bIsPackage = false;
+			CachedObject.bNoLoad = false;
+			CachedObject.bIsBroken = false;
+			CachedObject.bIsPending = false;
+			CachedObject.bIgnoreWhenMissing = false;
+
+			ReferenceHandleToCachedReference.Add(RefHandle, CachedObject);
+		}
+
+		ObjectToNetReferenceHandle.Add(Object, RefHandle);
+	}
+}
+
+void FObjectReferenceCache::RemoveReference(FNetHandle RefHandle, const UObject* Object)
+{
+	// Cleanup cached data for dynamic instance
+	if (ensure(RefHandle.IsDynamic()))
+	{
+		UE_LOG_REFERENCECACHE(Verbose, TEXT("ObjectReferenceCache::RemoveReference %s, Object: %s Pointer: %p"), *RefHandle.ToString(), Object ? *Object->GetName() : TEXT( "NULL" ), Object);
+
+		// We create an iterator to avoid having to find the key twice
+		TMap<FNetHandle, FCachedNetObjectReference>::TKeyIterator It = ReferenceHandleToCachedReference.CreateKeyIterator(RefHandle);
+		if (It)
+		{
+			FCachedNetObjectReference& CachedItemToRemove = It.Value();
+			if (Object == nullptr)
+			{
+				// Object might have been GC'd but we still want to cleanup 
+				Object = CachedItemToRemove.ObjectKey;
+			}
+			else
+			{
+				check(Object == CachedItemToRemove.ObjectKey);
+			}
+
+#if UE_NET_CLEANUP_REFERENCES_WITH_DYNAMIC_OUTER
+			// We want to remove cached references that referenced the removed dynamic references
+			TArray<FNetHandle> SubReferences;
+			HandleToDynamicOuter.MultiFind(RefHandle, SubReferences);
+			if (SubReferences.Num())
+			{
+				for (FNetHandle SubRefHandle : SubReferences)
+				{
+					FCachedNetObjectReference RemovedSubObjectRef;
+					if (ReferenceHandleToCachedReference.RemoveAndCopyValue(SubRefHandle, RemovedSubObjectRef))
+					{
+						const UObject* SubRefObject = RemovedSubObjectRef.Object.Get();
+						UE_LOG_REFERENCECACHE(Verbose, TEXT("ObjectReferenceCache::RemoveReference %s, Referenced dynamic %s, Object: %s Pointer: %p"), *SubRefHandle.ToString(), *RefHandle.ToString(), SubRefObject ? *SubRefObject->GetName() : TEXT( "NULL" ), RemovedSubObjectRef.ObjectKey);
+
+						ObjectToNetReferenceHandle.Remove(RemovedSubObjectRef.ObjectKey);
+					}
+				}
+
+				HandleToDynamicOuter.Remove(RefHandle);
+			}
+#endif
+
+			// Remove if RelativePath is invalid
+			if (!CachedItemToRemove.RelativePath.IsValid())
+			{
+				It.RemoveCurrent();
+			}
+			else
+			{
+				// Clear out object reference since it will be invalid
+				CachedItemToRemove.ObjectKey = nullptr;
+				CachedItemToRemove.Object.Reset();
+			}
+		}	
+		
+		if (Object)
+		{
+			// We create an iterator to avoid having to find the key twice
+			TMap<const UObject*, FNetHandle>::TKeyIterator ObjectIt = ObjectToNetReferenceHandle.CreateKeyIterator(Object);
+			if (ObjectIt)
+			{
+				const FNetHandle ExistingRefHandle = ObjectIt.Value();
+				if (ExistingRefHandle == RefHandle)
+				{
+					ObjectIt.RemoveCurrent();
+				}
+				else
+				{
+					// If we have an existing handle associated with the Object that differs from the one stored in the cache it means that the client might have destroyed the instance without notifying the replication system
+					// and reused the instance, if that is the case we just leave the association as it is
+					UE_LOG_REFERENCECACHE(Verbose, TEXT("FObjectReferenceCache::RemoveReference ExistingReference %s != RefHandle being removed: %s "), *(ExistingRefHandle.ToString()), *RefHandle.ToString());
+				}
+			}
+		}
+	}
+}
+
+UObject* FObjectReferenceCache::GetObjectFromReferenceHandle(FNetHandle RefHandle)
+{
+	FCachedNetObjectReference* CachedObject = ReferenceHandleToCachedReference.Find(RefHandle);
+
+	UObject* ResolvedObject = CachedObject ? CachedObject->Object.Get() : nullptr;
+
+	return ResolvedObject;
+}
+
+// $IRIS: $TODO: Most of the logic comes from GUIDCache::GetObjectFromNetGUID so we want to keep this in sync
+UObject* FObjectReferenceCache::ResolveObjectReferenceHandleInternal(FNetHandle RefHandle, const FNetObjectResolveContext& ResolveContext, bool& bOutMustBeMapped)
+{
+	FCachedNetObjectReference* CacheObjectPtr = ReferenceHandleToCachedReference.Find(RefHandle);
+	
+	if (!CacheObjectPtr)
+	{
+		bOutMustBeMapped = false;
+		return nullptr;
+	}
+
+	check(!CacheObjectPtr->bIsPending);
+	
+	UObject* Object = CacheObjectPtr->Object.Get();
+
+	bOutMustBeMapped = !CacheObjectPtr->bNoLoad;
+
+	if (Object != nullptr)
+	{
+		// Either the name should match, or this is dynamic, or we're on the server
+		// $IRIS: $TODO: Consider using FNames over string tokens, or store FNames in cache and fixup when resolving?
+#if UE_NET_VALIDATE_REFERENCECACHE
+		{			
+			FString ObjectPath(StringTokenStore->ResolveRemoteToken(CacheObjectPtr->RelativePath, *ResolveContext.RemoteNetTokenStoreState));
+			RenamePathForPie(ObjectPath, true);
+			check(ObjectPath == Object->GetFName().ToString() || RefHandle.IsDynamic() || IsAuthority());
+
+			FNetHandle* ExistingRefHandle = ObjectToNetReferenceHandle.Find(Object);
+			if (ExistingRefHandle && *ExistingRefHandle != RefHandle)
+			{
+				UE_LOG_REFERENCECACHE_WARNING(TEXT("ObjectReferenceCache::ResolveObjectRefererenceHandle detected potential conflicting references: %s and %s"), *RefHandle.ToString(), *(*ExistingRefHandle).ToString());
+			}
+		}
+#endif
+
+		check(IsValid(Object));
+
+		UE_LOG_REFERENCECACHE(VeryVerbose, TEXT("ObjectReferenceCache::ResolveObjectRefererenceHandle Found existing object for %s for Object: %s Pointer: %p"), *RefHandle.ToString(), ToCStr(Object->GetPathName()), Object);
+
+		return Object;
+	}
+
+	if (CacheObjectPtr->bIsBroken)
+	{
+		// This object is broken, we know it won't load
+		// At this stage, any warnings should have already been logged, so we just need to ignore from this point forward
+		return nullptr;
+	}
+
+	if (CacheObjectPtr->bIsPending)
+	{
+		// We're not done loading yet (and no error has been reported yet)
+		return nullptr;
+	}
+
+	// If we do not have a path and our handle is invalid, we cannot do much
+	if (!CacheObjectPtr->RelativePath.IsValid())
+	{
+		// If we don't have a path, assume this is a dynamic handle
+		check(RefHandle.IsDynamic());
+		
+		return nullptr;
+	}
+
+	// First, resolve the outer
+	UObject* ObjOuter = nullptr;
+
+	if (CacheObjectPtr->OuterNetHandle.IsValid())
+	{
+		// If we get here, we depend on an outer to fully load, don't go further until we know we have a fully loaded outer
+		FCachedNetObjectReference* OuterCacheObject = ReferenceHandleToCachedReference.Find(CacheObjectPtr->OuterNetHandle);
+		if (OuterCacheObject == nullptr)
+		{
+			// Shouldn't be possible, but just in case...
+			if (CacheObjectPtr->OuterNetHandle.IsStatic())
+			{
+				UE_LOG(LogIris, Error, TEXT("GetObjectFromRefHandle: Static outer not registered. FullPath: %s"), ToCStr(FullPath(RefHandle, ResolveContext)))
+				CacheObjectPtr->bIsBroken = 1;	// Set this to 1 so that we don't keep spamming
+			}
+			return nullptr;
+		}
+
+		// If outer is broken, we will never load, set ourselves to broken as well and bail
+		if (OuterCacheObject->bIsBroken)
+		{
+			UE_LOG(LogIris, Error, TEXT("GetObjectFromRefHandle:  Outer is broken. FullPath: %s"), ToCStr(FullPath(RefHandle, ResolveContext)));
+			CacheObjectPtr->bIsBroken = 1;	// Set this to 1 so that we don't keep spamming
+
+			return nullptr;
+		}
+
+		// Try to resolve the outer
+		bool bOuterMustBeMapped;
+		ObjOuter = ResolveObjectReferenceHandleInternal(CacheObjectPtr->OuterNetHandle, ResolveContext, bOuterMustBeMapped);
+
+		// If we can't resolve the outer, we cannot do more, we are probably waiting for a replicated object or we are waiting for a stringtoken
+		if (ObjOuter == nullptr)
+		{
+			return nullptr;
+		}
+	}
+
+	// At this point, we either have an outer, or we are a package
+	const uint32 TreatAsLoadedFlags = EPackageFlags::PKG_CompiledIn | EPackageFlags::PKG_PlayInEditor;
+
+	check(!CacheObjectPtr->bIsPending);
+
+	if (!ensure(ObjOuter == nullptr || ObjOuter->GetOutermost()->IsFullyLoaded() || ObjOuter->GetOutermost()->HasAnyPackageFlags(TreatAsLoadedFlags)))
+	{
+		UE_LOG(LogIris, Error, TEXT("GetObjectFromRefHandle: Outer is null or package is not fully loaded. FullPath: %s Outer: %s"), ToCStr(FullPath(RefHandle, ResolveContext)), *GetFullNameSafe(ObjOuter));
+	}
+
+	const TCHAR* ResolvedToken = StringTokenStore->ResolveRemoteToken(CacheObjectPtr->RelativePath, *ResolveContext.RemoteNetTokenStoreState);
+	if (ResolvedToken == nullptr)
+	{
+		return nullptr;
+	}
+
+	FString ObjectPath(ResolvedToken);
+	RenamePathForPie(ObjectPath, true);
+
+	// See if this object is in memory
+	Object = StaticFindObject(UObject::StaticClass(), ObjOuter, *ObjectPath, false);
+
+	// Assume this is a package if the outer is invalid and this is a static guid
+	const bool bIsPackage = CacheObjectPtr->bIsPackage;
+
+	if (Object == nullptr && !CacheObjectPtr->bNoLoad)
+	{
+		if (IsAuthority())
+		{
+			// Log when the server needs to re-load an object, it's probably due to a GC after initially loading as default guid
+			UE_LOG(LogIris, Error, TEXT("GetObjectFromRefHandle: Server re-loading object (might have been GC'd). FullPath: %s"), ToCStr(FullPath(RefHandle, ResolveContext)));
+		}
+
+		// $IRIS: $TODO: Implement Async loading support https://jira.it.epicgames.com/browse/UE-123487
+		if (bIsPackage)
+		{
+			// Async load the package if:
+			//	1. We are actually a package
+			//	2. We aren't already pending
+			//	3. We're actually suppose to load (levels don't load here for example)
+			//		(Refer to CanClientLoadObject, which is where we protect clients from trying to load levels)
+
+			//if (ShouldAsyncLoad())
+			//{
+			//	if (!PendingAsyncLoadRequests.Contains(Path))
+			//	{
+			//		CacheObjectPtr->bIsPending = true;
+
+			//		StartAsyncLoadingPackage(NetGUID, false);
+			//		UE_LOG(LogIris, Error, TEXT("GetObjectFromRefHandle: Async loading package. Path: %s, %s"), *ObjectPath, *RefHandle.ToString());
+			//	}
+			//	else
+			//	{
+			//		ValidateAsyncLoadingPackage(NetGUID);
+			//	}
+
+			//	// There is nothing else to do except wait on the delegate to tell us this package is done loading
+			//	return NULL;
+			//}
+			//else
+			{
+				// Async loading disabled
+				Object = LoadPackage(nullptr, *ObjectPath, LOAD_None);
+			}
+		}
+		else
+		{
+			// If we have a package, but for some reason didn't find the object then do a blocking load as a last attempt
+			// This can happen for a few reasons:
+			//	1. The object was GC'd, but the package wasn't, so we need to reload
+			//	2. Someone else started async loading the outer package, and it's not fully loaded yet
+			Object = StaticLoadObject(UObject::StaticClass(), ObjOuter, *ObjectPath, nullptr, LOAD_NoWarn);
+
+			//if (ShouldAsyncLoad())
+			//{
+			//	UE_LOG( LogNetPackageMap, Error, TEXT( "GetObjectFromRefHandle: Forced blocking load. Path: %s, %s"), *ObjectPath, *RefHandle.ToString());
+			//}
+		}
+	}
+
+	if (Object == nullptr)
+	{
+		if (!CacheObjectPtr->bIgnoreWhenMissing)
+		{
+			CacheObjectPtr->bIsBroken = 1;	// Set this to 1 so that we don't keep spamming
+			UE_LOG(LogIris, Error, TEXT("GetObjectFromRefHandle: Failed to resolve. FullPath: %s"), ToCStr(FullPath(RefHandle, ResolveContext)));
+		}
+
+		return nullptr;
+	}
+
+	if (bIsPackage)
+	{
+		UPackage* Package = Cast<UPackage>(Object);
+
+		if (Package == nullptr)
+		{
+			// This isn't really a package but it should be
+			CacheObjectPtr->bIsBroken = true;
+			//UE_LOG( LogNetPackageMap, Error, TEXT( "GetObjectFromNetGUID: Object is not a package but should be! Path: %s, NetGUID: %s" ), *Path.ToString(), *NetGUID.ToString() );
+			return nullptr;
+		}
+
+		if (!Package->IsFullyLoaded() 
+			&& !Package->HasAnyPackageFlags(TreatAsLoadedFlags)) //TODO: dependencies of CompiledIn could still be loaded asynchronously. Are they necessary at this point??
+		{
+			//if (ShouldAsyncLoad() && Package->HasAnyInternalFlags(EInternalObjectFlags::AsyncLoading))
+			//{
+			//	CacheObjectPtr->bIsPending = true;
+
+			//	// Something else is already async loading this package, calling load again will add our callback to the existing load request
+			//	StartAsyncLoadingPackage(NetGUID, true);
+			//	UE_LOG(LogNetPackageMap, Log, TEXT("GetObjectFromNetGUID: Listening to existing async load. Path: %s, NetGUID: %s"), *Path.ToString(), *NetGUID.ToString());
+			//}
+			//else
+			{
+				// If package isn't fully loaded, load it now
+				UE_LOG_REFERENCECACHE(Warning, TEXT("ObjectReferenceCache::GetObjectFromRefHandle Blocking load of %s, %s"), ToCStr(ObjectPath), *RefHandle.ToString());
+				Object = LoadPackage(nullptr, ToCStr(ObjectPath), LOAD_None);
+			}
+		}
+	}
+
+	//if (CacheObjectPtr->NetworkChecksum != 0 && !CVarIgnoreNetworkChecksumMismatch.GetValueOnAnyThread())
+	//{
+	//	const uint32 NetworkChecksum = GetNetworkChecksum( Object );
+
+	//	if (CacheObjectPtr->NetworkChecksum != NetworkChecksum )
+	//	{
+	//		if ( NetworkChecksumMode == ENetworkChecksumMode::SaveAndUse )
+	//		{
+	//			FString ErrorStr = FString::Printf(TEXT("GetObjectFromNetGUID: Network checksum mismatch. ``IDPath: %s, %u, %u"), *FullNetGUIDPath(NetGUID), CacheObjectPtr->NetworkChecksum, NetworkChecksum);
+	//			UE_LOG( LogNetPackageMap, Warning, TEXT("%s"), *ErrorStr );
+
+	//			CacheObjectPtr->bIsBroken = true;
+
+	//			BroadcastNetFailure(Driver, ENetworkFailure::NetChecksumMismatch, ErrorStr);
+	//			return NULL;
+	//		}
+	//		else
+	//		{
+	//			UE_LOG( LogNetPackageMap, Verbose, TEXT( "GetObjectFromNetGUID: Network checksum mismatch. ``IDPath: %s, %u, %u" ), *FullNetGUIDPath( NetGUID ), CacheObjectPtr->NetworkChecksum, NetworkChecksum );
+	//		}
+	//	}
+	//}
+
+	if (Object && !ReplicationBridge->ObjectLevelHasFinishedLoading(Object))
+	{
+		UE_LOG_REFERENCECACHE(Verbose, TEXT("GetObjectFromRefHandle: Forcing object to NULL since level is not loaded yet. Object: %s" ), *Object->GetFullName());
+
+		return nullptr;
+	}
+
+	// Assign the resolved object to this guid
+	if (!Object)
+	{
+		if (!CacheObjectPtr->bIgnoreWhenMissing)
+		{
+			CacheObjectPtr->bIsBroken = 1;	// Set this to 1 so that we don't keep spamming
+			UE_LOG(LogIris, Error, TEXT("GetObjectFromRefHandle: Failed to resolve. FullPath: %s non existing"), ToCStr(FullPath(RefHandle, ResolveContext)));
+		}
+		return nullptr;
+	}
+
+	if (!IsValid(Object))
+	{
+		UE_LOG_REFERENCECACHE_WARNING(TEXT("GetObjectFromRefHandle: Resolved object Pointer: %p that is not valid? FullPath: %s"), Object, ToCStr(FullPath(RefHandle, ResolveContext)));
+	}
+
+	CacheObjectPtr->Object = MakeWeakObjectPtr(Object);
+	CacheObjectPtr->ObjectKey = Object;
+
+	// Assign the guid to the object 
+	// We don't want to assign this guid to the object if this guid is timing out
+	// But we'll have to if there is no other guid yet
+	//const bool bAllowClientRemap = !bIsNetGUIDAuthority && GbAllowClientRemapCacheObject;
+	//const bool bIsNotReadOnlyOrAllowRemap = (CacheObjectPtr->ReadOnlyTimestamp == 0 || bAllowClientRemap);
+
+	//if (bIsNotReadOnlyOrAllowRemap || !NetGUIDLookup.Contains(Object))
+	//{
+	//	if (CacheObjectPtr->ReadOnlyTimestamp > 0)
+	//	{
+	//		UE_LOG( LogNetPackageMap, Warning, TEXT( "GetObjectFromNetGUID: Attempt to reassign read-only guid. FullNetGUIDPath: %s" ), *FullNetGUIDPath( NetGUID ) );
+
+	//		if (bAllowClientRemap)
+	//		{
+	//			CacheObjectPtr->ReadOnlyTimestamp = 0;
+	//		}
+	//	}
+
+	//	NetGUIDLookup.Add( Object, NetGUID );
+	//}
+
+	// Promote to full handle
+	FNetHandle CompleteRefHandle = FNetHandleManager::MakeNetHandle(RefHandle.GetId(), ReplicationSystem->GetId());
+
+	// If we already have an existing handle associated with this object, we need do do a few additional checks.
+	// Due to level streaming (streaming in and out the same level) we might have multiple path based references 
+	// to the same object which might potentially might resolve to the same object due to how streaming works
+	// This is fine, but we can only allow one of them to map from object to RefHandle
+	// The Current resolution strategy is to rely on the fact that we assign Static handles in an ever increasing manner
+	// and simply use the highest index as the correct reference	
+
+	bool bShouldAssign = true;
+	FNetHandle* ExistingHandle = ObjectToNetReferenceHandle.Find(Object);
+	if (ExistingHandle && (*ExistingHandle != CompleteRefHandle))
+	{		
+		FCachedNetObjectReference& ExistingCachedObject = ReferenceHandleToCachedReference.FindChecked(*ExistingHandle);
+		if (ExistingCachedObject.Object.IsStale())
+		{
+			UE_LOG_REFERENCECACHE_WARNING(TEXT("Invalidating object for Stale Ref %s with %s as the object has been reused"), *ExistingCachedObject.NetHandle.ToString(), *CompleteRefHandle.ToString());
+		}
+		else if (ExistingCachedObject.Object.HasSameIndexAndSerialNumber(CacheObjectPtr->Object))
+		{
+			bShouldAssign = CompleteRefHandle.GetId() > (*ExistingHandle).GetId();
+			if (bShouldAssign)
+			{
+				UE_LOG_REFERENCECACHE_WARNING(TEXT("Invalidating object for conflicting Ref %s with %s"), *ExistingCachedObject.NetHandle.ToString(), *CompleteRefHandle.ToString());
+			}	
+		}
+	}
+
+	LLM_SCOPE(ELLMTag::Iris);
+
+	if (bShouldAssign)
+	{
+		ObjectToNetReferenceHandle.Add(Object, CompleteRefHandle);
+		UE_NET_TRACE_NETHANDLE_CREATED(CompleteRefHandle, CreatePersistentNetDebugName(ResolvedToken), 0, IsAuthority() ? 1U : 0U);
+	}
+	
+	UE_LOG_REFERENCECACHE(Verbose, TEXT("ObjectReferenceCache::ResolveObjectRefererenceHandle %s for Object: %s Pointer: %p"), *RefHandle.ToString(), ToCStr(Object->GetPathName()), Object);
+
+	return Object;
+}
+
+ENetObjectReferenceResolveResult FObjectReferenceCache::ResolveObjectReference(const FNetObjectReference& Reference, const FNetObjectResolveContext& ResolveContext, UObject*& OutResolvedObject)
+{
+	UObject* ResolvedObject = nullptr;
+	bool bMustBeMapped = false;
+	
+	if (Reference.PathToken.IsValid())
+	{
+		// This path is only used by Client to Server references
+		//$TODO: $IRIS: Should we try to load or not? Currently we do not!
+
+		const TCHAR* ResolvedToken = StringTokenStore->ResolveRemoteToken(Reference.PathToken, *ResolveContext.RemoteNetTokenStoreState);
+		if (ResolvedToken)
+		{
+			FString ObjectPath(ResolvedToken);
+			RenamePathForPie(ObjectPath, true);
+
+			// If both path and handle is valid this is a reference to a dynamic object with a relative path
+			if (Reference.GetRefHandle().IsValid())
+			{
+				UObject* ReplicatedOuter = GetObjectFromReferenceHandle(Reference.GetRefHandle());
+				ResolvedObject = ReplicatedOuter ? StaticFindObject(UObject::StaticClass(), ReplicatedOuter, ToCStr(ObjectPath), false) : nullptr;
+			}
+			else
+			{
+				ResolvedObject = StaticFindObject(UObject::StaticClass(), nullptr, *ObjectPath, false);
+			}
+		}
+	}
+	else
+	{
+		ResolvedObject = ResolveObjectReferenceHandleInternal(Reference.GetRefHandle(), ResolveContext, bMustBeMapped);
+	}
+	
+	ENetObjectReferenceResolveResult Result = ENetObjectReferenceResolveResult::None;
+	if (!ResolvedObject && Reference.IsValid())
+	{
+		Result |= ENetObjectReferenceResolveResult::HasUnresolvedReferences;
+		Result |= bMustBeMapped ? ENetObjectReferenceResolveResult::HasUnresolvedMustBeMappedReferences : ENetObjectReferenceResolveResult::None;		
+	}
+
+	OutResolvedObject = ResolvedObject;
+
+	return Result;
+}
+
+// Find replicated outer
+FNetObjectReference FObjectReferenceCache::GetReplicatedOuter(const FNetObjectReference& Reference) const
+{
+	if (!Reference.IsValid())
+	{
+		return FNetObjectReference();
+	}
+
+	const FCachedNetObjectReference* CachedObject = ReferenceHandleToCachedReference.Find(Reference.GetRefHandle());
+	FNetHandle OuterNetHandle = CachedObject ? CachedObject->OuterNetHandle : FNetHandle();
+
+	while (OuterNetHandle.IsValid())
+	{
+		if (NetHandleManager->GetInternalIndex(OuterNetHandle))
+		{
+			return FNetObjectReference(OuterNetHandle);
+		}
+
+		CachedObject = ReferenceHandleToCachedReference.Find(OuterNetHandle);
+		OuterNetHandle = CachedObject ? CachedObject->OuterNetHandle : FNetHandle();
+	};
+
+	return FNetObjectReference();
+}
+
+FNetHandle FObjectReferenceCache::GetDynamicRoot(const FNetHandle Handle) const
+{
+	const FCachedNetObjectReference* CachedObject = ReferenceHandleToCachedReference.Find(Handle);
+	while (CachedObject && CachedObject->OuterNetHandle.IsValid())
+	{
+		CachedObject = ReferenceHandleToCachedReference.Find(CachedObject->OuterNetHandle);
+	}
+
+	if (CachedObject && CachedObject->NetHandle.IsDynamic())
+	{
+		return CachedObject->NetHandle;
+	}
+	else
+	{
+		return FNetHandle();
+	}
+}
+
+FNetObjectReference FObjectReferenceCache::GetOrCreateObjectReference(const UObject* Instance)
+{
+	FNetObjectReference ObjectReference;
+	CreateObjectReferenceInternal(Instance, ObjectReference);
+
+	return ObjectReference;
+}
+
+FNetObjectReference FObjectReferenceCache::GetOrCreateObjectReference(const FString& ObjectPath, const UObject* Outer)
+{
+	// Only allow creation of path based references on server
+	if (!IsAuthority())
+	{
+		return FNetObjectReference();
+	}
+	
+	// Generate a new Handle, we assume this to be a static handle
+	FNetObjectReference OutReference;
+
+	const bool bIsStatic = true;
+	FNetHandle Handle = NetHandleManager->AllocateNetHandle(bIsStatic);
+
+	if (Handle.IsValid())
+	{
+		LLM_SCOPE(ELLMTag::Iris);
+
+		// Register cached ref
+		FCachedNetObjectReference CachedObject;
+
+		CachedObject.ObjectKey = nullptr;
+		CachedObject.NetHandle = Handle;
+		CachedObject.RelativePath = StringTokenStore->GetOrCreateToken(ObjectPath);
+		CachedObject.OuterNetHandle = CreateObjectReferenceHandle(Outer);
+		CachedObject.bIsPackage = 0U;
+		CachedObject.bNoLoad = 1U;
+		CachedObject.bIgnoreWhenMissing = 1U;
+		CachedObject.bIsPending = 0U;
+		CachedObject.bIsBroken = 0U;
+
+		ReferenceHandleToCachedReference.Add(Handle, CachedObject);
+	
+		UE_LOG_REFERENCECACHE(Verbose, TEXT("ObjectReferenceCache::GetOrCreateObjectReference(FromPath) %s for Object %s, %s"), *CachedObject.NetHandle.ToString(), ToCStr(ObjectPath), *CachedObject.OuterNetHandle.ToString());
+	}
+
+	OutReference = FNetObjectReference(Handle);
+
+	return OutReference;
+}
+
+void FObjectReferenceCache::ConditionalWriteNetTokenData(FNetSerializationContext& Context, FNetExportContext* ExportContext, const FNetToken& NetToken) const
+{
+	FNetBitStreamWriter* Writer = Context.GetBitStreamWriter();
+
+	if (ExportContext)
+	{
+		if (Writer->WriteBool(!ExportContext->IsExported(NetToken)))
+		{
+			NetTokenStore->WriteTokenData(Context, NetToken);
+			ExportContext->AddExported(NetToken);			
+		}
+	}
+	else
+	{
+		Writer->WriteBool(true);
+		NetTokenStore->WriteTokenData(Context, NetToken);
+	}
+}
+
+void FObjectReferenceCache::ConditionalReadNetTokenData(FNetSerializationContext& Context, const FNetToken& NetToken) const
+{
+	FNetBitStreamReader* Reader = Context.GetBitStreamReader();
+
+	const bool bIsExportToken = Reader->ReadBool();
+	if (bIsExportToken)
+	{
+		if (Reader->IsOverflown())
+		{
+			return;
+		}
+
+		FNetObjectResolveContext& ResolveContext = Context.GetInternalContext()->ResolveContext;
+	
+		NetTokenStore->ReadTokenData(Context, NetToken, *ResolveContext.RemoteNetTokenStoreState);
+	}
+}
+
+void FObjectReferenceCache::WriteFullReferenceInternal(FNetSerializationContext& Context, const FNetObjectReference& Ref) const
+{
+	FNetBitStreamWriter* Writer = Context.GetBitStreamWriter();
+	const FNetHandle RefHandle = Ref.GetRefHandle();
+
+	UE_NET_TRACE_OBJECT_SCOPE(RefHandle, *Writer, Context.GetTraceCollector(), ENetTraceVerbosity::Verbose);
+	
+	// This is unfortunate, but as we defer writing of references, we can end up with stale references
+	const FCachedNetObjectReference* CachedObject = RefHandle.IsValid() ? ReferenceHandleToCachedReference.Find(RefHandle) : nullptr;
+	// Write invalid handle
+	if (!CachedObject)
+	{
+		UE_CLOG(RefHandle.IsValid(), LogIris, Log, TEXT("ObjectReferenceCache::WriteFullReference Trying to write Stale handle, %s"), *RefHandle.ToString());
+
+		WriteNetHandle(Context, FNetHandle());
+
+		return;			
+	}
+
+	// Write handle
+	WriteNetHandle(Context, RefHandle);
+
+	const bool bHasPath = CachedObject->RelativePath.IsValid();
+
+	FNetExportContext* ExportContext = Context.GetExportContext();
+
+	// We always export if we do not have a ExportContext
+	const bool bMustExport = bHasPath && IsAuthority() && (!ExportContext || !ExportContext->IsExported(RefHandle));
+
+	UE_LOG_REFERENCECACHE(VeryVerbose, TEXT("ObjectReferenceCache::WriteFullReferenceInternal Auth: %u %s, Outer: %s Export: %u"), IsAuthority() ? 1U : 0U, *CachedObject->NetHandle.ToString(), *CachedObject->OuterNetHandle.ToString(), bMustExport ? 1U : 0U);
+
+	if (Writer->WriteBool(bMustExport))
+	{
+		// $IRIS: $TODO: Implement Async loading support https://jira.it.epicgames.com/browse/UE-123487	
+		//if (ShouldAsyncLoad() && IsAuthority() && !CachedObject->bNoLoad)
+		//{
+		//	MustBeMappedGuidsInLastBunch.AddUnique( NetGUID );
+		//}
+	
+		// Write the data, bNoLoad is implicit true for dynamic objects
+		Writer->WriteBits(CachedObject->bNoLoad, 1U);
+
+		if (Writer->WriteBool(bHasPath))
+		{
+			WriteNetToken(Writer, CachedObject->RelativePath);
+			ConditionalWriteNetTokenData(Context, ExportContext, CachedObject->RelativePath);
+			WriteFullReferenceInternal(Context, FNetObjectReference(CachedObject->OuterNetHandle));
+		}
+	
+		if (ExportContext)
+		{
+			ExportContext->AddExported(RefHandle);
+		}
+	}
+}
+
+void FObjectReferenceCache::ReadFullReferenceInternal(FNetSerializationContext& Context, FNetObjectReference& OutObjectRef, uint32 RecursionCount)
+{
+	FNetBitStreamReader* Reader = Context.GetBitStreamReader();
+
+	if (RecursionCount > INTERNAL_READ_REF_RECURSION_LIMIT) 
+	{
+		UE_LOG_REFERENCECACHE_WARNING(TEXT("ReadFullReferenceInternal: Hit recursion limit."));
+		check(false);
+		Reader->DoOverflow();
+		OutObjectRef = FNetObjectReference();
+
+		return;
+	}
+
+	UE_NET_TRACE_NAMED_OBJECT_SCOPE(ReferenceScope, FNetHandle(), *Context.GetBitStreamReader(), Context.GetTraceCollector(), ENetTraceVerbosity::Verbose);
+
+	const FNetHandle NetHandle = ReadNetHandle(Context);
+
+	if (!NetHandle.IsValid())
+	{
+		OutObjectRef = FNetObjectReference();
+		return;
+	}
+
+	const bool bIsExported = Reader->ReadBool();
+	if (Reader->IsOverflown())
+	{
+		check(false);
+		OutObjectRef = FNetObjectReference();
+		return;
+	}
+
+	UE_LOG_REFERENCECACHE(VeryVerbose, TEXT("ObjectReferenceCache::ReadFullReferenceInternal Auth: %u %s, IsExported: %u"), IsAuthority() ? 1U : 0U, *NetHandle.ToString(), bIsExported ? 1U : 0U);
+
+	// For dynamic we assume no load, otherwise we will get it through export
+	uint8 bNoLoad = NetHandle.IsDynamic();
+
+	// Do we have a path
+	FNetObjectReference OuterRef;
+	FNetToken RelativePath;
+
+	// Exported object 
+	if (bIsExported)
+	{
+		// Can we load this object?
+		bNoLoad = Reader->ReadBits(1U);
+		if (Reader->IsOverflown())
+		{
+			check(false);
+			OutObjectRef = FNetObjectReference();
+			return;
+		}
+
+		if (Reader->ReadBool())
+		{
+			RelativePath = ReadNetToken(Reader);
+			ConditionalReadNetTokenData(Context, RelativePath);
+			ReadFullReferenceInternal(Context, OuterRef, RecursionCount + 1U);
+		}
+	}
+
+	if (Reader->IsOverflown())
+	{
+		check(false);
+		OutObjectRef = FNetObjectReference();
+		return;
+	}
+
+	UE_NET_TRACE_SET_SCOPE_OBJECTID(ReferenceScope, NetHandle);
+
+	// Already known? Validate that it is correct!
+	if (const FCachedNetObjectReference* CachedObject = ReferenceHandleToCachedReference.Find(NetHandle))
+	{
+		// Dynamic objects can be referenced by either NetHandle or ParentNetHandle/PathCombo
+		check(!bIsExported || ((CachedObject->NetHandle == NetHandle && CachedObject->RelativePath == RelativePath && CachedObject->OuterNetHandle == OuterRef.GetRefHandle()) || 
+				(NetHandle.IsDynamic() && (NetHandle == CachedObject->NetHandle))));
+
+		OutObjectRef = FNetObjectReference(NetHandle);
+
+		return;
+	}
+	else
+	{
+		if (!(bIsExported || NetHandle.IsDynamic()))
+		{
+			UE_LOG_REFERENCECACHE(Error, TEXT("ObjectReferenceCache::ReadFullReferenceInternal Fail reading %s IsExported:%u"), *NetHandle.ToString(), bIsExported ? 1U : 0U);
+			Context.SetError(GNetError_InvalidNetHandle);
+			return;
+		}
+	}
+
+	// Register in cache
+	if (!IsAuthority())
+	{
+		// Register cached ref
+		FCachedNetObjectReference CachedObject;
+
+		CachedObject.ObjectKey = nullptr;
+		CachedObject.NetHandle = NetHandle;
+		CachedObject.RelativePath = RelativePath;
+		CachedObject.OuterNetHandle = OuterRef.GetRefHandle();
+		CachedObject.bIsPackage = NetHandle.IsStatic() && !OuterRef.GetRefHandle().IsValid();
+		CachedObject.bNoLoad = bNoLoad;
+		CachedObject.bIsBroken = false;
+		CachedObject.bIsPending = false;
+		CachedObject.bIgnoreWhenMissing = bNoLoad;
+
+		ReferenceHandleToCachedReference.Add(NetHandle, CachedObject);
+	}
+
+	OutObjectRef = FNetObjectReference(NetHandle);
+
+#if UE_NET_CLEANUP_REFERENCES_WITH_DYNAMIC_OUTER
+	// We want to track references with a dynamic root so that we can purge them when the dynamic root is removed
+	if (NetHandle.IsDynamic() && OuterRef.IsValid())
+	{
+		FNetHandle DynamicOuterHandle = GetDynamicRoot(NetHandle);
+		if (DynamicOuterHandle.IsValid())
+		{
+			HandleToDynamicOuter.Add(DynamicOuterHandle, NetHandle);
+		}
+	}
+#endif
+}
+
+void FObjectReferenceCache::ReadFullReference(FNetSerializationContext& Context, FNetObjectReference& OutRef)
+{
+	FNetObjectReference ObjectRef;
+	FNetBitStreamReader* Reader = Context.GetBitStreamReader();
+
+	UE_NET_TRACE_SCOPE(FullNetObjectReference, *Reader, Context.GetTraceCollector(), ENetTraceVerbosity::Trace);
+	UE_NET_TRACE_NAMED_OBJECT_SCOPE(ReferenceScope, FNetHandle(), *Reader, Context.GetTraceCollector(), ENetTraceVerbosity::Verbose);
+	LLM_SCOPE(ELLMTag::Iris);
+
+	// Handle client assigned reference
+	if (const bool bIsClientAssignedReference = Reader->ReadBool())
+	{
+		const FNetHandle NetHandle = ReadNetHandle(Context);
+		FNetToken RelativePath = ReadNetToken(Reader);
+		ConditionalReadNetTokenData(Context, RelativePath);
+
+		OutRef.RefHandle = NetHandle;
+		OutRef.PathToken = RelativePath;
+
+		return;
+	}
+
+	ReadFullReferenceInternal(Context, ObjectRef, 0U);
+
+	UE_NET_TRACE_SET_SCOPE_OBJECTID(ReferenceScope, ObjectRef.GetRefHandle());
+
+	OutRef = ObjectRef;
+}
+
+void FObjectReferenceCache::WriteFullReference(FNetSerializationContext& Context, FNetObjectReference Ref) const
+{
+	UE_NET_TRACE_SCOPE(FullNetObjectReference, *Context.GetBitStreamWriter(), Context.GetTraceCollector(), ENetTraceVerbosity::Trace);
+
+	// Handle client assigned reference 
+	FNetBitStreamWriter* Writer = Context.GetBitStreamWriter();
+	const bool bIsClientAssignedReference = Ref.PathToken.IsValid();
+	if (Writer->WriteBool(bIsClientAssignedReference))
+	{
+		WriteNetHandle(Context, Ref.GetRefHandle());
+		WriteNetToken(Writer, Ref.PathToken);
+		FNetExportContext* ExportContext = Context.GetExportContext();
+		ConditionalWriteNetTokenData(Context, ExportContext, Ref.PathToken);
+
+		return;
+	}
+
+	WriteFullReferenceInternal(Context, Ref);
+}
+
+void FObjectReferenceCache::WriteReference(FNetSerializationContext& Context, FNetObjectReference Ref) const
+{
+	UE_NET_TRACE_SCOPE(NetObjectReference, *Context.GetBitStreamWriter(), Context.GetTraceCollector(), ENetTraceVerbosity::Trace);
+	UE_NET_TRACE_OBJECT_SCOPE(Ref.GetRefHandle(), *Context.GetBitStreamWriter(), Context.GetTraceCollector(), ENetTraceVerbosity::Trace);
+
+	FNetBitStreamWriter* Writer = Context.GetBitStreamWriter();
+	const FNetHandle RefHandle = Ref.GetRefHandle();
+
+	// Only client assigned references contains a pathtoken
+	const bool bIsClientAssignedReference = Ref.PathToken.IsValid();
+	if (Writer->WriteBool(bIsClientAssignedReference))
+	{
+		WriteNetHandle(Context, Ref.GetRefHandle());
+		WriteNetToken(Writer, Ref.PathToken);
+
+		return;
+	}
+
+	// This is unfortunate, but as we defer writing of references, we can end up with stale references so we need to do this lookup
+	const FCachedNetObjectReference* CachedObject = RefHandle.IsValid() ? ReferenceHandleToCachedReference.Find(RefHandle) : nullptr;
+
+	// Write invalid handle
+	if (!CachedObject)
+	{
+		UE_CLOG(RefHandle.IsValid(), LogIris, Verbose, TEXT("ObjectReferenceCache::WriteReference Trying to write Stale handle, %s"), *RefHandle.ToString());
+
+		WriteNetHandle(Context, FNetHandle());
+
+		return;
+	}
+
+	UE_LOG_REFERENCECACHE(Verbose, TEXT("ObjectReferenceCache::WriteReference Auth: %u %s"), IsAuthority() ? 1U : 0U, *RefHandle.ToString());
+
+	checkSlow(Ref.CanBeExported() || RefHandle.IsDynamic());
+
+	// Write handle
+	WriteNetHandle(Context, RefHandle);
+}
+
+void FObjectReferenceCache::ReadReference(FNetSerializationContext& Context, FNetObjectReference& OutRef)
+{
+	UE_NET_TRACE_SCOPE(NetObjectReference, *Context.GetBitStreamReader(), Context.GetTraceCollector(), ENetTraceVerbosity::Trace);
+	UE_NET_TRACE_NAMED_OBJECT_SCOPE(ReferenceScope, FNetHandle(), *Context.GetBitStreamReader(), Context.GetTraceCollector(), ENetTraceVerbosity::Verbose);
+	LLM_SCOPE(ELLMTag::Iris);
+
+	FNetBitStreamReader* Reader = Context.GetBitStreamReader();
+
+	if (const bool bIsClientAssignedReference = Reader->ReadBool())
+	{
+		const FNetHandle NetHandle = ReadNetHandle(Context);
+		FNetToken RelativePath = ReadNetToken(Reader);
+
+		OutRef.RefHandle = NetHandle;
+		OutRef.PathToken = RelativePath;
+
+		return;
+	}
+
+	const FNetHandle NetHandle = ReadNetHandle(Context);
+
+	if (Reader->IsOverflown())
+	{
+		check(false);
+		OutRef = FNetObjectReference();
+		return;
+	}
+
+	if (!NetHandle.IsValid() )
+	{
+		OutRef = FNetObjectReference();
+		return;
+	}
+
+	UE_NET_TRACE_SET_SCOPE_OBJECTID(ReferenceScope, NetHandle);
+
+	// The handle must exist in the cache or it must be a dynamic handle that will be explicitly registered
+	if (ensureAlwaysMsgf(NetHandle.IsDynamic() || ReferenceHandleToCachedReference.Find(NetHandle), TEXT("FObjectReferenceCache::ReadReference: Handle %s must already be known or be dynamic"), *NetHandle.ToString()))
+	{
+		UE_LOG_REFERENCECACHE(Verbose, TEXT("ObjectReferenceCache::ReadReference Auth: %u %s"), IsAuthority() ? 1U : 0U, *NetHandle.ToString());
+		OutRef = FNetObjectReference(NetHandle);
+	}
+	else
+	{
+		OutRef = FNetObjectReference();
+	}
+}
+
+bool FObjectReferenceCache::WriteExports(FNetSerializationContext& Context, TArrayView<const FNetObjectReference> ExportsView) const
+{
+	FNetBitStreamWriter& Writer = *Context.GetBitStreamWriter();
+
+	if (FNetExportContext* ExportContext = Context.GetExportContext())
+	{
+		UE_NET_TRACE_SCOPE(Exports, Writer, Context.GetTraceCollector(), ENetTraceVerbosity::Verbose);
+
+		for (const FNetObjectReference& Reference : ExportsView)
+		{
+			const bool bIsClientAssigned = Reference.PathToken.IsValid();
+			const FNetHandle& Handle = Reference.GetRefHandle();
+
+			if ((bIsClientAssigned && !ExportContext->IsExported(Reference.GetPathToken())) || (IsAuthority() && Reference.CanBeExported() && !ExportContext->IsExported(Handle)))
+			{
+				Writer.WriteBool(true);
+				WriteFullReference(Context, Reference);
+			}
+		}
+	}
+
+	// Write stop bit
+	Writer.WriteBool(false);
+
+	return !Writer.IsOverflown();
+}
+
+bool FObjectReferenceCache::ReadExports(FNetSerializationContext& Context)
+{
+	FNetBitStreamReader& Reader = *Context.GetBitStreamReader();
+
+	bool bHasExportsToRead = Reader.ReadBool(); 
+	UE_NET_TRACE_SCOPE(Exports, Reader, Context.GetTraceCollector(), ENetTraceVerbosity::Verbose);
+
+	while (bHasExportsToRead && !Context.HasErrorOrOverflow())
+	{
+		FNetObjectReference Import;
+		ReadFullReference(Context, Import);
+		bHasExportsToRead  = Reader.ReadBool();
+	}
+
+	return !Context.HasErrorOrOverflow();
+}
+
+FString FObjectReferenceCache::FullPath(FNetHandle RefHandle, const FNetObjectResolveContext& ResolveContext) const
+{
+	FString FullPath;
+
+	GenerateFullPath_r(RefHandle, ResolveContext, FullPath);
+
+	return FullPath;
+}
+
+void FObjectReferenceCache::GenerateFullPath_r(FNetHandle RefHandle, const FNetObjectResolveContext& ResolveContext, FString& FullPath) const
+{
+	if ( !RefHandle.IsValid() )
+	{
+		// This is the end of the outer chain, we're done
+		return;
+	}
+
+	const FCachedNetObjectReference* CacheObject = ReferenceHandleToCachedReference.Find(RefHandle);
+
+	if (CacheObject == nullptr)
+	{
+		// Doh, this shouldn't be possible, but if this happens, we can't continue
+		// So warn, and return
+		FullPath += FString::Printf(TEXT("[%s]NOT_IN_CACHE" ), *RefHandle.ToString());
+		return;
+	}
+
+	GenerateFullPath_r(CacheObject->OuterNetHandle, ResolveContext, FullPath);
+
+	if (!FullPath.IsEmpty())
+	{
+		FullPath += TEXT(".");
+	}
+
+	// Prefer the object name first, since non stable named objects don't store the path
+	if (CacheObject->Object.IsValid())
+	{
+		// Sanity check that the names match if the path was stored
+		// $IRIS: $TODO: 
+		// Consider resolving -> FName rather than using NetTokens
+		//if (CacheObject->RelativePath.IsValid() && CacheObject->Object->GetName() != StringTokenStore->ResolveRemoteToken(CacheObject->RelativePath, ResolveContext.RemoteNetTokenStoreState))
+		//{
+		//	UE_LOG_REFERENCECACHE_WARNING(TEXT( "GenerateFullPath_r: Name mismatch! %s != %s" ), *CacheObject->PathName.ToString(), *CacheObject->Object->GetName() );
+		//}
+
+		FullPath += FString::Printf(TEXT("[%s]%s"), *RefHandle.ToString(), *CacheObject->Object->GetName());
+	}
+	else
+	{
+		if (!CacheObject->RelativePath.IsValid())
+		{
+			// This can happen when a non stably named object is nullptr
+			FullPath += FString::Printf(TEXT("[%s]EMPTY" ), *RefHandle.ToString());
+		}
+		else
+		{
+			const TCHAR* ResolvedPath = StringTokenStore->ResolveRemoteToken(CacheObject->RelativePath, *ResolveContext.RemoteNetTokenStoreState);
+			FullPath += FString::Printf(TEXT("[%s]%s"), *RefHandle.ToString(), ResolvedPath);
+		}
+	}
+}
+
+}
diff --git a/Engine/Source/Runtime/Experimental/Iris/Core/Private/Iris/ReplicationSystem/ObjectReferenceCache.h b/Engine/Source/Runtime/Experimental/Iris/Core/Private/Iris/ReplicationSystem/ObjectReferenceCache.h
new file mode 100644
--- /dev/null
+++ b/Engine/Source/Runtime/Experimental/Iris/Core/Private/Iris/ReplicationSystem/ObjectReferenceCache.h
@@ -0,0 +1,176 @@
+// Copyright Epic Games, Inc. All Rights Reserved.
+
+#pragma once
+
+#include "CoreTypes.h"
+#include "Iris/Core/NetObjectReference.h"
+#include "Iris/ReplicationSystem/NetHandle.h"
+#include "UObject/WeakObjectPtr.h"
+#include "Containers/Map.h"
+#include "ObjectReferenceCacheFwd.h"
+
+namespace UE::Net
+{
+	class FStringTokenStore;
+	class FNetTokenStore;
+	class FNetSerializationContext;
+	namespace Private
+	{
+		class FNetHandleManager;
+		class FNetExportContext;
+	}
+}
+
+namespace UE::Net::Private
+{
+
+// A lot of code in this class is extracted from the re-factored GUIDCache/ClientPackageMap and must be kept up to sync once they are submitted
+// Hopefully we can merge parts of this back together later on
+
+class FObjectReferenceCache
+{
+public:
+	FObjectReferenceCache();
+
+	void Init(UReplicationSystem* ReplicationSystem);
+
+	// Determine if the object is dynamic
+	bool IsDynamicObject(const UObject* Object) const;
+
+	// Are we allowed to create new NetHandles to reference objects?
+	bool IsAuthority() const;
+
+	// Create and assign a new NetHandle to the object
+	FNetHandle CreateObjectReferenceHandle(const UObject* Object);
+
+	// Get existing handle for object
+	FNetHandle GetObjectReferenceHandleFromObject(const UObject* Object) const;
+
+	// Get object from handle, only if the object is in the cache.
+	UObject* GetObjectFromReferenceHandle(FNetHandle RefHandle);
+
+	// Try to resolve the object reference and try to load it if the object cannot be found
+	UObject* ResolveObjectReferenceHandle(FNetHandle RefHandle, const FNetObjectResolveContext& ResolveContext);
+
+	UObject* ResolveObjectReference(const FNetObjectReference& ObjectRef, const FNetObjectResolveContext& ResolveContext);
+	ENetObjectReferenceResolveResult ResolveObjectReference(const FNetObjectReference& ObjectRef, const FNetObjectResolveContext& ResolveContext, UObject*& OutResolvedObject);
+
+	// Find replicated outer
+	FNetObjectReference GetReplicatedOuter(const FNetObjectReference& Reference) const;
+
+	// Get or create a NetObjectReference from the object
+	FNetObjectReference GetOrCreateObjectReference(const UObject* Instance);
+
+	// Get or create a NetObjectReference from the object identifed by path relative to outer
+	FNetObjectReference GetOrCreateObjectReference(const FString& ObjectPath, const UObject* Outer);
+
+	// Add reference for dynamically spawned object
+	void AddRemoteReference(FNetHandle RefHandle, const UObject* Object);
+
+	// Remove references to dynamic objects
+	void RemoveReference(FNetHandle RefHandle, const UObject* Object);
+
+	// Write full chain of object references for RefHandle
+	void WriteFullReference(FNetSerializationContext& Context, FNetObjectReference Ref) const;
+
+	// Read/load full reference data, this will populate the cache on the receiving end, but will not try to resolve the actual objects
+	void ReadFullReference(FNetSerializationContext& Context, FNetObjectReference& OutRef);
+
+	// Write reference, the reference must already be exported
+	void WriteReference(FNetSerializationContext& Context, FNetObjectReference Ref) const;
+	
+	// Read reference, as written by WriteReference
+	void ReadReference(FNetSerializationContext& Context, FNetObjectReference& OutRef);
+ 
+	// Exports are expected to be part of the written state, so if WriteExports returns false or the BitStream is overflown
+	// it us up to the caller to roll back written data and exports
+	bool WriteExports(FNetSerializationContext& Context, TArrayView<const FNetObjectReference> ExportsView) const;
+
+	bool ReadExports(FNetSerializationContext& Context);
+
+	static FNetObjectReference MakeNetObjectReference(FNetHandle Handle) { return FNetObjectReference(Handle); }
+
+private:
+
+	struct FCachedNetObjectReference
+	{
+		TWeakObjectPtr<UObject> Object;
+		const UObject* ObjectKey;
+
+		// NetHandle
+		FNetHandle NetHandle;
+
+		// RelativePath to outer
+		FNetToken RelativePath;
+
+		// Ref to outer
+		FNetHandle OuterNetHandle;
+
+		// Flags
+		uint8 bNoLoad : 1;				// Don't load this, only do a find
+		uint8 bIgnoreWhenMissing : 1;
+		uint8 bIsPackage : 1;
+		uint8 bIsBroken : 1;
+		uint8 bIsPending : 1;
+	};
+
+	bool CreateObjectReferenceInternal(const UObject* Object, FNetObjectReference& OutReference);
+
+	void ConditionalWriteNetTokenData(FNetSerializationContext& Context, Private::FNetExportContext* ExportContext, const FNetToken& NetToken) const;
+	void ConditionalReadNetTokenData(FNetSerializationContext& Context, const FNetToken& NetToken) const;
+
+	void ReadFullReferenceInternal(FNetSerializationContext& Context, FNetObjectReference& OutRef, uint32 RecursionCount);
+	void WriteFullReferenceInternal(FNetSerializationContext& Context, const FNetObjectReference& Ref) const;
+	UObject* ResolveObjectReferenceHandleInternal(FNetHandle RefHandle, const FNetObjectResolveContext& ResolveContext, bool& bOutMustBeMapped);
+	bool IsDynamicInternal(const UObject* Object) const;
+	bool SupportsObjectInternal(const UObject* Object) const;
+	bool CanClientLoadObjectInternal(const UObject* Object, bool bIsDynamic) const;
+	bool ShouldIgnoreWhenMissing(FNetHandle RefHandle) const;
+	bool RenamePathForPie(FString &Str, bool bReading);
+
+	// Get the string path of RefHandle
+	FString FullPath(FNetHandle RefHandle, const FNetObjectResolveContext& ResolveContext) const;
+	void GenerateFullPath_r(FNetHandle RefHandle, const FNetObjectResolveContext& ResolveContext, FString& OutFullPath) const;
+
+	// Find dynamic root
+	FNetHandle GetDynamicRoot(const FNetHandle Handle) const;
+
+	static FNetObjectReference MakeNetObjectReference(FNetHandle NetHandle, FNetToken RelativePath) { return FNetObjectReference(NetHandle, RelativePath, RelativePath.IsValid() ? ENetObjectReferenceTraits::CanBeExported : ENetObjectReferenceTraits::None); }
+	static FNetObjectReference MakeNetObjectReference(const FCachedNetObjectReference& CachedReference) { return FNetObjectReference(CachedReference.NetHandle, FNetToken(), CachedReference.RelativePath.IsValid() ? ENetObjectReferenceTraits::CanBeExported : ENetObjectReferenceTraits::None); }
+
+private:
+
+	// Map raw UObject pointer -> Handle
+	// To verify that the reference is valid we need to check the weakpointer stored in the cache
+	TMap<const UObject*, FNetHandle> ObjectToNetReferenceHandle;
+
+	// Map ReferenceHandle -> CachedReference
+	TMap<FNetHandle, FCachedNetObjectReference> ReferenceHandleToCachedReference;
+
+	// To properly clean up stale references referencing dynamic objects we need to track them
+	TMultiMap<FNetHandle, FNetHandle> HandleToDynamicOuter;
+	
+	UReplicationSystem* ReplicationSystem;
+	UObjectReplicationBridge* ReplicationBridge;
+	FNetTokenStore* NetTokenStore;
+	FStringTokenStore* StringTokenStore;
+	FNetHandleManager* NetHandleManager;
+	
+	// Do we have authority to create references?
+	uint32 bIsAuthority : 1;
+};
+
+inline UObject* FObjectReferenceCache::ResolveObjectReferenceHandle(FNetHandle RefHandle, const FNetObjectResolveContext& ResolveContext)
+{
+	bool bMustBeMapped;
+	return ResolveObjectReferenceHandleInternal(RefHandle, ResolveContext, bMustBeMapped);
+}
+
+inline UObject* FObjectReferenceCache::ResolveObjectReference(const FNetObjectReference& ObjectRef, const FNetObjectResolveContext& ResolveContext)
+{
+	UObject* ResolvedObject = nullptr;
+	ResolveObjectReference(ObjectRef, ResolveContext, ResolvedObject);
+	return ResolvedObject;
+}
+
+}
