Iris experimental replication system.

diff --git a/Engine/Source/Runtime/Experimental/Iris/Core/Private/Iris/ReplicationSystem/ObjectReplicationBridge.cpp b/Engine/Source/Runtime/Experimental/Iris/Core/Private/Iris/ReplicationSystem/ObjectReplicationBridge.cpp
new file mode 100644
--- /dev/null
+++ b/Engine/Source/Runtime/Experimental/Iris/Core/Private/Iris/ReplicationSystem/ObjectReplicationBridge.cpp
@@ -0,0 +1,1331 @@
+// Copyright Epic Games, Inc. All Rights Reserved.
+
+#include "Iris/ReplicationSystem/ObjectReplicationBridge.h"
+
+#include "Iris/IrisConfigInternal.h"
+
+#include "Iris/Core/IrisLog.h"
+#include "Iris/Core/IrisProfiler.h"
+
+#include "Net/Core/Trace/NetTrace.h"
+#include "Net/Core/Trace/NetDebugName.h"
+
+#include "Iris/ReplicationSystem/LegacyPushModel.h"
+#include "Iris/ReplicationSystem/ObjectPollFrequencyLimiter.h"
+#include "Iris/ReplicationSystem/ObjectReplicationBridgeConfig.h"
+#include "Iris/ReplicationSystem/Prioritization/NetObjectPrioritizer.h"
+#include "Iris/ReplicationSystem/ReplicationFragment.h"
+#include "Iris/ReplicationSystem/ReplicationFragmentInternal.h"
+#include "Iris/ReplicationState/ReplicationStateDescriptorBuilder.h"
+#include "Iris/ReplicationState/ReplicationStateUtil.h"
+#include "Iris/ReplicationSystem/ReplicationSystem.h"
+#include "Iris/ReplicationSystem/ReplicationSystemTypes.h"
+#include "Iris/ReplicationSystem/ReplicationSystemInternal.h"
+#include "Iris/ReplicationSystem/ReplicationOperations.h"
+#include "Iris/ReplicationSystem/ReplicationOperationsInternal.h"
+#include "Iris/ReplicationSystem/RepTag.h"
+#include "Iris/Serialization/NetSerializationContext.h"
+#include "Iris/Serialization/InternalNetSerializationContext.h"
+#include "Iris/Serialization/NetBitStreamUtil.h"
+
+#if UE_NET_ALLOW_MULTIPLE_REPLICATION_SYSTEMS
+#define UE_LOG_OBJECTREPLICATIONBRIDGE(Category, Format, ...)  UE_LOG(LogIrisBridge, Category, TEXT("ObjectReplicationBridge(%u)::") Format, GetReplicationSystem()->GetId(), ##__VA_ARGS__)
+#else
+#define UE_LOG_OBJECTREPLICATIONBRIDGE(Category, Format, ...)  UE_LOG(LogIrisBridge, Category, TEXT("ObjectReplicationBridge::") Format, ##__VA_ARGS__)
+#endif
+
+static bool bUseFrequencyBasedPolling = true;
+static FAutoConsoleVariableRef CVarUseFrequencyBasedPolling(
+		TEXT("net.Iris.UseFrequencyBasedPolling"),
+		bUseFrequencyBasedPolling,
+		TEXT("Whether to use frequency based polling or not. Default is true.")
+		);
+
+static bool bUseDormancyToFilterPolling = true;
+static FAutoConsoleVariableRef CVarUseDormancyToFilterPolling(
+		TEXT("net.Iris.UseDormancyToFilterPolling"),
+		bUseDormancyToFilterPolling,
+		TEXT("Whether we should use dormancy to filter out objects that we should not poll. Default is true.")
+		);
+
+static bool bAllowPollPeriodOverrides = true;
+static FAutoConsoleVariableRef CVarAllowPollPeriodOverrides(
+		TEXT("net.Iris.AllowPollPeriodOverrides"),
+		bAllowPollPeriodOverrides,
+		TEXT("Whether we allow poll period overrides set in ObjectReplicationBridgeConfig. Default is true.")
+		);
+
+static bool bEnableFilterMappings = true;
+static FAutoConsoleVariableRef CVarEnableFilterMappings(
+		TEXT("net.Iris.EnableFilterMappings"),
+		bEnableFilterMappings,
+		TEXT("Whether we honor filter mappings set in ObjectReplicationBridgeConfig. If filter mappings are enabled then objects may also be assigned the default spatial filter even if there aren't any specific mappings. Default is true.")
+		);
+
+UObjectReplicationBridge::FCreateNetHandleParams UObjectReplicationBridge::DefaultCreateNetHandleParams =
+{
+	0U, // bCanReceive
+	0U, // bNeedsPreUpdate
+	0U, // bNeedsWorldLocationUpdate
+	1U, // bAllowDynamicFilter
+	0.0f, // StaticPriority
+	0U, // PollFramePeriod
+};
+
+namespace UE::Net::Private
+{
+	void CallRegisterReplicationFragments(UObject* Object, FFragmentRegistrationContext& Context, EFragmentRegistrationFlags RegistrationFlags)
+	{
+#if UE_WITH_IRIS
+		Object->RegisterReplicationFragments(Context, RegistrationFlags);
+#endif
+	}
+
+	struct FReplicationInstanceProtocolDeleter
+	{
+		void operator()(FReplicationInstanceProtocol* InstanceProtocol)
+		{ 
+			if (InstanceProtocol != nullptr)
+			{
+				FReplicationProtocolManager::DestroyInstanceProtocol(InstanceProtocol);
+			}
+		}
+	};
+	typedef TUniquePtr<FReplicationInstanceProtocol, FReplicationInstanceProtocolDeleter> FReplicationInstanceProtocolPtr;
+}
+
+UObjectReplicationBridge::UObjectReplicationBridge()
+: Super()
+, PollFrequencyLimiter(new UE::Net::Private::FObjectPollFrequencyLimiter())
+, DefaultSpatialFilterHandle(UE::Net::InvalidNetObjectFilterHandle)
+{
+	SetShouldUseDefaultSpatialFilterFunction([](const UClass*){ return false; });
+	SetShouldSubclassUseSameFilterFunction([](const UClass*,const UClass*){ return true; });
+}
+
+UObjectReplicationBridge::~UObjectReplicationBridge()
+{
+	delete PollFrequencyLimiter;
+}
+
+void UObjectReplicationBridge::Initialize(UReplicationSystem* InReplicationSystem)
+{
+	LLM_SCOPE(ELLMTag::Iris);
+
+	Super::Initialize(InReplicationSystem);
+
+	const uint32 MaxActiveObjectCount = NetHandleManager->GetMaxActiveObjectCount();
+	PollFrequencyLimiter->Init(MaxActiveObjectCount);
+	ObjectsWithObjectReferences.Init(MaxActiveObjectCount);
+	GarbageCollectionAffectedObjects.Init(MaxActiveObjectCount);
+
+	LoadConfig();
+}
+
+void UObjectReplicationBridge::Deinitialize()
+{
+	PollFrequencyLimiter->Deinit();
+	Super::Deinitialize();
+}
+
+UObject* UObjectReplicationBridge::GetObjectFromReferenceHandle(FNetHandle NetHandle) const
+{
+	return GetObjectReferenceCache()->GetObjectFromReferenceHandle(NetHandle);
+}
+
+UObject* UObjectReplicationBridge::ResolveObjectReference(const UE::Net::FNetObjectReference& Reference, const UE::Net::FNetObjectResolveContext& ResolveContext)
+{
+	return GetObjectReferenceCache()->ResolveObjectReference(Reference, ResolveContext);
+}
+
+UE::Net::FNetObjectReference UObjectReplicationBridge::GetOrCreateObjectReference(const UObject* Instance) const
+{
+	return GetObjectReferenceCache()->GetOrCreateObjectReference(Instance);
+}
+
+UE::Net::FNetObjectReference UObjectReplicationBridge::GetOrCreateObjectReference(const FString& Path, const UObject* Outer) const
+{
+	return GetObjectReferenceCache()->GetOrCreateObjectReference(Path, Outer);
+}
+
+void UObjectReplicationBridge::AddStaticDestructionInfo(const FString& ObjectPath, const UObject* Outer, const FEndReplicationParameters& Parameters)
+{
+	UE::Net::FNetObjectReference ObjectRef = GetOrCreateObjectReference(ObjectPath, Outer);
+	if (ObjectRef.IsValid())
+	{
+		InternalAddDestructionInfo(ObjectRef.GetRefHandle(), Parameters);
+	}
+}
+
+UE::Net::FNetHandle UObjectReplicationBridge::GetReplicatedHandle(const UObject* Object) const
+{
+	FNetHandle Handle = GetObjectReferenceCache()->GetObjectReferenceHandleFromObject(Object);
+
+	return IsReplicatedHandle(Handle) ? Handle : FNetHandle();
+}
+
+UObject* UObjectReplicationBridge::GetReplicatedObject(FNetHandle Handle) const
+{
+	return IsReplicatedHandle(Handle) ? GetObjectFromReferenceHandle(Handle) : nullptr;
+};
+
+UE::Net::FNetHandle UObjectReplicationBridge::BeginReplication(UObject* Instance, const FCreateNetHandleParams& Params)
+{
+	LLM_SCOPE(ELLMTag::IrisState);
+
+	using namespace UE::Net;
+	using namespace UE::Net::Private;
+
+	FNetHandle AllocatedHandle = GetObjectReferenceCache()->CreateObjectReferenceHandle(Instance);
+
+	// If we failed to assign a handle, or if the Handle already is replicating, just return the handle
+	if (!AllocatedHandle.IsValid() || IsReplicatedHandle(AllocatedHandle))
+	{
+		return AllocatedHandle;
+	}
+
+	// Register fragments
+	EReplicationFragmentTraits Traits = EReplicationFragmentTraits::CanReplicate;
+	Traits |= Params.bCanReceive ? EReplicationFragmentTraits::CanReceive : EReplicationFragmentTraits::None;
+	Traits |= Params.bNeedsPreUpdate ? EReplicationFragmentTraits::NeedsPreSendUpdate : EReplicationFragmentTraits::None;
+	Traits |= Params.bNeedsWorldLocationUpdate ? EReplicationFragmentTraits::NeedsWorldLocationUpdate : EReplicationFragmentTraits::None;
+	
+	FFragmentRegistrationContext FragmentRegistrationContext(GetReplicationStateDescriptorRegistry(), Traits);
+
+	// For everything derived from UObject we can call the virtual function call RegisterReplicationFragments	
+	CallRegisterReplicationFragments(Instance, FragmentRegistrationContext, EFragmentRegistrationFlags::None);
+
+	const FReplicationFragments& RegisteredFragments = FFragmentRegistrationContextPrivateAccessor::GetReplicationFragments(FragmentRegistrationContext);
+
+	// We currently identify protocols by local archetype or CDO pointer and verified the protocol id received from server
+	// We also should verify the default state that we use for delta compression https://jira.it.epicgames.com/browse/UE-131344
+	const UObject* ArchetypeOrCDOUsedAsKey = Instance->GetArchetype();
+
+	// Create Protocols
+	FReplicationProtocolManager* ProtocolManager = GetReplicationProtocolManager();
+
+	FReplicationInstanceProtocolPtr InstanceProtocol(ProtocolManager->CreateInstanceProtocol(RegisteredFragments));
+		
+	FReplicationProtocolIdentifier ProtocolIdentifier = ProtocolManager->CalculateProtocolIdentifier(RegisteredFragments);
+	const FReplicationProtocol* ReplicationProtocol = ProtocolManager->GetReplicationProtocol(ProtocolIdentifier, ArchetypeOrCDOUsedAsKey);
+	if (!ReplicationProtocol)
+	{
+		ReplicationProtocol = ProtocolManager->CreateReplicationProtocol(ArchetypeOrCDOUsedAsKey, ProtocolIdentifier, RegisteredFragments, *Instance->GetClass()->GetName(), false);
+	}
+#if !(UE_BUILD_SHIPPING || UE_BUILD_TEST)
+	else
+	{
+		const bool bIsValidProtocol = ProtocolManager->ValidateReplicationProtocol(ReplicationProtocol, RegisteredFragments);
+		if (!bIsValidProtocol)
+		{
+			UE_LOG_OBJECTREPLICATIONBRIDGE(Error, TEXT("BeginReplication Found invalid protocol ProtocolId:0x%" UINT64_x_FMT " for Object named %s"), ReplicationProtocol->ProtocolIdentifier, *Instance->GetName());
+			return FNetHandle();
+		}
+	}
+#endif
+		
+	if (ReplicationProtocol)
+	{
+		// Create NetHandle and bind instance
+		// $TODO: Rename to BeginReplication 
+		FNetHandle Handle = InternalCreateNetObject(AllocatedHandle, ReplicationProtocol);
+		if (Handle.IsValid())
+		{
+			// Attach the instance and Bind the Instance protocol to dirty tracking
+			constexpr bool bBindInstanceProtocol = true;
+			InternalAttachInstanceToNetHandle(Handle, bBindInstanceProtocol, InstanceProtocol.Get(), Instance);
+
+			const FInternalNetHandle InternalReplicationIndex = NetHandleManager->GetInternalIndex(Handle);
+
+			// Keep track of handles with object references for garbage collection's sake.
+			ObjectsWithObjectReferences.SetBitValue(InternalReplicationIndex, EnumHasAnyFlags(InstanceProtocol->InstanceTraits, EReplicationInstanceProtocolTraits::HasObjectReference));
+
+			// Set poll frame period
+			uint8 PollFramePeriod = Params.PollFramePeriod;
+			FindPollInfo(Instance->GetClass(), PollFramePeriod);
+			SetPollFramePeriod(InternalReplicationIndex, PollFramePeriod);
+
+			UE_LOG_OBJECTREPLICATIONBRIDGE(Verbose, TEXT("BeginReplication Created %s with ProtocolId:0x%" UINT64_x_FMT " for Object named %s"), *Handle.ToString(), ReplicationProtocol->ProtocolIdentifier, *Instance->GetName());
+
+			{
+				FWorldLocations& WorldLocations = ReplicationSystem->GetReplicationSystemInternal()->GetWorldLocations();
+				WorldLocations.SetHasWorldLocation(InternalReplicationIndex, Params.bNeedsWorldLocationUpdate);
+			}
+
+			// Set prioritizer
+			if (Params.StaticPriority > 0.0f)
+			{
+				ReplicationSystem->SetStaticPriority(Handle, Params.StaticPriority);
+			}
+			else
+			{
+				const FNetObjectPrioritizerHandle PrioritizerHandle = GetPrioritizer(Instance->GetClass());
+				if (PrioritizerHandle != InvalidNetObjectPrioritizerHandle)
+				{
+					ReplicationSystem->SetPrioritizer(Handle, PrioritizerHandle);
+				}
+				else if (Params.bNeedsWorldLocationUpdate || HasRepTag(ReplicationProtocol, RepTag_WorldLocation))
+				{
+					ReplicationSystem->SetPrioritizer(Handle, DefaultSpatialNetObjectPrioritizerHandle);
+				}
+			}
+
+			if (bEnableFilterMappings && Params.bAllowDynamicFilter)
+			{
+				const FNetObjectFilterHandle FilterHandle = GetDynamicFilter(Instance->GetClass());
+				if (FilterHandle != InvalidNetObjectFilterHandle)
+				{
+					ReplicationSystem->SetFilter(Handle, FilterHandle);
+				}
+			}
+
+			if (ShouldClassBeDeltaCompressed(Instance->GetClass()))
+			{
+				ReplicationSystem->SetDeltaCompressionStatus(Handle, ENetObjectDeltaCompressionStatus::Allow);
+			}
+
+			// Release instance protocol from the uniquePtr as it is now successfully bound to the handle
+			InstanceProtocol.Release();
+
+			return Handle;
+		}
+		UE_LOG_OBJECTREPLICATIONBRIDGE(Warning, TEXT("BeginReplication Failed to create Handle with ProtocolId:0x%" UINT64_x_FMT " for Object named %s"), (ReplicationProtocol != nullptr ? ReplicationProtocol->ProtocolIdentifier : FReplicationProtocolIdentifier(0)), *Instance->GetName());
+
+	}
+
+	// If we get here, it means that we failed to assign an internal handle for the object, we probably have ran out of handles which currently is a fatal error.
+	UE_LOG(LogIris, Error, TEXT("UObjectReplicationBridge::CreateNetHandle - Failed to create NetHandle for object %s"), ToCStr(Instance->GetPathName()));
+
+	return FNetHandle();
+}
+
+UE::Net::FNetHandle UObjectReplicationBridge::BeginReplication(FNetHandle OwnerHandle, UObject* Instance, FNetHandle InsertRelativeToSubObjectHandle, ESubObjectInsertionOrder InsertionOrder, const FCreateNetHandleParams& Params)
+{
+	LLM_SCOPE(ELLMTag::IrisState);
+
+	using namespace UE::Net;
+	using namespace UE::Net::Private;
+
+	const FReplicationSystemInternal* ReplicationSystemInternal = GetReplicationSystem()->GetReplicationSystemInternal();
+	const FNetHandleManager& LocalNetHandleManager = ReplicationSystemInternal->GetNetHandleManager();
+
+	// Owner must be replicated
+	check(IsReplicatedHandle(OwnerHandle));
+	// verify assumptions
+	check(!Instance->HasAnyFlags(RF_ArchetypeObject | RF_ClassDefaultObject));
+
+	FNetHandle SubObjectHandle = GetReplicatedHandle(Instance);
+	if (SubObjectHandle.IsValid())
+	{
+		// Verify that the existing object is a subobject of the owner
+		check(OwnerHandle == LocalNetHandleManager.GetSubObjectOwner(SubObjectHandle));
+		return SubObjectHandle;
+	}
+	else
+	{
+		FCreateNetHandleParams SubObjectCreateParams = Params;
+		// The filtering system ignores subobjects so let's not waste cycles figuring out which filter to use.
+		SubObjectCreateParams.bAllowDynamicFilter = 0U;
+		SubObjectHandle = BeginReplication(Instance, SubObjectCreateParams);
+	}
+
+	if (SubObjectHandle.IsValid())
+	{
+		// Add subobject
+		UE_LOG_OBJECTREPLICATIONBRIDGE(Verbose, TEXT("BeginReplication Added SubObject %s to Owner %s RelativeToSubObjectHandle %s"), *SubObjectHandle.ToString(), *OwnerHandle.ToString(), *InsertRelativeToSubObjectHandle.ToString());
+
+		InternalAddSubObject(OwnerHandle, SubObjectHandle, InsertRelativeToSubObjectHandle, InsertionOrder);
+
+		// SubObjects should always poll with owner
+		SetPollWithObject(OwnerHandle, SubObjectHandle);
+
+		// Copy pending dormancy from owner
+		SetObjectWantsToBeDormant(SubObjectHandle, GetObjectWantsToBeDormant(OwnerHandle));
+	
+		return SubObjectHandle;
+	}
+
+	return FNetHandle();
+}
+
+void UObjectReplicationBridge::SetSubObjectNetCondition(FNetHandle SubObjectHandle, ELifetimeCondition Condition)
+{
+	using namespace UE::Net::Private;
+
+	// We assume that we can store the condition in an int8;
+	static_assert(ELifetimeCondition::COND_Max <= 127);
+
+	FReplicationSystemInternal* ReplicationSystemInternal = GetReplicationSystem()->GetReplicationSystemInternal();
+	FNetHandleManager& LocalNetHandleManager = ReplicationSystemInternal->GetNetHandleManager();
+
+	const FInternalNetHandle SubObjectInternalIndex = LocalNetHandleManager.GetInternalIndex(SubObjectHandle);
+	if (LocalNetHandleManager.SetSubObjectNetCondition(SubObjectInternalIndex, (int8)Condition))
+	{
+		UE_LOG_OBJECTREPLICATIONBRIDGE(Verbose, TEXT("SetSubObjectNetCondition for SubObject %s Condition %s"), *SubObjectHandle.ToString(), *UEnum::GetValueAsString<ELifetimeCondition>(Condition));
+		MarkNetObjectStateDirty(ReplicationSystem->GetId(), SubObjectInternalIndex);
+	}
+	else
+	{
+		UE_LOG_OBJECTREPLICATIONBRIDGE(Verbose, TEXT("Failed to SetSubObjectNetCondition for SubObject %s Condition %s"), *SubObjectHandle.ToString(), *UEnum::GetValueAsString<ELifetimeCondition>(Condition));
+	}
+}
+
+void UObjectReplicationBridge::AddDependentObject(FNetHandle ParentHandle, FNetHandle DependentHandle, EDependentWarnFlags WarnFlags)
+{
+	using namespace UE::Net::Private;
+
+	FReplicationSystemInternal* ReplicationSystemInternal = GetReplicationSystem()->GetReplicationSystemInternal();
+	FNetHandleManager& LocalNetHandleManager = ReplicationSystemInternal->GetNetHandleManager();
+
+	// New logic to treat Dependent objects 
+	if (LocalNetHandleManager.AddDependentObject(ParentHandle, DependentHandle))
+	{
+		FReplicationFiltering& Filtering = ReplicationSystemInternal->GetFiltering();
+		const FInternalNetHandle DependentInternalIndex = LocalNetHandleManager.GetInternalIndex(DependentHandle);
+		Filtering.NotifyAddedDependentObject(DependentInternalIndex);
+
+		UE_LOG_OBJECTREPLICATIONBRIDGE(Verbose, TEXT("AddDependentObject Added dependent object %s to parent %s"), *DependentHandle.ToString(), *ParentHandle.ToString());
+	}
+	else
+	{
+		UE_LOG_OBJECTREPLICATIONBRIDGE(Verbose, TEXT("AddDependentObject Failed to add dependent object %s to parent %s"), *DependentHandle.ToString(), *ParentHandle.ToString());
+	}
+}
+
+void UObjectReplicationBridge::RemoveDependentObject(FNetHandle ParentHandle, FNetHandle DependentHandle)
+{
+	using namespace UE::Net::Private;
+
+	UE_LOG_OBJECTREPLICATIONBRIDGE(Verbose, TEXT("RemoveDependentObject Removing dependent object %s from parent %s"), *DependentHandle.ToString(), *ParentHandle.ToString());
+
+	FReplicationSystemInternal* ReplicationSystemInternal = GetReplicationSystem()->GetReplicationSystemInternal();
+
+	// Remove dependent object
+	FNetHandleManager& LocalNetHandleManager = ReplicationSystemInternal->GetNetHandleManager();
+	LocalNetHandleManager.RemoveDependentObject(ParentHandle, DependentHandle);
+
+	const FInternalNetHandle DependentInternalIndex = LocalNetHandleManager.GetInternalIndex(DependentHandle);
+	if (DependentInternalIndex != FNetHandleManager::InvalidInternalIndex)
+	{
+		FReplicationFiltering& Filtering = ReplicationSystemInternal->GetFiltering();
+		Filtering.NotifyRemovedDependentObject(DependentInternalIndex);
+	}
+}
+
+bool UObjectReplicationBridge::WriteNetHandleCreationInfo(FReplicationBridgeSerializationContext& Context, FNetHandle Handle)
+{
+	// Write ProtocolId
+	const UE::Net::FReplicationProtocol* Protocol = GetReplicationSystem()->GetReplicationProtocol(Handle);
+	WriteUint64(Context.SerializationContext.GetBitStreamWriter(), Protocol->ProtocolIdentifier);
+
+	// Write Type header
+	return WriteCreationHeader(Context.SerializationContext, Handle);
+}
+
+void UObjectReplicationBridge::EndReplication(UObject* Instance, EEndReplicationFlags EndReplicationFlags, FEndReplicationParameters* Parameters)
+{
+	const FNetHandle Handle = GetReplicatedHandle(Instance);
+	UReplicationBridge::EndReplication(Handle, EndReplicationFlags, Parameters);
+}
+
+void UObjectReplicationBridge::DetachInstanceFromRemote(FNetHandle Handle, bool bTearOff, bool bShouldDestroyInstance)
+{
+	UE_LOG_OBJECTREPLICATIONBRIDGE(Verbose, TEXT("OnDetachInstanceFromRemote %s TearOff: %u BShouldDestroy: %u"), *Handle.ToString(), (uint32)bTearOff, (uint32)bShouldDestroyInstance);
+	UnregisterRemoteInstance(Handle, bTearOff, bShouldDestroyInstance);
+}
+
+void UObjectReplicationBridge::DetachInstance(FNetHandle Handle)
+{
+	UnregisterRemoteInstance(Handle, false, false);
+}
+
+void UObjectReplicationBridge::RegisterRemoteInstance(FNetHandle Handle, UObject* Instance, const UE::Net::FReplicationProtocol* Protocol, UE::Net::FReplicationInstanceProtocol* InstanceProtocol, const FCreationHeader* Header, uint32 ConnectionId)
+{
+	// Attach the instance protocol and instance to the Handle
+	constexpr bool bBindInstanceProtocol = false;
+	InternalAttachInstanceToNetHandle(Handle, bBindInstanceProtocol, InstanceProtocol, Instance);
+
+	// Dynamic references needs to be promoted to find the instantiated object
+	if (Handle.IsDynamic())
+	{
+		GetObjectReferenceCache()->AddRemoteReference(Handle, Instance);
+	}
+
+	UE_LOG_OBJECTREPLICATIONBRIDGE(Verbose, TEXT("RegisterRemoteInstance %s %s with ProtocolId:0x%" UINT64_x_FMT), *Handle.ToString(), ToCStr(Instance->GetName()), Protocol->ProtocolIdentifier);
+}
+
+void UObjectReplicationBridge::UnregisterRemoteInstance(FNetHandle Handle, bool bTearOff, bool bShouldDestroyInstance)
+{
+	// Lookup the instance and remove it	
+	UObject* Instance = GetObjectFromReferenceHandle(Handle);
+	
+	// Try to remove any references to dynamic objects
+	if (Handle.IsDynamic())
+	{
+		GetObjectReferenceCache()->RemoveReference(Handle, Instance);
+	}
+
+	// Destroy instance if we should	
+	if ((bTearOff || bShouldDestroyInstance) && Instance)
+	{
+		DestroyInstanceFromRemote(Instance, bTearOff);
+	}
+
+	// $TODO: Cleanup any pending creation data if we have not yet instantiated the instance
+}
+	
+UE::Net::FNetHandle UObjectReplicationBridge::CreateNetHandleFromRemote(FNetHandle SubObjectOwnerNetHandle, FNetHandle WantedNetHandle, FReplicationBridgeSerializationContext& Context)
+{
+	LLM_SCOPE(ELLMTag::IrisState);
+
+	using namespace UE::Net;
+	using namespace UE::Net::Private;
+
+	FNetBitStreamReader* Reader = Context.SerializationContext.GetBitStreamReader();
+
+	FReplicationProtocolIdentifier ReceivedProtocolId = ReadUint64(Context.SerializationContext.GetBitStreamReader());
+
+	// Read creation header
+	TUniquePtr<FCreationHeader> Header(ReadCreationHeader(Context.SerializationContext));
+	if (Context.SerializationContext.HasErrorOrOverflow())
+	{
+		return FNetHandle();
+	}
+
+	// Currently remote objects can only receive replicated data
+	FFragmentRegistrationContext FragmentRegistrationContext(GetReplicationStateDescriptorRegistry(), EReplicationFragmentTraits::CanReceive);
+	FReplicationProtocolManager* ProtocolManager = GetReplicationProtocolManager();
+
+#if (UE_BUILD_SHIPPING || UE_BUILD_TEST)
+	const bool bVerifyExistingProtocol = false;
+#else
+	const bool bVerifyExistingProtocol = true;
+#endif
+
+	// Currently we need to always instantiate remote objects, moving forward we want to make this optional so that can be deferred until it is time to apply received state data.
+	// https://jira.it.epicgames.com/browse/UE-127369	
+	UObject* InstancePtr = BeginInstantiateFromRemote(SubObjectOwnerNetHandle, Context.SerializationContext.GetInternalContext()->ResolveContext, Header.Get());
+	if (!ensureAlwaysMsgf(InstancePtr, TEXT("Failed to instantiate Handle: %s"), *WantedNetHandle.ToString()))
+	{
+		return FNetHandle();
+	}
+
+	// Register all fragments
+	CallRegisterReplicationFragments(InstancePtr, FragmentRegistrationContext, EFragmentRegistrationFlags::None);
+
+	const FReplicationFragments& RegisteredFragments = FFragmentRegistrationContextPrivateAccessor::GetReplicationFragments(FragmentRegistrationContext);
+
+	// We currently identify protocols by local archetype or CDO pointer and verified the protocol id received from server
+	// We also should verify the default state that we use for delta compression
+	const UObject* ArchetypeOrCDOUsedAsKey = InstancePtr->GetArchetype();
+
+	// Create Protocols
+	FReplicationInstanceProtocolPtr InstanceProtocol(ProtocolManager->CreateInstanceProtocol(RegisteredFragments));
+
+	// See if the protocol already is known
+	const FReplicationProtocol* ReplicationProtocol = ProtocolManager->GetReplicationProtocol(ReceivedProtocolId, ArchetypeOrCDOUsedAsKey);
+	if (!ReplicationProtocol)
+	{
+		ReplicationProtocol = ProtocolManager->CreateReplicationProtocol(ArchetypeOrCDOUsedAsKey, ReceivedProtocolId, RegisteredFragments, *(InstancePtr->GetClass()->GetName()), true);
+	}
+	else
+	{
+		if (!ensureAlways(!bVerifyExistingProtocol || ProtocolManager->ValidateReplicationProtocol(ReplicationProtocol, RegisteredFragments)))
+		{
+			return FNetHandle();
+		}
+	}
+
+	if (ReplicationProtocol)
+	{
+		// Create NetHandle
+		FNetHandle Handle = InternalCreateNetObjectFromRemote(WantedNetHandle, ReplicationProtocol);
+		if (Handle.IsValid())
+		{
+			RegisterRemoteInstance(Handle, InstancePtr, ReplicationProtocol, InstanceProtocol.Get(), Header.Get(), Context.ConnectionId);
+
+			// Release instance protocol from the uniquePtr as it is now successfully bound to the handle
+			InstanceProtocol.Release();
+
+			// Now it is safe to issue OnActorChannelOpen callback
+			ensureAlwaysMsgf(OnInstantiatedFromRemote(InstancePtr, Header.Get(), Context.ConnectionId), TEXT("Failed to invoke OnInstantiatedFromRemote for Instance named %s %s"), *InstancePtr->GetName(), *Handle.ToString());
+
+			return Handle;
+		}
+	}
+
+	return FNetHandle();
+}
+
+void UObjectReplicationBridge::PostApplyInitialState(FNetHandle Handle)
+{
+	EndInstantiateFromRemote(Handle);
+}
+
+void UObjectReplicationBridge::PreSendUpdateSingleHandle(FNetHandle Handle)
+{
+	PreUpdateAndPollImpl(Handle);
+}
+
+void UObjectReplicationBridge::PreSendUpdate()
+{
+	IRIS_PROFILER_SCOPE(UObjectReplicationBridge_OnPreSendUpdate);
+
+	// Invalid Handle means update all objects
+	FNetHandle Handle;
+	PreUpdateAndPollImpl(Handle);
+}
+
+
+void UObjectReplicationBridge::PruneStaleObjects()
+{
+	using namespace UE::Net;
+	using namespace UE::Net::Private;
+
+	IRIS_PROFILER_SCOPE(UObjectReplicationBridge_PruneStaleObjects);
+
+	// Mark all objects with object references as potentially affected by GC
+	GarbageCollectionAffectedObjects = ObjectsWithObjectReferences;
+
+	FNetHandleManager& LocalNetHandleManager = GetReplicationSystem()->GetReplicationSystemInternal()->GetNetHandleManager();
+	const TArray<UObject*>& ReplicatedInstances = LocalNetHandleManager.GetReplicatedInstances();
+
+	TArray<FNetHandle> StaleObjects;
+
+	// Detect stale references and try to kill/report them
+	auto DetectStaleObjectsFunc = [&LocalNetHandleManager, &StaleObjects, &ReplicatedInstances](uint32 InternalNetHandleIndex)
+	{
+		if (ReplicatedInstances[InternalNetHandleIndex] == nullptr)
+		{
+			const FNetHandleManager::FReplicatedObjectData& ObjectData = LocalNetHandleManager.GetReplicatedObjectDataNoCheck(InternalNetHandleIndex);
+			if (ObjectData.InstanceProtocol)
+			{
+				const FReplicationProtocol* Protocol = ObjectData.Protocol;
+				const FNetDebugName* DebugName = Protocol ? Protocol->DebugName : nullptr;
+				UE_LOG(LogIrisBridge, Warning, TEXT("UObjectReplicationBridge::PruneStaleObjects ObjectInstance replicated as: %s of Type named:%s has been destroyed without notifying the ReplicationSystem %u"), *ObjectData.Handle.ToString(), ToCStr(DebugName), ObjectData.Handle.GetReplicationSystemId());
+
+				// If the instance protocol is bound, then this is an error and we cannot safely cleanup as unbinding abound instance protocol will modify bound states
+				checkf(!EnumHasAnyFlags(ObjectData.InstanceProtocol->InstanceTraits, EReplicationInstanceProtocolTraits::IsBound), TEXT("UObjectReplicationBridge::PruneStaleObjects Bound ObjectInstance replicated as: %s has been destroyed without notifying the ReplicationSystem."), *ObjectData.Handle.ToString(), ToCStr(DebugName), ObjectData.Handle.GetReplicationSystemId());
+
+				StaleObjects.Push(ObjectData.Handle);
+			}
+		}
+	};
+
+	// Iterate over assigned indices and detect if any of the replicated instances has been garbagecollected (excluding DestroyedStartupObjectInternalIndices) as they never have an instance
+	FNetBitArray::ForAllSetBits(LocalNetHandleManager.GetAssignedInternalIndices(), LocalNetHandleManager.GetDestroyedStartupObjectInternalIndices(), FNetBitArray::AndNotOp, DetectStaleObjectsFunc);
+
+	// EndReplication/detach stale instances
+	for (FNetHandle Handle : MakeArrayView(StaleObjects.GetData(), StaleObjects.Num()))
+	{
+		UReplicationBridge::EndReplication(Handle);
+	}
+
+	// Mark poll override info as dirty
+	bHasDirtyClassesInPollPeriodOverrides = ClassHierarchyPollPeriodOverrides.Num() > 0;
+}
+
+void UObjectReplicationBridge::SetInstancePreUpdateFunction(FInstancePreUpdateFunction InPreUpdateFunction)
+{
+	PreUpdateInstanceFunction = InPreUpdateFunction;
+}
+
+void UObjectReplicationBridge::SetInstanceGetWorldLocationFunction(FInstanceGetWorldLocationFunction InGetWorldLocationFunction)
+{
+	GetInstanceWorldLocationFunction = InGetWorldLocationFunction;
+}
+
+void UObjectReplicationBridge::PreUpdateAndPollImpl(FNetHandle Handle)
+{
+	using namespace UE::Net;
+	using namespace UE::Net::Private;
+
+	FReplicationSystemInternal* ReplicationSystemInternal = GetReplicationSystem()->GetReplicationSystemInternal();
+	const FNetHandleManager& LocalNetHandleManager = ReplicationSystemInternal->GetNetHandleManager();
+	const TArray<UObject*>& ReplicatedInstances = LocalNetHandleManager.GetReplicatedInstances();
+	const bool bIsUsingPushModel = IsIrisPushModelEnabled();
+	const FNetBitArrayView DirtyObjects = ReplicationSystemInternal->GetDirtyNetObjectTracker().GetDirtyNetObjects();
+
+	struct FPreUpdateAndPollStats
+	{
+		uint32 PreUpdatedObjectCount = 0;
+		uint32 PolledObjectCount = 0;
+		uint32 PolledReferencesObjectCount = 0;
+	};
+
+	FPreUpdateAndPollStats Stats;
+
+	auto UpdateAndPollFunction = [&ReplicatedInstances, this, &LocalNetHandleManager, &Stats](uint32 InternalObjectIndex)
+	{
+		const FNetHandleManager::FReplicatedObjectData& ObjectData = LocalNetHandleManager.GetReplicatedObjectDataNoCheck(InternalObjectIndex);
+		if (ObjectData.InstanceProtocol)
+		{
+			// Call per-instance PreUpdate function
+			if (PreUpdateInstanceFunction && EnumHasAnyFlags(ObjectData.InstanceProtocol->InstanceTraits, EReplicationInstanceProtocolTraits::NeedsPreSendUpdate))
+			{
+#if !(UE_BUILD_SHIPPING || UE_BUILD_TEST)
+				IRIS_PROFILER_SCOPE(PreReplicationUpdate);
+				IRIS_PROFILER_SCOPE_TEXT(ObjectData.Protocol->DebugName->Name);
+#endif
+				(*PreUpdateInstanceFunction)(ObjectData.Handle, ReplicatedInstances[InternalObjectIndex], this);
+				++Stats.PreUpdatedObjectCount;
+			}
+
+			// Poll properties if the instance protocol requires it
+			if (EnumHasAnyFlags(ObjectData.InstanceProtocol->InstanceTraits, EReplicationInstanceProtocolTraits::NeedsPoll))
+			{
+#if !(UE_BUILD_SHIPPING || UE_BUILD_TEST)
+				IRIS_PROFILER_SCOPE(PreReplicationUpdate);
+				IRIS_PROFILER_SCOPE_TEXT(ObjectData.Protocol->DebugName->Name);
+#endif
+
+				const bool bIsGCAffectedObject = GarbageCollectionAffectedObjects.GetBit(InternalObjectIndex);
+				GarbageCollectionAffectedObjects.ClearBit(InternalObjectIndex);
+
+				// If this object has been around for a garbage collect and it has object references we must make sure that we update all cached object references
+				EReplicationFragmentPollFlags PollOptions = EReplicationFragmentPollFlags::PollAllState;
+				PollOptions |= bIsGCAffectedObject ? EReplicationFragmentPollFlags::ForceRefreshCachedObjectReferencesAfterGC : EReplicationFragmentPollFlags::None;
+
+				FReplicationInstanceOperations::PollAndRefreshCachedPropertyData(ObjectData.InstanceProtocol, PollOptions);
+				++Stats.PolledObjectCount;
+			}
+		}
+	};
+
+	const FNetBitArray& PrevScopableObjects = LocalNetHandleManager.GetPrevFrameScopableInternalIndices();
+
+	auto PushModelUpdateAndPollFunction = [&ReplicatedInstances, this, &LocalNetHandleManager, &DirtyObjects, &PrevScopableObjects, &Stats](uint32 InternalObjectIndex)
+	{
+		const FNetHandleManager::FReplicatedObjectData& ObjectData = LocalNetHandleManager.GetReplicatedObjectDataNoCheck(InternalObjectIndex);
+		if (const FReplicationInstanceProtocol* InstanceProtocol = ObjectData.InstanceProtocol)
+		{
+			const EReplicationInstanceProtocolTraits InstanceTraits = InstanceProtocol->InstanceTraits;
+			// Call per-instance PreUpdate function
+			if (PreUpdateInstanceFunction && EnumHasAnyFlags(InstanceTraits, EReplicationInstanceProtocolTraits::NeedsPreSendUpdate))
+			{
+#if !(UE_BUILD_SHIPPING || UE_BUILD_TEST)
+				IRIS_PROFILER_SCOPE(PreReplicationUpdate);
+				IRIS_PROFILER_SCOPE_TEXT(ObjectData.Protocol->DebugName->Name);
+#endif
+				(*PreUpdateInstanceFunction)(ObjectData.Handle, ReplicatedInstances[InternalObjectIndex], this);
+				++Stats.PreUpdatedObjectCount;
+			}
+
+			const bool bIsDirtyObject = DirtyObjects.GetBit(InternalObjectIndex);
+			const bool bIsNewInScope = !PrevScopableObjects.GetBit(InternalObjectIndex);
+			const bool bIsGCAffectedObject = GarbageCollectionAffectedObjects.GetBit(InternalObjectIndex);			
+			GarbageCollectionAffectedObjects.ClearBit(InternalObjectIndex);
+
+			// Early out if the instance does not require polling
+			if (!EnumHasAnyFlags(InstanceTraits, EReplicationInstanceProtocolTraits::NeedsPoll))
+			{
+				return;
+			}
+
+#if !(UE_BUILD_SHIPPING || UE_BUILD_TEST)
+			IRIS_PROFILER_SCOPE(PollPushBased);
+			IRIS_PROFILER_SCOPE_TEXT(ObjectData.Protocol->DebugName->Name);
+#endif
+
+			// If the object is fully push model we only need to poll it if it's dirty, unless it's a new object or was garbage collected.
+			if (EnumHasAnyFlags(InstanceTraits, EReplicationInstanceProtocolTraits::HasFullPushBasedDirtiness))
+			{
+				if (bIsDirtyObject | bIsNewInScope)
+				{
+					// We need to do a poll if object is marked as dirty
+					EReplicationFragmentPollFlags PollOptions = EReplicationFragmentPollFlags::PollAllState;
+					PollOptions |= bIsGCAffectedObject ? EReplicationFragmentPollFlags::ForceRefreshCachedObjectReferencesAfterGC : EReplicationFragmentPollFlags::None;
+					FReplicationInstanceOperations::PollAndRefreshCachedPropertyData(InstanceProtocol, EReplicationFragmentTraits::None, PollOptions);
+					++Stats.PolledObjectCount;
+				}
+				else if (bIsGCAffectedObject)
+				{
+					// If this object might have been affected by GC, only refresh cached references
+					const EReplicationFragmentTraits RequiredTraits = EReplicationFragmentTraits::HasPushBasedDirtiness;
+					FReplicationInstanceOperations::PollAndRefreshCachedObjectReferences(InstanceProtocol, RequiredTraits);
+					++Stats.PolledReferencesObjectCount;
+				}
+			}
+			else
+			{
+				// If the object has pushed based fragments, and not is marked dirty and object is affected by GC we need to make sure that we refresh cached references for all push based fragments
+				const bool bIsPushBasedObject = EnumHasAnyFlags(InstanceTraits, EReplicationInstanceProtocolTraits::HasPartialPushBasedDirtiness | EReplicationInstanceProtocolTraits::HasFullPushBasedDirtiness);
+				const bool bHasObjectReferences = EnumHasAnyFlags(InstanceTraits, EReplicationInstanceProtocolTraits::HasObjectReference);
+				const bool bNeedsRefreshOfCachedObjectReferences = (!(bIsNewInScope | bIsDirtyObject) & bIsGCAffectedObject & bIsPushBasedObject & bHasObjectReferences);
+				if (bNeedsRefreshOfCachedObjectReferences)
+				{
+					// Only states which has push based dirtiness need to be updated as the other states will be polled in full anyway.
+					const EReplicationFragmentTraits RequiredTraits = EReplicationFragmentTraits::HasPushBasedDirtiness;
+					FReplicationInstanceOperations::PollAndRefreshCachedObjectReferences(InstanceProtocol, RequiredTraits);
+					++Stats.PolledReferencesObjectCount;
+				}
+
+				// If this object has been around for a garbage collect and it has object references we must make sure that we update all cached object references 
+				EReplicationFragmentPollFlags PollOptions = EReplicationFragmentPollFlags::PollAllState;
+				PollOptions |= bIsGCAffectedObject ? EReplicationFragmentPollFlags::ForceRefreshCachedObjectReferencesAfterGC : EReplicationFragmentPollFlags::None;
+
+				// If the object is not new or dirty at this point we only need to poll non-push based fragments as we know that pushed based states have not been modified
+				const EReplicationFragmentTraits ExcludeTraits = (bIsDirtyObject | bIsNewInScope) ? EReplicationFragmentTraits::None : EReplicationFragmentTraits::HasPushBasedDirtiness;
+				FReplicationInstanceOperations::PollAndRefreshCachedPropertyData(InstanceProtocol, ExcludeTraits, PollOptions);
+				++Stats.PolledObjectCount;
+			}
+		}
+	};
+
+	// Update all objects
+	if (!Handle.IsValid())
+	{
+		const FNetBitArray& ScopableObjects = LocalNetHandleManager.GetScopableInternalIndices();
+		const FNetBitArray& WantToBeDormantObjects = LocalNetHandleManager.GetWantToBeDormantInternalIndices();
+
+		// Filter the set of objects considered for pre-update and polling
+		// We always want to consider objects marked as dirty and their subobjects
+		FNetBitArray ObjectsConsideredForPolling(LocalNetHandleManager.GetMaxActiveObjectCount());
+		if (bUseFrequencyBasedPolling)
+		{
+			FNetBitArrayView FrequencyBasedObjectsToPollView = MakeNetBitArrayView(ObjectsConsideredForPolling);
+			PollFrequencyLimiter->Update(MakeNetBitArrayView(ScopableObjects), DirtyObjects, FrequencyBasedObjectsToPollView);
+		}
+		else
+		{
+			ObjectsConsideredForPolling.Copy(ScopableObjects);
+		}
+
+		// Mask off objects pending dormancy as we do not want to poll/pre-update them unless they are marked for flush or are dirty
+		if (bUseDormancyToFilterPolling)
+		{
+			IRIS_PROFILER_SCOPE(UObjectReplicationBridge_PreUpdateAndPollImpl_Dormancy);
+
+			// Mask off objects pending dormancy, masked off dirty objects are restored later
+			ObjectsConsideredForPolling.Combine(WantToBeDormantObjects, FNetBitArrayView::AndNotOp);
+
+			// Add objects that have requested to flush dormancy
+			for (FNetHandle HandlePendingFlush : MakeArrayView(DormantHandlesPendingFlush))
+			{
+				if (const uint32 InternalObjectIndex = LocalNetHandleManager.GetInternalIndex(HandlePendingFlush))
+				{
+					if (ObjectsConsideredForPolling.GetBit(InternalObjectIndex))
+					{
+						continue;
+					}
+
+					ObjectsConsideredForPolling.SetBit(InternalObjectIndex);
+					for (const FInternalNetHandle SubObjectIndex : LocalNetHandleManager.GetSubObjects(InternalObjectIndex))
+					{
+						ObjectsConsideredForPolling.SetBit(SubObjectIndex);
+					}
+				}
+			}
+		}
+		DormantHandlesPendingFlush.Reset();
+
+		/**
+		 * Make sure to propagate polling for owners to subobjects and vice versa. If an actor is not due to update due to
+		 * polling frequency it can still be dirty or a dormant object marked for flush and polled for that reason. In order to make sure all recent state updates
+		 * are replicated atomically this polling propagation is required.
+		 */
+		{
+			IRIS_PROFILER_SCOPE(UObjectReplicationBridge_PreUpdateAndPollImpl_PropagateDirtyness);
+
+			auto PropagateSubObjectDirtinessToOwner = [&LocalNetHandleManager, &ObjectsConsideredForPolling](uint32 InternalObjectIndex)
+			{
+				const FNetHandleManager::FReplicatedObjectData& ObjectData = LocalNetHandleManager.GetReplicatedObjectDataNoCheck(InternalObjectIndex);
+				ObjectsConsideredForPolling.SetBit(ObjectData.SubObjectRootIndex);
+			};
+
+			auto PropagateOwnerDirtinessToSubObjectsAndDependentObjects = [&LocalNetHandleManager, &ObjectsConsideredForPolling](uint32 InternalObjectIndex)
+			{
+				const FNetHandleManager::FReplicatedObjectData& ObjectData = LocalNetHandleManager.GetReplicatedObjectDataNoCheck(InternalObjectIndex);
+
+				// Mark owner as well as we might have masked it out in the dormancy pass
+				ObjectsConsideredForPolling.SetBit(InternalObjectIndex);
+
+				for (const FInternalNetHandle SubObjectInternalIndex : LocalNetHandleManager.GetSubObjects(InternalObjectIndex))
+				{
+					ObjectsConsideredForPolling.SetBit(SubObjectInternalIndex);
+				}
+			};
+
+			// Update subobjects' owner first and owners' subobjects second. It's the only way to properly mark all groups of objects in two passes.
+			const FNetBitArrayView SubObjects = MakeNetBitArrayView(LocalNetHandleManager.GetSubObjectInternalIndices());
+			FNetBitArrayView::ForAllSetBits(DirtyObjects, SubObjects, FNetBitArray::AndOp, PropagateSubObjectDirtinessToOwner);
+			FNetBitArrayView::ForAllSetBits(DirtyObjects, SubObjects, FNetBitArray::AndNotOp, PropagateOwnerDirtinessToSubObjectsAndDependentObjects);
+			
+			// Patch in dependent objects, not subtle as this point
+			{
+				IRIS_PROFILER_SCOPE(UObjectReplicationBridge_PreUpdateAndPollImpl Patch in depedent objects);
+
+				FNetBitArray TempObjectsConsideredForPolling(ObjectsConsideredForPolling);
+				FNetBitArray::ForAllSetBits(TempObjectsConsideredForPolling, NetHandleManager->GetObjectsWithDependentObjectsInternalIndices(), FNetBitArray::AndOp,
+					[&LocalNetHandleManager, &ObjectsConsideredForPolling](FInternalNetHandle ObjectIndex) 
+					{
+						LocalNetHandleManager.ForAllDependentObjectsRecursive(ObjectIndex, [&ObjectsConsideredForPolling, ObjectIndex](FInternalNetHandle DependentObjectIndex) { ObjectsConsideredForPolling.SetBit(DependentObjectIndex); });
+					}
+				);
+			}
+		}
+
+		if (IsIrisPushModelEnabled())
+		{
+			ObjectsConsideredForPolling.ForAllSetBits(PushModelUpdateAndPollFunction);
+		}
+		else
+		{
+			ObjectsConsideredForPolling.ForAllSetBits(UpdateAndPollFunction);
+		}
+	}
+	else if (uint32 InternalObjectIndex = LocalNetHandleManager.GetInternalIndex(Handle))
+	{
+		UpdateAndPollFunction(InternalObjectIndex);
+	}
+
+	// Report stats
+	UE_NET_TRACE_FRAME_STATSCOUNTER(ReplicationSystem->GetId(), ReplicationSystem.PreUpdatedObjectCount, Stats.PreUpdatedObjectCount, ENetTraceVerbosity::Trace);
+	UE_NET_TRACE_FRAME_STATSCOUNTER(ReplicationSystem->GetId(), ReplicationSystem.PolledObjectCount, Stats.PolledObjectCount, ENetTraceVerbosity::Trace);
+	UE_NET_TRACE_FRAME_STATSCOUNTER(ReplicationSystem->GetId(), ReplicationSystem.PolledReferencesObjectCount, Stats.PolledReferencesObjectCount, ENetTraceVerbosity::Trace);
+}
+
+void UObjectReplicationBridge::UpdateInstancesWorldLocation()
+{
+	using namespace UE::Net;
+	using namespace UE::Net::Private;
+
+	if (!GetInstanceWorldLocationFunction)
+	{
+		return;
+	}
+
+	FReplicationSystemInternal* ReplicationSystemInternal = GetReplicationSystem()->GetReplicationSystemInternal();
+	const FNetHandleManager& LocalNetHandleManager = ReplicationSystemInternal->GetNetHandleManager();
+	FWorldLocations& WorldLocations = ReplicationSystemInternal->GetWorldLocations();
+	const TArray<UObject*>& ReplicatedInstances = LocalNetHandleManager.GetReplicatedInstances();
+
+
+	// Retrieve the world location for instances that supports it. Only dirty objects are considered.
+	auto UpdateInstanceWorldLocation = [this, &ReplicatedInstances, &LocalNetHandleManager, &WorldLocations](uint32 InternalObjectIndex)
+	{
+		const FNetHandleManager::FReplicatedObjectData& ObjectData = LocalNetHandleManager.GetReplicatedObjectDataNoCheck(InternalObjectIndex);
+		if (ObjectData.InstanceProtocol && EnumHasAnyFlags(ObjectData.InstanceProtocol->InstanceTraits, EReplicationInstanceProtocolTraits::NeedsWorldLocationUpdate))
+		{
+			FVector WorldLocation = (*GetInstanceWorldLocationFunction)(ObjectData.Handle, ReplicatedInstances[InternalObjectIndex]);
+			WorldLocations.SetWorldLocation(InternalObjectIndex, WorldLocation);
+		}
+	};
+
+	const FNetBitArrayView DirtyObjects = ReplicationSystemInternal->GetDirtyNetObjectTracker().GetDirtyNetObjects();
+	DirtyObjects.ForAllSetBits(UpdateInstanceWorldLocation);
+
+}
+
+void UObjectReplicationBridge::SetPollFramePeriod(UE::Net::Private::FInternalNetHandle InternalReplicationIndex, uint8 PollFramePeriod)
+{
+	PollFrequencyLimiter->SetPollFramePeriod(InternalReplicationIndex, PollFramePeriod);
+}
+
+void UObjectReplicationBridge::SetPollWithObject(FNetHandle ObjectToPollWithHandle, FNetHandle ObjectHandle)
+{
+	const UE::Net::Private::FInternalNetHandle PollWithInternalReplicationIndex = NetHandleManager->GetInternalIndex(ObjectToPollWithHandle);
+	const UE::Net::Private::FInternalNetHandle InternalReplicationIndex = NetHandleManager->GetInternalIndex(ObjectHandle);
+	PollFrequencyLimiter->SetPollWithObject(PollWithInternalReplicationIndex, InternalReplicationIndex);
+}
+
+bool UObjectReplicationBridge::GetObjectWantsToBeDormant(FNetHandle Handle) const
+{
+	using namespace UE::Net::Private;
+
+	const FReplicationSystemInternal* ReplicationSystemInternal = GetReplicationSystem()->GetReplicationSystemInternal();
+	const FNetHandleManager& LocalNetHandleManager = ReplicationSystemInternal->GetNetHandleManager();
+
+	if (const FInternalNetHandle InternalObjectIndex = LocalNetHandleManager.GetInternalIndex(Handle))
+	{
+		return LocalNetHandleManager.GetWantToBeDormantInternalIndices().GetBit(InternalObjectIndex);
+	}
+
+	return false;
+}
+
+void UObjectReplicationBridge::SetObjectWantsToBeDormant(FNetHandle Handle, bool bWantsToBeDormant)
+{
+	using namespace UE::Net::Private;
+
+	FReplicationSystemInternal* ReplicationSystemInternal = GetReplicationSystem()->GetReplicationSystemInternal();
+	FNetHandleManager& LocalNetHandleManager = ReplicationSystemInternal->GetNetHandleManager();
+
+	if (const FInternalNetHandle InternalObjectIndex = LocalNetHandleManager.GetInternalIndex(Handle))
+	{
+		UE::Net::FNetBitArray& WantToBeDormantObjects = LocalNetHandleManager.GetWantToBeDormantInternalIndices();
+
+		// Update pending dormancy status
+		WantToBeDormantObjects.SetBitValue(InternalObjectIndex, bWantsToBeDormant);
+
+		// If we want to be dormant we want to make sure we poll the object one more time
+		if (bWantsToBeDormant)
+		{
+			// Note: this will override update frequency forcing a last poll/pre-replication-update if necessary
+			// No need to do this for subobjects as they always are handled with owner
+			DormantHandlesPendingFlush.Add(Handle);
+		}
+
+		// Since we use this as a mask when updating objects we must include subobjects as well
+		// Subobjects added later will copy status from owner when they are added
+		for (const FInternalNetHandle SubObjectInternalIndex : LocalNetHandleManager.GetSubObjects(InternalObjectIndex))
+		{
+			WantToBeDormantObjects.SetBitValue(SubObjectInternalIndex, bWantsToBeDormant);
+		}
+	}
+}
+
+void UObjectReplicationBridge::ForceUpdateWantsToBeDormantObject(FNetHandle Handle)
+{
+	DormantHandlesPendingFlush.Add(Handle);
+}
+
+bool UObjectReplicationBridge::FindPollInfo(const UClass* Class, uint8& OutPollPeriod)
+{
+	if (!(bAllowPollPeriodOverrides & bHasPollOverrides))
+	{
+		return false;
+	}
+
+	const FName ClassName = Class->GetFName();
+	if (const FPollInfo* PollInfo = ClassesWithPollPeriodOverride.Find(ClassName))
+	{
+		OutPollPeriod = PollInfo->PollFramePeriod;
+		return true;
+	}
+
+	if (ClassesWithoutPollPeriodOverride.Find(ClassName))
+	{
+		return false;
+	}
+
+	/**
+	 * Only if there are poll period overrides for class hierarchies does it make sense to add
+	 * more class names to the exact match containers.
+	 */
+	if (ClassHierarchyPollPeriodOverrides.Num() == 0)
+	{
+		return false;
+	}
+
+	// We have not encountered this class before. Let's add it to the appropriate container for faster lookup next time.
+	if (bHasDirtyClassesInPollPeriodOverrides)
+	{
+		FindClassesInPollPeriodOverrides();
+	}
+
+	const UClass* SuperclassWithPollInfo = nullptr;
+	FPollInfo SuperclassPollInfo;
+	for (const auto& ClassNameAndPollInfo : ClassHierarchyPollPeriodOverrides)
+	{
+		const UClass* ClassWithPollInfo = ClassNameAndPollInfo.Value.Class.Get();
+		if (ClassWithPollInfo == nullptr)
+		{
+			continue;
+		}
+
+		if (Class->IsChildOf(ClassWithPollInfo))
+		{
+			// If we've already found a superclass with a config, see which one is closer in the hierarchy.
+			if (SuperclassWithPollInfo != nullptr)
+			{
+				if (ClassWithPollInfo->IsChildOf(SuperclassWithPollInfo))
+				{
+					SuperclassWithPollInfo = ClassWithPollInfo;
+					SuperclassPollInfo = ClassNameAndPollInfo.Value;
+				}
+			}
+			else
+			{
+				SuperclassWithPollInfo = ClassWithPollInfo;
+				SuperclassPollInfo = ClassNameAndPollInfo.Value;
+			}
+		}
+	}
+
+	if (SuperclassWithPollInfo != nullptr)
+	{
+		/*
+		 * Reset class weak pointer as it's not used for exact class matches
+		 * and we have no interest in maintaining a valid weak pointer for this case.
+		 */
+		SuperclassPollInfo.Class.Reset();
+
+		// Add the class hiearchy to our set of classes with overrides.
+		for (const UClass* ClassToAdd = Class; ClassToAdd != nullptr; ClassToAdd = ClassToAdd->GetSuperClass())
+		{
+			const FName ClassToAddName = ClassToAdd->GetFName();
+			ClassesWithPollPeriodOverride.FindOrAdd(ClassToAddName, SuperclassPollInfo);
+		}
+
+		OutPollPeriod = SuperclassPollInfo.PollFramePeriod;
+		return true;
+	}
+	else
+	{
+		// Add the class hiearchy to our set of classes without overrides.
+		for (const UClass* ClassToAdd = Class; ClassToAdd != nullptr; ClassToAdd = ClassToAdd->GetSuperClass())
+		{
+			// We avoid adding classes that are in the exact match container, even though it's not strictly necessary.
+			// It makes it easier to reason about things as a class will only be found in exactly one exact cast container.
+			const FName ClassToAddName = ClassToAdd->GetFName();
+			if (ClassesWithPollPeriodOverride.Contains(ClassToAddName))
+			{
+				continue;
+			}
+
+			ClassesWithoutPollPeriodOverride.Add(ClassToAddName);
+		}
+	}
+
+	return false;
+}
+
+FName UObjectReplicationBridge::GetConfigClassPathName(const UClass* Class)
+{
+	if (FName* CachedPathName = ConfigClassPathNameCache.Find(Class))
+	{
+		return *CachedPathName;
+	}
+
+	const FName ClassPathName(Class->GetPathName());
+
+	ConfigClassPathNameCache.Add(Class, ClassPathName);
+
+	return ClassPathName;
+}
+
+bool UObjectReplicationBridge::ShouldClassBeDeltaCompressed(const UClass* Class)
+{
+	using namespace UE::Net;
+
+	if (ClassesWithDeltaCompression.Num() > 0)
+	{
+		for (; Class != nullptr; Class = Class->GetSuperClass())
+		{
+			if (bool* bShouldBeDeltaCompressed = ClassesWithDeltaCompression.Find(GetConfigClassPathName(Class)))
+			{
+				return *bShouldBeDeltaCompressed;
+			}
+		}
+	}
+
+	return false;
+}
+
+UE::Net::FNetObjectFilterHandle UObjectReplicationBridge::GetDynamicFilter(const UClass* Class)
+{
+	using namespace UE::Net;
+
+	if (ClassesWithDynamicFilter.Num() > 0)
+	{
+		const FName ClassName = GetConfigClassPathName(Class);
+
+		// Try exact match first.
+		if (FNetObjectFilterHandle* FilterHandlePtr = ClassesWithDynamicFilter.Find(ClassName))
+		{
+			return *FilterHandlePtr;
+		}
+
+		/**
+		 * Try to find superclass. If we find it and the classes are considered equal we copy the filter setting.
+		 * If it's not equal we check whether it can be spatialized or not and use the result of that.
+		 * In all cases we add the result to the mapping for faster lookup next time.
+		 */
+		for (const UClass* SuperClass = Class->GetSuperClass(); SuperClass != nullptr; SuperClass = SuperClass->GetSuperClass())
+		{
+			const FName SuperClassName = GetConfigClassPathName(SuperClass);
+
+			// Try to get exact match first.
+			if (FNetObjectFilterHandle* FilterHandlePtr = ClassesWithDynamicFilter.Find(SuperClassName))
+			{
+				if (ShouldSubclassUseSameFilterFunction(SuperClass, Class))
+				{
+					const FNetObjectFilterHandle FilterHandle = *FilterHandlePtr;
+					ClassesWithDynamicFilter.Add(ClassName, FilterHandle);
+					return FilterHandle;
+				}
+
+				// Here's a good place to put a line of code and set a breakpoint to debug inheritance issues.
+
+				break;
+			}
+		}
+
+		// Either super class wasn't found or it wasn't considered equal. Let's add a new filter mapping.
+		const FNetObjectFilterHandle FilterHandle = ShouldUseDefaultSpatialFilterFunction(Class) ? DefaultSpatialFilterHandle : InvalidNetObjectFilterHandle;
+		ClassesWithDynamicFilter.Add(ClassName, FilterHandle);
+		return FilterHandle;
+	}
+
+	/**
+	 * For the cases when there are no configured filter mappings we just check whether to use a spatial filter or not.
+	 * We don't add anything to the filter mapping.
+	 */
+	const FNetObjectFilterHandle FilterHandle = ShouldUseDefaultSpatialFilterFunction(Class) ? DefaultSpatialFilterHandle : InvalidNetObjectFilterHandle;
+	return FilterHandle;
+}
+
+UE::Net::FNetObjectPrioritizerHandle UObjectReplicationBridge::GetPrioritizer(const UClass* Class)
+{
+	using namespace UE::Net;
+
+	if (ClassesWithPrioritizer.Num() > 0)
+	{
+		const FName ClassName = GetConfigClassPathName(Class);
+
+		// Try exact match first.
+		if (FNetObjectPrioritizerHandle* PrioritizerHandlePtr = ClassesWithPrioritizer.Find(ClassName))
+		{
+			return *PrioritizerHandlePtr;
+		}
+
+		/**
+		 * Try to find superclass. If we find it and the classes are considered equal we copy the filter setting.
+		 * If it's not equal we check whether it can be spatialized or not and use the result of that.
+		 * In all cases we add the result to the mapping for faster lookup next time.
+		 */
+		for (const UClass* SuperClass = Class->GetSuperClass(); SuperClass != nullptr; SuperClass = SuperClass->GetSuperClass())
+		{
+			const FName SuperClassName = GetConfigClassPathName(SuperClass);
+
+			// Try to get exact match first.
+			if (FNetObjectPrioritizerHandle* PrioritizerHandlePtr = ClassesWithPrioritizer.Find(SuperClassName))
+			{
+				// Assume this class should use the same prioritizer
+				{
+					const FNetObjectPrioritizerHandle PrioritizerHandle = *PrioritizerHandlePtr;
+					ClassesWithPrioritizer.Add(ClassName, PrioritizerHandle);
+					return PrioritizerHandle;
+				}
+			}
+		}
+	}
+
+	// No prioritizer has been configured for this class.
+	return InvalidNetObjectPrioritizerHandle;
+}
+
+void UObjectReplicationBridge::LoadConfig()
+{
+	// Clear everything related to the config.
+	bHasPollOverrides = false;
+	bHasDirtyClassesInPollPeriodOverrides = false;
+	ClassHierarchyPollPeriodOverrides.Empty();
+	ClassesWithPollPeriodOverride.Empty();
+	ClassesWithoutPollPeriodOverride.Empty();
+	ClassesWithDynamicFilter.Empty();
+	ClassesWithPrioritizer.Empty();
+	ClassesWithDeltaCompression.Empty();
+
+	// Reset PathNameCache
+	ConfigClassPathNameCache.Empty();
+
+	const UObjectReplicationBridgeConfig* BridgeConfig = GetDefault<UObjectReplicationBridgeConfig>();
+
+	// Load poll configs
+	
+	// These classes are forbidden to override due to being too generic and could cause memory and performance issues.
+	// If there's need for a global poll period override it should be implemented separately and not via class overrides.
+	const FName ForbiddenNames[] = {NAME_Object, NAME_Actor};
+	const TConstArrayView<FName> ForbiddenNamesArray = MakeArrayView(ForbiddenNames);
+
+	for (const FObjectReplicationBridgePollConfig& PollOverride : BridgeConfig->GetPollConfigs())
+	{
+		if (!ensure(!ForbiddenNamesArray.Contains(PollOverride.ClassName)))
+		{
+			continue;
+		}
+
+		bHasPollOverrides = true;
+
+		FPollInfo PollInfo;
+		PollInfo.PollFramePeriod = FPlatformMath::Min(PollOverride.PollFramePeriod, 255U) & 255U;
+		if (PollOverride.bIncludeSubclasses)
+		{
+			bHasDirtyClassesInPollPeriodOverrides = true;
+			ClassHierarchyPollPeriodOverrides.Add(PollOverride.ClassName, PollInfo);
+		}
+		else
+		{
+			ClassesWithPollPeriodOverride.Add(PollOverride.ClassName, PollInfo);
+		}
+	}
+
+	if (bHasDirtyClassesInPollPeriodOverrides)
+	{
+		FindClassesInPollPeriodOverrides();
+	}
+
+	// Filter mappings.
+	DefaultSpatialFilterHandle = ReplicationSystem->GetFilterHandle(BridgeConfig->GetDefaultSpatialFilterName());
+
+	{
+		for (const FObjectReplicationBridgeFilterConfig& FilterConfig : BridgeConfig->GetFilterConfigs())
+		{
+			const UE::Net::FNetObjectFilterHandle FilterHandle = ReplicationSystem->GetFilterHandle(FilterConfig.DynamicFilterName);
+			ClassesWithDynamicFilter.Add(FilterConfig.ClassName, FilterHandle);
+		}
+	}
+
+	// Prioritizer mappings
+	{
+		for (const FObjectReplicationBridgePrioritizerConfig& PrioritizerConfig : BridgeConfig->GetPrioritizerConfigs())
+		{
+			const UE::Net::FNetObjectPrioritizerHandle PrioritizerHandle = ReplicationSystem->GetPrioritizerHandle(PrioritizerConfig.PrioritizerName);
+			ClassesWithPrioritizer.Add(PrioritizerConfig.ClassName, PrioritizerHandle);
+		}
+	}
+
+	// Load delta compression settings
+	{
+		for (const FObjectReplicationBridgeDeltaCompressionConfig& DCConfig : BridgeConfig->GetDeltaCompressionConfigs())
+		{
+			if (!ensure(!ForbiddenNamesArray.Contains(DCConfig.ClassName)))
+			{
+				continue;
+			}
+
+			ClassesWithDeltaCompression.Add(DCConfig.ClassName, DCConfig.bEnableDeltaCompression);
+		}
+	}
+}
+
+void UObjectReplicationBridge::FindClassesInPollPeriodOverrides()
+{
+	bool bFailedToFindClass = false;
+	for (auto& ClassNameAndPollInfo : ClassHierarchyPollPeriodOverrides)
+	{
+		FPollInfo& PollInfo = ClassNameAndPollInfo.Value;
+		if (!PollInfo.Class.IsValid())
+		{
+			constexpr UObject* ClassOuter = nullptr;
+			constexpr bool bExactClass = true;
+			const UClass* Class = CastChecked<UClass>(StaticFindObject(UClass::StaticClass(), ClassOuter, ToCStr(ClassNameAndPollInfo.Key.ToString()), bExactClass));
+			bFailedToFindClass |= (Class == nullptr);
+			PollInfo.Class = Class;
+		}
+	}
+
+	bHasDirtyClassesInPollPeriodOverrides = bFailedToFindClass;
+}
+
+void UObjectReplicationBridge::SetShouldUseDefaultSpatialFilterFunction(TFunction<bool(const UClass*)> InShouldUseDefaultSpatialFilterFunction)
+{
+	if (!ensureAlwaysMsgf((bool)InShouldUseDefaultSpatialFilterFunction, TEXT("%s"), TEXT("A valid function must be provided for SetShouldUseDefaultSpatialFilterFunction.")))
+	{
+		return;
+	}
+
+	ShouldUseDefaultSpatialFilterFunction = InShouldUseDefaultSpatialFilterFunction;
+}
+
+void UObjectReplicationBridge::SetShouldSubclassUseSameFilterFunction(TFunction<bool(const UClass* Class, const UClass* Subclass)> InShouldSubclassUseSameFilterFunction)
+{
+	if (!ensureAlwaysMsgf((bool)InShouldSubclassUseSameFilterFunction, TEXT("%s"), TEXT("A valid function must be provided for SetShouldSubclassUseSameFilterFunction.")))
+	{
+		return;
+	}
+
+	ShouldSubclassUseSameFilterFunction = InShouldSubclassUseSameFilterFunction;
+}
diff --git a/Engine/Source/Runtime/Experimental/Iris/Core/Private/Iris/ReplicationSystem/ReplicationBridge.cpp b/Engine/Source/Runtime/Experimental/Iris/Core/Private/Iris/ReplicationSystem/ReplicationBridge.cpp
new file mode 100644
--- /dev/null
+++ b/Engine/Source/Runtime/Experimental/Iris/Core/Private/Iris/ReplicationSystem/ReplicationBridge.cpp
@@ -0,0 +1,829 @@
+// Copyright Epic Games, Inc. All Rights Reserved.
+
+#include "Iris/ReplicationSystem/ReplicationBridge.h"
+#include "CoreTypes.h"
+#include "Containers/ArrayView.h"
+#include "Iris/IrisConfigInternal.h"
+#include "Iris/Core/IrisLog.h"
+#include "Iris/Core/IrisProfiler.h"
+#include "Net/Core/NetBitArray.h"
+#include "Net/Core/Trace/NetTrace.h"
+#include "Net/Core/Trace/NetDebugName.h"
+#include "Iris/ReplicationState/ReplicationStateUtil.h"
+#include "Iris/ReplicationSystem/LegacyPushModel.h"
+#include "Iris/ReplicationSystem/Prioritization/NetObjectPrioritizer.h"
+#include "Iris/ReplicationSystem/ReplicationFragment.h"
+#include "Iris/ReplicationSystem/ReplicationFragmentInternal.h"
+#include "Iris/ReplicationSystem/ReplicationSystem.h"
+#include "Iris/ReplicationSystem/ReplicationSystemInternal.h"
+#include "Iris/ReplicationSystem/RepTag.h"
+#include "Iris/ReplicationSystem/WorldLocations.h"
+#include "ReplicationOperationsInternal.h"
+#include "Iris/Serialization/ObjectNetSerializer.h"
+#include "Iris/Serialization/NetBitStreamWriter.h"
+#include "Iris/ReplicationSystem/ChangeMaskCache.h"
+#include "Iris/ReplicationSystem/ReplicationWriter.h"
+#include "Iris/Serialization/InternalNetSerializationContext.h"
+#include "ReplicationFragmentInternal.h"
+
+DEFINE_LOG_CATEGORY(LogIrisBridge)
+
+#if UE_NET_ALLOW_MULTIPLE_REPLICATION_SYSTEMS
+#define UE_LOG_REPLICATIONBRIDGE(Category, Format, ...)  UE_LOG(LogIrisBridge, Category, TEXT("ReplicationBridge(%u)::") Format, ReplicationSystem->GetId(), ##__VA_ARGS__)
+#else
+#define UE_LOG_REPLICATIONBRIDGE(Category, Format, ...)  UE_LOG(LogIrisBridge, Category, TEXT("ReplicationBridge::") Format, ##__VA_ARGS__)
+#endif
+
+/*
+* ReplicationBridge Implementation
+*/
+UReplicationBridge::UReplicationBridge()
+: ReplicationSystem(nullptr)
+, ReplicationProtocolManager(nullptr)
+, ReplicationStateDescriptorRegistry(nullptr)
+, NetHandleManager(nullptr)
+, DestructionInfoProtocol(nullptr)
+{
+}
+
+bool UReplicationBridge::WriteNetHandleCreationInfo(FReplicationBridgeSerializationContext& Context, FNetHandle Handle)
+{
+	return true;
+}
+
+UE::Net::FNetHandle UReplicationBridge::CreateNetHandleFromRemote(FNetHandle SubObjectOwnerNetHandle, FNetHandle WantedNetHandle, FReplicationBridgeSerializationContext& Context)
+{
+	return FNetHandle();
+};
+
+void UReplicationBridge::PreSendUpdateSingleHandle(FNetHandle Handle)
+{
+}
+
+void UReplicationBridge::PreSendUpdate()
+{
+}
+
+void UReplicationBridge::UpdateInstancesWorldLocation()
+{
+}
+
+void UReplicationBridge::PostApplyInitialState(FNetHandle Handle)
+{
+}
+
+void UReplicationBridge::DetachInstanceFromRemote(FNetHandle Handle, bool bTearOff, bool bShouldDestroyInstance)
+{
+}
+
+void UReplicationBridge::DetachInstance(FNetHandle Handle)
+{
+}
+
+void UReplicationBridge::PruneStaleObjects()
+{
+}
+
+void UReplicationBridge::SetNetDriver(UNetDriver* NetDriver)
+{
+}
+
+void UReplicationBridge::GetInitialDependencies(FNetHandle Handle, FNetDependencyInfoArray& OutDependencies) const
+{
+	return;
+}
+
+void UReplicationBridge::CallGetInitialDependencies(FNetHandle Handle, FNetDependencyInfoArray& OutDependencies) const
+{
+	using namespace UE::Net::Private;
+
+	// if the Handle is static, the initial dependency is the handle itself
+	if (Handle.IsStatic())
+	{
+		OutDependencies.Emplace(FObjectReferenceCache::MakeNetObjectReference(Handle));
+		return;
+	}
+	else
+	{
+		GetInitialDependencies(Handle, OutDependencies);
+	}
+}
+
+void UReplicationBridge::DetachSubObjectInstancesFromRemote(FNetHandle OwnerHandle, bool bTearOff, bool bShouldDestroyInstance)
+{
+	using namespace UE::Net::Private;
+
+	FInternalNetHandle OwnerInternalIndex = NetHandleManager->GetInternalIndex(OwnerHandle);
+	if (OwnerInternalIndex)
+	{
+		for (const FInternalNetHandle SubObjectInternalIndex : NetHandleManager->GetChildSubObjects(OwnerInternalIndex))
+		{
+			FNetHandleManager::FReplicatedObjectData& SubObjectData = NetHandleManager->GetReplicatedObjectDataNoCheck(SubObjectInternalIndex);
+			const FNetHandle SubObjectHandle = SubObjectData.Handle;
+			SubObjectData.bTearOff = bTearOff;
+
+			CallDetachInstanceFromRemote(SubObjectHandle, bTearOff, bShouldDestroyInstance);
+		}
+	}
+}
+
+void UReplicationBridge::DestroyNetObjectFromRemote(FNetHandle Handle, bool bTearOff, bool bDestroyInstance)
+{
+	using namespace UE::Net::Private;
+
+	if (Handle.IsValid())
+	{
+		FInternalNetHandle OwnerInternalIndex = NetHandleManager->GetInternalIndex(Handle);
+		FNetHandleManager::FReplicatedObjectData& ObjectData = NetHandleManager->GetReplicatedObjectDataNoCheck(OwnerInternalIndex);
+		ObjectData.bTearOff = bTearOff;
+
+		// if the a subobject owner is to be destroyed we want to detach all subobjects before doing so to ensure we execute expected callbacks
+		// We keep tracking them internally
+		DetachSubObjectInstancesFromRemote(Handle, bTearOff, bDestroyInstance);
+		
+		// Allow derived bridges to cleanup any instance info they have stored
+		CallDetachInstanceFromRemote(Handle, bTearOff, bDestroyInstance);
+
+		// Detach instance protocol
+		InternalDetachInstanceFromNetHandle(Handle);
+	
+		// Destroy the NetHandle
+		UReplicationBridge::InternalDestroyNetObject(Handle);
+
+		UE_LOG_REPLICATIONBRIDGE(Verbose, TEXT("DestroyNetObjectFromRemote Remote %s"), *Handle.ToString());
+	}
+}
+
+void UReplicationBridge::CallPreSendUpdate(float DeltaSeconds)
+{
+	// Tear-off all handles pending tear-off
+	TearOffHandlesPendingTearOff();
+
+	PreSendUpdate();
+}
+
+void UReplicationBridge::CallPreSendUpdateSingleHandle(FNetHandle Handle)
+{
+	PreSendUpdateSingleHandle(Handle);
+}
+
+void UReplicationBridge::CallUpdateInstancesWorldLocation()
+{
+	UpdateInstancesWorldLocation();
+}
+
+void UReplicationBridge::CallDetachInstance(FNetHandle Handle)
+{
+	DetachInstance(Handle);
+}
+
+void UReplicationBridge::CallDetachInstanceFromRemote(FNetHandle Handle, bool bTearOff, bool bShouldDestroyInstance)
+{
+	DetachInstanceFromRemote(Handle, bTearOff, bShouldDestroyInstance);
+}
+
+void UReplicationBridge::CallPruneStaleObjects()
+{
+	PruneStaleObjects();
+}
+
+bool UReplicationBridge::CallWriteNetHandleCreationInfo(FReplicationBridgeSerializationContext& Context, FNetHandle Handle)
+{
+	using namespace UE::Net;
+
+	// Special path for destruction info
+	if (Context.bIsDestructionInfo)
+	{
+		UE_NET_TRACE_SCOPE(DestructionInfo, *Context.SerializationContext.GetBitStreamWriter(), Context.SerializationContext.GetTraceCollector(), ENetTraceVerbosity::Trace);
+
+		const FDestructionInfo* Info = StaticObjectsPendingDestroy.Find(Handle);
+		if (ensureAlwaysMsgf(Info, TEXT("Failed to write destructionInfo for %s"), *Handle.ToString()))
+		{
+			// Write destruction info
+			WriteFullNetObjectReference(Context.SerializationContext, Info->StaticRef);
+		}
+		else
+		{
+			// Write invalid reference
+			WriteFullNetObjectReference(Context.SerializationContext, FNetObjectReference());
+		}
+		return !Context.SerializationContext.HasErrorOrOverflow();	
+	}
+	else
+	{
+		return WriteNetHandleCreationInfo(Context, Handle);
+	}
+}
+
+UE::Net::FNetHandle UReplicationBridge::CallCreateNetHandleFromRemote(FNetHandle SubObjectOwnerHandle, FNetHandle WantedNetHandle, FReplicationBridgeSerializationContext& Context)
+{
+	using namespace UE::Net;
+
+	check(!Context.bIsDestructionInfo);
+
+	FNetHandle ResultingHandle = CreateNetHandleFromRemote(SubObjectOwnerHandle, WantedNetHandle, Context);
+
+	// Track subobjects on clients
+	if (ResultingHandle.IsValid() && SubObjectOwnerHandle.IsValid())
+	{
+		NetHandleManager->AddSubObject(SubObjectOwnerHandle, ResultingHandle);
+	}
+
+	return ResultingHandle;
+}
+
+void UReplicationBridge::ReadAndExecuteDestructionInfoFromRemote(FReplicationBridgeSerializationContext& Context)
+{
+	using namespace UE::Net;
+
+	check(Context.bIsDestructionInfo);
+
+	UE_NET_TRACE_SCOPE(DestructionInfo, *Context.SerializationContext.GetBitStreamReader(), Context.SerializationContext.GetTraceCollector(), ENetTraceVerbosity::Trace);
+
+	// Destroy instance here, or defer to later?
+	FNetObjectReference ReferenceToDestroy;
+	ReadFullNetObjectReference(Context.SerializationContext, ReferenceToDestroy);
+
+	// Destroy the reference
+	// Resolve the reference in order to be able to destroy it
+	if (ObjectReferenceCache->ResolveObjectReference(ReferenceToDestroy, Context.SerializationContext.GetInternalContext()->ResolveContext) != nullptr)
+	{
+		// Need to forward this as we do not really know how to destroy the actual instance
+		const bool bTearOff = false;
+		CallDetachInstanceFromRemote(ReferenceToDestroy.GetRefHandle(), bTearOff, true);
+	}
+}
+
+void UReplicationBridge::CallPostApplyInitialState(FNetHandle Handle)
+{
+	PostApplyInitialState(Handle);
+}
+
+UReplicationBridge::~UReplicationBridge()
+{
+}
+
+void UReplicationBridge::Initialize(UReplicationSystem* InReplicationSystem)
+{
+	using namespace UE::Net;
+
+	Private::FReplicationSystemInternal* ReplicationSystemInternal = InReplicationSystem->GetReplicationSystemInternal();
+
+	ReplicationSystem = InReplicationSystem;
+	ReplicationProtocolManager = &ReplicationSystemInternal->GetReplicationProtocolManager();
+	ReplicationStateDescriptorRegistry = &ReplicationSystemInternal->GetReplicationStateDescriptorRegistry();
+	NetHandleManager = &ReplicationSystemInternal->GetNetHandleManager();
+	ObjectReferenceCache = &ReplicationSystemInternal->GetObjectReferenceCache();
+	Groups = &ReplicationSystemInternal->GetGroups();
+
+	// Create destruction info protocol
+	{
+		const FReplicationFragments RegisteredFragments;
+		DestructionInfoProtocol = ReplicationProtocolManager->CreateReplicationProtocol(ReplicationProtocolManager->CalculateProtocolIdentifier(RegisteredFragments), RegisteredFragments, TEXT("InternalDestructionInfo"), false);
+	}
+}
+
+void UReplicationBridge::Deinitialize()
+{
+	// Just set the protocol to null, it will be destroyed with the protocolmanager
+	DestructionInfoProtocol = nullptr;
+
+	ReplicationSystem = nullptr;
+	ReplicationProtocolManager = nullptr;
+	ReplicationStateDescriptorRegistry = nullptr;
+	NetHandleManager = nullptr;
+	ObjectReferenceCache = nullptr;
+	Groups = nullptr;
+}
+
+UE::Net::FNetHandle UReplicationBridge::InternalCreateNetObject(FNetHandle AllocatedHandle, const UE::Net::FReplicationProtocol* ReplicationProtocol)
+{
+	check(AllocatedHandle.IsValid() && AllocatedHandle.IsCompleteHandle());
+
+	FNetHandle Handle = NetHandleManager->CreateNetObject(AllocatedHandle, ReplicationProtocol);
+	
+	if (Handle.IsValid())
+	{
+		UE_NET_TRACE_NETHANDLE_CREATED(Handle, ReplicationProtocol->DebugName, ReplicationProtocol->ProtocolIdentifier, 0/*Local*/);
+	}
+
+	return Handle;
+}
+
+UE::Net::FNetHandle UReplicationBridge::InternalCreateNetObjectFromRemote(FNetHandle WantedNetHandle, const UE::Net::FReplicationProtocol* ReplicationProtocol)
+{
+	FNetHandle Handle = NetHandleManager->CreateNetObjectFromRemote(WantedNetHandle, ReplicationProtocol);
+
+	if (Handle.IsValid())
+	{
+		UE_NET_TRACE_NETHANDLE_CREATED(Handle, ReplicationProtocol->DebugName, ReplicationProtocol->ProtocolIdentifier, 1/*Remote*/);
+	}
+
+	return Handle;
+}
+
+void UReplicationBridge::InternalAttachInstanceToNetHandle(FNetHandle Handle, bool bBindInstanceProtocol, UE::Net::FReplicationInstanceProtocol* InstanceProtocol, UObject* Instance)
+{
+	using namespace UE::Net;
+	using namespace UE::Net::Private;
+
+	const uint32 ReplicationSystemId = Handle.GetReplicationSystemId();
+	const uint32 InternalReplicationIndex = NetHandleManager->GetInternalIndex(Handle);
+
+	NetHandleManager->AttachInstanceProtocol(InternalReplicationIndex, InstanceProtocol, Instance);
+	UE_LOG_REPLICATIONBRIDGE(Verbose, TEXT("InternalAttachInstanceToNetHandle Attached: %s %s to ( InternalIndex: %u )"), *Instance->GetName(), *Handle.ToString(), InternalReplicationIndex);
+
+	// Bind instance protocol to dirty state tracking
+	if (bBindInstanceProtocol)
+	{
+		// Set push ID only if any state supports it. If no state supports it then we might crash if setting the ID.
+#if WITH_PUSH_MODEL
+		if (EnumHasAnyFlags(InstanceProtocol->InstanceTraits, EReplicationInstanceProtocolTraits::HasPartialPushBasedDirtiness | EReplicationInstanceProtocolTraits::HasFullPushBasedDirtiness))
+		{
+			const FNetPushObjectHandle PushHandle(InternalReplicationIndex, ReplicationSystemId);
+			TArrayView<const FReplicationFragment* const> Fragments(InstanceProtocol->Fragments, InstanceProtocol->FragmentCount);
+			SetNetPushIdOnFragments(Fragments, PushHandle);
+		}
+#endif
+
+		FReplicationInstanceOperationsInternal::BindInstanceProtocol(ReplicationSystemId, InternalReplicationIndex, InstanceProtocol, NetHandleManager->GetReplicatedObjectDataNoCheck(InternalReplicationIndex).Protocol);
+		MarkNetObjectStateDirty(Handle.GetReplicationSystemId(), InternalReplicationIndex);
+	}
+}
+
+void UReplicationBridge::InternalDetachInstanceFromNetHandle(FNetHandle Handle)
+{
+	using namespace UE::Net;
+	using namespace UE::Net::Private;
+
+	const uint32 InternalReplicationIndex = NetHandleManager->GetInternalIndex(Handle);
+
+	if (FReplicationInstanceProtocol* InstanceProtocol = const_cast<FReplicationInstanceProtocol*>(NetHandleManager->DetachInstanceProtocol(InternalReplicationIndex)))
+	{
+		if (EnumHasAnyFlags(InstanceProtocol->InstanceTraits, EReplicationInstanceProtocolTraits::IsBound))
+		{
+#if WITH_PUSH_MODEL
+			if (EnumHasAnyFlags(InstanceProtocol->InstanceTraits, EReplicationInstanceProtocolTraits::HasPartialPushBasedDirtiness | EReplicationInstanceProtocolTraits::HasFullPushBasedDirtiness))
+			{
+				TArrayView<const FReplicationFragment* const> Fragments(InstanceProtocol->Fragments, InstanceProtocol->FragmentCount);
+				ClearNetPushIdOnFragments(Fragments);
+			}
+#endif
+			FReplicationInstanceOperationsInternal::UnbindInstanceProtocol(InstanceProtocol, NetHandleManager->GetReplicatedObjectDataNoCheck(InternalReplicationIndex).Protocol);
+		}
+		ReplicationProtocolManager->DestroyInstanceProtocol(InstanceProtocol);
+	}
+}
+
+void UReplicationBridge::InternalDestroyNetObject(FNetHandle Handle)
+{
+	using namespace UE::Net;
+	using namespace UE::Net::Private;
+
+	if (const FInternalNetHandle ObjectInternalIndex = NetHandleManager->GetInternalIndex(Handle))
+	{
+		const FNetHandleManager::FReplicatedObjectData& ObjectData = NetHandleManager->GetReplicatedObjectDataNoCheck(ObjectInternalIndex);
+
+		if (ObjectData.bTearOff && NetHandleManager->GetNetObjectRefCount(ObjectInternalIndex) > 0U)
+		{
+			// We need to explicitly notify all ReplicationWriters that we are destroying objects pending tearoff
+			// The handle will automatically be removed from HandlesPendingTearOff after the next update
+
+			FReplicationConnections& Connections = ReplicationSystem->GetReplicationSystemInternal()->GetConnections();
+
+			auto NotifyDestroyedObjectPendingTearOff = [&Connections, &ObjectInternalIndex](uint32 ConnectionId)
+			{
+				FReplicationConnection* Conn = Connections.GetConnection(ConnectionId);			
+				Conn->ReplicationWriter->NotifyDestroyedObjectPendingTearOff(ObjectInternalIndex);
+			};
+
+			const FNetBitArray& ValidConnections = Connections.GetValidConnections();
+			ValidConnections.ForAllSetBits(NotifyDestroyedObjectPendingTearOff);					
+		}
+	}
+
+	NetHandleManager->DestroyNetObject(Handle);
+}
+
+void UReplicationBridge::DestroyLocalNetHandle(FNetHandle Handle)
+{
+	// Detach instance protocol
+	InternalDetachInstanceFromNetHandle(Handle);
+
+	// Allow derived bridges to cleanup any instance info they have stored
+	DetachInstance(Handle);
+
+	// If the object is in any groups we need to remove it to make sure that we update filtering
+	GetReplicationSystem()->RemoveFromAllGroups(Handle);
+
+	// If we have any attached SubObjects, tag them for destroy as well
+	InternalDestroySubObjects(Handle);
+
+	// Tell ReplicationSystem to destroy the NetHandle
+	InternalDestroyNetObject(Handle);
+
+	UE_LOG_REPLICATIONBRIDGE(Verbose, TEXT("DestroyLocalNetHandle Local %s"), *Handle.ToString());
+}
+
+UE::Net::FNetHandle UReplicationBridge::InternalGetSubObjectOwner(FNetHandle SubObjectHandle) const
+{
+	using namespace UE::Net::Private;
+
+	const FInternalNetHandle InternalHandle = NetHandleManager->GetInternalIndex(SubObjectHandle);
+	if (InternalHandle != FNetHandleManager::InvalidInternalIndex)
+	{
+		const FNetHandleManager::FReplicatedObjectData& ObjectData = NetHandleManager->GetReplicatedObjectDataNoCheck(InternalHandle);
+		const FInternalNetHandle OwnerInternalHandle = ObjectData.SubObjectRootIndex;
+
+		if (OwnerInternalHandle != FNetHandleManager::InvalidInternalIndex)
+		{
+			return NetHandleManager->GetReplicatedObjectDataNoCheck(OwnerInternalHandle).Handle;
+		}
+	}
+
+	return FNetHandle();
+}
+
+void UReplicationBridge::InternalAddSubObject(FNetHandle OwnerHandle, FNetHandle SubObjectHandle, FNetHandle InsertRelativeToSubObjectHandle, ESubObjectInsertionOrder InsertionOrder)
+{
+	using namespace UE::Net::Private;
+
+	EAddSubObjectFlags AddSubObjectFlags = EAddSubObjectFlags::Default;
+	AddSubObjectFlags |= InsertionOrder == ESubObjectInsertionOrder::None ? EAddSubObjectFlags::None : EAddSubObjectFlags::ReplicateWithSubObject;
+
+	if (NetHandleManager->AddSubObject(OwnerHandle, SubObjectHandle, InsertRelativeToSubObjectHandle, AddSubObjectFlags))
+	{
+		// If the subobject is not new we need to mark it as dirty to pick it up for replication with its new parent
+		MarkNetObjectStateDirty(ReplicationSystem->GetId(), NetHandleManager->GetInternalIndex(SubObjectHandle));
+
+		// We set the priority of subobjects to be static as they will be prioritized with owner
+		ReplicationSystem->SetStaticPriority(SubObjectHandle, 1.f);
+	}
+}
+
+void UReplicationBridge::InternalDestroySubObjects(FNetHandle OwnerHandle)
+{
+	using namespace UE::Net::Private;
+
+	// Destroy SubObjects
+	FInternalNetHandle OwnerInternalIndex = NetHandleManager->GetInternalIndex(OwnerHandle);
+	if (OwnerInternalIndex)
+	{
+		for (FInternalNetHandle SubObjectInternalIndex : NetHandleManager->GetChildSubObjects(OwnerInternalIndex))
+		{
+			const FNetHandleManager::FReplicatedObjectData& SubObjectData = NetHandleManager->GetReplicatedObjectData(SubObjectInternalIndex);
+			const FNetHandle SubObjectHandle = SubObjectData.Handle;
+			const bool bDestroySubObjectWithOwner = SubObjectData.bDestroySubObjectWithOwner;
+				
+			// Tag subobject for destroy, the check against the scope is needed since the subobjects array might contain subobjects already pending destroy
+			if (bDestroySubObjectWithOwner && NetHandleManager->IsScopableIndex(SubObjectInternalIndex))
+			{
+				UE_LOG_REPLICATIONBRIDGE(Verbose, TEXT("InternalDestroySubObjects %s - SubObject %s"), *OwnerHandle.ToString(), *SubObjectHandle.ToString());
+				DestroyLocalNetHandle(SubObjectHandle);
+			}
+		}
+	}
+}
+
+void UReplicationBridge::EndReplication(FNetHandle Handle, EEndReplicationFlags EndReplicationFlags, FEndReplicationParameters* Parameters)
+{
+	if (!IsReplicatedHandle(Handle))
+	{
+		return;
+	}
+
+	if (NetHandleManager->IsLocalNetHandle(Handle))
+	{
+		if (EnumHasAnyFlags(EndReplicationFlags, EEndReplicationFlags::TearOff))
+		{
+			// Add handle to list of objects pending tear off and indicate that it should be destroyed during next update
+			// We need to do this to cover the case where the torn off object not yet has been added to scope,
+			TearOff(Handle, true);
+
+			// We do however copy the final state data and mark object to stop propagating state changes
+			// Note: Current implementation keeps the mapping between instance and handle alive so that we can handle the case where we both start replication 
+			// and tear the object off during the same frame
+			InternalTearOff(Handle);
+		}
+		else 
+		{
+			const bool bShouldCreateDestructionInfo = Handle.IsStatic() && EnumHasAnyFlags(EndReplicationFlags, EEndReplicationFlags::Destroy);		
+			if (bShouldCreateDestructionInfo && Parameters)
+			{
+				// Add static destruction info for the Handle, should we just copy group memberships?
+				// need to copy cached reference as well so that we can "keep" it around even after pointer reuse
+				InternalAddDestructionInfo(Handle, *Parameters);
+			}
+
+			DestroyLocalNetHandle(Handle);	
+		}
+	}
+	else
+	{
+		// If we get a call to end replication on the client, we need to detach the instance as it might be garbage collected
+		InternalDetachInstanceFromNetHandle(Handle);
+	}
+}
+
+void UReplicationBridge::RemoveDestructionInfosForGroup(UE::Net::FNetObjectGroupHandle GroupHandle)
+{
+	using namespace UE::Net;
+	using namespace UE::Net::Private;
+
+	FWorldLocations& WorldLocations = GetReplicationSystem()->GetReplicationSystemInternal()->GetWorldLocations();
+
+	if (GroupHandle)
+	{
+		UE_LOG_REPLICATIONBRIDGE(Verbose, TEXT("RemoveDestructionInfosForGroup GroupIndex: %u"), GroupHandle);
+
+		FNetObjectGroup* Group = Groups->GetGroup(GroupHandle);
+		check(Group);
+
+		TArray<FNetHandle, TInlineAllocator<384>> ObjectsToRemove;	
+		TArray<FInternalNetHandle, TInlineAllocator<384>> ObjectIndicesToRemove;	
+		for (uint32 InternalObjectIndex : MakeArrayView(Group->Members))
+		{
+			if (NetHandleManager->GetIsDestroyedStartupObject(InternalObjectIndex))
+			{
+				const FNetHandleManager::FReplicatedObjectData& ObjectData = NetHandleManager->GetReplicatedObjectDataNoCheck(InternalObjectIndex);
+				if (StaticObjectsPendingDestroy.Remove(ObjectData.Handle))
+				{
+					ObjectsToRemove.Add(ObjectData.Handle);
+					ObjectIndicesToRemove.Add(InternalObjectIndex);
+				}
+			}			
+		}
+
+		for (FNetHandle Handle : MakeArrayView(ObjectsToRemove))
+		{
+			NetHandleManager->DestroyNetObject(Handle);
+		}
+
+		for (FInternalNetHandle InternalReplicationIndex : MakeArrayView(ObjectIndicesToRemove))
+		{
+			constexpr bool bHasWorldLocation = false;
+			WorldLocations.SetHasWorldLocation(InternalReplicationIndex, bHasWorldLocation);
+		}
+	}
+	else
+	{
+		// We should remove all destruction infos and objects
+		for (const auto& It : StaticObjectsPendingDestroy)
+		{
+			NetHandleManager->DestroyNetObject(It.Key);
+		}
+
+		// Remove from WorldLocations
+		for (const auto& It : StaticObjectsPendingDestroy)
+		{
+			constexpr bool bHasWorldLocation = false;
+			WorldLocations.SetHasWorldLocation(It.Value.InternalReplicationIndex, bHasWorldLocation);
+		}
+
+		StaticObjectsPendingDestroy.Empty();
+	}
+}
+
+void UReplicationBridge::TearOffHandlesPendingTearOff()
+{
+	// Initiate tear off
+	for (FTearOffInfo Info : MakeArrayView(HandlesPendingTearOff))
+	{
+		InternalTearOff(Info.Handle);
+	}
+}
+
+void UReplicationBridge::UpdateHandlesPendingTearOff()
+{
+	TArray<FTearOffInfo, TInlineAllocator<32>> ObjectsStillPendingTearOff;
+	for (FTearOffInfo Info : MakeArrayView(HandlesPendingTearOff))
+	{
+		if (uint32 ObjectInternalIndex = NetHandleManager->GetInternalIndex(Info.Handle))
+		{
+			// Immediate tear-off or object that no longer are referenced by any connections are destroyed
+			if (NetHandleManager->GetNetObjectRefCount(ObjectInternalIndex) == 0U || Info.bIsImmediate)
+			{
+				DestroyLocalNetHandle(Info.Handle);
+			}
+			else
+			{
+				// If the object is still in scope remove it from scope as objects being torn-off should not be added to new connections
+				if (NetHandleManager->IsScopableIndex(ObjectInternalIndex))
+				{
+					// Mark object as no longer scopable, and that we should not propagate changed states
+					NetHandleManager->RemoveFromScope(ObjectInternalIndex);
+				}
+			
+				// Keep object in the pending tear-off list until the object is no longer referenced by any ReplicationWriter
+				ObjectsStillPendingTearOff.Add(FTearOffInfo(Info.Handle, false));
+			}
+		}
+	}
+
+	HandlesPendingTearOff.Reset();
+	HandlesPendingTearOff.Insert(ObjectsStillPendingTearOff.GetData(), ObjectsStillPendingTearOff.Num(), 0);
+}
+
+void UReplicationBridge::TearOff(FNetHandle Handle, bool bIsImmediate)
+{
+	if (!HandlesPendingTearOff.FindByPredicate([&](const FTearOffInfo& Info){ return Info.Handle == Handle; }))
+	{
+		HandlesPendingTearOff.Add(FTearOffInfo(Handle, bIsImmediate));
+	}	
+}
+
+void UReplicationBridge::InternalTearOff(FNetHandle Handle)
+{
+	using namespace UE::Net;
+	using namespace UE::Net::Private;
+
+	IRIS_PROFILER_SCOPE(InternalTearOff);
+
+	const uint32 InternalObjectIndex = NetHandleManager->GetInternalIndex(Handle);
+
+	if (!InternalObjectIndex)
+	{
+		return;
+	}
+
+	FNetHandleManager::FReplicatedObjectData& ObjectData = NetHandleManager->GetReplicatedObjectDataNoCheck(InternalObjectIndex);
+
+	// Copy state data and tear off now
+	if (!ObjectData.bTearOff)
+	{
+		// Mark subobjects for TearOff as well.
+		for (FInternalNetHandle SubObjectInternalIndex : NetHandleManager->GetChildSubObjects(InternalObjectIndex))
+		{
+			InternalTearOff(NetHandleManager->GetNetHandleFromInternalIndex(SubObjectInternalIndex));
+		}	
+
+		UE_LOG_REPLICATIONBRIDGE(Verbose, TEXT("TearOff %s"), *Handle.ToString());
+
+		// Force copy of final state data as we will detach the object after scope update
+		FChangeMaskCache ChangeMaskCache;
+		FNetBitStreamWriter ChangeMaskWriter;
+
+		// Setup context
+		FNetSerializationContext SerializationContext;
+		FInternalNetSerializationContext InternalContext(ReplicationSystem);
+		SerializationContext.SetInternalContext(&InternalContext);
+
+		if (EnumHasAnyFlags(ObjectData.InstanceProtocol->InstanceTraits, EReplicationInstanceProtocolTraits::NeedsPoll | EReplicationInstanceProtocolTraits::NeedsPreSendUpdate))
+		{
+			CallPreSendUpdateSingleHandle(Handle);
+		} 
+
+		if (FReplicationInstanceOperationsInternal::CopyObjectStateData(ChangeMaskWriter, ChangeMaskCache, *NetHandleManager, SerializationContext, InternalObjectIndex))
+		{
+			// Propagate changes to all connections that we currently have in scope
+			FReplicationConnections& Connections = ReplicationSystem->GetReplicationSystemInternal()->GetConnections();
+
+			// Iterate over connections and propagate dirty changemasks to all connections already scoping this object
+			auto UpdateDirtyChangeMasks = [&Connections, &ChangeMaskCache](uint32 ConnectionId)
+			{
+				FReplicationConnection* Conn = Connections.GetConnection(ConnectionId);			
+				Conn->ReplicationWriter->TearOffAndUpdateDirtyChangeMasks(ChangeMaskCache);
+			};
+			const FNetBitArray& ValidConnections = Connections.GetValidConnections();
+			ValidConnections.ForAllSetBits(UpdateDirtyChangeMasks);		
+		}
+
+		// Mark object as being torn-off and that we should no longer propagate state changes
+		ObjectData.bTearOff = 1U;
+		ObjectData.bShouldPropagateChangedStates = 0U;
+	}
+}
+
+UE::Net::FNetHandle UReplicationBridge::InternalAddDestructionInfo(FNetHandle Handle, const FEndReplicationParameters& Parameters)
+{
+	using namespace UE::Net;
+	using namespace UE::Net::Private;
+
+	FReplicationProtocolManager* ProtocolManager = GetReplicationProtocolManager();
+
+	// Create DestructionInfo NetHandle used to scope destruction infos
+	const FNetObjectGroupHandle LevelGroupHandle = GetLevelGroup(Parameters.Level);
+	FNetHandle DestructionInfoHandle = NetHandleManager->CreateHandleForDestructionInfo(Handle, DestructionInfoProtocol);
+	GetReplicationSystem()->AddToGroup(LevelGroupHandle, DestructionInfoHandle);
+
+	const FInternalNetHandle InternalReplicationIndex = NetHandleManager->GetInternalIndex(DestructionInfoHandle);
+
+	// We also need to store the actual data we send to destroy static references when they are scoped
+	FDestructionInfo PersistentDestructionInfo;
+	PersistentDestructionInfo.LevelGroupHandle = LevelGroupHandle;
+	PersistentDestructionInfo.StaticRef = FObjectReferenceCache::MakeNetObjectReference(Handle);
+	PersistentDestructionInfo.InternalReplicationIndex = InternalReplicationIndex;
+
+	StaticObjectsPendingDestroy.Add(DestructionInfoHandle, PersistentDestructionInfo);
+
+	// If we use distance based Prioritization for destruction infos we need to populate the quantized state used for prioritization.
+	const bool bUseDynamicPrioritization = Parameters.bUseDistanceBasedPrioritization;
+	if (bUseDynamicPrioritization)
+	{
+		// Use WorldLocations to feed the location of the destruction info so that it can be prioritized properly.
+		FWorldLocations& WorldLocations = GetReplicationSystem()->GetReplicationSystemInternal()->GetWorldLocations();
+		WorldLocations.SetHasWorldLocation(InternalReplicationIndex, true);
+		WorldLocations.SetWorldLocation(InternalReplicationIndex, Parameters.Location);
+
+		GetReplicationSystem()->SetPrioritizer(DestructionInfoHandle, DefaultSpatialNetObjectPrioritizerHandle);
+	}
+
+	UE_LOG_REPLICATIONBRIDGE(Verbose, TEXT("UReplicationBridge::InternalAddDestructionInfo %s ( InternalIndex: %u ) for %s GroupIndex: %u"), *DestructionInfoHandle.ToString(), InternalReplicationIndex,  *Handle.ToString(), LevelGroupHandle);
+	
+	return DestructionInfoHandle;
+}
+
+bool UReplicationBridge::IsReplicatedHandle(FNetHandle Handle) const
+{
+	return Handle.IsValid() && ReplicationSystem->IsValidHandle(Handle);
+}
+
+void UReplicationBridge::SetNetPushIdOnFragments(const TArrayView<const UE::Net::FReplicationFragment*const>& Fragments, const UE::Net::Private::FNetPushObjectHandle& PushHandle)
+{
+#if WITH_PUSH_MODEL
+	using namespace UE::Net;
+
+	constexpr uint32 MaxFragmentOwnerCount = 1U;
+	UObject* FragmentOwners[MaxFragmentOwnerCount] = {};
+	FReplicationStateOwnerCollector FragmentOwnerCollector(FragmentOwners, MaxFragmentOwnerCount);
+	UObject* Instance = nullptr;
+	for (const FReplicationFragment* Fragment : Fragments)
+	{
+		if (!EnumHasAnyFlags(Fragment->GetTraits(), EReplicationFragmentTraits::HasPushBasedDirtiness))
+		{
+			continue;
+		}
+
+		FragmentOwnerCollector.Reset();
+		Fragment->CollectOwner(&FragmentOwnerCollector);
+		UObject* FragmentOwner = (FragmentOwnerCollector.GetOwnerCount() > 0 ? FragmentOwnerCollector.GetOwners()[0] : static_cast<UObject*>(nullptr));
+		if (FragmentOwner != nullptr && FragmentOwner != Instance)
+		{
+			Instance = FragmentOwner;
+			UE_NET_IRIS_SET_PUSH_ID(Instance, PushHandle);
+		}
+	}
+#endif
+}
+
+void UReplicationBridge::ClearNetPushIdOnFragments(const TArrayView<const UE::Net::FReplicationFragment*const>& Fragments)
+{
+#if WITH_PUSH_MODEL
+	using namespace UE::Net;
+
+	constexpr uint32 MaxFragmentOwnerCount = 1U;
+	UObject* FragmentOwners[MaxFragmentOwnerCount] = {};
+	FReplicationStateOwnerCollector FragmentOwnerCollector(FragmentOwners, MaxFragmentOwnerCount);
+	UObject* Instance = nullptr;
+	for (const FReplicationFragment* Fragment : Fragments)
+	{
+		if (!EnumHasAnyFlags(Fragment->GetTraits(), EReplicationFragmentTraits::HasPushBasedDirtiness))
+		{
+			continue;
+		}
+
+		FragmentOwnerCollector.Reset();
+		Fragment->CollectOwner(&FragmentOwnerCollector);
+		UObject* FragmentOwner = (FragmentOwnerCollector.GetOwnerCount() > 0 ? FragmentOwnerCollector.GetOwners()[0] : static_cast<UObject*>(nullptr));
+		if (FragmentOwner != nullptr && FragmentOwner != Instance)
+		{
+			Instance = FragmentOwner;
+			if (IsValid(FragmentOwner))
+			{
+				UE_NET_IRIS_CLEAR_PUSH_ID(FragmentOwner);
+			}
+		}
+	}
+#endif
+}
+
+void UReplicationBridge::NotifyStreamingLevelUnload(const UObject* Level)
+{
+	// Destroy group associated with level
+	UE::Net::FNetObjectGroupHandle LevelGroupHandle;
+	if (LevelGroups.RemoveAndCopyValue(FObjectKey(Level), LevelGroupHandle))
+	{
+		RemoveDestructionInfosForGroup(LevelGroupHandle);
+		ReplicationSystem->DestroyGroup(LevelGroupHandle);
+	}
+}
+
+UE::Net::FNetObjectGroupHandle UReplicationBridge::CreateLevelGroup(const UObject* Level)
+{
+	using namespace UE::Net;
+
+	FNetObjectGroupHandle LevelGroupHandle = ReplicationSystem->CreateGroup();
+	if (ensure(LevelGroupHandle != InvalidNetObjectGroupHandle))
+	{
+		ReplicationSystem->AddGroupFilter(LevelGroupHandle);
+		LevelGroups.Emplace(FObjectKey(Level), LevelGroupHandle);
+	}
+
+	return LevelGroupHandle;
+}
+
+UE::Net::FNetObjectGroupHandle UReplicationBridge::GetLevelGroup(const UObject* Level) const
+{
+	const UE::Net::FNetObjectGroupHandle* LevelGroupHandle = LevelGroups.Find(FObjectKey(Level));
+	return (LevelGroupHandle != nullptr ? *LevelGroupHandle : UE::Net::InvalidNetObjectGroupHandle);
+}
diff --git a/Engine/Source/Runtime/Experimental/Iris/Core/Private/Iris/ReplicationSystem/ReplicationConnections.cpp b/Engine/Source/Runtime/Experimental/Iris/Core/Private/Iris/ReplicationSystem/ReplicationConnections.cpp
new file mode 100644
--- /dev/null
+++ b/Engine/Source/Runtime/Experimental/Iris/Core/Private/Iris/ReplicationSystem/ReplicationConnections.cpp
@@ -0,0 +1,84 @@
+// Copyright Epic Games, Inc. All Rights Reserved.
+
+#include "ReplicationConnections.h"
+#include "Iris/DataStream/DataStreamManager.h"
+#include "Iris/ReplicationSystem/NetTokenDataStream.h"
+#include "Iris/ReplicationSystem/ReplicationDataStream.h"
+#include "Iris/ReplicationSystem/ReplicationWriter.h"
+#include "Iris/ReplicationSystem/ReplicationReader.h"
+
+namespace UE::Net::Private
+{
+
+void FReplicationConnections::Deinit()
+{
+	ValidConnections.ForAllSetBits([this](uint32 ConnectionId) { RemoveConnection(ConnectionId); });
+}
+
+void FReplicationConnections::InitDataStreams(uint32 ReplicationSystemId, uint32 ConnectionId, UDataStreamManager* DataStreamManager)
+{
+	FReplicationConnection* Connection = GetConnection(ConnectionId);
+	if (Connection == nullptr)
+	{
+		return;
+	}
+
+	// Init NetTokenDataStream if we should have one
+	UNetTokenDataStream* NetTokenDataStream = StaticCast<UNetTokenDataStream*>(DataStreamManager->GetStream(FName("NetToken")));
+	if (NetTokenDataStream != nullptr)
+	{
+		UNetTokenDataStream::FInitParameters Params;
+		Params.ConnectionId = ConnectionId;
+		Params.ReplicationSystemId = ReplicationSystemId;
+		Params.RemoteTokenStoreState = &RemoteNetTokenStoreStates[ConnectionId];
+		Params.NetExports = &DataStreamManager->GetNetExports();
+
+		NetTokenDataStream->Init(Params);
+	}
+
+	// Init ReplicationDataStream
+	Connection->ReplicationWriter->SetNetExports(DataStreamManager->GetNetExports());
+	Connection->ReplicationWriter->SetRemoteNetTokenStoreState(&RemoteNetTokenStoreStates[ConnectionId]);
+	Connection->ReplicationReader->SetRemoteNetTokenStoreState(&RemoteNetTokenStoreStates[ConnectionId]);
+
+	UReplicationDataStream* ReplicationDataStream = StaticCast<UReplicationDataStream*>(DataStreamManager->GetStream(FName("Replication")));
+	if (ReplicationDataStream != nullptr) 
+	{
+		checkf(NetTokenDataStream != nullptr, TEXT("%s"), TEXT("ReplicationDataStream requires a NetTokenDataStream"));
+		ReplicationDataStream->SetReaderAndWriter(Connection->ReplicationReader, Connection->ReplicationWriter);
+		Connection->ReplicationDataStream = ReplicationDataStream;
+	}
+}
+
+void FReplicationConnections::SetReplicationView(uint32 ConnectionId, const FReplicationView& View)
+{
+	ReplicationViews[ConnectionId] = View;
+}
+
+void FReplicationConnections::RemoveConnection(uint32 ConnectionId)
+{
+	check(ValidConnections.GetBit(ConnectionId));
+	SetReplicationView(ConnectionId, FReplicationView());
+	RemoteNetTokenStoreStates[ConnectionId] = FNetTokenStoreState();
+	DestroyReplicationReaderAndWriter(ConnectionId);
+
+	Connections[ConnectionId] = FReplicationConnection();
+	ValidConnections.ClearBit(ConnectionId);
+}
+
+void FReplicationConnections::DestroyReplicationReaderAndWriter(uint32 ConnectionId)
+{
+	FReplicationConnection* Connection = GetConnection(ConnectionId);
+	UReplicationDataStream* ReplicationDataStream = CastChecked<UReplicationDataStream>(Connection->ReplicationDataStream.Get(), ECastCheckedType::NullAllowed);
+	if (ReplicationDataStream)
+	{
+		ReplicationDataStream->SetReaderAndWriter(nullptr, nullptr);
+	}
+
+	delete Connection->ReplicationReader;
+	Connection->ReplicationReader = nullptr;
+	delete Connection->ReplicationWriter;
+	Connection->ReplicationWriter = nullptr;
+}
+
+}
diff --git a/Engine/Source/Runtime/Experimental/Iris/Core/Private/Iris/ReplicationSystem/ReplicationConnections.h b/Engine/Source/Runtime/Experimental/Iris/Core/Private/Iris/ReplicationSystem/ReplicationConnections.h
new file mode 100644
--- /dev/null
+++ b/Engine/Source/Runtime/Experimental/Iris/Core/Private/Iris/ReplicationSystem/ReplicationConnections.h
@@ -0,0 +1,99 @@
+// Copyright Epic Games, Inc. All Rights Reserved.
+
+#pragma once
+
+#include "CoreTypes.h"
+#include "Containers/Array.h"
+#include "Net/Core/NetBitArray.h"
+#include "Iris/ReplicationSystem/ReplicationView.h"
+#include "Iris/ReplicationSystem/NetTokenStoreState.h"
+#include "UObject/ObjectPtr.h"
+#include "UObject/WeakObjectPtr.h"
+
+class UDataStreamManager;
+namespace UE::Net::Private
+{
+	class FReplicationWriter;
+	class FReplicationReader;
+}
+
+namespace UE::Net::Private
+{
+
+struct FReplicationConnection
+{
+	FWeakObjectPtr ReplicationDataStream;
+	FReplicationWriter* ReplicationWriter = nullptr;
+	FReplicationReader* ReplicationReader = nullptr;
+	FObjectPtr UserData = nullptr;
+};
+
+class FReplicationConnections
+{
+public:
+	explicit FReplicationConnections(uint32 MaxConnections = 128)
+	: ValidConnections(MaxConnections)
+	{
+		Connections.SetNumZeroed(MaxConnections);
+		ReplicationViews.SetNum(MaxConnections);
+		RemoteNetTokenStoreStates.SetNum(MaxConnections);
+	}
+
+	void Deinit();
+
+	const FReplicationConnection* GetConnection(uint32 ConnectionId) const
+	{
+		if (ValidConnections.GetBit(ConnectionId))
+		{
+			return &Connections[ConnectionId];
+		}
+		
+		return nullptr;
+	}
+
+	FReplicationConnection* GetConnection(uint32 ConnectionId)
+	{
+		if (ValidConnections.GetBit(ConnectionId))
+		{
+			return &Connections[ConnectionId];
+		}
+		
+		return nullptr;
+	}
+
+	bool IsValidConnection(uint32 ConnectionId) const
+	{
+		return ConnectionId < GetMaxConnectionCount() && ValidConnections.GetBit(ConnectionId);
+	}
+
+	IRISCORE_API void AddConnection(uint32 ConnectionId)
+	{
+		check(ValidConnections.GetBit(ConnectionId) == false);
+		ValidConnections.SetBit(ConnectionId);
+	}
+
+	IRISCORE_API void RemoveConnection(uint32 ConnectionId);
+
+	uint32 GetMaxConnectionCount() const { return ValidConnections.GetNumBits(); }
+
+	const FNetBitArray& GetValidConnections() const { return ValidConnections; }
+
+	void InitDataStreams(uint32 ReplicationSystemId, uint32 ConnectionId, UDataStreamManager* DataStreamManager);
+
+	void SetReplicationView(uint32 ConnectionId, const FReplicationView& ViewInfo);
+	FReplicationView GetReplicationView(uint32 ConnectionId) const { return ReplicationViews[ConnectionId]; }
+
+	const FNetTokenStoreState& GetRemoteNetTokenStoreState(uint32 ConnectionId) const { return RemoteNetTokenStoreStates[ConnectionId]; }
+	FNetTokenStoreState& GetRemoteNetTokenStoreState(uint32 ConnectionId) { return RemoteNetTokenStoreStates[ConnectionId]; }
+
+private:
+	void DestroyReplicationReaderAndWriter(uint32 ConnectionId);
+
+private:
+	TArray<FReplicationConnection> Connections;
+	TArray<FReplicationView> ReplicationViews;
+	TArray<FNetTokenStoreState> RemoteNetTokenStoreStates;
+	FNetBitArray ValidConnections;
+};
+
+}
diff --git a/Engine/Source/Runtime/Experimental/Iris/Core/Private/Iris/ReplicationSystem/ReplicationFragmentInternal.h b/Engine/Source/Runtime/Experimental/Iris/Core/Private/Iris/ReplicationSystem/ReplicationFragmentInternal.h
new file mode 100644
--- /dev/null
+++ b/Engine/Source/Runtime/Experimental/Iris/Core/Private/Iris/ReplicationSystem/ReplicationFragmentInternal.h
@@ -0,0 +1,11 @@
+#include "Iris/ReplicationSystem/ReplicationFragment.h"
+
+#pragma once
+
+namespace UE::Net::Private
+{
+	struct FFragmentRegistrationContextPrivateAccessor
+	{
+		static const FReplicationFragments& GetReplicationFragments(const FFragmentRegistrationContext& Context) { return Context.Fragments; }
+	};
+}
