Iris experimental replication system.

diff --git a/Engine/Source/Runtime/Experimental/Iris/Core/Private/Iris/ReplicationSystem/ReplicationFragmentUtil.cpp b/Engine/Source/Runtime/Experimental/Iris/Core/Private/Iris/ReplicationSystem/ReplicationFragmentUtil.cpp
new file mode 100644
--- /dev/null
+++ b/Engine/Source/Runtime/Experimental/Iris/Core/Private/Iris/ReplicationSystem/ReplicationFragmentUtil.cpp
@@ -0,0 +1,67 @@
+// Copyright Epic Games, Inc. All Rights Reserved.
+
+#include "Iris/ReplicationSystem/ReplicationFragmentUtil.h"
+#include "Iris/ReplicationSystem/ReplicationFragment.h"
+#include "Iris/ReplicationSystem/PropertyReplicationFragment.h"
+#include "Iris/ReplicationSystem/FastArrayReplicationFragment.h"
+#include "Iris/ReplicationState/ReplicationStateDescriptorBuilder.h"
+#include "Iris/Core/IrisProfiler.h"
+#include "Iris/Core/IrisLog.h"
+
+namespace UE::Net
+{
+
+uint32 FReplicationFragmentUtil::CreateAndRegisterFragmentsForObject(UObject* Object, FFragmentRegistrationContext& Context, EFragmentRegistrationFlags RegistrationFlags, TArray<FReplicationFragment*>* OutCreatedFragments)
+{
+	IRIS_PROFILER_SCOPE(FReplicationFragmentUtil_CreateAndRegisterFragmentsForObject);
+
+	FReplicationStateDescriptorBuilder::FParameters BuilderParameters;
+	BuilderParameters.DescriptorRegistry = Context.GetReplicationStateRegistry();
+
+	// By default we use the archetype for our default values, but if InitializeDefaultStateFromClassDefaults is set we 
+	// use the class default instead
+	if (!EnumHasAnyFlags(RegistrationFlags, EFragmentRegistrationFlags::InitializeDefaultStateFromClassDefaults))
+	{
+		BuilderParameters.DefaultStateSource = Object->GetArchetype();
+
+		if (BuilderParameters.DefaultStateSource == nullptr)
+		{
+			UE_LOG(LogIris, Error, TEXT("FPropertyReplicationFragment::CreateAndRegisterFragmentsForObject: Invalid object archetype for object %s, default state will be used from CDO:"), *GetFullNameSafe(Object));
+		}
+	}
+
+	FReplicationStateDescriptorBuilder::FResult Result;
+	FReplicationStateDescriptorBuilder::CreateDescriptorsForClass(Result, Object->GetClass(), BuilderParameters);
+
+	const bool bRegisterFunctionsOnly = EnumHasAnyFlags(RegistrationFlags, EFragmentRegistrationFlags::RegisterRPCsOnly);
+
+	uint32 NumCreatedReplicationFragments = 0U;
+	// create fragments and let the instance protocol own them.
+	for (auto& Desc : Result)
+	{
+		if (bRegisterFunctionsOnly && (Desc->FunctionCount == 0))
+		{
+			continue;
+		}
+
+		FReplicationFragment* Fragment = nullptr;
+		// If descriptor provides CreateAndRegisterReplicationFragment function we use that to instantiate fragment
+		if (Desc->CreateAndRegisterReplicationFragmentFunction)
+		{
+			Fragment = Desc->CreateAndRegisterReplicationFragmentFunction(Object, Desc.GetReference(), Context);
+		}
+		else
+		{
+			Fragment = FPropertyReplicationFragment::CreateAndRegisterFragment(Object, Desc.GetReference(), Context);
+		}
+		if (OutCreatedFragments)
+		{
+			OutCreatedFragments->Add(Fragment);
+			++NumCreatedReplicationFragments;
+		}
+	}
+
+	return NumCreatedReplicationFragments;
+}
+
+}
diff --git a/Engine/Source/Runtime/Experimental/Iris/Core/Private/Iris/ReplicationSystem/ReplicationProtocolManager.h b/Engine/Source/Runtime/Experimental/Iris/Core/Private/Iris/ReplicationSystem/ReplicationProtocolManager.h
new file mode 100644
--- /dev/null
+++ b/Engine/Source/Runtime/Experimental/Iris/Core/Private/Iris/ReplicationSystem/ReplicationProtocolManager.h
@@ -0,0 +1,86 @@
+// Copyright Epic Games, Inc. All Rights Reserved.
+
+#pragma once
+
+#include "CoreTypes.h"
+#include "Containers/Map.h"
+#include "Iris/ReplicationState/ReplicationStateDescriptor.h"
+#include "Iris/ReplicationSystem/ReplicationFragment.h"
+#include "Iris/ReplicationSystem/ReplicationProtocol.h"
+
+namespace UE::Net::Private
+{
+
+class FReplicationProtocolManager
+{
+public:
+	~FReplicationProtocolManager();
+
+	/* Calculate protocol Identifier from registered fragment data */
+	IRISCORE_API FReplicationProtocolIdentifier CalculateProtocolIdentifier(const FReplicationFragments& Fragments) const;
+
+	/* Create protocol from registered fragment data with provided Id, verification is optional */
+	IRISCORE_API const FReplicationProtocol* CreateReplicationProtocol(const UObject* ArchetypeOrCDOUsedAsKey, const FReplicationProtocolIdentifier ProtocolId, const FReplicationFragments& Fragments, const TCHAR* DebugName, bool bVerifyId = false);
+
+	/* Create protocol from registered fragment data with provided Id, verification is optional */
+	IRISCORE_API const FReplicationProtocol* CreateReplicationProtocol(const FReplicationProtocolIdentifier ProtocolId, const FReplicationFragments& Fragments, const TCHAR* DebugName, bool bVerifyId = false)
+	{
+		return CreateReplicationProtocol(nullptr, ProtocolId, Fragments, DebugName, bVerifyId);
+	}
+
+	/* Get an existing replication protocol */
+	IRISCORE_API const FReplicationProtocol* GetReplicationProtocol(FReplicationProtocolIdentifier ProtocolId, const UObject* ArchetypeOrCDOUsedAsKey = nullptr);
+
+	/** Iterate over all protocols matching the ProtocolId, mostly used by debug functionality with no sideeffects */
+	template<typename T>
+	void ForEachProtocol(FReplicationProtocolIdentifier ProtocolId, T&& Functor) const;
+	
+	/* Validate that a existing protocol matches the FragmentList of an instance, returns true of it is a match */
+	IRISCORE_API bool ValidateReplicationProtocol(const FReplicationProtocol*, const FReplicationFragments& Fragments) const;
+
+	/* Destroy existing replication protocol */
+	IRISCORE_API void DestroyReplicationProtocol(const FReplicationProtocol* ReplicationProtocol);
+
+	/** Create instance protocol from registered fragment data
+		Lifetime either the same as the lifetime of a remotely created instance (we are the client) or controlled by
+		calls to the bridge api to replicate an object or not
+	*/
+	IRISCORE_API static FReplicationInstanceProtocol* CreateInstanceProtocol(const FReplicationFragments& Fragments);
+	IRISCORE_API static void DestroyInstanceProtocol(FReplicationInstanceProtocol*);
+
+	IRISCORE_API void InvalidateDescriptor(const FReplicationStateDescriptor* InvalidatedReplicationStateDescriptor);
+
+private:
+	void InternalDestroyReplicationProtocol(const FReplicationProtocol* Protocol);
+	void InternalDeferDestroyReplicationProtocol(const FReplicationProtocol* Protocol);
+	void PruneProtocolsPendingDestroy();
+
+	struct FRegisteredProtocolInfo
+	{
+		const FReplicationProtocol* Protocol;
+		const UObject* ArchetypeOrCDOUsedAsKey;
+	
+		bool operator==(const FRegisteredProtocolInfo& Other) const { return Protocol == Other.Protocol && ArchetypeOrCDOUsedAsKey == Other.ArchetypeOrCDOUsedAsKey; };
+	};
+
+	TMultiMap<FReplicationProtocolIdentifier, FRegisteredProtocolInfo> RegisteredProtocols;
+	TMap<const FReplicationProtocol*, FRegisteredProtocolInfo> ProtocolToInfoMap;
+	
+	// We use the pointer as key, we have full control over lifetime of descriptors so this should not be a problem
+	TMultiMap<const FReplicationStateDescriptor*, const FReplicationProtocol*> DescriptorToProtocolMap;
+	TArray<const FReplicationProtocol*> PendingDestroyProtocols;
+};
+
+template<typename T>
+void FReplicationProtocolManager::ForEachProtocol(FReplicationProtocolIdentifier ProtocolId, T&& Functor) const
+{
+	// Find protocols using the descriptor being invalidated
+	for (auto It = RegisteredProtocols.CreateConstKeyIterator(ProtocolId); It; ++It)
+	{
+		const FRegisteredProtocolInfo& Info = It.Value();
+		Functor(Info.Protocol, Info.ArchetypeOrCDOUsedAsKey);
+	}
+}
+
+
+}
diff --git a/Engine/Source/Runtime/Experimental/Iris/Core/Private/Iris/ReplicationSystem/ReplicationReader.cpp b/Engine/Source/Runtime/Experimental/Iris/Core/Private/Iris/ReplicationSystem/ReplicationReader.cpp
new file mode 100644
--- /dev/null
+++ b/Engine/Source/Runtime/Experimental/Iris/Core/Private/Iris/ReplicationSystem/ReplicationReader.cpp
@@ -0,0 +1,1562 @@
+// Copyright Epic Games, Inc. All Rights Reserved.
+
+#include "Iris/ReplicationSystem/ReplicationReader.h"
+#include "Iris/IrisConfigInternal.h"
+#include "Iris/Core/IrisLog.h"
+#include "Iris/Core/IrisProfiler.h"
+#include "Iris/Core/IrisDebugging.h"
+#include "Net/Core/Trace/NetTrace.h"
+#include "Iris/DataStream/DataStream.h"
+#include "Iris/ReplicationSystem/DeltaCompression/DeltaCompressionBaselineManager.h"
+#include "Iris/ReplicationSystem/NetHandleManager.h"
+#include "Iris/Serialization/NetReferenceCollector.h"
+#include "Iris/ReplicationSystem/ReplicationOperations.h"
+#include "Iris/ReplicationSystem/ReplicationOperationsInternal.h"
+#include "Iris/ReplicationSystem/DequantizeAndApplyHelper.h"
+#include "Iris/ReplicationSystem/ReplicationSystem.h"
+#include "Iris/ReplicationSystem/ReplicationSystemInternal.h"
+#include "Iris/ReplicationSystem/NetBlob/NetObjectBlobHandler.h"
+#include "Iris/ReplicationSystem/NetBlob/PartialNetObjectAttachmentHandler.h"
+#include "Iris/Serialization/InternalNetSerializationContext.h"
+#include "Iris/Serialization/NetBitStreamReader.h"
+#include "Iris/Serialization/NetBitStreamUtil.h"
+#include "Iris/Serialization/NetSerializer.h"
+#include "HAL/IConsoleManager.h"
+#include "UObject/Class.h"
+
+#if UE_NET_ENABLE_REPLICATIONREADER_LOG
+#	define UE_LOG_REPLICATIONREADER(Format, ...)  UE_LOG(LogIris, Log, Format, ##__VA_ARGS__)
+#	define UE_LOG_REPLICATIONREADER_CONN(Format, ...)  UE_LOG(LogIris, Log, TEXT("Conn: %u ") Format, Parameters.ConnectionId, ##__VA_ARGS__)
+#else
+#	define UE_LOG_REPLICATIONREADER(...)
+#	define UE_LOG_REPLICATIONWRITER_CONN(Format, ...)
+#endif
+
+#define UE_LOG_REPLICATIONREADER_WARNING(Format, ...)  UE_LOG(LogIris, Warning, Format, ##__VA_ARGS__)
+#define UE_LOG_REPLICATIONREADER_CONN_WARNING(Format, ...) UE_LOG(LogIris, Warning, TEXT("Conn: %u ") Format, Parameters.ConnectionId, ##__VA_ARGS__)
+#define UE_LOG_REPLICATIONREADER_ERROR(Format, ...)  UE_LOG(LogIris, Error, Format, ##__VA_ARGS__)
+
+namespace UE::Net::Private
+{
+
+static bool bExecuteReliableRPCsBeforeOnReps = false;
+static FAutoConsoleVariableRef CVarExecuteReliableRPCsBeforeOnReps(
+		TEXT("net.Iris.ExecuteReliableRPCsBeforeOnReps"),
+		bExecuteReliableRPCsBeforeOnReps,
+		TEXT("If true and Iris runs in backwards compatibility mode then reliable RPCs will be executed before OnReps on the target object. Default is false."
+		));
+
+static bool bDeferEndReplication = true;
+static FAutoConsoleVariableRef CVarDeferEndReplication(
+	TEXT("net.Iris.DeferEndReplication"),
+	bDeferEndReplication,
+	TEXT("bDeferEndReplication if true calls to EndReplication will be defered until after we have applied statedata. Default is true."
+	));
+
+static const FName NetError_FailedToFindAttachmentQueue("Failed to find attachment queue");
+
+class FResolveAndCollectUnresolvedAndResolvedReferenceCollector
+{
+public:
+	typedef FNetReferenceCollector::FReferenceInfoArray FReferenceInfoArray;
+
+	void CollectReferences(FObjectReferenceCache& ObjectReferenceCache,
+		const FNetObjectResolveContext& ResolveContext,
+		bool bInIncludeInitState, 
+		const FNetBitArrayView* ChangeMask,
+		uint8* RESTRICT InternalBuffer, const FReplicationProtocol* Protocol)
+	{
+		bIncludeInitState = bInIncludeInitState;
+
+		// Setup context
+		FNetSerializationContext Context;
+		Context.SetChangeMask(ChangeMask);
+		Context.SetIsInitState(bInIncludeInitState);
+
+		FNetReferenceCollector Collector;
+		FReplicationProtocolOperationsInternal::CollectReferences(Context, Collector, InternalBuffer, Protocol);
+
+		// Iterate over result and process results
+		for (const FNetReferenceCollector::FReferenceInfo& Info : MakeArrayView(Collector.GetCollectedReferences()))
+		{
+			if (!ObjectReferenceCache.ResolveObjectReference(Info.Reference, ResolveContext))
+			{
+				UnresolvedReferenceInfos.Add(Info);
+			}
+			else
+			{
+				ResolvedReferenceInfos.Add(Info);
+			}		
+		}		
+	}
+
+	void Reset()
+	{ 
+		UnresolvedReferenceInfos.Reset();
+		ResolvedReferenceInfos.Reset();
+	}
+
+	const FReferenceInfoArray& GetResolvedReferences() const { return ResolvedReferenceInfos; }
+	const FReferenceInfoArray& GetUnresolvedReferences() const { return UnresolvedReferenceInfos; }
+	bool IsInitStateIncluded() const { return bIncludeInitState; }
+
+private:
+	FReferenceInfoArray UnresolvedReferenceInfos;
+	FReferenceInfoArray ResolvedReferenceInfos;
+	bool bIncludeInitState = false;
+};
+
+FReplicationReader::FReplicatedObjectInfo::FReplicatedObjectInfo()
+: InternalIndex(FNetHandleManager::InvalidInternalIndex)
+, Value(0U)
+{
+	FMemory::Memzero(StoredBaselines);
+	LastStoredBaselineIndex = FDeltaCompressionBaselineManager::InvalidBaselineIndex;
+	PrevStoredBaselineIndex = FDeltaCompressionBaselineManager::InvalidBaselineIndex;
+}
+
+FReplicationReader::FReplicationReader()
+: TempLinearAllocator()
+, TempChangeMaskAllocator(&TempLinearAllocator)
+, ReplicationSystemInternal(nullptr)
+, NetHandleManager(nullptr)
+, StateStorage(nullptr)
+, ObjectsToDispatch(nullptr)
+, ObjectsToDispatchCount(0U)
+, ObjectsToDispatchCapacity(0U)
+, NetBlobHandlerManager(nullptr)
+, NetObjectBlobType(InvalidNetBlobType)
+, DelayAttachmentsWithUnresolvedReferences(IConsoleManager::Get().FindConsoleVariable(TEXT("net.DelayUnmappedRPCs")))
+{
+}
+
+FReplicationReader::~FReplicationReader()
+{
+	Deinit();
+}
+
+void FReplicationReader::Init(const FReplicationParameters& InParameters)
+{
+	// Store copy of parameters
+	Parameters = InParameters;
+
+	// Cache internal systems
+	ReplicationSystemInternal = Parameters.ReplicationSystem->GetReplicationSystemInternal();
+	NetHandleManager = &ReplicationSystemInternal->GetNetHandleManager();
+	StateStorage = &ReplicationSystemInternal->GetReplicationStateStorage();
+	NetBlobHandlerManager = &ReplicationSystemInternal->GetNetBlobHandlerManager();
+	ObjectReferenceCache = &ReplicationSystemInternal->GetObjectReferenceCache();
+	ReplicationBridge = Parameters.ReplicationSystem->GetReplicationBridge();
+
+	// Find out if there's a PartialNetObjectAttachmentHandler so we can re-assemble split blobs
+	if (const UNetBlobHandler* Handler = ReplicationSystemInternal->GetNetBlobManager().GetPartialNetObjectAttachmentHandler())
+	{
+		Attachments.SetPartialNetBlobType(Handler->GetNetBlobType());
+	}
+
+	if (const UNetBlobHandler* Handler = ReplicationSystemInternal->GetNetBlobManager().GetNetObjectBlobHandler())
+	{
+		NetObjectBlobType = Handler->GetNetBlobType();
+	}
+
+	// reserve index 0
+	StartReplication(0);
+}
+
+void FReplicationReader::Deinit()
+{
+	// Cleanup any allocation stored in the per object info
+	for (auto& ObjectIt : ReplicatedObjects)
+	{
+		CleanupObjectData(ObjectIt.Value);
+	}
+	ReplicatedObjects.Empty();
+}
+
+// Read incomplete handle
+FNetHandle FReplicationReader::ReadNetHandleId(FNetBitStreamReader& Reader) const
+{
+	const uint32 NetId = Reader.ReadBits(FNetHandle::IdBits);
+	return FNetHandleManager::MakeNetHandleFromId(NetId);
+}
+	
+uint16 FReplicationReader::ReadObjectsPendingDestroy(FNetSerializationContext& Context)
+{
+	FNetBitStreamReader& Reader = *Context.GetBitStreamReader();
+
+	UE_NET_TRACE_SCOPE(ObjectsPendingDestroy, Reader, Context.GetTraceCollector(), ENetTraceVerbosity::Trace);
+
+	// Read how many destroyed objects we have
+	const uint16 ObjectsToRead = Reader.ReadBits(16);
+	
+	if (!Reader.IsOverflown())
+	{
+		for (uint32 It = 0; It < ObjectsToRead; ++It)
+		{
+			UE_NET_TRACE_NAMED_OBJECT_SCOPE(DestroyedObjectScope, FNetHandle(), Reader, Context.GetTraceCollector(), ENetTraceVerbosity::Trace);
+
+			FNetHandle IncompleteHandle = ReadNetHandleId(Reader);
+			const bool bShouldDestroyInstance = Reader.ReadBool();
+			if (!Reader.IsOverflown())
+			{
+				// Resolve handle and destroy using bridge
+				const uint32 InternalIndex = NetHandleManager->GetInternalIndex(IncompleteHandle);
+				if (InternalIndex)
+				{
+					UE_NET_TRACE_SET_SCOPE_OBJECTID(DestroyedObjectScope, IncompleteHandle);
+
+					// Defer EndReplication until after applying state data
+					if (bDeferEndReplication)
+					{
+						FDispatchObjectInfo& Info = ObjectsToDispatch[ObjectsToDispatchCount];
+						Info = FDispatchObjectInfo();
+
+						Info.bDestroy = bShouldDestroyInstance;
+						Info.bTearOff = false;
+						Info.bDeferredEndReplication = true;
+						Info.InternalIndex = InternalIndex;
+						Info.bIsInitialState = 0U;
+						Info.bHasState = 0U;
+						Info.bHasAttachments = 0U;
+
+						// Mark for dispatch
+						++ObjectsToDispatchCount;
+					}
+					else
+					{
+						EndReplication(InternalIndex, false, bShouldDestroyInstance);
+					}
+
+					continue;
+				}
+
+				// If we did not find the object or associated bridge something is wrong and we should disconnect.
+				UE_LOG_REPLICATIONREADER_WARNING(TEXT("FReplicationReader::Read Tried to destroy object with %s (This can occur if the server sends destroy for an object that has not yet been confirmed as created)"), *IncompleteHandle.ToString());
+			}
+		}
+	}
+
+	check(!Reader.IsOverflown());
+
+	return ObjectsToRead;
+}
+
+FReplicationReader::FReplicatedObjectInfo& FReplicationReader::StartReplication(uint32 InternalIndex)
+{
+	check(!ReplicatedObjects.Contains(InternalIndex));
+
+	// Create ReadObjectInfo
+	FReplicatedObjectInfo& ObjectInfo = ReplicatedObjects.Emplace(InternalIndex);
+	ObjectInfo = FReplicatedObjectInfo();
+	ObjectInfo.InternalIndex = InternalIndex;
+
+	// Allocate changemask (if needed)
+	if (InternalIndex != 0U)
+	{		
+		const FNetHandleManager::FReplicatedObjectData& Data = NetHandleManager->GetReplicatedObjectDataNoCheck(InternalIndex);
+		ObjectInfo.ChangeMaskBitCount = Data.Protocol->ChangeMaskBitCount;
+
+		// Alloc and init changemask
+		FNetBitArrayView ChangeMask = FChangeMaskStorageOrPointer::AllocAndInitBitArray(ObjectInfo.UnresolvedChangeMaskOrPointer, ObjectInfo.ChangeMaskBitCount, PersistentChangeMaskAllocator);
+	}
+
+	return ObjectInfo;
+}
+
+FReplicationReader::FReplicatedObjectInfo* FReplicationReader::GetReplicatedObjectInfo(uint32 InternalIndex)
+{
+	return ReplicatedObjects.Find(InternalIndex);
+}
+
+void FReplicationReader::CleanupObjectData(FReplicatedObjectInfo& ObjectInfo)
+{
+	// Remove from pending resolve
+	if (ObjectInfo.InternalIndex != 0U)
+	{
+		FChangeMaskStorageOrPointer::Free(ObjectInfo.UnresolvedChangeMaskOrPointer, ObjectInfo.ChangeMaskBitCount, PersistentChangeMaskAllocator);
+	}
+
+	// Release stored baselines
+	if (ObjectInfo.LastStoredBaselineIndex != FDeltaCompressionBaselineManager::InvalidBaselineIndex)
+	{
+		StateStorage->FreeBaseline(ObjectInfo.InternalIndex, ObjectInfo.StoredBaselines[ObjectInfo.LastStoredBaselineIndex]);
+	}
+	if (ObjectInfo.PrevStoredBaselineIndex != FDeltaCompressionBaselineManager::InvalidBaselineIndex)
+	{
+		StateStorage->FreeBaseline(ObjectInfo.InternalIndex, ObjectInfo.StoredBaselines[ObjectInfo.PrevStoredBaselineIndex]);
+	}
+}
+
+void FReplicationReader::EndReplication(uint32 InternalIndex, bool bTearOff, bool bDestroyInstance)
+{
+	if (FReplicatedObjectInfo* ObjectInfo = ReplicatedObjects.Find(InternalIndex))
+	{
+		const FNetHandleManager::FReplicatedObjectData& Data = NetHandleManager->GetReplicatedObjectDataNoCheck(InternalIndex);
+
+		CleanupReferenceTracking(ObjectInfo);
+		Attachments.DropAllAttachments(ENetObjectAttachmentType::Normal, InternalIndex);
+		ReplicationBridge->DestroyNetObjectFromRemote(Data.Handle, bTearOff, bDestroyInstance);
+
+		CleanupObjectData(*ObjectInfo);
+
+		ReplicatedObjects.Remove(InternalIndex);
+	}
+}
+
+void FReplicationReader::DeserializeObjectStateDelta(FNetSerializationContext& Context, uint32 InternalIndex, FDispatchObjectInfo& Info, FReplicatedObjectInfo& ObjectInfo, const FNetHandleManager::FReplicatedObjectData& ObjectData, uint32& OutNewBaselineIndex)
+{
+	FNetBitStreamReader& Reader = *Context.GetBitStreamReader();
+
+	uint32 BaselineIndex = FDeltaCompressionBaselineManager::InvalidBaselineIndex;
+	BaselineIndex = Reader.ReadBits(FDeltaCompressionBaselineManager::BaselineIndexBitCount);
+	if (BaselineIndex != FDeltaCompressionBaselineManager::InvalidBaselineIndex)
+	{ 
+		const bool bIsNewBaseline = Reader.ReadBool();
+
+		if (Reader.IsOverflown())
+		{
+			UE_LOG_REPLICATIONREADER_ERROR(TEXT("FReplicationReader::DeserializeObjectStateDelta Bitstream corrupted."));
+			return;
+		}
+
+		if (bIsNewBaseline)
+		{
+			OutNewBaselineIndex = (BaselineIndex + 1) % FDeltaCompressionBaselineManager::MaxBaselineCount;
+		}
+
+		// If we are compressing against the LastStoredBaselineIndex we can release older baselines to reduce memory overhead
+		if (!bIsNewBaseline && BaselineIndex == ObjectInfo.LastStoredBaselineIndex && ObjectInfo.PrevStoredBaselineIndex != FDeltaCompressionBaselineManager::InvalidBaselineIndex)
+		{
+			StateStorage->FreeBaseline(InternalIndex, ObjectInfo.StoredBaselines[ObjectInfo.PrevStoredBaselineIndex]);
+			ObjectInfo.StoredBaselines[ObjectInfo.PrevStoredBaselineIndex] = nullptr;
+			ObjectInfo.PrevStoredBaselineIndex = FDeltaCompressionBaselineManager::InvalidBaselineIndex;
+		}
+
+		check(ObjectInfo.StoredBaselines[BaselineIndex]);
+
+		UE_NET_TRACE_SCOPE(DeltaCompressed, Reader, Context.GetTraceCollector(), ENetTraceVerbosity::Trace);
+		FReplicationProtocolOperations::DeserializeWithMaskDelta(Context, Info.ChangeMaskOrPointer.GetPointer(ObjectInfo.ChangeMaskBitCount), ObjectData.ReceiveStateBuffer, ObjectInfo.StoredBaselines[BaselineIndex], ObjectData.Protocol);
+	}
+	else
+	{
+		const uint32 NewBaselineIndex = Reader.ReadBits(FDeltaCompressionBaselineManager::BaselineIndexBitCount);
+		if (Reader.IsOverflown())
+		{
+			UE_LOG_REPLICATIONREADER_ERROR(TEXT("FReplicationReader::DeserializeObjectStateDelta Bitstream corrupted."));
+			return;
+		}
+		OutNewBaselineIndex = NewBaselineIndex;
+		FReplicationProtocolOperations::DeserializeWithMask(Context, Info.ChangeMaskOrPointer.GetPointer(ObjectInfo.ChangeMaskBitCount), ObjectData.ReceiveStateBuffer, ObjectData.Protocol);
+	}
+}
+
+void FReplicationReader::ReadObject(FNetSerializationContext& Context)
+{
+	FNetBitStreamReader& Reader = *Context.GetBitStreamReader();
+
+	if (const bool bIsDestructionInfo = Reader.ReadBool())
+	{
+		FReplicationBridgeSerializationContext BridgeContext(Context, Parameters.ConnectionId, true);
+
+		ReplicationBridge->ReadAndExecuteDestructionInfoFromRemote(BridgeContext);
+
+#if UE_NET_USE_READER_WRITER_SENTINEL
+		ReadAndVerifySentinelBits(&Reader, TEXT("DestructionInfo"), 8);
+#endif
+		
+		return;
+	}
+
+#if UE_NET_USE_READER_WRITER_SENTINEL
+	ReadAndVerifySentinelBits(&Reader, TEXT("ReadObject"), 8);
+#endif
+
+	const FNetHandle IncompleteHandle = ReadNetHandleId(Reader);
+	const bool bTearOff = Reader.ReadBool();
+
+	bool bSubObjectPendingDestroy = false;
+	const bool bSubObjectPendingEndReplication = Reader.ReadBool();
+	if (bSubObjectPendingEndReplication)
+	{
+		UE_NET_TRACE_SCOPE(SubObjectPendingDestroy, Reader, Context.GetTraceCollector(), ENetTraceVerbosity::Trace);
+		bSubObjectPendingDestroy = Reader.ReadBool();
+	}
+
+	const bool bHasState = Reader.ReadBool();
+	uint32 NewBaselineIndex = FDeltaCompressionBaselineManager::InvalidBaselineIndex;
+
+	if (bHasState)
+	{
+		ObjectReferenceCache->ReadExports(Context);
+#if UE_NET_USE_READER_WRITER_SENTINEL
+		ReadAndVerifySentinelBits(&Reader, TEXT("Exports"), 8);
+#endif
+	}
+
+	const bool bIsInitialState = bHasState && Reader.ReadBool();
+	uint32 InternalIndex = ObjectIndexForOOBAttachment;
+
+	UE_NET_TRACE_OBJECT_SCOPE(IncompleteHandle, Reader, Context.GetTraceCollector(), ENetTraceVerbosity::Trace);
+
+	//UE_LOG_REPLICATIONREADER(TEXT("FReplicationReader::Read Object with %s InitialState: %u"), *IncompleteHandle.ToString(), bIsInitialState ? 1u : 0u);
+
+	bool bHasErrors = false;
+
+	// Read creation data
+	Context.SetIsInitState(bIsInitialState);
+	if (bIsInitialState)
+	{
+		UE_NET_TRACE_SCOPE(CreationInfo, *Context.GetBitStreamReader(), Context.GetTraceCollector(), ENetTraceVerbosity::Trace);
+
+#if UE_NET_REPLICATION_SUPPORT_SKIP_INITIAL_STATE
+		const uint32 NumBitsUsedForInitialStateSize = Parameters.NumBitsUsedForInitialStateSize;
+		const uint32 SkipSeekPos = Reader.ReadBits(NumBitsUsedForInitialStateSize) + Reader.GetPosBits();
+#endif
+
+		// SubObject data
+		FNetHandle SubObjectOwnerHandle;
+		if (Reader.ReadBool())
+		{
+			const FNetHandle IncompleteOwnerHandle = ReadNetHandleId(Reader);
+				
+			FInternalNetHandle SubObjectOwnerInternalIndex = NetHandleManager->GetInternalIndex(IncompleteOwnerHandle);
+			if (Reader.IsOverflown() || SubObjectOwnerInternalIndex == FNetHandleManager::InvalidInternalIndex)
+			{
+				UE_LOG_REPLICATIONREADER_ERROR(TEXT("FReplicationReader::ReadObject Invalid subobjectowner handle. %s"), ToCStr(IncompleteOwnerHandle.ToString()));
+				const FName& NetError = (Reader.IsOverflown() ? GNetError_BitStreamOverflow : GNetError_InvalidNetHandle);
+				Context.SetError(NetError);
+				return;			
+			}
+
+			SubObjectOwnerHandle = NetHandleManager->GetReplicatedObjectDataNoCheck(SubObjectOwnerInternalIndex).Handle;
+		}
+
+		const bool bIsDeltaCompressed = Reader.ReadBool();
+		if (bIsDeltaCompressed)
+		{
+			UE_LOG_REPLICATIONREADER(TEXT("DeltaCompression is enabled for Handle %s"), *IncompleteHandle.ToString());
+			NewBaselineIndex = Reader.ReadBits(FDeltaCompressionBaselineManager::BaselineIndexBitCount);
+		}
+		
+		// We got a read error
+		if (Reader.IsOverflown() || IsObjectIndexForOOBAttachment(IncompleteHandle.GetId()))
+		{
+			UE_LOG_REPLICATIONREADER_ERROR(TEXT("FReplicationReader::ReadObject Bitstream corrupted."));
+			const FName& NetError = (Reader.IsOverflown() ? GNetError_BitStreamOverflow : GNetError_BitStreamError);
+			Context.SetError(NetError);
+			return;
+		}
+	
+		// Get Bridge
+		FReplicationBridgeSerializationContext BridgeContext(Context, Parameters.ConnectionId);
+
+		FNetHandle NetHandle = ReplicationBridge->CallCreateNetHandleFromRemote(SubObjectOwnerHandle, IncompleteHandle, BridgeContext);
+
+		if (!NetHandle.IsValid())
+		{
+#if UE_NET_REPLICATION_SUPPORT_SKIP_INITIAL_STATE
+			// If we support skipping
+			UE_LOG_REPLICATIONREADER_WARNING(TEXT("FReplicationReader::ReadObject Failed to instantiate %s, skipping over it assuming that object was streamed out"), *NetHandle.ToString());
+			Reader.Seek(SkipSeekPos);
+			return;
+#endif
+	
+			UE_LOG_REPLICATIONREADER_ERROR(TEXT("FReplicationReader::ReadObject Unable to create handle for %s."), *IncompleteHandle.ToString());
+			Context.SetError(GNetError_InvalidNetHandle);
+			bHasErrors = true;
+			goto ErrorHandling;
+		}
+
+		InternalIndex = NetHandleManager->GetInternalIndex(NetHandle);		 
+		FReplicatedObjectInfo& ObjectInfo = StartReplication(InternalIndex);
+
+		ObjectInfo.bIsDeltaCompressionEnabled = bIsDeltaCompressed;
+	}
+	else
+	{
+		bHasErrors = bHasErrors || Context.HasErrorOrOverflow();
+		if (bHasErrors || IsObjectIndexForOOBAttachment(IncompleteHandle.GetId()))
+		{
+			InternalIndex = ObjectIndexForOOBAttachment;
+		}
+		else
+		{
+			// If we get back an invalid internal index then either the object has been deleted or there's bitstream corruption.
+			InternalIndex = NetHandleManager->GetInternalIndex(IncompleteHandle);
+
+			// If this is a subobject that is being destroyed this was no error as we send destroy info for unconfirmed objects
+			if (!bSubObjectPendingEndReplication)
+			{
+				bHasErrors = InternalIndex == FNetHandleManager::InvalidInternalIndex;
+			}
+		}
+	}
+
+	if (bHasErrors)
+	{
+		Context.SetError(GNetError_InvalidNetHandle);
+		goto ErrorHandling;
+	}
+
+	// Read state data and attachments
+	{
+		const FNetHandleManager::FReplicatedObjectData& ObjectData = NetHandleManager->GetReplicatedObjectDataNoCheck(InternalIndex);
+
+		// Add entry in our received data as we postpone state application until we have received all data in order to be able to properly resolve references
+		FDispatchObjectInfo& Info = ObjectsToDispatch[ObjectsToDispatchCount];
+		Info = FDispatchObjectInfo();
+
+		Info.bDestroy = bTearOff || bSubObjectPendingDestroy;
+		Info.bTearOff = bTearOff;
+		Info.bDeferredEndReplication = bTearOff || bSubObjectPendingEndReplication;
+
+		if (bHasState)
+		{
+			bHasErrors = IsObjectIndexForOOBAttachment(InternalIndex);
+			if (bHasErrors)
+			{
+				UE_LOG_REPLICATIONREADER_WARNING(TEXT("FReplicationReader::ReadObject Bitstream corrupted. Getting state when only expecting RPCs."));
+				Context.SetError(GNetError_BitStreamError);
+				goto ErrorHandling;
+			}
+
+			FReplicatedObjectInfo* ObjectInfo = GetReplicatedObjectInfo(InternalIndex);
+			checkSlow(ObjectInfo);
+
+			const uint32 ChangeMaskBitCount = ObjectData.Protocol->ChangeMaskBitCount;
+
+			// Temporary changemask
+			FChangeMaskStorageOrPointer::Alloc(Info.ChangeMaskOrPointer, ChangeMaskBitCount, TempChangeMaskAllocator);
+
+			if (bIsInitialState)
+			{
+				FReplicationProtocolOperations::DeserializeInitialStateWithMask(Context, Info.ChangeMaskOrPointer.GetPointer(ChangeMaskBitCount), ObjectData.ReceiveStateBuffer, ObjectData.Protocol);
+			}
+			else
+			{
+				if (ObjectInfo->bIsDeltaCompressionEnabled)
+				{
+					DeserializeObjectStateDelta(Context, InternalIndex, Info, *ObjectInfo, ObjectData, NewBaselineIndex);
+				}
+				else
+				{
+					FReplicationProtocolOperations::DeserializeWithMask(Context, Info.ChangeMaskOrPointer.GetPointer(ChangeMaskBitCount), ObjectData.ReceiveStateBuffer, ObjectData.Protocol);
+				}
+			}
+
+			// Should we store a new baseline?
+			if (NewBaselineIndex != FDeltaCompressionBaselineManager::InvalidBaselineIndex)
+			{
+				// This object uses delta compression, store the last received state as a baseline with the specified index
+				UE_LOG_REPLICATIONREADER(TEXT("Storing new baselineindex: %u for (:%u) Handle %s"), NewBaselineIndex, InternalIndex, *ObjectData.Handle.ToString());
+
+				check(NewBaselineIndex < FDeltaCompressionBaselineManager::MaxBaselineCount);
+				if (ObjectInfo->StoredBaselines[NewBaselineIndex])
+				{
+					// Clone into already allocated state, unfortunately we have to free dynamic state
+					FReplicationProtocolOperations::FreeDynamicState(Context, ObjectInfo->StoredBaselines[NewBaselineIndex], ObjectData.Protocol);
+					FReplicationProtocolOperationsInternal::CloneQuantizedState(Context, ObjectInfo->StoredBaselines[NewBaselineIndex], ObjectData.ReceiveStateBuffer, ObjectData.Protocol);
+				}
+				else
+				{
+					// Allocate new baseline and initialize from current RecvState
+					ObjectInfo->StoredBaselines[NewBaselineIndex] = StateStorage->AllocBaseline(InternalIndex, EReplicationStateType::CurrentRecvState);
+				}
+	
+				// Make sure that PrevStoredBaselineIndex is not set to the same as the NewBaselineIndex
+				const uint32 OldPrevStoredBaselineIndex = ObjectInfo->PrevStoredBaselineIndex;
+				const uint32 NewPrevStoredBaselineIndex = NewBaselineIndex != ObjectInfo->LastStoredBaselineIndex ? ObjectInfo->LastStoredBaselineIndex : FDeltaCompressionBaselineManager::InvalidBaselineIndex;
+				ObjectInfo->PrevStoredBaselineIndex = NewPrevStoredBaselineIndex;
+				if (NewPrevStoredBaselineIndex == FDeltaCompressionBaselineManager::InvalidBaselineIndex && NewPrevStoredBaselineIndex != OldPrevStoredBaselineIndex)
+				{
+					uint8* PrevBaseline = ObjectInfo->StoredBaselines[OldPrevStoredBaselineIndex];
+					ObjectInfo->StoredBaselines[OldPrevStoredBaselineIndex] = nullptr;
+					StateStorage->FreeBaseline(InternalIndex, PrevBaseline);
+				}
+				ObjectInfo->LastStoredBaselineIndex = NewBaselineIndex;
+			}
+		}
+
+		const bool bHasAttachments = Reader.ReadBool();
+		ENetObjectAttachmentType AttachmentType = ENetObjectAttachmentType::Normal;
+		if (bHasAttachments)
+		{
+			if (IsObjectIndexForOOBAttachment(InternalIndex))
+			{
+				const bool bIsHugeObject = Reader.ReadBool();
+				AttachmentType = (bIsHugeObject ? ENetObjectAttachmentType::HugeObject : ENetObjectAttachmentType::OutOfBand);
+				bHasErrors = bHasErrors || (!Parameters.bAllowReceivingAttachmentsFromRemoteObjectsNotInScope && AttachmentType == ENetObjectAttachmentType::OutOfBand);
+				if (bHasErrors)
+				{
+					Context.SetError(GNetError_InvalidNetHandle);
+					goto ErrorHandling;
+				}
+			}
+
+			Attachments.Deserialize(Context, AttachmentType, InternalIndex, ObjectData.Handle);
+		}
+
+		bHasErrors = bHasErrors || Context.HasErrorOrOverflow();
+		if (bHasErrors)
+		{
+			goto ErrorHandling;
+		}
+
+		// Fill in ReadObjectInfo, we skip HugeObjects as they are not added to the dispatch list until they are fully assembled
+		if (AttachmentType != ENetObjectAttachmentType::HugeObject)
+		{
+			Info.InternalIndex = InternalIndex;
+			Info.bIsInitialState = bIsInitialState ? 1U : 0U;
+			Info.bHasState = bHasState ? 1U : 0U;
+			Info.bHasAttachments = bHasAttachments ? 1U : 0U;
+
+			// Mark for dispatch
+			++ObjectsToDispatchCount;
+		}
+	}
+	
+ErrorHandling:
+	if (bHasErrors)
+	{
+		UE_LOG_REPLICATIONREADER_ERROR(TEXT("FReplicationReader::ReadObject Failed to read replicated object with %s. Error '%s'."), *IncompleteHandle.ToString(), (Context.HasError() ? ToCStr(Context.GetError().ToString()) : TEXT("BitStream Overflow")));
+	}
+}
+
+// Update reference tracking maps for the current object. It is assumed the ObjectReferenceTracker do no include duplicates for a given key.
+void FReplicationReader::UpdateObjectReferenceTracking(FReplicatedObjectInfo* ReplicationInfo, FNetBitArrayView ChangeMask, bool bIncludeInitState, const FObjectReferenceTracker& NewUnresolvedReferences, const FObjectReferenceTracker& NewMappedDynamicReferences)
+{
+	IRIS_PROFILER_SCOPE(FReplicationReader_UpdateObjectReferenceTracking)
+
+	/*
+	 * As we store references per changemask we need to construct a set of all unresolved references and
+	 * compare with the new set of unresolved references. The new set is found by first updating the 
+	 * references that were found in the changemask.
+	 */
+	{
+		// Try to avoid dynamic allocations during the update of the UnresolvedObjectReferences.
+		ReplicationInfo->UnresolvedObjectReferences.Reserve(ReplicationInfo->UnresolvedObjectReferences.Num() + NewUnresolvedReferences.Num());
+
+		TSet<FNetHandle> OldUnresolvedSet;
+		OldUnresolvedSet.Reserve(ReplicationInfo->UnresolvedObjectReferences.Num());
+		for (const FObjectReferenceTracker::ElementType& Element : ReplicationInfo->UnresolvedObjectReferences)
+		{
+			OldUnresolvedSet.Add(Element.Value);
+		}
+
+		// Replace each entry in UnresolvedObjectReferences for the given changemask
+		auto&& UpdateUnresolvedReferencesForChange = [ReplicationInfo, &NewUnresolvedReferences](uint32 ChangeBit)
+		{
+			FObjectReferenceTracker& UnresolvedObjectReferences = ReplicationInfo->UnresolvedObjectReferences;
+
+			// Use the Hash variants of the map methods to provide a tiny speedup.
+			const uint32 KeyHash = ::GetTypeHash(ChangeBit);
+			UnresolvedObjectReferences.RemoveByHash(KeyHash, ChangeBit);
+
+			for (auto It = NewUnresolvedReferences.CreateConstKeyIterator(ChangeBit); It; ++It)
+			{
+				const FNetHandle NetHandle = It.Value();
+				UnresolvedObjectReferences.AddByHash(KeyHash, ChangeBit, NetHandle);
+			}
+		};
+
+		ChangeMask.ForAllSetBits(UpdateUnresolvedReferencesForChange);
+		if (bIncludeInitState)
+		{
+			UpdateUnresolvedReferencesForChange(FakeInitChangeMaskOffset);
+		}
+
+		// The unresolved set is now updated with the current status of unresolved references.
+		TSet<FNetHandle> NewUnresolvedSet;
+		NewUnresolvedSet.Reserve(NewUnresolvedReferences.Num());
+		for (const FObjectReferenceTracker::ElementType& Element : ReplicationInfo->UnresolvedObjectReferences)
+		{
+			NewUnresolvedSet.Add(Element.Value);
+		}
+
+		// Update ReplicationInfo with the status of unresolved references
+		ReplicationInfo->bHasUnresolvedReferences = NewUnresolvedSet.Num() > 0;
+		ReplicationInfo->bHasUnresolvedInitialReferences = ReplicationInfo->UnresolvedObjectReferences.Find(FakeInitChangeMaskOffset) != nullptr;
+
+		// Remove resolved or no longer existing references
+		const uint32 OwnerInternalIndex = ReplicationInfo->InternalIndex;
+		for (FNetHandle Handle : OldUnresolvedSet)
+		{
+			if (!NewUnresolvedSet.Contains(Handle))
+			{
+				// Remove from tracking
+				UnresolvedHandleToDependents.RemoveSingle(Handle, OwnerInternalIndex);
+				UE_LOG(LogIris, Verbose, TEXT("FReplicationReader::UpdateObjectReferenceTracking Removing unresolved reference %s for %s"), ToCStr(Handle.ToString()), ToCStr(NetHandleManager->GetNetHandleFromInternalIndex(OwnerInternalIndex).ToString()));
+			}
+		}
+
+		// Add new unresolved references
+		for (FNetHandle Handle : NewUnresolvedSet)
+		{
+			if (!OldUnresolvedSet.Contains(Handle))
+			{
+				// Add to tracking
+				UnresolvedHandleToDependents.Add(Handle, OwnerInternalIndex);
+				UE_LOG(LogIris, Verbose, TEXT("FReplicationReader::UpdateObjectReferenceTracking Adding unresolved reference %s for %s"), ToCStr(Handle.ToString()), ToCStr(NetHandleManager->GetNetHandleFromInternalIndex(OwnerInternalIndex).ToString()));
+			}
+		}
+	}
+
+	// Update tracking for resolved dynamic references
+#if 0
+	{
+		// Try to avoid dynamic allocations during the update of the ResolvedDynamicObjectReferences.
+		ReplicationInfo->ResolvedDynamicObjectReferences.Reserve(ReplicationInfo->ResolvedDynamicObjectReferences.Num() + NewMappedDynamicReferences.Num());
+
+		TSet<FNetHandle> OldResolvedSet;
+		OldResolvedSet.Reserve(ReplicationInfo->ResolvedDynamicObjectReferences.Num());
+		for (const FObjectReferenceTracker::ElementType& Element : ReplicationInfo->ResolvedDynamicObjectReferences)
+		{
+			OldResolvedSet.Add(Element.Value);
+		}
+
+		// Replace each entry in ResolvedDynamicObjectReferences for the given changemask
+		auto&& UpdateResolvedReferencesForChange = [ReplicationInfo, &NewMappedDynamicReferences](uint32 ChangeBit)
+		{
+			FObjectReferenceTracker& ResolvedDynamicObjectReferences = ReplicationInfo->ResolvedDynamicObjectReferences;
+
+			const uint32 KeyHash = ::GetTypeHash(ChangeBit);
+			ResolvedDynamicObjectReferences.RemoveByHash(KeyHash, ChangeBit);
+
+			for (auto It = NewMappedDynamicReferences.CreateConstKeyIterator(ChangeBit); It; ++It)
+			{
+				const FNetHandle NetHandle = It.Value();
+				ResolvedDynamicObjectReferences.AddByHash(KeyHash, ChangeBit, NetHandle);
+			}
+		};
+
+		ChangeMask.ForAllSetBits(UpdateResolvedReferencesForChange);
+		// Intentionally leaving out init state. It seems weird to call rep notifies and update init only properties after
+		// the initial state has already been applied.
+
+		// The resolved set is now updated with the current status of resolved references.
+		TSet<FNetHandle> NewResolvedSet;
+		NewResolvedSet.Reserve(ReplicationInfo->ResolvedDynamicObjectReferences.Num());
+		for (const FObjectReferenceTracker::ElementType& Element : ReplicationInfo->ResolvedDynamicObjectReferences)
+		{
+			NewResolvedSet.Add(Element.Value);
+		}
+
+		// Remove now unresolved or no longer existing references
+		const uint32 OwnerInternalIndex = ReplicationInfo->InternalIndex;
+		for (FNetHandle Handle : OldResolvedSet)
+		{
+			if (!NewResolvedSet.Contains(Handle))
+			{
+				// Remove from tracking
+				ResolvedDynamicHandleToDependents.RemoveSingle(Handle, OwnerInternalIndex);
+				UE_LOG(LogIris, Verbose, TEXT("FReplicationReader::UpdateObjectReferenceTracking Removing resolved dynamic reference %s for %s"), ToCStr(Handle.ToString()), ToCStr(NetHandleManager->GetNetHandleFromInternalIndex(OwnerInternalIndex).ToString()));
+			}
+		}
+
+		// Add new resolved dynamic references
+		for (FNetHandle Handle : NewResolvedSet)
+		{
+			if (Handle.IsDynamic() && !OldResolvedSet.Contains(Handle))
+			{
+				// Add to tracking
+				ResolvedDynamicHandleToDependents.Add(Handle, OwnerInternalIndex);
+				UE_LOG(LogIris, Verbose, TEXT("FReplicationReader::UpdateObjectReferenceTracking Adding resolved dynamic reference %s for %s"), ToCStr(Handle.ToString()), ToCStr(NetHandleManager->GetNetHandleFromInternalIndex(OwnerInternalIndex).ToString()));
+			}
+		}
+	}
+#endif
+}
+
+void FReplicationReader::RemoveUnresolvedObjectReferenceInReplicationInfo(FReplicatedObjectInfo* ReplicationInfo, FNetHandle Handle)
+{
+	for (FObjectReferenceTracker::TIterator It = ReplicationInfo->UnresolvedObjectReferences.CreateIterator(); It; ++It)
+	{
+		const FNetHandle NetHandle = It->Value;
+		if (NetHandle == Handle)
+		{
+			It.RemoveCurrent();
+		}
+	}
+}
+
+void FReplicationReader::RemoveResolvedObjectReferenceInReplicationInfo(FReplicatedObjectInfo* ReplicationInfo, FNetHandle Handle)
+{
+	for (FObjectReferenceTracker::TIterator It = ReplicationInfo->ResolvedDynamicObjectReferences.CreateIterator(); It; ++It)
+	{
+		const FNetHandle NetHandle = It->Value;
+		if (NetHandle == Handle)
+		{
+			It.RemoveCurrent();
+		}
+	}
+}
+
+bool FReplicationReader::MoveResolvedObjectReferenceToUnresolvedInReplicationInfo(FReplicatedObjectInfo* ReplicationInfo, FNetHandle UnresolvableHandle)
+{
+	bool bFoundHandle = false;
+	bool bHasUnresolvedReferences = ReplicationInfo->bHasUnresolvedReferences;
+	bool bHasUnresolvedInitialReferences = ReplicationInfo->bHasUnresolvedInitialReferences;
+	FNetBitArrayView UnresolvedChangeMask = FChangeMaskUtil::MakeChangeMask(ReplicationInfo->UnresolvedChangeMaskOrPointer, ReplicationInfo->ChangeMaskBitCount);
+	FObjectReferenceTracker& UnresolvedObjectReferences = ReplicationInfo->UnresolvedObjectReferences;
+	for (FObjectReferenceTracker::TIterator It = ReplicationInfo->ResolvedDynamicObjectReferences.CreateIterator(); It; ++It)
+	{
+		const FNetHandle NetHandle = It->Value;
+		if (NetHandle == UnresolvableHandle)
+		{
+			bFoundHandle = true;
+
+			const uint32 ChangemaskOffset = It->Key;
+			if (ChangemaskOffset == FakeInitChangeMaskOffset)
+			{
+				bHasUnresolvedInitialReferences = true;
+			}
+			else
+			{
+				bHasUnresolvedReferences = true;
+				UnresolvedChangeMask.SetBit(ChangemaskOffset);
+			}
+
+			// At this point we'd like to skip iteration to the next key as a handle can only be found once per changemask.
+			It.RemoveCurrent();
+
+			// This handle should only have existed once in the ResolvedDynamicObjectReferences map and should not be able to
+			// already exist in the UnresolvedObjectReferences map, so no need to call AddUnique.
+			UnresolvedObjectReferences.Add(ChangemaskOffset, UnresolvableHandle);
+
+			UE_LOG(LogIris, Verbose, TEXT("FReplicationReader::MoveResolvedObjectReferenceToUnresolvedInReplicationInfo Moving from resolved to unresolved reference %s for %s"), ToCStr(UnresolvableHandle.ToString()), ToCStr(NetHandleManager->GetNetHandleFromInternalIndex(ReplicationInfo->InternalIndex).ToString()));
+		}
+	}
+
+	ReplicationInfo->bHasUnresolvedInitialReferences = bHasUnresolvedInitialReferences;
+	ReplicationInfo->bHasUnresolvedReferences = bHasUnresolvedReferences;
+
+	return bFoundHandle;
+}
+
+// Remove all references for object
+void FReplicationReader::CleanupReferenceTracking(FReplicatedObjectInfo* ObjectInfo)
+{
+	const uint32 ObjectIndex = ObjectInfo->InternalIndex;
+
+	// Remove from unresolved references
+	for (FObjectReferenceTracker::ElementType Element : ObjectInfo->UnresolvedObjectReferences)
+	{
+		// Remove from tracking
+		FNetHandle Handle = Element.Value;
+		UnresolvedHandleToDependents.RemoveSingle(Handle, ObjectIndex);
+		UE_LOG(LogIris, Verbose, TEXT("FReplicationReader::CleanupReferenceTracking Removing unresolved reference %s for %s"), *Handle.ToString(), *(NetHandleManager->GetNetHandleFromInternalIndex(ObjectIndex).ToString()));
+	}
+	ObjectInfo->UnresolvedObjectReferences.Reset();
+
+	// Remove from resolved dynamic references
+	for (FObjectReferenceTracker::ElementType Element : ObjectInfo->ResolvedDynamicObjectReferences)
+	{
+		// Remove from tracking
+		const FNetHandle Handle = Element.Value;
+		ResolvedDynamicHandleToDependents.RemoveSingle(Handle, ObjectIndex);
+		UE_LOG(LogIris, Verbose, TEXT("FReplicationReader::CleanupReferenceTracking Removing resolved dynamic reference %s for %s"), *Handle.ToString(), *(NetHandleManager->GetNetHandleFromInternalIndex(ObjectIndex).ToString()));
+	}
+	ObjectInfo->ResolvedDynamicObjectReferences.Reset();
+
+	// Remove from attachment resolve
+	ObjectsWithAttachmentPendingResolve.Remove(ObjectIndex);
+}
+
+void FReplicationReader::BuildUnresolvedChangeMaskAndUpdateObjectReferenceTracking(const FResolveAndCollectUnresolvedAndResolvedReferenceCollector& Collector, FNetBitArrayView CollectorChangeMask, FReplicatedObjectInfo* ReplicationInfo, FNetBitArrayView& OutUnresolvedChangeMask)
+{
+	OutUnresolvedChangeMask.Reset();
+	bool bHasUnresolvedInitReferences = false;
+
+	FObjectReferenceTracker UnresolvedReferences;
+	for (const auto& RefInfo : Collector.GetUnresolvedReferences())
+	{
+		const FNetSerializerChangeMaskParam& ChangeMaskInfo = RefInfo.ChangeMaskInfo;
+		if (ChangeMaskInfo.BitCount)
+		{
+			OutUnresolvedChangeMask.SetBit(ChangeMaskInfo.BitOffset);
+		}
+		else
+		{
+			bHasUnresolvedInitReferences = true;
+		}
+
+		const uint32 BitOffset = (ChangeMaskInfo.BitCount > 0U ? ChangeMaskInfo.BitOffset : FakeInitChangeMaskOffset);
+		UnresolvedReferences.AddUnique(BitOffset, RefInfo.Reference.GetRefHandle());
+	}
+
+	FObjectReferenceTracker MappedDynamicReferences;
+	for (const auto& RefInfo : Collector.GetResolvedReferences())
+	{
+		if (RefInfo.Reference.GetRefHandle().IsDynamic())
+		{
+			const uint32 BitOffset = (RefInfo.ChangeMaskInfo.BitCount > 0U ? RefInfo.ChangeMaskInfo.BitOffset : FakeInitChangeMaskOffset);
+			MappedDynamicReferences.AddUnique(BitOffset, RefInfo.Reference.GetRefHandle());
+		}
+	}
+
+	// Update object specific 
+	UpdateObjectReferenceTracking(ReplicationInfo, CollectorChangeMask, Collector.IsInitStateIncluded(), UnresolvedReferences, MappedDynamicReferences);
+}
+
+void FReplicationReader::ResolveAndDispatchUnresolvedReferencesForObject(FNetSerializationContext& Context, uint32 InternalIndex)
+{
+	IRIS_PROFILER_SCOPE(FReplicationReader_ResolveAndDispatchUnresolvedReferencesForObject);
+
+	FReplicatedObjectInfo* ReplicationInfo = GetReplicatedObjectInfo(InternalIndex);
+
+	const bool bObjectHasAttachments = ReplicationInfo->bHasAttachments;
+	const bool bObjectHasReferences = ReplicationInfo->bHasUnresolvedInitialReferences | ReplicationInfo->bHasUnresolvedReferences;
+
+	ENetObjectAttachmentDispatchFlags AttachmentDispatchedFlags = ENetObjectAttachmentDispatchFlags::None;
+	if (bObjectHasReferences)
+	{
+		const FNetHandleManager::FReplicatedObjectData& ObjectData = NetHandleManager->GetReplicatedObjectDataNoCheck(ReplicationInfo->InternalIndex);
+		const uint32 ChangeMaskBitCount = ReplicationInfo->ChangeMaskBitCount;
+		
+		// Try to resolve references and collect unresolved references
+		const bool bOldHasUnresolvedInitReferences = ReplicationInfo->bHasUnresolvedInitialReferences;
+
+		FNetBitArrayView UnresolvedChangeMask = FChangeMaskUtil::MakeChangeMask(ReplicationInfo->UnresolvedChangeMaskOrPointer, ChangeMaskBitCount);
+	
+		// Need a temporary changemask for the unresolved changes due to UnresolvedChangeMask being written to by BuildUnresolvedChangeMaskAndUpdateObjectReferenceTracking
+		FChangeMaskStorageOrPointer TempChangeMaskOrPointer;
+		FNetBitArrayView TempChangeMask = FChangeMaskStorageOrPointer::AllocAndInitBitArray(TempChangeMaskOrPointer, ChangeMaskBitCount, TempChangeMaskAllocator);
+		FNetBitArrayView TempUnresolvedChangeMask = TempChangeMask;
+		TempUnresolvedChangeMask.Copy(UnresolvedChangeMask);
+
+		// Try to resolve references and collect resolved and references still pending resolve
+		FResolveAndCollectUnresolvedAndResolvedReferenceCollector Collector;
+		Collector.CollectReferences(*ObjectReferenceCache, ResolveContext, ReplicationInfo->bHasUnresolvedInitialReferences, &UnresolvedChangeMask, ObjectData.ReceiveStateBuffer, ObjectData.Protocol);
+
+		// Build UnresolvedChangeMask from collected data and update replication info
+		BuildUnresolvedChangeMaskAndUpdateObjectReferenceTracking(Collector, TempUnresolvedChangeMask, ReplicationInfo, UnresolvedChangeMask);
+
+		// Repurpose temp changemask for members that has resolved references.
+		TempChangeMask.Reset();
+		FNetBitArrayView ResolvedChangeMask = TempChangeMask;
+
+		bool bHasResolvedInitReferences = false;
+		for (const auto& RefInfo : Collector.GetResolvedReferences())
+		{
+			const FNetSerializerChangeMaskParam& ChangeMaskInfo = RefInfo.ChangeMaskInfo;
+			if (ChangeMaskInfo.BitCount)
+			{
+				ResolvedChangeMask.SetBit(ChangeMaskInfo.BitOffset);
+			}
+			else
+			{
+				// If we had old unresolved init dependencies we need to include the init state when we update references
+				bHasResolvedInitReferences = bOldHasUnresolvedInitReferences;
+			}
+		}
+
+		if (ObjectData.InstanceProtocol)
+		{
+			// Apply resolved references, this is a blunt tool as we currently push out full dirty properties rather than only the resolved references
+			if (ResolvedChangeMask.IsAnyBitSet() || bHasResolvedInitReferences)
+			{
+				if (bObjectHasAttachments && bExecuteReliableRPCsBeforeOnReps && !bHasResolvedInitReferences)
+				{
+					AttachmentDispatchedFlags = ENetObjectAttachmentDispatchFlags::Reliable;
+					ResolveAndDispatchAttachments(Context, ReplicationInfo, ENetObjectAttachmentDispatchFlags::Reliable);
+				}
+
+				Context.SetIsInitState(bHasResolvedInitReferences);
+
+				FDequantizeAndApplyParameters Params;
+				Params.Allocator = &TempLinearAllocator;
+				Params.ChangeMaskData = ResolvedChangeMask.GetData();
+				Params.UnresolvedReferencesChangeMaskData = ReplicationInfo->bHasUnresolvedReferences ? ReplicationInfo->UnresolvedChangeMaskOrPointer.GetPointer(ChangeMaskBitCount) : nullptr;
+				Params.InstanceProtocol = ObjectData.InstanceProtocol;
+				Params.Protocol = ObjectData.Protocol;
+				Params.SrcObjectStateBuffer = ObjectData.ReceiveStateBuffer;
+				Params.bHasUnresolvedInitReferences = ReplicationInfo->bHasUnresolvedInitialReferences;
+
+				FReplicationInstanceOperations::DequantizeAndApply(Context, Params);
+			}
+		}
+		else
+		{
+			// $IRIS: $TODO: Figure out how to handle this, currently we do not crash but we probably want to
+			// handle this properly by accumulating changemask for later instantiation
+			UE_LOG_REPLICATIONREADER_WARNING(TEXT("Cannot dispatch state data for not instantiated %s"), *ObjectData.Handle.ToString());
+		}
+	}
+
+	// Dispatch attachment and enqueue for later resolving
+	if (bObjectHasAttachments)
+	{
+		// If we haven't dispatched reliable attachments for this object then do so now in addition to unreliable attachments.
+		const ENetObjectAttachmentDispatchFlags AttachmentDispatchFlags = ENetObjectAttachmentDispatchFlags::Unreliable | (AttachmentDispatchedFlags ^ ENetObjectAttachmentDispatchFlags::Reliable);
+		ResolveAndDispatchAttachments(Context, ReplicationInfo, AttachmentDispatchFlags);
+	}
+}
+
+// Dispatch all data received for the frame, this includes trying to resolve object references
+void FReplicationReader::DispatchStateData(FNetSerializationContext& Context)
+{
+	IRIS_PROFILER_SCOPE(FReplicationReader::DispatchStateData);
+
+	// In order to execute PostNetRecv/PostRepNotifies after we have applied the actual state 
+	// we need to cache some information during dispatch and execute the logic in multiple passes
+	// Note: Currently all objects received in the packet are treated a a single batch
+	struct FPostDispatchObjectInfo
+	{
+		FReplicatedObjectInfo* ReplicationInfo;
+		FDispatchObjectInfo* Info;
+		FDequantizeAndApplyHelper::FContext* DequantizeAndApplyContext;
+		ENetObjectAttachmentDispatchFlags AttachmentDispatchedFlags;		
+	};
+
+	// Allocate temporary space for post dispatch
+	FPostDispatchObjectInfo* PostDispatchObjectInfos = new (TempLinearAllocator) FPostDispatchObjectInfo[ObjectsToDispatchCount];
+	uint32 NumObjectsPendingPostDistpatch = 0U;
+
+	// Function to flush all objects pending post dispatch
+	auto&& FlushPostDispatchForBatch = [PostDispatchObjectInfos, &NumObjectsPendingPostDistpatch, this, &Context]()
+	{
+		// When all received states have been applied we invoke PostReplicate and RepNotifies
+		for (FPostDispatchObjectInfo& PostDispatchObjectInfo : MakeArrayView(PostDispatchObjectInfos, NumObjectsPendingPostDistpatch))
+		{
+			FDispatchObjectInfo& Info = *PostDispatchObjectInfo.Info;
+
+			// If we are running in backwards compatibility mode, execute Reliable RPC`s before RepNotify callbacks
+			if (Info.bHasAttachments && bExecuteReliableRPCsBeforeOnReps && !Info.bIsInitialState)
+			{
+				PostDispatchObjectInfo.AttachmentDispatchedFlags |= ENetObjectAttachmentDispatchFlags::Reliable;
+				ResolveAndDispatchAttachments(Context, PostDispatchObjectInfo.ReplicationInfo, ENetObjectAttachmentDispatchFlags::Reliable);
+			}
+
+			// Execute legacy post replicate functions
+			if (Info.bHasState && PostDispatchObjectInfo.DequantizeAndApplyContext)
+			{
+				Context.SetIsInitState(Info.bIsInitialState);
+				FDequantizeAndApplyHelper::CallLegacyPostApplyFunctions(PostDispatchObjectInfo.DequantizeAndApplyContext, Context);
+			}
+		}
+
+		// In the last pass, RPC`s and cleanup cached data
+		for (FPostDispatchObjectInfo& PostDispatchObjectInfo : MakeArrayView(PostDispatchObjectInfos, NumObjectsPendingPostDistpatch))
+		{
+			FDispatchObjectInfo& Info = *PostDispatchObjectInfo.Info;
+
+			// If object was created this frame it`s initial state is now applied
+			if (Info.bHasState && Info.bIsInitialState)
+			{
+				const FNetHandleManager::FReplicatedObjectData& ObjectData = NetHandleManager->GetReplicatedObjectDataNoCheck(Info.InternalIndex);
+				ReplicationBridge->CallPostApplyInitialState(ObjectData.Handle);
+			}
+
+			// Dispatch attachment and enqueue for later resolving
+			if (Info.bHasAttachments)
+			{
+				// If we haven't dispatched reliable attachments for this object then do so now in addition to unreliable attachments.
+				const ENetObjectAttachmentDispatchFlags AttachmentDispatchFlags = ENetObjectAttachmentDispatchFlags::Unreliable | (PostDispatchObjectInfo.AttachmentDispatchedFlags ^ ENetObjectAttachmentDispatchFlags::Reliable);
+				ResolveAndDispatchAttachments(Context, PostDispatchObjectInfo.ReplicationInfo, AttachmentDispatchFlags);
+			}
+
+			// Cleanup temporary state data
+			if (PostDispatchObjectInfo.DequantizeAndApplyContext)
+			{
+				FDequantizeAndApplyHelper::Deinitialize(PostDispatchObjectInfo.DequantizeAndApplyContext);
+			}
+		}
+
+		NumObjectsPendingPostDistpatch = 0U;
+	};
+
+	// In order to properly execute legacy callbacks we need to batch apply state data for owner/subobjects
+	FInternalNetHandle LastDispatchedRootInternalIndex = 0U;
+	
+	// Dispatch and apply received state data
+	for (FDispatchObjectInfo& Info : MakeArrayView(ObjectsToDispatch, ObjectsToDispatchCount))
+	{
+		FReplicatedObjectInfo* ReplicationInfo = GetReplicatedObjectInfo(Info.InternalIndex);
+
+		FPostDispatchObjectInfo PostDispatchObjectInfo;
+		PostDispatchObjectInfo.ReplicationInfo = ReplicationInfo;
+		PostDispatchObjectInfo.Info = &Info;
+		PostDispatchObjectInfo.DequantizeAndApplyContext = nullptr;
+		PostDispatchObjectInfo.AttachmentDispatchedFlags = ENetObjectAttachmentDispatchFlags::None;
+
+		// If we have any object references we want to update any unresolved ones, including previously unresolved references
+		if (Info.bHasState)
+		{
+			const FNetHandleManager::FReplicatedObjectData& ObjectData = NetHandleManager->GetReplicatedObjectDataNoCheck(Info.InternalIndex);
+
+			// We only need to flush if we are switching to a new root object with state data
+			const FInternalNetHandle RootInternalIndex = ObjectData.SubObjectRootIndex == FNetHandleManager::InvalidInternalIndex ? Info.InternalIndex : ObjectData.SubObjectRootIndex;
+			if (RootInternalIndex != LastDispatchedRootInternalIndex)
+			{
+				FlushPostDispatchForBatch();
+				LastDispatchedRootInternalIndex = RootInternalIndex;
+			}
+
+			const uint32 ChangeMaskBitCount = ReplicationInfo->ChangeMaskBitCount;
+
+			FNetBitArrayView ChangeMask = FChangeMaskUtil::MakeChangeMask(Info.ChangeMaskOrPointer, ChangeMaskBitCount);
+			// If we have pending unresolved changes we include them as well 
+			FNetBitArrayView UnresolvedChangeMask = FChangeMaskUtil::MakeChangeMask(ReplicationInfo->UnresolvedChangeMaskOrPointer, ChangeMaskBitCount);
+
+			if (ReplicationInfo->bHasUnresolvedReferences)
+			{
+				ChangeMask.Combine(UnresolvedChangeMask, FNetBitArrayView::OrOp);
+			}
+
+			// Collect all unresolvable references, including old pending references
+			FResolveAndCollectUnresolvedAndResolvedReferenceCollector Collector;
+			Collector.CollectReferences(*ObjectReferenceCache, ResolveContext, Info.bIsInitialState | ReplicationInfo->bHasUnresolvedInitialReferences, &ChangeMask, ObjectData.ReceiveStateBuffer, ObjectData.Protocol);
+
+			// If we have any object references we need to track them
+			if (Collector.GetUnresolvedReferences().Num() > 0 || Collector.GetResolvedReferences().Num() > 0)
+			{
+				BuildUnresolvedChangeMaskAndUpdateObjectReferenceTracking(Collector, ChangeMask, ReplicationInfo, UnresolvedChangeMask);
+			
+				// $IRIS: $TODO: For now we always apply, even if we cannot resolve all references for a property
+				// We need to investigate how this is handled best as we do not want to prevent arrays(fastarrays) from applying data just because a single element wont resolve?
+				// Mask off any unresolvable states before we dispatch state data
+				//ChangeMask.Combine(UnresolvedChangeMask, FNetBitArrayView::AndNotOp);
+			}
+
+			// Apply state data
+			if (ObjectData.InstanceProtocol)
+			{
+				Context.SetIsInitState(Info.bIsInitialState);
+
+				FDequantizeAndApplyParameters Params;
+				Params.Allocator = &TempLinearAllocator;
+				Params.ChangeMaskData = Info.ChangeMaskOrPointer.GetPointer(ChangeMaskBitCount);
+				Params.UnresolvedReferencesChangeMaskData = ReplicationInfo->bHasUnresolvedReferences ? ReplicationInfo->UnresolvedChangeMaskOrPointer.GetPointer(ChangeMaskBitCount) : nullptr;
+				Params.InstanceProtocol = ObjectData.InstanceProtocol;
+				Params.Protocol = ObjectData.Protocol;
+				Params.SrcObjectStateBuffer = ObjectData.ReceiveStateBuffer;
+				Params.bHasUnresolvedInitReferences = ReplicationInfo->bHasUnresolvedInitialReferences;
+
+				// Dequantize state data, call PreReplicate and apply received state
+				PostDispatchObjectInfo.DequantizeAndApplyContext = FDequantizeAndApplyHelper::Initialize(Context, Params);
+				FDequantizeAndApplyHelper::ApplyAndCallLegacyPreApplyFunction(PostDispatchObjectInfo.DequantizeAndApplyContext, Context);
+			}
+			else
+			{
+				// $IRIS: $TODO: Figure out how to handle this, currently we do not crash but we probably want to
+				// handle this properly by accumulating changemask for later instantiation
+				UE_LOG_REPLICATIONREADER_WARNING(TEXT("Cannot dispatch state data for not instantiated %s"), *(ObjectData.Handle.ToString()));
+			}
+		}
+
+		// Add to post dispatch
+		PostDispatchObjectInfos[NumObjectsPendingPostDistpatch++] = PostDispatchObjectInfo;
+	}
+
+	FlushPostDispatchForBatch();
+
+	ObjectsToDispatchCapacity = 0U;
+}
+
+void FReplicationReader::ResolveAndDispatchUnresolvedReferences()
+{
+	IRIS_PROFILER_SCOPE(FReplicationReader_ResolveAndDispatchUnresolvedReferences);
+
+	// Setup context for dispatch
+	FNetSerializationContext Context;
+	FInternalNetSerializationContext InternalContext(Parameters.ReplicationSystem, ResolveContext.RemoteNetTokenStoreState);
+	Context.SetLocalConnectionId(Parameters.ConnectionId);
+	Context.SetInternalContext(&InternalContext);
+
+	// Currently we brute force this by iterating over all handles pending resolve and update all objects pending resolve
+	TArray<FNetHandle> UpdatedHandles;
+	UpdatedHandles.Reserve(128);
+	UnresolvedHandleToDependents.GenerateKeyArray(UpdatedHandles);
+	
+	TSet<uint32> InternalObjectsToResolve;
+	InternalObjectsToResolve.Reserve(UnresolvedHandleToDependents.Num());
+
+	for (FNetHandle Handle : UpdatedHandles)
+	{
+		// Only make sense to update dependant objects if handle is resolvable
+		if (ObjectReferenceCache->ResolveObjectReferenceHandle(Handle, ResolveContext) != nullptr)
+		{
+			for (auto It = UnresolvedHandleToDependents.CreateConstKeyIterator(Handle); It; ++It)
+			{
+				InternalObjectsToResolve.Add(It.Value());
+			}
+		}
+	}
+
+	// Add in any handles with pending attachments to resolve
+	InternalObjectsToResolve.Append(ObjectsWithAttachmentPendingResolve);
+
+	// Try to resolve objects with updated references
+	for (uint32 InternalIndex : InternalObjectsToResolve)
+	{
+		ResolveAndDispatchUnresolvedReferencesForObject(Context, InternalIndex);
+	}
+
+	if (NumHandlesPendingResolveLastUpdate != UpdatedHandles.Num() || ObjectsWithAttachmentPendingResolve.Num() > 0)
+	{
+		UE_LOG_REPLICATIONREADER(TEXT("FReplicationReader::ResolveAndDispatchUnresolvedReferences NetHandles pending: %u Attachments pending: %u)"), UpdatedHandles.Num(), ObjectsWithAttachmentPendingResolve.Num());
+		NumHandlesPendingResolveLastUpdate = UpdatedHandles.Num();
+	}
+}
+
+
+void FReplicationReader::UpdateUnresolvableReferenceTracking()
+{
+	constexpr uint32 AssumedMaxDependentCount = 256;
+	TArray<uint32, TInlineAllocator<AssumedMaxDependentCount>> Dependents;
+
+	// Naively go through every object pending destroy, see if it's dynamic and update dependent's unresolved tracking
+	TArrayView<const FInternalNetHandle> ObjectsPendingDestroy = NetHandleManager->GetObjectsPendingDestroy();
+	for (const FInternalNetHandle InternalIndex : ObjectsPendingDestroy)
+	{
+		const FNetHandleManager::FReplicatedObjectData& ObjectData = NetHandleManager->GetReplicatedObjectDataNoCheck(InternalIndex);
+		const FNetHandle DestroyedHandle = ObjectData.Handle;
+		if (!DestroyedHandle.IsDynamic())
+		{
+			continue;
+		}
+
+		// For torn off objects we want to remove both resolved and unresolved references to it as it will never be replicated again.
+		if (ObjectData.bTearOff)
+		{
+			constexpr bool bMaintainOrder = false;
+			Dependents.Reset();
+			UnresolvedHandleToDependents.MultiFind(DestroyedHandle, Dependents, bMaintainOrder);
+			UnresolvedHandleToDependents.Remove(DestroyedHandle);
+			for (const uint32 DependentObjectIndex : Dependents)
+			{
+				FReplicatedObjectInfo* ReplicationInfo = GetReplicatedObjectInfo(DependentObjectIndex);
+				RemoveUnresolvedObjectReferenceInReplicationInfo(ReplicationInfo, DestroyedHandle);
+			}
+		}
+
+		// For any previously resolved handles make sure to move them to unresolved status.
+		constexpr bool bMaintainOrder = false;
+		Dependents.Reset();
+		ResolvedDynamicHandleToDependents.MultiFind(DestroyedHandle, Dependents, bMaintainOrder);
+		if (Dependents.Num() > 0)
+		{
+			ResolvedDynamicHandleToDependents.Remove(DestroyedHandle);
+			// Torn off objects will get new handles if replicated again so they can never be remapped.
+			if (ObjectData.bTearOff)
+			{
+				for (const uint32 DependentObjectIndex : Dependents)
+				{
+					FReplicatedObjectInfo* ReplicationInfo = GetReplicatedObjectInfo(DependentObjectIndex);
+					RemoveResolvedObjectReferenceInReplicationInfo(ReplicationInfo, DestroyedHandle);
+				}
+			}
+			else
+			{
+				for (const uint32 DependentObjectIndex : Dependents)
+				{
+					FReplicatedObjectInfo* ReplicationInfo = GetReplicatedObjectInfo(DependentObjectIndex);
+					if (MoveResolvedObjectReferenceToUnresolvedInReplicationInfo(ReplicationInfo, DestroyedHandle))
+					{
+						UnresolvedHandleToDependents.Add(DestroyedHandle, DependentObjectIndex);
+					}
+				}
+			}
+		}
+	}
+}
+
+void FReplicationReader::DispatchEndReplication(FNetSerializationContext& Context)
+{
+	for (FDispatchObjectInfo& Info : MakeArrayView(ObjectsToDispatch, ObjectsToDispatchCount))
+	{
+		if (Info.bDeferredEndReplication)
+		{
+			// Detach destroy object
+			EndReplication(Info.InternalIndex, Info.bTearOff, Info.bDestroy);
+		}
+	}
+}
+
+void FReplicationReader::ReadObjects(FNetSerializationContext& Context, uint32 ObjectCountToRead)
+{
+	IRIS_PROFILER_SCOPE(ReplicationReader_ReadObjects);
+
+	FNetBitStreamReader& Reader = *Context.GetBitStreamReader();
+	
+	while (ObjectCountToRead && !Context.HasErrorOrOverflow())
+	{
+		ReadObject(Context);
+		--ObjectCountToRead;
+	}
+
+	ensureAlwaysMsgf(!Context.HasErrorOrOverflow(), TEXT("Overflow: %c Error: %s Bit stream bits left: %u position: %u"), "YN"[Context.HasError()], ToCStr(Context.GetError().ToString()), Reader.GetBitsLeft(), Reader.GetPosBits());
+}
+
+void FReplicationReader::ProcessHugeObjectAttachment(FNetSerializationContext& Context, const TRefCountPtr<FNetBlob>& Attachment)
+{
+	if (Attachment->GetCreationInfo().Type != NetObjectBlobType)
+	{
+		Context.SetError(GNetError_UnsupportedNetBlob);
+		return;
+	}
+
+	FNetTraceCollector* HugeObjectTraceCollector = nullptr;
+#if UE_NET_TRACE_ENABLED
+	FNetTraceCollector HugeObjectTraceCollectorOnStack;
+	HugeObjectTraceCollector = &HugeObjectTraceCollectorOnStack;
+#endif
+
+	const FNetObjectBlob& NetObjectBlob = *static_cast<FNetObjectBlob*>(Attachment.GetReference());
+
+	FNetBitStreamReader HugeObjectReader;
+	HugeObjectReader.InitBits(NetObjectBlob.GetRawData().GetData(), NetObjectBlob.GetRawDataBitCount());
+	FNetSerializationContext HugeObjectSerializationContext = Context.MakeSubContext(&HugeObjectReader);
+	HugeObjectSerializationContext.SetNetTraceCollector(HugeObjectTraceCollector);
+
+	UE_NET_TRACE_NAMED_SCOPE(HugeObjectTraceScope, HugeObjectState, HugeObjectReader, HugeObjectTraceCollector, ENetTraceVerbosity::Trace);
+
+	// Find out how many objects to read so we can reserve object dispatch infos.
+	FNetObjectBlob::FHeader HugeObjectHeader = {};
+	FNetObjectBlob::DeserializeHeader(HugeObjectSerializationContext, HugeObjectHeader);
+	if (HugeObjectSerializationContext.HasErrorOrOverflow() || HugeObjectHeader.ObjectCount < 1U)
+	{
+		if (!Context.HasError())
+		{
+			Context.SetError(GNetError_BitStreamError);
+			return;
+		}
+	}
+
+	if (ObjectsToDispatchCapacity < (ObjectsToDispatchCount + HugeObjectHeader.ObjectCount))
+	{
+		// Need to reallocate ObjectDispatchInfos and copy the old data. We allocate a bit extra just in case.
+		ObjectsToDispatchCapacity = ObjectsToDispatchCount + HugeObjectHeader.ObjectCount + ObjectsToDispatchSlackCount;
+		FDispatchObjectInfo* NewObjectsToDispatch = new (TempLinearAllocator) FDispatchObjectInfo[ObjectsToDispatchCapacity];
+		FPlatformMemory::Memcpy(NewObjectsToDispatch, ObjectsToDispatch, ObjectsToDispatchCount*sizeof(FDispatchObjectInfo));
+		ObjectsToDispatch = NewObjectsToDispatch;
+	}
+
+	ReadObjects(HugeObjectSerializationContext, HugeObjectHeader.ObjectCount);
+	if (HugeObjectSerializationContext.HasErrorOrOverflow())
+	{
+		Context.SetError(GNetError_BitStreamError);
+		return;
+	}
+
+#if UE_NET_TRACE_ENABLED
+	UE_NET_TRACE_EXIT_NAMED_SCOPE(HugeObjectTraceScope);
+
+	// Append huge object state at end of stream.
+	if (FNetTraceCollector* TraceCollector = Context.GetTraceCollector())
+	{
+		FNetBitStreamReader& Reader = *Context.GetBitStreamReader();
+		// Need a valid fake bitstream position, a position that starts before the end of the packet.
+		FNetTrace::FoldTraceCollector(TraceCollector, HugeObjectTraceCollector, GetBitStreamPositionForNetTrace(Reader) - 1U);
+	}
+#endif
+}
+
+void FReplicationReader::ProcessHugeObject(FNetSerializationContext& Context)
+{
+	if (!Attachments.HasUnprocessedAttachments(ENetObjectAttachmentType::HugeObject, ObjectIndexForOOBAttachment))
+	{
+		return;
+	}
+
+	FNetObjectAttachmentReceiveQueue* AttachmentQueue = Attachments.GetQueue(ENetObjectAttachmentType::HugeObject, ObjectIndexForOOBAttachment);
+	while (const TRefCountPtr<FNetBlob>* Attachment = AttachmentQueue->PeekReliable())
+	{
+		ProcessHugeObjectAttachment(Context, *Attachment);
+		AttachmentQueue->PopReliable();
+		if (Context.HasError())
+		{
+			return;
+		}
+	}
+	while (const TRefCountPtr<FNetBlob>* Attachment = AttachmentQueue->PeekUnreliable())
+	{
+		ProcessHugeObjectAttachment(Context, *Attachment);
+		AttachmentQueue->PopUnreliable();
+		if (Context.HasError())
+		{
+			return;
+		}
+	}
+}
+
+void FReplicationReader::Read(FNetSerializationContext& Context)
+{
+	FNetBitStreamReader& Reader = *Context.GetBitStreamReader();
+
+	// Setup internal context
+	FInternalNetSerializationContext InternalContext(Parameters.ReplicationSystem, ResolveContext.RemoteNetTokenStoreState);
+	Context.SetLocalConnectionId(Parameters.ConnectionId);
+	Context.SetInternalContext(&InternalContext);
+	Context.SetNetBlobReceiver(&ReplicationSystemInternal->GetNetBlobHandlerManager());
+
+	UE_NET_TRACE_SCOPE(ReplicationData, *Context.GetBitStreamReader(), Context.GetTraceCollector(), ENetTraceVerbosity::Trace);
+
+	FMemMark TempAllocatorScope(TempLinearAllocator);
+
+	// Sanity check received obejct count
+	const uint32 MaxObjectCountToRead = 8192U;
+	const uint32 ReceivedObjectCountToRead = Reader.ReadBits(16);
+	uint32 ObjectCountToRead = ReceivedObjectCountToRead;
+
+	if (Reader.IsOverflown() || ObjectCountToRead >= MaxObjectCountToRead)
+	{
+		const FName& NetError = (Reader.IsOverflown() ? GNetError_BitStreamOverflow : GNetError_BitStreamError);
+		Context.SetError(NetError);
+
+		return;
+	}
+
+	if (ObjectCountToRead == 0)
+	{
+		return;
+	}
+
+	// Allocate tracking info for objects we receive this packet
+	ObjectsToDispatchCapacity = ObjectCountToRead + ObjectsToDispatchSlackCount;
+	ObjectsToDispatchCount = 0U;
+	ObjectsToDispatch = new (TempLinearAllocator) FDispatchObjectInfo[ObjectsToDispatchCapacity];
+
+	uint32 DestroyedObjectCount = ReadObjectsPendingDestroy(Context);
+
+	ObjectCountToRead -= DestroyedObjectCount;
+
+	// Nothing more to do or we failed and should disconnect
+	if (Reader.IsOverflown() || (ObjectCountToRead == 0 && ObjectsToDispatchCount == 0))
+	{
+		return;
+	}
+
+	ReadObjects(Context, ObjectCountToRead);
+	if (Context.HasErrorOrOverflow())
+	{
+		return;
+	}
+
+	// Assemble and deserialize huge object if present
+	ProcessHugeObject(Context);
+	if (Context.HasErrorOrOverflow())
+	{
+		return;
+	}
+
+	// Stats
+	UE_NET_TRACE_PACKET_STATSCOUNTER(Parameters.ReplicationSystem->GetId(), Parameters.ConnectionId, ReplicationReader.ReadObjectCOunt, ReceivedObjectCountToRead, ENetTraceVerbosity::Trace);
+	UE_NET_TRACE_PACKET_STATSCOUNTER(Parameters.ReplicationSystem->GetId(), Parameters.ConnectionId, ReplicationReader.ObjectsToDispatchCount, ObjectsToDispatchCount, ENetTraceVerbosity::Trace);
+
+	// Apply received data and resolve dependencies
+	DispatchStateData(Context);
+	if (Context.HasErrorOrOverflow())
+	{
+		return;
+	}
+
+	// Resolve
+	ResolveAndDispatchUnresolvedReferences();
+
+	// EndReplication for all objects that should no longer replicate
+	DispatchEndReplication(Context);
+
+	// Drop temporary data
+	ObjectsToDispatch = nullptr;
+	ObjectsToDispatchCount = 0U;
+}
+
+void FReplicationReader::ResolveAndDispatchAttachments(FNetSerializationContext& Context, FReplicatedObjectInfo* ReplicationInfo, ENetObjectAttachmentDispatchFlags DispatchFlags)
+{
+	if (DispatchFlags == ENetObjectAttachmentDispatchFlags::None)
+	{
+		return;
+	}
+
+	// Cache configurables before processing attachments
+	const bool bDispatchReliableAttachments = EnumHasAnyFlags(DispatchFlags, ENetObjectAttachmentDispatchFlags::Reliable);
+	const bool bDispatchUnreliableAttachments = EnumHasAnyFlags(DispatchFlags, ENetObjectAttachmentDispatchFlags::Unreliable);
+	const bool bCanDelayAttachments = Parameters.bAllowDelayingAttachmentsWithUnresolvedReferences && (DelayAttachmentsWithUnresolvedReferences != nullptr && DelayAttachmentsWithUnresolvedReferences->GetInt() > 0);
+	const uint32 InternalIndex = ReplicationInfo->InternalIndex;
+
+	/**
+	* This code path handles all cases where the initial state has already been applied. An object can have multiple entries in ObjectsPendingResolve.
+	* Reliable attachments will be delivered if they can be resolved or if CVarDelayUnmappedRPCs is <= 0
+	*/
+	bool bHasUnresolvedReferences = false;
+	const ENetObjectAttachmentType AttachmentType = (IsObjectIndexForOOBAttachment(InternalIndex) ? ENetObjectAttachmentType::OutOfBand : ENetObjectAttachmentType::Normal);
+	if (FNetObjectAttachmentReceiveQueue* AttachmentQueue = Attachments.GetQueue(AttachmentType, InternalIndex))
+	{
+		if (bDispatchReliableAttachments)
+		{
+			while (const TRefCountPtr<FNetBlob>* Attachment = AttachmentQueue->PeekReliable())
+			{
+				// Delay attachments with unresolved references
+				if (bCanDelayAttachments)
+				{
+					// Only block reliable stream if we have unresolved references that must be mapped
+					const ENetObjectReferenceResolveResult ResolveResult = Attachment->GetReference()->ResolveObjectReferences(Context);
+					if (EnumHasAnyFlags(ResolveResult, ENetObjectReferenceResolveResult::HasUnresolvedMustBeMappedReferences))
+					{
+						const FReplicationStateDescriptor* Descriptor = Attachment->GetReference()->GetReplicationStateDescriptor();
+
+						UE_LOG(LogIris, Warning, TEXT("Unable to resolve references in %s for InternalIndex %u"), (Descriptor != nullptr ? ToCStr(Descriptor->DebugName) : TEXT("N/A")), InternalIndex);
+						break;
+					}
+				}
+
+				NetBlobHandlerManager->OnNetBlobReceived(Context, *reinterpret_cast<const TRefCountPtr<FNetBlob>*>(Attachment));
+				AttachmentQueue->PopReliable();
+
+				if (Context.HasError())
+				{
+					return;
+				}
+			}
+		}
+
+		if (bDispatchUnreliableAttachments)
+		{
+			while (const TRefCountPtr<FNetBlob>* Attachment = AttachmentQueue->PeekUnreliable())
+			{
+				NetBlobHandlerManager->OnNetBlobReceived(Context, *Attachment);
+				AttachmentQueue->PopUnreliable();
+
+				if (Context.HasError())
+				{
+					return;
+				}
+			}
+		}
+		
+		if (AttachmentQueue->IsSafeToDestroy())
+		{
+			// N.B. AttachmentQueue is no longer valid after this call
+			Attachments.DropAllAttachments(AttachmentType, InternalIndex);
+		}
+		else
+		{
+			bHasUnresolvedReferences = AttachmentQueue->HasUnprocessed();
+		}
+	}
+	else
+	{
+		// Should not get here, if we do something is out of sync and we should disconnect
+		Context.SetError(NetError_FailedToFindAttachmentQueue);
+		ensure(AttachmentType == ENetObjectAttachmentType::OutOfBand);
+	}
+
+	// Update tracking of objects with attachments pending resolve
+	if (bHasUnresolvedReferences && !ReplicationInfo->bHasAttachments)
+	{
+		ObjectsWithAttachmentPendingResolve.Add(InternalIndex);
+	}
+	else if (!bHasUnresolvedReferences && ReplicationInfo->bHasAttachments)
+	{
+		ObjectsWithAttachmentPendingResolve.RemoveSwap(InternalIndex);
+	}
+	ReplicationInfo->bHasAttachments = bHasUnresolvedReferences;
+}
+
+void FReplicationReader::SetRemoteNetTokenStoreState(FNetTokenStoreState* InRemoteTokenStoreState)
+{
+	ResolveContext.RemoteNetTokenStoreState = InRemoteTokenStoreState;
+}
+
+}
diff --git a/Engine/Source/Runtime/Experimental/Iris/Core/Private/Iris/ReplicationSystem/ReplicationReader.h b/Engine/Source/Runtime/Experimental/Iris/Core/Private/Iris/ReplicationSystem/ReplicationReader.h
new file mode 100644
--- /dev/null
+++ b/Engine/Source/Runtime/Experimental/Iris/Core/Private/Iris/ReplicationSystem/ReplicationReader.h
@@ -0,0 +1,233 @@
+// Copyright Epic Games, Inc. All Rights Reserved.
+
+#pragma once
+
+// This class will never be included from public headers
+#include "Iris/ReplicationSystem/AttachmentReplication.h"
+#include "Iris/ReplicationSystem/ChangeMaskUtil.h"
+#include "Iris/ReplicationSystem/NetHandleManager.h"
+#include "Iris/ReplicationSystem/ObjectReferenceCache.h"
+#include "Iris/ReplicationSystem/ReplicationTypes.h"
+#include "Containers/Map.h"
+#include "Misc/MemStack.h"
+
+// Forward declaration
+class UReplicationBridge;
+namespace UE::Net
+{
+	class FNetSerializationContext;
+	class FNetTokenStoreState;
+	class FReplicationStateStorage;
+	namespace Private
+	{
+		class FReplicationSystemInternal;
+		class FResolveAndCollectUnresolvedAndResolvedReferenceCollector;
+		class FNetBlobHandlerManager;
+		class FNetHandleManager;
+	}
+}
+
+namespace UE::Net::Private
+{
+
+enum class ENetObjectAttachmentDispatchFlags : uint32
+{
+	None = 0,
+	Reliable = 1U << 0U,
+	Unreliable = Reliable << 1U,
+};
+ENUM_CLASS_FLAGS(ENetObjectAttachmentDispatchFlags);
+
+/** Deals with Incoming replication Data and dispatch */
+class FReplicationReader
+{
+public:
+	FReplicationReader();
+	~FReplicationReader();
+
+	// Init
+	void Init(const FReplicationParameters& Parameters);
+	void Deinit();
+
+	// Read incoming replication data
+	void Read(FNetSerializationContext& Context);
+
+	// Read index part of handle
+	FNetHandle ReadNetHandleId(FNetBitStreamReader& Reader) const;
+
+	void SetRemoteNetTokenStoreState(FNetTokenStoreState* RemoteTokenStoreState);
+	
+	// Mark objects pending destroy as unresolvable.
+	void UpdateUnresolvableReferenceTracking();
+
+private:
+	// ChangeMaskOffset -> FNetHandle
+	enum EConstants : uint32
+	{
+		FakeInitChangeMaskOffset = 0xFFFFFFFFU,
+	};
+	typedef TMultiMap<uint32, FNetHandle> FObjectReferenceTracker;
+
+	struct FReplicatedObjectInfo
+	{
+		FReplicatedObjectInfo();
+
+		// We accumulate unresolved changes in this changemask
+		FChangeMaskStorageOrPointer UnresolvedChangeMaskOrPointer;
+
+		/* In order to be able to do partial updates the changemask bit is stored with the reference.
+		 * That also means a reference can have many entries, but at most one per changemask bit. */
+		FObjectReferenceTracker UnresolvedObjectReferences;
+		FObjectReferenceTracker ResolvedDynamicObjectReferences;
+
+		// Baselines
+		uint8* StoredBaselines[2];
+
+		uint32 InternalIndex;							// InternalIndex
+		union
+		{
+			uint32 Value;
+			struct	 
+			{
+				uint32 ChangeMaskBitCount : 16;					// This is cached to avoid having to look it up in the protocol		
+				uint32 bHasUnresolvedReferences : 1;			// Do we have unresolved references in the changemask?
+				uint32 bHasUnresolvedInitialReferences : 1;		// Do we have unresolved initial only references
+				uint32 bHasAttachments : 1;
+				uint32 bDestroy : 1;
+				uint32 bTearOff : 1;
+				uint32 bIsDeltaCompressionEnabled : 1;
+				uint32 LastStoredBaselineIndex : 2;				// Last stored baseline, as soon as we receive data compressed against the baseline we know that we can release older baselines
+				uint32 PrevStoredBaselineIndex : 2;				// Previous stored baselines index
+				uint32 Padding : 7;
+			};
+		};
+	};
+
+	// Temporary Data to dispatch
+	struct FDispatchObjectInfo
+	{
+		uint32 InternalIndex;
+		FChangeMaskStorageOrPointer ChangeMaskOrPointer;
+		uint32 bIsInitialState : 1;
+		uint32 bHasState : 1;
+		uint32 bHasAttachments : 1;
+		uint32 bDestroy : 1;
+		uint32 bTearOff : 1;
+		uint32 bDeferredEndReplication : 1;
+	};
+
+	enum : uint32
+	{
+		ObjectIndexForOOBAttachment = 0U,
+		// Try to avoid reallocations for dispatch in the case we need to process a huge object
+		ObjectsToDispatchSlackCount = 16U,
+	};
+
+	bool IsObjectIndexForOOBAttachment(uint32 InternalIndex) const { return InternalIndex == ObjectIndexForOOBAttachment; }
+
+	// Read index reference to replicated object
+	uint32 ReadInternalIndex(FNetBitStreamReader& Reader) const;
+
+	// Read a new or updated object
+	void ReadObject(FNetSerializationContext& Context);
+
+	// Read all objects pending destroy
+	uint16 ReadObjectsPendingDestroy(FNetSerializationContext& Context);
+
+	// Read state data for all incoming objects
+	void ReadObjects(FNetSerializationContext& Context, uint32 ObjectCountToRead);
+
+	// Process a single huge object attachment
+	void ProcessHugeObjectAttachment(FNetSerializationContext& Context, const TRefCountPtr<FNetBlob>& Attachment);
+
+	// Assemble and deserialize huge object if present
+	void ProcessHugeObject(FNetSerializationContext& Context);
+
+	// Resolve and dispatch unresolved references
+	void ResolveAndDispatchUnresolvedReferencesForObject(FNetSerializationContext& Context, uint32 InternalIndex);
+
+	// Resolve and dispatch unresolved references
+	void ResolveAndDispatchUnresolvedReferences();
+
+	// Dispatch all data received for the frame, this includes trying to resolve object references
+	void DispatchStateData(FNetSerializationContext& Context);
+
+	// Dispatch resolved attachments
+	void ResolveAndDispatchAttachments(FNetSerializationContext& Context, FReplicatedObjectInfo* ReplicationInfo, ENetObjectAttachmentDispatchFlags DispatchFlags);
+
+	// End replication for all objects that the server has told us to destroy or tear off
+	void DispatchEndReplication(FNetSerializationContext& Context);
+
+	// Create tracking info for the object with the given InternalInfo
+	FReplicatedObjectInfo& StartReplication(uint32 InternalIndex);
+
+	// Remove tracking info for the object with InternalIndex
+	void EndReplication(uint32 InternalIndex, bool bTearOff, bool bDestroyInstance);
+
+	// Free any data allocated per object
+	void CleanupObjectData(FReplicatedObjectInfo& ObjectInfo);
+
+	// Lookup the tracking info for the object with IntnernalIndex
+	FReplicatedObjectInfo* GetReplicatedObjectInfo(uint32 InternalIndex);
+
+	// Update reference tracking maps for the current object
+	void UpdateObjectReferenceTracking(FReplicatedObjectInfo* ReplicationInfo, FNetBitArrayView ChangeMask, bool bIncludeInitState, const FObjectReferenceTracker& NewUnresolvedReferences, const FObjectReferenceTracker& NewMappedDynamicReferences);
+
+	// Remove all references for object
+	void CleanupReferenceTracking(FReplicatedObjectInfo* ObjectInfo);
+
+	// Update ReplicationInfo and OutUnresolvedChangeMask based on data collected by the Collector
+	void BuildUnresolvedChangeMaskAndUpdateObjectReferenceTracking(const FResolveAndCollectUnresolvedAndResolvedReferenceCollector& Collector, FNetBitArrayView CollectorChangeMask, FReplicatedObjectInfo* ReplicationInfo, FNetBitArrayView& OutUnresolvedChangeMask);
+
+	void RemoveUnresolvedObjectReferenceInReplicationInfo(FReplicatedObjectInfo* ReplicationInfo, FNetHandle Handle);
+	void RemoveResolvedObjectReferenceInReplicationInfo(FReplicatedObjectInfo* ReplicationInfo, FNetHandle Handle);
+
+	// A previously resolved dynamic reference is now unresolvable. The ReplicationInfo needs to be updated to reflect this.
+	// Returns true if the reference was found.
+	bool MoveResolvedObjectReferenceToUnresolvedInReplicationInfo(FReplicatedObjectInfo* ReplicationInfo, FNetHandle UnresolvableHandle);
+
+	void DeserializeObjectStateDelta(FNetSerializationContext& Context, uint32 InternalIndex, FDispatchObjectInfo& Info, FReplicatedObjectInfo& ObjectInfo, const FNetHandleManager::FReplicatedObjectData& ObjectData, uint32& OutNewBaselineIndex);
+
+private:
+	FReplicationParameters Parameters;
+
+	FMemStackBase TempLinearAllocator;
+	FMemStackChangeMaskAllocator TempChangeMaskAllocator;
+
+	FGlobalChangeMaskAllocator PersistentChangeMaskAllocator;
+
+	FReplicationSystemInternal* ReplicationSystemInternal;
+	FNetHandleManager* NetHandleManager;
+	FReplicationStateStorage* StateStorage;
+	UReplicationBridge* ReplicationBridge;
+
+	// We track some data about incoming objects
+	// Stored in a map for now
+	TMap<uint32, FReplicatedObjectInfo> ReplicatedObjects;
+
+	// Temporary data valid during receive
+	FDispatchObjectInfo* ObjectsToDispatch;
+	uint32 ObjectsToDispatchCount;
+	uint32 ObjectsToDispatchCapacity;
+
+	// We need to keep some data around for objects with pending dependencies
+	// For now just use array and brute force the updates
+	TArray<uint32> ObjectsWithAttachmentPendingResolve;
+
+	// Track all objects waiting for this handle to be resolvable
+	TMultiMap<FNetHandle, uint32> UnresolvedHandleToDependents;
+	
+	// Track all objects with a dynamic handle in case it becomes unresolvable
+	TMultiMap<FNetHandle, uint32> ResolvedDynamicHandleToDependents;
+
+	FNetBlobHandlerManager* NetBlobHandlerManager;
+	FNetBlobType NetObjectBlobType;
+	FNetObjectAttachmentsReader Attachments;
+	FObjectReferenceCache* ObjectReferenceCache;
+	FNetObjectResolveContext ResolveContext;
+
+	IConsoleVariable const* DelayAttachmentsWithUnresolvedReferences;
+	uint32 NumHandlesPendingResolveLastUpdate = 0U;
+};
+
+}
