Iris experimental replication system.

diff --git a/Engine/Source/Runtime/Experimental/Iris/Core/Private/Iris/ReplicationSystem/ReplicationSystem.cpp b/Engine/Source/Runtime/Experimental/Iris/Core/Private/Iris/ReplicationSystem/ReplicationSystem.cpp
new file mode 100644
--- /dev/null
+++ b/Engine/Source/Runtime/Experimental/Iris/Core/Private/Iris/ReplicationSystem/ReplicationSystem.cpp
@@ -0,0 +1,1253 @@
+// Copyright Epic Games, Inc. All Rights Reserved.
+
+#include "Iris/ReplicationSystem/ReplicationSystem.h"
+#include "Iris/Core/IrisDebugging.h"
+#include "Iris/Core/IrisLog.h"
+#include "Iris/Core/IrisProfiler.h"
+#include "Iris/Core/NetObjectReference.h"
+#include "Net/Core/Trace/NetTrace.h"
+#include "Net/Core/Misc/NetConditionGroupManager.h"
+#include "Iris/ReplicationState/ReplicationStateUtil.h"
+#include "Iris/ReplicationSystem/ChangeMaskCache.h"
+#include "Iris/ReplicationSystem/Filtering/NetObjectFilter.h"
+#include "Iris/ReplicationSystem/Filtering/ReplicationFiltering.h"
+#include "Iris/ReplicationSystem/ReplicationBridge.h"
+#include "Iris/ReplicationSystem/ReplicationOperationsInternal.h"
+#include "Iris/ReplicationSystem/ReplicationReader.h"
+#include "Iris/ReplicationSystem/ReplicationSystemInternal.h"
+#include "Iris/ReplicationSystem/ReplicationTypes.h"
+#include "Iris/ReplicationSystem/ReplicationWriter.h"
+#include "Iris/Serialization/InternalNetSerializationContext.h"
+#include "Iris/Serialization/NetBitStreamWriter.h"
+#include "Iris/Serialization/NetSerializer.h"
+#include "HAL/IConsoleManager.h"
+#include "UObject/UObjectGlobals.h"
+
+namespace ReplicationSystemCVars
+{
+#if !UE_BUILD_SHIPPING
+static bool bForcePruneBeforeUpdate = false;
+static FAutoConsoleVariableRef CVarForcePruneBeforeUpdate(TEXT("net.Iris.ForcePruneBeforeUpdate"), bForcePruneBeforeUpdate, TEXT("Verify integrity of all tracked instances at the start of every update."));
+#endif
+}
+
+namespace UE::Net::Private
+{
+
+class FReplicationSystemImpl
+{
+public:
+	UReplicationSystem* ReplicationSystem;
+	FReplicationSystemInternal ReplicationSystemInternal;
+	uint64 IrisDebugHelperDummy = 0U;
+
+	explicit FReplicationSystemImpl(UReplicationSystem* InReplicationSystem, const UReplicationSystem::FReplicationSystemParams& Params)
+	: ReplicationSystem(InReplicationSystem)
+	, ReplicationSystemInternal(FReplicationSystemInternalInitParams({ InReplicationSystem->GetId(), Params.MaxReplicatedObjectCount }))
+	{
+	}
+
+	~FReplicationSystemImpl()
+	{
+	}
+
+	void InitDefaultFilteringGroups()
+	{
+		FNetObjectGroupHandle LocalNotReplicatedGroupHande = ReplicationSystem->CreateGroup();
+		check(LocalNotReplicatedGroupHande == NotReplicatedNetObjectGroupHandle);
+		ReplicationSystem->AddGroupFilter(NotReplicatedNetObjectGroupHandle);
+		
+		// Setup SubObjectFiltering groups
+		FNetObjectGroupHandle LocalNetGroupOwnerNetObjectGroupHandle = ReplicationSystem->GetOrCreateSubObjectFilter(UE::Net::NetGroupOwner);
+		check(LocalNetGroupOwnerNetObjectGroupHandle == NetGroupOwnerNetObjectGroupHandle);
+
+		FNetObjectGroupHandle LocalNetGroupReplayNetObjectGroupHandle = ReplicationSystem->GetOrCreateSubObjectFilter(UE::Net::NetGroupReplay);
+		check(LocalNetGroupReplayNetObjectGroupHandle == NetGroupReplayNetObjectGroupHandle);
+	}
+
+	void Init(const UReplicationSystem::FReplicationSystemParams& Params)
+	{
+		LLM_SCOPE(ELLMTag::Iris);
+
+#if !UE_BUILD_SHIPPING
+		IrisDebugHelperDummy = UE::Net::IrisDebugHelper::Init();
+#endif
+
+		const uint32 ReplicationSystemId = ReplicationSystem->GetId();
+
+		// $IRIS TODO: Need object ID range. Currently abusing hardcoded values from FNetHandleManager
+		FDirtyNetObjectTrackerInitParams DirtyNetObjectTrackerInitParams = {};
+		DirtyNetObjectTrackerInitParams.ReplicationSystemId = ReplicationSystemId;
+		DirtyNetObjectTrackerInitParams.NetObjectIndexRangeStart = 1;
+		DirtyNetObjectTrackerInitParams.NetObjectIndexRangeEnd = ReplicationSystemInternal.GetNetHandleManager().GetMaxActiveObjectCount() - 1U;
+		FDirtyNetObjectTracker& DirtyNetObjectTracker = ReplicationSystemInternal.GetDirtyNetObjectTracker();
+
+		DirtyNetObjectTracker.Init(DirtyNetObjectTrackerInitParams);
+
+		FReplicationStateStorage& StateStorage = ReplicationSystemInternal.GetReplicationStateStorage();
+		{
+			FReplicationStateStorageInitParams InitParams;
+			InitParams.ReplicationSystem = ReplicationSystem;
+			InitParams.NetHandleManager = &ReplicationSystemInternal.GetNetHandleManager();
+			InitParams.MaxObjectCount = DirtyNetObjectTrackerInitParams.NetObjectIndexRangeEnd + 1U;
+			InitParams.MaxConnectionCount = ReplicationSystemInternal.GetConnections().GetMaxConnectionCount();
+			InitParams.MaxDeltaCompressedObjectCount = Params.MaxDeltaCompressedObjectCount;
+			StateStorage.Init(InitParams);
+		}
+
+		FNetObjectGroups& Groups = ReplicationSystemInternal.GetGroups();
+		{
+			FNetObjectGroupInitParams InitParams = {};
+			InitParams.MaxObjectCount = DirtyNetObjectTrackerInitParams.NetObjectIndexRangeEnd + 1U;
+			InitParams.MaxGroupCount = Params.MaxNetObjectGroupCount;
+
+			Groups.Init(InitParams);
+		}
+
+		FReplicationStateDescriptorRegistry& Registry = ReplicationSystemInternal.GetReplicationStateDescriptorRegistry();
+		{
+			FReplicationStateDescriptorRegistryInitParams InitParams = {};
+			InitParams.ProtocolManager = &ReplicationSystemInternal.GetReplicationProtocolManager();
+
+			Registry.Init(InitParams);
+		}		
+
+		FWorldLocations& WorldLocations = ReplicationSystemInternal.GetWorldLocations();
+		{
+			FWorldLocationsInitParams InitParams;
+			InitParams.MaxObjectCount = DirtyNetObjectTrackerInitParams.NetObjectIndexRangeEnd + 1U;
+			WorldLocations.Init(InitParams);
+		}
+	
+		FDeltaCompressionBaselineInvalidationTracker& DeltaCompressionBaselineInvalidationTracker = ReplicationSystemInternal.GetDeltaCompressionBaselineInvalidationTracker();
+		FDeltaCompressionBaselineManager& DeltaCompressionBaselineManager = ReplicationSystemInternal.GetDeltaCompressionBaselineManager();
+		{
+			FDeltaCompressionBaselineInvalidationTrackerInitParams InitParams;
+			InitParams.BaselineManager = &DeltaCompressionBaselineManager;
+			InitParams.MaxObjectCount = DirtyNetObjectTrackerInitParams.NetObjectIndexRangeEnd + 1U;
+			DeltaCompressionBaselineInvalidationTracker.Init(InitParams);
+		}
+		{
+			FDeltaCompressionBaselineManagerInitParams InitParams;
+			InitParams.BaselineInvalidationTracker = &DeltaCompressionBaselineInvalidationTracker;
+			InitParams.Connections = &ReplicationSystemInternal.GetConnections();
+			InitParams.NetHandleManager = &ReplicationSystemInternal.GetNetHandleManager();
+			InitParams.ReplicationStateStorage = &StateStorage;
+			InitParams.MaxObjectCount = DirtyNetObjectTrackerInitParams.NetObjectIndexRangeEnd + 1U;
+			InitParams.MaxDeltaCompressedObjectCount = Params.MaxDeltaCompressedObjectCount;
+			InitParams.ReplicationSystem = ReplicationSystem;
+			DeltaCompressionBaselineManager.Init(InitParams);
+		}
+
+		FReplicationFiltering& ReplicationFiltering = ReplicationSystemInternal.GetFiltering();
+		{
+			FReplicationFilteringInitParams InitParams;
+			InitParams.ReplicationSystem = ReplicationSystem;
+			InitParams.Connections = &ReplicationSystemInternal.GetConnections();
+			InitParams.NetHandleManager = &ReplicationSystemInternal.GetNetHandleManager();
+			InitParams.Groups = &Groups;
+			InitParams.BaselineInvalidationTracker = &ReplicationSystemInternal.GetDeltaCompressionBaselineInvalidationTracker();
+			InitParams.MaxObjectCount = DirtyNetObjectTrackerInitParams.NetObjectIndexRangeEnd + 1U;
+			InitParams.MaxGroupCount = Params.MaxNetObjectGroupCount;
+			ReplicationFiltering.Init(InitParams);
+		}
+
+		InitDefaultFilteringGroups();
+
+		FReplicationConditionals& ReplicationConditionals = ReplicationSystemInternal.GetConditionals();
+		{
+			FReplicationConditionalsInitParams InitParams = {};
+			InitParams.NetHandleManager = &ReplicationSystemInternal.GetNetHandleManager();
+			InitParams.ReplicationConnections = &ReplicationSystemInternal.GetConnections();
+			InitParams.ReplicationFiltering = &ReplicationFiltering;
+			InitParams.NetObjectGroups = &Groups;
+			InitParams.BaselineInvalidationTracker = &ReplicationSystemInternal.GetDeltaCompressionBaselineInvalidationTracker();
+			InitParams.MaxObjectCount = DirtyNetObjectTrackerInitParams.NetObjectIndexRangeEnd + 1U;
+			InitParams.MaxConnectionCount = ReplicationSystemInternal.GetConnections().GetMaxConnectionCount();
+			ReplicationConditionals.Init(InitParams);
+		}
+
+		FReplicationPrioritization& ReplicationPrioritization = ReplicationSystemInternal.GetPrioritization();
+		{
+			FReplicationPrioritizationInitParams InitParams;
+			InitParams.ReplicationSystem = ReplicationSystem;
+			InitParams.Connections = &ReplicationSystemInternal.GetConnections();
+			InitParams.NetHandleManager = &ReplicationSystemInternal.GetNetHandleManager();
+			InitParams.MaxObjectCount = DirtyNetObjectTrackerInitParams.NetObjectIndexRangeEnd + 1U;
+			ReplicationPrioritization.Init(InitParams);
+		}
+
+		ReplicationSystemInternal.SetReplicationBridge(Params.ReplicationBridge);
+
+		// Init replication bridge
+		ReplicationSystemInternal.GetReplicationBridge()->Initialize(ReplicationSystem);
+
+		ReplicationSystemInternal.GetObjectReferenceCache().Init(ReplicationSystem);
+
+		FNetBlobManager& BlobManager = ReplicationSystemInternal.GetNetBlobManager();
+		{
+			FNetBlobManagerInitParams InitParams = {};
+			InitParams.ReplicationSystem = ReplicationSystem;
+			InitParams.bSendAttachmentsWithObject = ReplicationSystem->IsServer();
+			BlobManager.Init(InitParams);
+		}
+	}
+
+	void Deinit()
+	{
+		ReplicationSystemInternal.GetConnections().Deinit();
+
+		// Reset replication bridge
+		ReplicationSystemInternal.GetReplicationBridge()->Deinitialize();
+	}
+
+	/**
+	 * Due to objects having been marked as dirty and later removed we must make sure
+	 * that all dirty objects are still in scope.
+	 */
+	void UpdateDirtyObjectList()
+	{
+		FNetBitArrayView DirtyObjects = ReplicationSystemInternal.GetDirtyNetObjectTracker().GetDirtyNetObjects();
+		const FNetBitArrayView ScopableInternalIndices = MakeNetBitArrayView(ReplicationSystemInternal.GetNetHandleManager().GetScopableInternalIndices());
+		DirtyObjects.Combine(ScopableInternalIndices, FNetBitArrayView::AndOp);
+	}
+
+	void UpdateWorldLocations()
+	{
+		IRIS_PROFILER_SCOPE(FReplicationSystem_UpdateWorldLocations);
+
+		ReplicationSystemInternal.GetReplicationBridge()->CallUpdateInstancesWorldLocation();
+	}
+
+	void UpdateFiltering(const FNetBitArrayView& DirtyObjects)
+	{
+		IRIS_PROFILER_SCOPE(FReplicationSystem_UpdateFiltering);
+		LLM_SCOPE(ELLMTag::Iris);
+
+		FReplicationFiltering& Filtering = ReplicationSystemInternal.GetFiltering();
+		Filtering.Filter(DirtyObjects);
+
+		// Iterate over all valid connections and propagate updated scopes
+		FReplicationConnections& Connections = ReplicationSystemInternal.GetConnections();
+		auto&& UpdateConnectionScope = [&Filtering, &Connections](uint32 ConnectionId)
+		{
+			FReplicationConnection* Conn = Connections.GetConnection(ConnectionId);
+			const FNetBitArrayView& ObjectsInScope = Filtering.GetObjectsInScope(ConnectionId);
+			Conn->ReplicationWriter->UpdateScope(ObjectsInScope);
+		};
+		const FNetBitArray& ValidConnections = Connections.GetValidConnections();
+		ValidConnections.ForAllSetBits(UpdateConnectionScope);
+	}
+
+	// Can run at any time between scoping and replication.
+	void UpdateConditionals()
+	{
+		IRIS_PROFILER_SCOPE(FReplicationSystem_UpdateConditionals);
+
+		FReplicationConditionals& Conditionals = ReplicationSystemInternal.GetConditionals();
+		Conditionals.Update();
+	}
+
+	// Runs after filtering
+	void UpdatePrioritization(const FNetBitArrayView& ReplicatingConnections, const FNetBitArrayView& DirtyObjects)
+	{
+		IRIS_PROFILER_SCOPE(FReplicationSystem::FImpl::UpdatePrioritization);
+		LLM_SCOPE(ELLMTag::Iris);
+
+		FReplicationPrioritization& Prioritization = ReplicationSystemInternal.GetPrioritization();
+		Prioritization.Prioritize(ReplicatingConnections, DirtyObjects);
+	}
+
+	void PropagateDirtyChanges()
+	{
+		IRIS_PROFILER_SCOPE(FReplicationSystem_PropagateDirtyChanges);
+
+		FReplicationConnections& Connections = ReplicationSystemInternal.GetConnections();
+		const FChangeMaskCache& UpdatedChangeMasks = ReplicationSystemInternal.GetChangeMaskCache();
+
+		// Iterate over connections and propagate dirty changemasks
+		auto&& UpdateDirtyChangeMasks = [&Connections, &UpdatedChangeMasks](uint32 ConnectionId)
+		{
+			FReplicationConnection* Conn = Connections.GetConnection(ConnectionId);
+			Conn->ReplicationWriter->UpdateDirtyChangeMasks(UpdatedChangeMasks);
+		};
+		const FNetBitArray& ValidConnections = Connections.GetValidConnections();
+		ValidConnections.ForAllSetBits(UpdateDirtyChangeMasks);
+	}
+
+	void CopyDirtyStateData()
+	{
+		IRIS_PROFILER_SCOPE(FReplicationSystem_CopyDirtyStateData);
+		LLM_SCOPE(ELLMTag::IrisState);
+
+		FNetHandleManager& NetHandleManager = ReplicationSystemInternal.GetNetHandleManager();
+		FChangeMaskCache& Cache = ReplicationSystemInternal.GetChangeMaskCache();
+
+		const uint32 ReplicationSystemId = ReplicationSystem->GetId();
+		uint32 CopiedObjectCount = 0;
+		
+		// Prepare cache
+		constexpr uint32 ReservedIndexCount = 2048;
+		constexpr uint32 ReservedStorageCount = 16536;
+		Cache.PrepareCache(ReservedIndexCount, ReservedStorageCount);
+
+		// We use this ChangeMaskWriter to capture changemasks for all copied objects
+		FNetBitStreamWriter ChangeMaskWriter;
+
+		// Setup context
+		FNetSerializationContext SerializationContext;
+		FInternalNetSerializationContext InternalContext(ReplicationSystem);
+
+		SerializationContext.SetInternalContext(&InternalContext);
+
+		auto CopyFunction = [&ChangeMaskWriter, &Cache, &NetHandleManager, &CopiedObjectCount, &SerializationContext, ReplicationSystemId](uint32 DirtyIndex)
+		{
+			CopiedObjectCount += FReplicationInstanceOperationsInternal::CopyObjectStateData(ChangeMaskWriter, Cache, NetHandleManager, SerializationContext, DirtyIndex);
+		};
+
+		FDirtyNetObjectTracker& DirtyNetObjectTracker = ReplicationSystemInternal.GetDirtyNetObjectTracker();
+
+		// Only iterate over dirty objects
+		FNetBitArrayView DirtyObjects = DirtyNetObjectTracker.GetDirtyNetObjects();
+
+		// Copy all ReplicatedObjects with dirty state data
+		DirtyObjects.ForAllSetBits(CopyFunction);
+
+		UE_NET_TRACE_FRAME_STATSCOUNTER(ReplicationSystemId, ReplicationSystem.CopiedObjectCount, CopiedObjectCount, ENetTraceVerbosity::Trace);
+	}
+
+	void ProcessNetObjectAttachmentSendQueue()
+	{
+		IRIS_PROFILER_SCOPE(FReplicationSystem_ProcessNetObjectAttachmentSendQueue);
+
+		FNetBlobManager& NetBlobManager = ReplicationSystemInternal.GetNetBlobManager();
+		NetBlobManager.ProcessNetObjectAttachmentSendQueue();
+	}
+
+	// send data
+	void SendUpdate()
+	{
+	}
+
+	void AddConnection(uint32 ConnectionId)
+	{
+		LLM_SCOPE(ELLMTag::IrisConnection);
+
+		FReplicationConnections& Connections = ReplicationSystemInternal.GetConnections();
+		{
+			Connections.AddConnection(ConnectionId);
+
+			FReplicationConnection* Connection = Connections.GetConnection(ConnectionId);
+
+			FReplicationParameters Params;
+			Params.ReplicationSystem = ReplicationSystem;
+			Params.PacketSendWindowSize = 256;
+			Params.ConnectionId = ConnectionId;
+			Params.MaxActiveReplicatedObjectCount = ReplicationSystemInternal.GetNetHandleManager().GetMaxActiveObjectCount();
+			/** 
+			  * Currently we expect all objects to be replicated from server to client.
+			  * That means we will have to support sending attachments such as RPCs from
+			  * the client to the server, if the RPC is allowed to be sent in the first place.
+			  */
+			Params.bAllowSendingAttachmentsToObjectsNotInScope = !ReplicationSystem->IsServer();
+			Params.bAllowReceivingAttachmentsFromRemoteObjectsNotInScope = !Params.bAllowSendingAttachmentsToObjectsNotInScope;
+			// Delaying attachments with unresolved references on the server could cause massive queues of RPCs, potentially an OOM situation.
+			Params.bAllowDelayingAttachmentsWithUnresolvedReferences = !ReplicationSystem->IsServer();
+
+			Connection->ReplicationWriter = new FReplicationWriter();
+			Connection->ReplicationReader = new FReplicationReader();
+
+			Connection->ReplicationWriter->Init(Params);
+			Connection->ReplicationReader->Init(Params);
+		}
+
+		{
+			FReplicationConditionals& ReplicationConditionals = ReplicationSystemInternal.GetConditionals();
+			ReplicationConditionals.AddConnection(ConnectionId);
+		}
+
+		{
+			FReplicationFiltering& ReplicationFiltering = ReplicationSystemInternal.GetFiltering();
+			ReplicationFiltering.AddConnection(ConnectionId);
+		}
+
+		{
+			FReplicationPrioritization& ReplicationPrioritization = ReplicationSystemInternal.GetPrioritization();
+			ReplicationPrioritization.AddConnection(ConnectionId);
+		}
+
+		{
+			FDeltaCompressionBaselineManager& DC = ReplicationSystemInternal.GetDeltaCompressionBaselineManager();
+			DC.AddConnection(ConnectionId);
+		}
+	}
+
+	void RemoveConnection(uint32 ConnectionId)
+	{
+		{
+			FDeltaCompressionBaselineManager& DC = ReplicationSystemInternal.GetDeltaCompressionBaselineManager();
+			DC.RemoveConnection(ConnectionId);
+		}
+
+		{
+			FReplicationPrioritization& ReplicationPrioritization = ReplicationSystemInternal.GetPrioritization();
+			ReplicationPrioritization.RemoveConnection(ConnectionId);
+		}
+
+		{
+			FReplicationFiltering& ReplicationFiltering = ReplicationSystemInternal.GetFiltering();
+			ReplicationFiltering.RemoveConnection(ConnectionId);
+		}
+
+		{
+			FReplicationConditionals& ReplicationConditionals = ReplicationSystemInternal.GetConditionals();
+			ReplicationConditionals.RemoveConnection(ConnectionId);
+		}
+
+		FReplicationConnections& Connections = ReplicationSystemInternal.GetConnections();
+		{
+			Connections.RemoveConnection(ConnectionId);
+		}
+	}
+
+	void UpdateUnresolvableReferenceTracking()
+	{
+		FReplicationConnections& Connections = ReplicationSystemInternal.GetConnections();
+		auto&& UpdateUnresolvableReferenceTracking = [&Connections](uint32 ConnectionId)
+		{
+			FReplicationConnection* Conn = Connections.GetConnection(ConnectionId);
+			Conn->ReplicationReader->UpdateUnresolvableReferenceTracking();
+		};
+		const FNetBitArray& ValidConnections = Connections.GetValidConnections();
+		ValidConnections.ForAllSetBits(UpdateUnresolvableReferenceTracking);
+	}
+};
+
+}
+
+UReplicationSystem::UReplicationSystem()
+: Super()
+, Impl(nullptr)
+, Id(~0U)
+, bIsServer(0)
+, bAllowObjectReplication(0)
+, bDoCollectGarbage(0)
+{
+}
+
+void UReplicationSystem::Init(uint32 InId, const FReplicationSystemParams& Params)
+{
+	LLM_SCOPE(ELLMTag::Iris);
+
+	Id = InId;
+	bIsServer = Params.bIsServer;
+	bAllowObjectReplication = Params.bAllowObjectReplication;
+
+	ReplicationBridge = Params.ReplicationBridge;
+
+	Impl = MakePimpl<UE::Net::Private::FReplicationSystemImpl>(this, Params);
+	Impl->Init(Params);
+
+	PostGarbageCollectHandle = FCoreUObjectDelegates::GetPostGarbageCollect().AddUObject(this, &UReplicationSystem::PostGarbageCollection);
+}
+
+void UReplicationSystem::Shutdown()
+{
+	FCoreUObjectDelegates::GetPostGarbageCollect().Remove(PostGarbageCollectHandle);
+
+	// Destroy impl
+	Impl->Deinit();
+	Impl.Reset();
+
+	ReplicationBridge = nullptr;
+}
+
+UReplicationSystem::~UReplicationSystem()
+{
+}
+
+UE::Net::Private::FReplicationSystemInternal* UReplicationSystem::GetReplicationSystemInternal()
+{
+	return &Impl->ReplicationSystemInternal;
+}
+
+const UE::Net::Private::FReplicationSystemInternal* UReplicationSystem::GetReplicationSystemInternal() const
+{
+	return &Impl->ReplicationSystemInternal;
+}
+
+void UReplicationSystem::PreSendUpdate(float DeltaSeconds)
+{
+	IRIS_PROFILER_SCOPE(FReplicationSystem_PreSendUpdate);
+
+#if !UE_BUILD_SHIPPING
+	// Force a integrity check of all replicated instances
+	if (bDoCollectGarbage || ReplicationSystemCVars::bForcePruneBeforeUpdate)
+	{
+		CollectGarbage();
+	}
+#endif
+
+	UE::Net::Private::FReplicationSystemInternal& InternalSys = Impl->ReplicationSystemInternal;
+
+	// $IRIS TODO. There may be some throttling of connections to tick that we should take into account.
+	const UE::Net::FNetBitArrayView& ReplicatingConnections = MakeNetBitArrayView(Impl->ReplicationSystemInternal.GetConnections().GetValidConnections());
+
+	if (bAllowObjectReplication)
+	{
+		UE_NET_TRACE_FRAME_STATSCOUNTER(GetId(), ReplicationSystem.ReplicatedObjectCount, InternalSys.GetNetHandleManager().GetActiveObjectCount(), ENetTraceVerbosity::Verbose);
+
+		// Refresh dirty object list
+		Impl->UpdateDirtyObjectList();
+
+		// Invoke any operations we need to do before copying state data
+		InternalSys.GetReplicationBridge()->CallPreSendUpdate(DeltaSeconds);
+
+		// Update world locations. We need this to happen before both filtering and prioritization.
+		Impl->UpdateWorldLocations();
+
+		// Update conditionals
+		Impl->UpdateConditionals();
+
+		// Copy dirty state data. We need this to happen before both filtering and prioritization
+		Impl->CopyDirtyStateData();
+
+		// Update filtering and scope for all connections
+		{
+			UE::Net::FNetBitArrayView DirtyObjects = InternalSys.GetDirtyNetObjectTracker().GetDirtyNetObjects();
+			Impl->UpdateFiltering(DirtyObjects);
+		}
+
+		// Propagate dirty changes to all connections
+		Impl->PropagateDirtyChanges();
+	}
+
+	// Forward attachments to the connections now that we know which objects are in scope
+	Impl->ProcessNetObjectAttachmentSendQueue();
+
+	if (bAllowObjectReplication)
+	{
+		// Update object priorities
+		UE::Net::FNetBitArrayView DirtyObjects = InternalSys.GetDirtyNetObjectTracker().GetDirtyNetObjects();
+		Impl->UpdatePrioritization(ReplicatingConnections, DirtyObjects);
+
+		// Delta compression preparations before send
+		{
+			UE::Net::Private::FDeltaCompressionBaselineManagerPreSendUpdateParams UpdateParams;
+			UpdateParams.ChangeMaskCache = &InternalSys.GetChangeMaskCache();
+			InternalSys.GetDeltaCompressionBaselineManager().PreSendUpdate(UpdateParams);
+		}
+	}
+
+	// Destroy objects pending destroy
+	{
+		Impl->UpdateUnresolvableReferenceTracking();
+		InternalSys.GetNetHandleManager().DestroyObjectsPendingDestroy();
+	}
+}
+
+void UReplicationSystem::SendUpdate()
+{
+	Impl->SendUpdate();
+}
+
+void UReplicationSystem::PostSendUpdate()
+{
+	using namespace UE::Net::Private;
+
+	IRIS_PROFILER_SCOPE(FReplicationSystem_PostSendUpdate);
+	
+	FReplicationSystemInternal& InternalSys = Impl->ReplicationSystemInternal;
+
+	InternalSys.GetDirtyNetObjectTracker().ClearDirtyNetObjects();
+
+	InternalSys.GetChangeMaskCache().ResetCache();
+
+	// Store the state of the previous frames scopable objects
+	InternalSys.GetNetHandleManager().SetPrevFrameScopableInternalIndicesToCurrent();
+
+	// Update handles pending tear-off
+	InternalSys.GetReplicationBridge()->UpdateHandlesPendingTearOff();
+
+	// Reset baseline invalidation
+	InternalSys.GetDeltaCompressionBaselineInvalidationTracker().PostSendUpdate();
+
+	if (bAllowObjectReplication)
+	{
+		FDeltaCompressionBaselineManagerPostSendUpdateParams UpdateParams;
+		InternalSys.GetDeltaCompressionBaselineManager().PostSendUpdate(UpdateParams);
+	}
+}
+
+void UReplicationSystem::PostGarbageCollection()
+{
+	bDoCollectGarbage = 1U;
+}
+
+void UReplicationSystem::CollectGarbage()
+{
+	IRIS_PROFILER_SCOPE(ReplicationSystem_CollectGarbage);
+
+	// Prune stale object instances before descriptors and protocols are pruned
+	Impl->ReplicationSystemInternal.GetReplicationBridge()->CallPruneStaleObjects();
+	Impl->ReplicationSystemInternal.GetReplicationStateDescriptorRegistry().PruneStaleDescriptors();
+
+	bDoCollectGarbage = 0;
+}
+
+void UReplicationSystem::ResetGameWorldState()
+{
+	Impl->ReplicationSystemInternal.GetReplicationBridge()->RemoveDestructionInfosForGroup(UE::Net::InvalidNetObjectGroupHandle);
+}
+
+void UReplicationSystem::NotifyStreamingLevelUnload(const UObject* Level)
+{
+	Impl->ReplicationSystemInternal.GetReplicationBridge()->NotifyStreamingLevelUnload(Level);
+}
+
+void UReplicationSystem::AddConnection(uint32 ConnectionId)
+{
+	Impl->AddConnection(ConnectionId);
+}
+
+void UReplicationSystem::RemoveConnection(uint32 ConnectionId)
+{
+	Impl->RemoveConnection(ConnectionId);
+}
+
+bool UReplicationSystem::IsValidConnection(uint32 ConnectionId) const
+{
+	UE::Net::Private::FReplicationConnections& Connections = Impl->ReplicationSystemInternal.GetConnections();
+	return Connections.GetConnection(ConnectionId) != nullptr;
+}
+
+void UReplicationSystem::SetReplicationEnabledForConnection(uint32 ConnectionId, bool bReplicationEnabled)
+{
+	UE::Net::Private::FReplicationConnections& Connections = Impl->ReplicationSystemInternal.GetConnections();
+	UE::Net::Private::FReplicationConnection* Connection = Connections.GetConnection(ConnectionId);
+
+	check(Connection);
+
+	Connection->ReplicationWriter->SetReplicationEnabled(bReplicationEnabled);
+}
+
+void UReplicationSystem::SetReplicationView(uint32 ConnectionId, const UE::Net::FReplicationView& View)
+{
+	UE::Net::Private::FReplicationConnections& Connections = Impl->ReplicationSystemInternal.GetConnections();
+	Connections.SetReplicationView(ConnectionId, View);
+}
+
+void UReplicationSystem::SetStaticPriority(FNetHandle Handle, float Priority)
+{
+	const uint32 ObjectIndex = Impl->ReplicationSystemInternal.GetNetHandleManager().GetInternalIndex(Handle);
+	if (ObjectIndex == UE::Net::Private::FNetHandleManager::InvalidInternalIndex)
+	{
+		return;
+	}
+
+	return Impl->ReplicationSystemInternal.GetPrioritization().SetStaticPriority(ObjectIndex, Priority);
+}
+
+bool UReplicationSystem::SetPrioritizer(FNetHandle Handle, FNetObjectPrioritizerHandle Prioritizer)
+{
+	const uint32 ObjectIndex = Impl->ReplicationSystemInternal.GetNetHandleManager().GetInternalIndex(Handle);
+	if (ObjectIndex == UE::Net::Private::FNetHandleManager::InvalidInternalIndex)
+	{
+		return false;
+	}
+
+	return Impl->ReplicationSystemInternal.GetPrioritization().SetPrioritizer(ObjectIndex, Prioritizer);
+}
+
+FNetObjectPrioritizerHandle UReplicationSystem::GetPrioritizerHandle(const FName PrioritizerName) const
+{
+	return Impl->ReplicationSystemInternal.GetPrioritization().GetPrioritizerHandle(PrioritizerName);
+}
+
+UNetObjectPrioritizer* UReplicationSystem::GetPrioritizer(const FName PrioritizerName) const
+{
+	return Impl->ReplicationSystemInternal.GetPrioritization().GetPrioritizer(PrioritizerName);
+}
+
+#if 0
+UNetObjectPrioritizer* UReplicationSystem::GetPrioritizer(const FName PrioritizerName)
+{
+	return Impl->ReplicationSystemInternal.GetPrioritization().GetPrioritizer(PrioritizerName);
+}
+#endif
+
+const UE::Net::FStringTokenStore* UReplicationSystem::GetStringTokenStore() const
+{
+	return &Impl->ReplicationSystemInternal.GetStringTokenStore();
+}
+
+UE::Net::FStringTokenStore* UReplicationSystem::GetStringTokenStore()
+{
+	return &Impl->ReplicationSystemInternal.GetStringTokenStore();
+}
+
+bool UReplicationSystem::RegisterNetBlobHandler(UNetBlobHandler* Handler)
+{
+	UE::Net::Private::FNetBlobManager& NetBlobManager = Impl->ReplicationSystemInternal.GetNetBlobManager();
+	return NetBlobManager.RegisterNetBlobHandler(Handler);
+}
+
+bool UReplicationSystem::QueueNetObjectAttachment(uint32 ConnectionId, const UE::Net::FNetObjectReference& TargetRef, const TRefCountPtr<UE::Net::FNetObjectAttachment>& Attachment)
+{
+	UE::Net::Private::FNetBlobManager& NetBlobManager = Impl->ReplicationSystemInternal.GetNetBlobManager();
+	return NetBlobManager.QueueNetObjectAttachment(ConnectionId, TargetRef, Attachment);
+}
+
+bool UReplicationSystem::SendRPC(const UObject* Object, const UObject* SubObject, const UFunction* Function, const void* Parameters)
+{
+	UE::Net::Private::FNetBlobManager& NetBlobManager = Impl->ReplicationSystemInternal.GetNetBlobManager();
+	return NetBlobManager.SendRPC(Object, SubObject, Function, Parameters);
+}
+
+bool UReplicationSystem::SendRPC(uint32 ConnectionId, const UObject* Object, const UObject* SubObject, const UFunction* Function, const void* Parameters)
+{
+	UE::Net::Private::FNetBlobManager& NetBlobManager = Impl->ReplicationSystemInternal.GetNetBlobManager();
+	return NetBlobManager.SendRPC(ConnectionId, Object, SubObject, Function, Parameters);
+}
+
+void UReplicationSystem::InitDataStreams(uint32 ConnectionId, UDataStreamManager* DataStreamManager)
+{
+	UE::Net::Private::FReplicationConnections& Connections = Impl->ReplicationSystemInternal.GetConnections();
+	Connections.InitDataStreams(GetId(), ConnectionId, DataStreamManager);
+}
+
+void UReplicationSystem::SetConnectionUserData(uint32 ConnectionId, UObject* InUserData)
+{
+	UE::Net::Private::FReplicationConnections& Connections = Impl->ReplicationSystemInternal.GetConnections();
+	UE::Net::Private::FReplicationConnection* Connection = Connections.GetConnection(ConnectionId);
+	
+	check(Connection);
+
+	Connection->UserData = InUserData;
+}
+
+UObject* UReplicationSystem::GetConnectionUserData(uint32 ConnectionId) const
+{
+	UE::Net::Private::FReplicationConnections& Connections = Impl->ReplicationSystemInternal.GetConnections();
+	if (UE::Net::Private::FReplicationConnection* Connection = Connections.GetConnection(ConnectionId))
+	{
+		return Connection->UserData.Get();
+	}
+	return nullptr;
+}
+
+UReplicationBridge* UReplicationSystem::GetReplicationBridge() const
+{
+	return Impl->ReplicationSystemInternal.GetReplicationBridge();
+}
+
+bool UReplicationSystem::IsValidHandle(FNetHandle Handle) const
+{
+	return Handle.IsValid() && Impl->ReplicationSystemInternal.GetNetHandleManager().IsValidNetHandle(Handle);
+}
+
+const UE::Net::FReplicationProtocol* UReplicationSystem::GetReplicationProtocol(FNetHandle Handle) const
+{
+	using namespace UE::Net::Private;
+
+	FNetHandleManager& NetHandleManager = Impl->ReplicationSystemInternal.GetNetHandleManager();
+
+	const uint32 ObjectInternalIndex = NetHandleManager.GetInternalIndex(Handle);
+	if (ObjectInternalIndex == FNetHandleManager::InvalidInternalIndex)
+	{
+		return nullptr;
+	}
+
+	return NetHandleManager.GetReplicatedObjectDataNoCheck(ObjectInternalIndex).Protocol;
+}
+
+void UReplicationSystem::SetOwningNetConnection(FNetHandle Handle, uint32 ConnectionId)
+{
+	using namespace UE::Net::Private;
+
+	FNetHandleManager& NetHandleManager = Impl->ReplicationSystemInternal.GetNetHandleManager();
+	const uint32 ObjectInternalIndex = NetHandleManager.GetInternalIndex(Handle);
+	if (ObjectInternalIndex == FNetHandleManager::InvalidInternalIndex)
+	{
+		return;
+	}
+
+	FReplicationFiltering& Filtering = Impl->ReplicationSystemInternal.GetFiltering();
+	Filtering.SetOwningConnection(ObjectInternalIndex, ConnectionId);
+}
+
+uint32 UReplicationSystem::GetOwningNetConnection(FNetHandle Handle) const
+{
+	using namespace UE::Net::Private;
+
+	FNetHandleManager& NetHandleManager = Impl->ReplicationSystemInternal.GetNetHandleManager();
+	const uint32 ObjectInternalIndex = NetHandleManager.GetInternalIndex(Handle);
+	if (ObjectInternalIndex == FNetHandleManager::InvalidInternalIndex)
+	{
+		return 0U;
+	}
+
+	const FReplicationFiltering& Filtering = Impl->ReplicationSystemInternal.GetFiltering();
+	return Filtering.GetOwningConnection(ObjectInternalIndex);
+}
+
+bool UReplicationSystem::SetFilter(FNetHandle Handle, UE::Net::FNetObjectFilterHandle Filter)
+{
+	using namespace UE::Net::Private;
+
+	FNetHandleManager& NetHandleManager = Impl->ReplicationSystemInternal.GetNetHandleManager();
+	const uint32 ObjectInternalIndex = NetHandleManager.GetInternalIndex(Handle);
+	if (ObjectInternalIndex == FNetHandleManager::InvalidInternalIndex)
+	{
+		return false;
+	}
+
+	FReplicationFiltering& Filtering = Impl->ReplicationSystemInternal.GetFiltering();
+	return Filtering.SetFilter(ObjectInternalIndex, Filter);
+}
+
+UE::Net::FNetObjectFilterHandle UReplicationSystem::GetFilterHandle(const FName FilterName) const
+{
+	return Impl->ReplicationSystemInternal.GetFiltering().GetFilterHandle(FilterName);
+}
+
+UNetObjectFilter* UReplicationSystem::GetFilter(const FName FilterName) const
+{
+	return Impl->ReplicationSystemInternal.GetFiltering().GetFilter(FilterName);
+}
+
+bool UReplicationSystem::SetConnectionFilter(FNetHandle Handle, const TBitArray<>& Connections, UE::Net::ENetFilterStatus ReplicationStatus)
+{
+	using namespace UE::Net;
+	using namespace UE::Net::Private;
+
+	FNetHandleManager& NetHandleManager = Impl->ReplicationSystemInternal.GetNetHandleManager();
+	const uint32 ObjectInternalIndex = NetHandleManager.GetInternalIndex(Handle);
+	if (ObjectInternalIndex == FNetHandleManager::InvalidInternalIndex)
+	{
+		return false;
+	}
+
+	FReplicationFiltering& Filtering = Impl->ReplicationSystemInternal.GetFiltering();
+	return Filtering.SetConnectionFilter(ObjectInternalIndex, MakeNetBitArrayView(Connections.GetData(), Connections.Num()), ReplicationStatus);
+}
+
+UE::Net::FNetObjectGroupHandle UReplicationSystem::GetOrCreateSubObjectFilter(FName GroupName)
+{
+	using namespace UE::Net;
+	using namespace UE::Net::Private;
+
+	FNetObjectGroups& Groups = Impl->ReplicationSystemInternal.GetGroups();
+	FReplicationFiltering& Filtering = Impl->ReplicationSystemInternal.GetFiltering();
+
+	FNetObjectGroupHandle GroupHandle = Groups.GetNamedGroupHandle(GroupName);
+	if (GroupHandle != InvalidNetObjectGroupHandle)
+	{
+		check(Filtering.IsSubObjectFilterGroup(GroupHandle));
+		return GroupHandle;
+	}
+
+	GroupHandle = Groups.CreateNamedGroup(GroupName);
+	if (GroupHandle != InvalidNetObjectGroupHandle)
+	{
+		Filtering.AddSubObjectFilter(GroupHandle);
+	}
+	return GroupHandle;
+}
+
+UE::Net::FNetObjectGroupHandle UReplicationSystem::GetSubObjectFilterGroupHandle(FName GroupName) const
+{
+	using namespace UE::Net;
+	using namespace UE::Net::Private;
+
+	FNetObjectGroups& Groups = Impl->ReplicationSystemInternal.GetGroups();
+	FReplicationFiltering& Filtering = Impl->ReplicationSystemInternal.GetFiltering();
+
+	FNetObjectGroupHandle GroupHandle = Groups.GetNamedGroupHandle(GroupName);
+	if (GroupHandle != InvalidNetObjectGroupHandle)
+	{
+		if (ensureAlwaysMsgf(Filtering.IsSubObjectFilterGroup(GroupHandle), TEXT("UReplicationSystem::GetSubObjectFilterGroupHandle Trying to lookup NetObjectGroupHandle for NetGroup %s that is not a subobject filter"), *GroupName.ToString()))
+		{
+			return GroupHandle;
+		}
+	}
+	return InvalidNetObjectGroupHandle;
+}
+
+void UReplicationSystem::SetSubObjectFilterStatus(FName GroupName, uint32 ConnectionId, UE::Net::ENetFilterStatus ReplicationStatus)
+{
+	using namespace UE::Net;
+	using namespace UE::Net::Private;
+
+	if (UE::Net::IsSpecialNetConditionGroup(GroupName))
+	{
+		ensureAlwaysMsgf(false, TEXT("UReplicationSystem::SetSubObjectFilterStatus Cannot SetSubObjectFilterStatus for special NetGroup %s"), *GroupName.ToString());
+		return;
+	}
+
+	FNetObjectGroupHandle GroupHandle = GetSubObjectFilterGroupHandle(GroupName);
+	if (GroupHandle != InvalidNetObjectGroupHandle)
+	{
+		FReplicationFiltering& Filtering = Impl->ReplicationSystemInternal.GetFiltering();
+		Filtering.SetSubObjectFilterStatus(GroupHandle, ConnectionId, ReplicationStatus);
+	}
+}
+
+void UReplicationSystem::RemoveSubObjectFilter(FName GroupName)
+{
+	using namespace UE::Net;
+	using namespace UE::Net::Private;
+
+	FNetObjectGroups& Groups = Impl->ReplicationSystemInternal.GetGroups();
+	FReplicationFiltering& Filtering = Impl->ReplicationSystemInternal.GetFiltering();
+
+	FNetObjectGroupHandle GroupHandle = GetSubObjectFilterGroupHandle(GroupName);
+	if (GroupHandle != InvalidNetObjectGroupHandle)
+	{
+		Filtering.RemoveSubObjectFilter(GroupHandle);
+		Groups.DestroyGroup(GroupHandle);
+	}
+}
+
+UE::Net::FNetObjectGroupHandle UReplicationSystem::CreateGroup()
+{
+	LLM_SCOPE(ELLMTag::Iris);
+
+	return Impl->ReplicationSystemInternal.GetGroups().CreateGroup();
+}
+
+void UReplicationSystem::AddToGroup(FNetObjectGroupHandle GroupHandle, FNetHandle Handle)
+{
+	LLM_SCOPE(ELLMTag::Iris);
+
+	using namespace UE::Net::Private;
+
+	FNetObjectGroups& Groups = Impl->ReplicationSystemInternal.GetGroups();
+	FNetHandleManager& NetHandleManager = Impl->ReplicationSystemInternal.GetNetHandleManager();
+	FReplicationFiltering& Filtering = Impl->ReplicationSystemInternal.GetFiltering();	
+
+	const uint32 ObjectInternalIndex = NetHandleManager.GetInternalIndex(Handle);
+
+	if (GroupHandle && ObjectInternalIndex)
+	{
+		Groups.AddToGroup(GroupHandle, ObjectInternalIndex);
+		Filtering.NotifyObjectAddedToGroup(GroupHandle, ObjectInternalIndex);
+	}
+}
+
+void UReplicationSystem::RemoveFromGroup(FNetObjectGroupHandle GroupHandle, FNetHandle Handle)
+{	
+	using namespace UE::Net::Private;
+
+	FNetHandleManager& NetHandleManager = Impl->ReplicationSystemInternal.GetNetHandleManager();
+
+	const uint32 ObjectInternalIndex = NetHandleManager.GetInternalIndex(Handle);
+	if (GroupHandle && ObjectInternalIndex)
+	{
+		FNetObjectGroups& Groups = Impl->ReplicationSystemInternal.GetGroups();
+		FReplicationFiltering& Filtering = Impl->ReplicationSystemInternal.GetFiltering();	
+
+		Groups.RemoveFromGroup(GroupHandle, ObjectInternalIndex);
+		Filtering.NotifyObjectRemovedFromGroup(GroupHandle, ObjectInternalIndex);
+	}
+}
+
+void UReplicationSystem::RemoveFromAllGroups(FNetHandle Handle)
+{
+	using namespace UE::Net::Private;
+
+	FNetHandleManager& NetHandleManager = Impl->ReplicationSystemInternal.GetNetHandleManager();
+	FNetObjectGroups& Groups = Impl->ReplicationSystemInternal.GetGroups();
+	FReplicationFiltering& Filtering = Impl->ReplicationSystemInternal.GetFiltering();
+
+	uint32 NumGroupMemberShips;
+	const uint32 ObjectInternalIndex = NetHandleManager.GetInternalIndex(Handle);
+	if (const FNetObjectGroupHandle* GroupHandles = Groups.GetGroupMemberships(ObjectInternalIndex, NumGroupMemberShips))
+	{
+		// We copy the membership array as it is modified during removal
+		TArray<FNetObjectGroupHandle> CopiedGroupHandles(MakeArrayView(GroupHandles, NumGroupMemberShips));
+		for (FNetObjectGroupHandle GroupHandle : MakeArrayView(CopiedGroupHandles.GetData(), CopiedGroupHandles.Num()))
+		{
+			Groups.RemoveFromGroup(GroupHandle, ObjectInternalIndex);
+			Filtering.NotifyObjectRemovedFromGroup(GroupHandle, ObjectInternalIndex);
+		}
+	}	
+}
+
+bool UReplicationSystem::IsInGroup(FNetObjectGroupHandle GroupHandle, FNetHandle Handle) const
+{
+	const UE::Net::Private::FNetObjectGroups& Groups = Impl->ReplicationSystemInternal.GetGroups();
+	const UE::Net::Private::FNetHandleManager& NetHandleManager = Impl->ReplicationSystemInternal.GetNetHandleManager();
+
+	const uint32 ObjectInternalIndex = NetHandleManager.GetInternalIndex(Handle);
+
+	return Groups.Contains(GroupHandle, ObjectInternalIndex);
+}
+
+bool UReplicationSystem::IsValidGroup(FNetObjectGroupHandle GroupHandle) const
+{
+	const UE::Net::Private::FNetObjectGroups& Groups = Impl->ReplicationSystemInternal.GetGroups();
+
+	return GroupHandle && Groups.IsValidGroup(GroupHandle);
+}
+
+void UReplicationSystem::DestroyGroup(FNetObjectGroupHandle GroupHandle)
+{
+	using namespace UE::Net;
+	using namespace UE::Net::Private;
+
+	// We do not allow client code to remove reserved groups
+	if (IsReservedNetObjectGroupHandle(GroupHandle))
+	{
+		check(false);
+		return;
+	}
+
+	FReplicationFiltering& Filtering = Impl->ReplicationSystemInternal.GetFiltering();	
+	FNetObjectGroups& Groups = Impl->ReplicationSystemInternal.GetGroups();
+
+	Filtering.RemoveGroupFilter(GroupHandle);
+	Filtering.RemoveSubObjectFilter(GroupHandle);
+
+	Groups.DestroyGroup(GroupHandle);
+}
+
+void UReplicationSystem::AddGroupFilter(FNetObjectGroupHandle GroupHandle)
+{
+	UE::Net::Private::FReplicationFiltering& Filtering = Impl->ReplicationSystemInternal.GetFiltering();	
+	Filtering.AddGroupFilter(GroupHandle);
+}
+
+void UReplicationSystem::RemoveGroupFilter(FNetObjectGroupHandle GroupHandle)
+{
+	UE::Net::Private::FReplicationFiltering& Filtering = Impl->ReplicationSystemInternal.GetFiltering();	
+	Filtering.RemoveGroupFilter(GroupHandle);
+}
+
+void UReplicationSystem::SetGroupFilterStatus(FNetObjectGroupHandle GroupHandle, uint32 ConnectionId, UE::Net::ENetFilterStatus ReplicationStatus)
+{
+	UE::Net::Private::FReplicationFiltering& Filtering = Impl->ReplicationSystemInternal.GetFiltering();	
+	Filtering.SetGroupFilterStatus(GroupHandle, ConnectionId, ReplicationStatus);
+}
+
+void UReplicationSystem::SetGroupFilterStatus(FNetObjectGroupHandle GroupHandle, const UE::Net::FNetBitArray& Connections, UE::Net::ENetFilterStatus ReplicationStatus)
+{
+	UE::Net::Private::FReplicationFiltering& Filtering = Impl->ReplicationSystemInternal.GetFiltering();	
+	Filtering.SetGroupFilterStatus(GroupHandle, UE::Net::MakeNetBitArrayView(Connections), ReplicationStatus);
+}
+
+void UReplicationSystem::SetGroupFilterStatus(FNetObjectGroupHandle GroupHandle, UE::Net::ENetFilterStatus ReplicationStatus)
+{
+	UE::Net::Private::FReplicationFiltering& Filtering = Impl->ReplicationSystemInternal.GetFiltering();	
+	Filtering.SetGroupFilterStatus(GroupHandle, ReplicationStatus);
+}
+
+bool UReplicationSystem::SetReplicationConditionConnectionFilter(FNetHandle Handle, UE::Net::EReplicationCondition Condition, uint32 ConnectionId, bool bEnable)
+{
+	using namespace UE::Net::Private;
+
+	FNetHandleManager& NetHandleManager = Impl->ReplicationSystemInternal.GetNetHandleManager();
+	const uint32 ObjectInternalIndex = NetHandleManager.GetInternalIndex(Handle);
+	if (ObjectInternalIndex == FNetHandleManager::InvalidInternalIndex)
+	{
+		return false;
+	}
+
+	FReplicationConditionals& Conditionals = Impl->ReplicationSystemInternal.GetConditionals();
+	return Conditionals.SetConditionConnectionFilter(ObjectInternalIndex, Condition, ConnectionId, bEnable);
+}
+
+bool UReplicationSystem::SetReplicationCondition(FNetHandle Handle, UE::Net::EReplicationCondition Condition, bool bEnable)
+{
+	using namespace UE::Net::Private;
+
+	FNetHandleManager& NetHandleManager = Impl->ReplicationSystemInternal.GetNetHandleManager();
+	const uint32 ObjectInternalIndex = NetHandleManager.GetInternalIndex(Handle);
+	if (ObjectInternalIndex == FNetHandleManager::InvalidInternalIndex)
+	{
+		return false;
+	}
+
+	FReplicationConditionals& Conditionals = Impl->ReplicationSystemInternal.GetConditionals();
+	return Conditionals.SetCondition(ObjectInternalIndex, Condition, bEnable);
+}
+
+bool UReplicationSystem::SetPropertyCustomCondition(FNetHandle Handle, const void* Owner, uint16 RepIndex, bool bEnable)
+{
+	using namespace UE::Net::Private;
+
+	FNetHandleManager& NetHandleManager = Impl->ReplicationSystemInternal.GetNetHandleManager();
+	const uint32 ObjectInternalIndex = NetHandleManager.GetInternalIndex(Handle);
+	if (ObjectInternalIndex == FNetHandleManager::InvalidInternalIndex)
+	{
+		return false;
+	}
+
+	FReplicationConditionals& Conditionals = Impl->ReplicationSystemInternal.GetConditionals();
+	return Conditionals.SetPropertyCustomCondition(ObjectInternalIndex, Owner, RepIndex, bEnable);
+}
+
+void UReplicationSystem::SetDeltaCompressionStatus(FNetHandle Handle, UE::Net::ENetObjectDeltaCompressionStatus Status)
+{
+	using namespace UE::Net::Private;
+
+	FReplicationSystemInternal& InternalSys = Impl->ReplicationSystemInternal;
+	FNetHandleManager& NetHandleManager = InternalSys.GetNetHandleManager();
+	const uint32 ObjectInternalIndex = NetHandleManager.GetInternalIndex(Handle);
+	if (ObjectInternalIndex == FNetHandleManager::InvalidInternalIndex)
+	{
+		return;
+	}
+
+	FDeltaCompressionBaselineManager& DC = InternalSys.GetDeltaCompressionBaselineManager();
+	return DC.SetDeltaCompressionStatus(ObjectInternalIndex, Status);
+}
+
+void UReplicationSystem::SetIsNetTemporary(FNetHandle Handle)
+{
+	UE::Net::Private::FNetHandleManager& NetHandleManager = Impl->ReplicationSystemInternal.GetNetHandleManager();
+	if (ensure(NetHandleManager.IsLocalNetHandle(Handle)))
+	{
+		// Set the object to not propagate changed states
+		NetHandleManager.SetShouldPropagateChangedStates(Handle, false);
+	}
+}
+
+void UReplicationSystem::TearOffNextUpdate(FNetHandle Handle)
+{
+	Impl->ReplicationSystemInternal.GetReplicationBridge()->TearOff(Handle, false);
+}
+
+void UReplicationSystem::MarkDirty(FNetHandle Handle)
+{
+	if (const uint32 InternalObjectIndex = Impl->ReplicationSystemInternal.GetNetHandleManager().GetInternalIndex(Handle))
+	{
+		 UE::Net::Private::MarkNetObjectStateDirty(GetId(), InternalObjectIndex);
+	}
+}
+
+uint32 UReplicationSystem::GetMaxConnectionCount() const
+{
+	return Impl->ReplicationSystemInternal.GetConnections().GetMaxConnectionCount();
+}
+
+void UReplicationSystem::AddReferencedObjects(UObject* InThis, FReferenceCollector& Collector)
+{
+	UReplicationSystem* This = CastChecked<UReplicationSystem>(InThis);
+	if (This->Impl.IsValid())
+	{
+		This->Impl->ReplicationSystemInternal.GetNetHandleManager().AddReferencedObjects(Collector);
+	}
+	Super::AddReferencedObjects(InThis, Collector);
+}
+
+const UE::Net::FWorldLocations& UReplicationSystem::GetWorldLocations() const
+{
+	return Impl->ReplicationSystemInternal.GetWorldLocations();
+}
+
+namespace UE::Net
+{
+
+#if !UE_NET_ALLOW_MULTIPLE_REPLICATION_SYSTEMS
+	UReplicationSystem* FReplicationSystemFactory::ReplicationSystem = nullptr;
+#endif
+
+FReplicationSystemCreatedDelegate& FReplicationSystemFactory::GetReplicationSystemCreatedDelegate()
+{
+	static FReplicationSystemCreatedDelegate Delegate;
+	return Delegate;
+}
+
+FReplicationSystemDestroyedDelegate& FReplicationSystemFactory::GetReplicationSystemDestroyedDelegate()
+{
+	static FReplicationSystemDestroyedDelegate Delegate;
+	return Delegate;
+}
+
+// FReplicationSystemFactory
+#if UE_NET_ALLOW_MULTIPLE_REPLICATION_SYSTEMS
+UReplicationSystem* FReplicationSystemFactory::ReplicationSystems[MaxReplicationSystemCount];
+#endif
+
+UReplicationSystem* FReplicationSystemFactory::CreateReplicationSystem(const UReplicationSystem::FReplicationSystemParams& Params)
+{
+	if (!Params.ReplicationBridge)
+	{
+		UE_LOG(LogIris, Error, TEXT("Cannot create ReplicationSystem without a ReplicationBridge"));
+
+		return nullptr;
+	}
+
+#if UE_NET_ALLOW_MULTIPLE_REPLICATION_SYSTEMS
+	for (uint32 It = 0, EndIt = MaxReplicationSystemCount; It != EndIt; ++It)
+	{
+		if (ReplicationSystems[It] != nullptr)
+		{
+			continue;
+		}
+
+		UReplicationSystem* ReplicationSystem = NewObject<UReplicationSystem>();
+		ReplicationSystems[It] = ReplicationSystem;
+
+		ReplicationSystem->AddToRoot();
+		ReplicationSystem->Init(It, Params);
+
+		if (GetReplicationSystemCreatedDelegate().IsBound())
+		{
+			GetReplicationSystemCreatedDelegate().Broadcast(ReplicationSystem);
+		}
+
+		return ReplicationSystem;
+	}
+
+	LowLevelFatalError(TEXT("Too many ReplicationSystems have already been created (%u)"), MaxReplicationSystemCount);
+	return nullptr;
+#else
+	if (ReplicationSystem != nullptr)
+	{
+		LowLevelFatalError(TEXT("There can only be one ReplicationSystem"));
+	}
+
+	ReplicationSystem = NewObject<UReplicationSystem>();
+	ReplicationSystem->AddToRoot();
+	ReplicationSystem->Init(0, Params);
+
+	if (GetReplicationSystemCreatedDelegate().IsBound())
+	{
+		GetReplicationSystemCreatedDelegate().Broadcast(ReplicationSystem);
+	}
+	
+	return ReplicationSystem;
+#endif
+}
+
+void FReplicationSystemFactory::DestroyReplicationSystem(UReplicationSystem* System)
+{
+	if (System == nullptr)
+		return;
+
+#if UE_NET_ALLOW_MULTIPLE_REPLICATION_SYSTEMS
+	const uint32 Id = System->GetId();
+	if (Id < MaxReplicationSystemCount)
+	{
+		ReplicationSystems[Id] = nullptr;
+	}
+#else
+	if (System == ReplicationSystem)
+	{
+		ReplicationSystem = nullptr;
+	}
+#endif
+
+	if (GetReplicationSystemDestroyedDelegate().IsBound())
+	{
+		GetReplicationSystemDestroyedDelegate().Broadcast(System);
+	}
+
+	System->Shutdown();
+	System->RemoveFromRoot();
+	System->MarkAsGarbage();
+}
+
+}
diff --git a/Engine/Source/Runtime/Experimental/Iris/Core/Private/Iris/ReplicationSystem/ReplicationTypes.h b/Engine/Source/Runtime/Experimental/Iris/Core/Private/Iris/ReplicationSystem/ReplicationTypes.h
new file mode 100644
--- /dev/null
+++ b/Engine/Source/Runtime/Experimental/Iris/Core/Private/Iris/ReplicationSystem/ReplicationTypes.h
@@ -0,0 +1,30 @@
+// Copyright Epic Games, Inc. All Rights Reserved.
+
+#pragma once
+#include "Net/Core/NetBitArray.h"
+
+class UReplicationSystem;
+
+// If set to 1 we include the size of our creation state so that we can skip that on clients
+#define UE_NET_REPLICATION_SUPPORT_SKIP_INITIAL_STATE 0
+
+namespace UE::Net::Private
+{
+
+struct FReplicationParameters
+{
+	uint32 MaxActiveReplicatedObjectCount;	
+	uint32 PacketSendWindowSize;
+	uint32 ConnectionId;
+	UReplicationSystem* ReplicationSystem = nullptr;
+	bool bAllowSendingAttachmentsToObjectsNotInScope = false;
+	bool bAllowReceivingAttachmentsFromRemoteObjectsNotInScope = false;
+	bool bAllowDelayingAttachmentsWithUnresolvedReferences = false;
+	uint32 SmallObjectBitThreshold = 160U; // Number of bits remaining in a packet for us to consider trying to serialize a replicated object
+	uint32 MaxFailedSmallObjectCount = 10U;	// Number of objects that we try to serialize after an initial stream overflow to fill up a packet, this can improve bandwidth usage but comes at a cpu cost
+#if UE_NET_REPLICATION_SUPPORT_SKIP_INITIAL_STATE
+	uint32 NumBitsUsedForInitialStateSize = 20U;
+#endif
+};
+
+}
