Iris experimental replication system.

diff --git a/Engine/Source/Runtime/Experimental/Iris/Core/Private/Iris/ReplicationSystem/ReplicationWriter.cpp b/Engine/Source/Runtime/Experimental/Iris/Core/Private/Iris/ReplicationSystem/ReplicationWriter.cpp
new file mode 100644
--- /dev/null
+++ b/Engine/Source/Runtime/Experimental/Iris/Core/Private/Iris/ReplicationSystem/ReplicationWriter.cpp
@@ -0,0 +1,2754 @@
+// Copyright Epic Games, Inc. All Rights Reserved.
+
+#include "Iris/ReplicationSystem/ReplicationWriter.h"
+#include "Containers/Array.h"
+#include "Iris/IrisConfigInternal.h"
+#include "Iris/Core/IrisLog.h"
+#include "Iris/Core/IrisProfiler.h"
+#include "Iris/DataStream/DataStream.h"
+#include "Net/Core/Trace/NetTrace.h"
+#include "Iris/PacketControl/PacketNotification.h"
+#include "Iris/ReplicationSystem/ReplicationSystemTypes.h"
+#include "Iris/ReplicationSystem/DeltaCompression/DeltaCompressionBaselineManager.h"
+#include "Iris/ReplicationSystem/Filtering/ReplicationFiltering.h"
+#include "Iris/ReplicationSystem/NetHandleManager.h"
+#include "Iris/ReplicationSystem/ReplicationOperations.h"
+#include "Iris/ReplicationSystem/ReplicationOperationsInternal.h"
+#include "Iris/ReplicationSystem/ReplicationSystem.h"
+#include "Iris/ReplicationSystem/ReplicationSystemInternal.h"
+#include "Iris/ReplicationSystem/NetBlob/NetObjectBlobHandler.h"
+#include "Iris/ReplicationSystem/NetBlob/PartialNetObjectAttachmentHandler.h"
+#include "Iris/Serialization/NetReferenceCollector.h"
+#include "Iris/ReplicationSystem/NetTokenStoreState.h"
+#include "Iris/Serialization/InternalNetSerializationContext.h"
+#include "Iris/Serialization/NetBitStreamWriter.h"
+#include "Iris/Serialization/NetBitStreamUtil.h"
+#include "Iris/Serialization/NetSerializer.h"
+#include "Iris/Serialization/NetExportContext.h"
+#include "HAL/IConsoleManager.h"
+#include <algorithm>
+
+#if UE_NET_ENABLE_REPLICATIONWRITER_LOG
+#	define UE_LOG_REPLICATIONWRITER(Format, ...)  UE_LOG(LogIris, Log, Format, ##__VA_ARGS__)
+#	define UE_LOG_REPLICATIONWRITER_CONN(Format, ...)  UE_LOG(LogIris, Log, TEXT("Conn: %u ") Format, Parameters.ConnectionId, ##__VA_ARGS__)
+#else
+#	define UE_LOG_REPLICATIONWRITER(...)
+#	define UE_LOG_REPLICATIONWRITER_CONN(Format, ...)
+#endif
+
+#define UE_LOG_REPLICATIONWRITER_WARNING(Format, ...)  UE_LOG(LogIris, Warning, Format, ##__VA_ARGS__)
+#define UE_CLOG_REPLICATIONWRITER_WARNING(Condition, Format, ...)  UE_CLOG(Condition, LogIris, Warning, Format, ##__VA_ARGS__)
+
+namespace UE::Net::Private
+{
+
+static bool bWarnAboutDroppedAttachmentsToObjectsNotInScope = false;
+static FAutoConsoleVariableRef CVarWarnAboutDroppedAttachmentsToObjectsNotInScope(
+	TEXT("net.Iris.WarnAboutDroppedAttachmentsToObjectsNotInScope"),
+	bWarnAboutDroppedAttachmentsToObjectsNotInScope,
+	TEXT("Warn when attachments are dropped due to object not in scope. Default is false."
+	));
+
+static const FName NetError_ObjectStateTooLarge("Object state is too large to be split.");
+
+const TCHAR* FReplicationWriter::LexToString(const EReplicatedObjectState State)
+{
+	static const TCHAR* Names[] = {
+		TEXT("Invalid"),
+		TEXT("AttachmentToObjectNotInScope"),
+		TEXT("HugeObject"),
+		TEXT("PendingCreate"),
+		TEXT("WaitOnCreateConfirmation"),
+		TEXT("Created"),
+		TEXT("PendingFlush"),
+		TEXT("WaitOnFlush"),
+		TEXT("Flushed"),
+		TEXT("Hibernating"),
+		TEXT("PendingTearOff"),
+		TEXT("SubObjectPendingDestroy"),
+		TEXT("CancelPendingDestroy"),
+		TEXT("PendingDestroy"),
+		TEXT("WaitOnDestroyConfirmation"),
+		TEXT("Destroyed"),
+		TEXT("PermanentlyDestroyed"),
+	};
+	static_assert(UE_ARRAY_COUNT(Names) == uint32(EReplicatedObjectState::Count), "Missing names for one or more values of EReplicatedObjectState.");
+
+	return State < EReplicatedObjectState::Count ? Names[(uint32)State] : TEXT("");
+}
+
+void FReplicationWriter::FReplicationInfo::SetState(EReplicatedObjectState NewState)
+{
+	EReplicatedObjectState CurrentState = GetState();
+	switch (NewState)
+	{
+		case EReplicatedObjectState::PendingCreate:
+		{
+			check(CurrentState == EReplicatedObjectState::Invalid || CurrentState == EReplicatedObjectState::WaitOnCreateConfirmation);
+		}
+		break;
+		case EReplicatedObjectState::WaitOnCreateConfirmation:
+		{
+			check(CurrentState == EReplicatedObjectState::PendingCreate || CurrentState == EReplicatedObjectState::CancelPendingDestroy);
+		}
+		break;
+		case EReplicatedObjectState::Created:
+		{
+			check(CurrentState == EReplicatedObjectState::PendingCreate || CurrentState == EReplicatedObjectState::WaitOnCreateConfirmation || CurrentState == EReplicatedObjectState::CancelPendingDestroy);
+		}
+		break;
+		case EReplicatedObjectState::PendingTearOff:
+		{
+			check(CurrentState == EReplicatedObjectState::PendingTearOff || CurrentState == EReplicatedObjectState::WaitOnCreateConfirmation || CurrentState == EReplicatedObjectState::Created || CurrentState == EReplicatedObjectState::WaitOnDestroyConfirmation);
+		}
+		break;
+		case EReplicatedObjectState::SubObjectPendingDestroy:
+		{
+			check(CurrentState == EReplicatedObjectState::PendingDestroy || CurrentState == EReplicatedObjectState::SubObjectPendingDestroy || CurrentState == EReplicatedObjectState::WaitOnCreateConfirmation || CurrentState == EReplicatedObjectState::Created || CurrentState == EReplicatedObjectState::WaitOnDestroyConfirmation);
+		}
+		break;
+		case EReplicatedObjectState::PendingDestroy:
+		{
+			check(CurrentState != EReplicatedObjectState::Invalid);
+		}
+		break;
+		case EReplicatedObjectState::WaitOnDestroyConfirmation:
+		{
+			check(CurrentState >= EReplicatedObjectState::PendingTearOff);
+		}
+		break;
+		case EReplicatedObjectState::CancelPendingDestroy:
+		{
+			check(CurrentState == EReplicatedObjectState::WaitOnDestroyConfirmation || CurrentState == EReplicatedObjectState::CancelPendingDestroy);
+		}
+		break;
+		case EReplicatedObjectState::Destroyed:
+		{
+			check(CurrentState == EReplicatedObjectState::WaitOnDestroyConfirmation || CurrentState == EReplicatedObjectState::PendingTearOff || CurrentState == EReplicatedObjectState::CancelPendingDestroy);
+		}
+		break;
+		case EReplicatedObjectState::PermanentlyDestroyed:
+		{
+			check(CurrentState == EReplicatedObjectState::Invalid || CurrentState == EReplicatedObjectState::WaitOnCreateConfirmation);
+		}
+		break;
+		case EReplicatedObjectState::Invalid:
+		{
+			check(CurrentState == EReplicatedObjectState::PermanentlyDestroyed || CurrentState == EReplicatedObjectState::Destroyed || CurrentState == EReplicatedObjectState::PendingCreate);
+		}
+		break;
+
+		default:
+			check(false);
+			break;
+	};
+
+	State = (uint32)NewState;
+}
+
+
+FReplicationWriter::FHugeObjectContext::FHugeObjectContext()
+: SendStatus(EHugeObjectSendStatus::Idle)
+, InternalIndex(0)
+, DebugName(CreatePersistentNetDebugName(TEXT("HugeObjectState"), UE_ARRAY_COUNT(TEXT("HugeObjectState"))))
+{
+#if UE_NET_TRACE_ENABLED
+	DebugName->DebugNameId = FNetTrace::TraceName(DebugName->Name);
+#endif
+}
+
+// Default allocator for changemasks
+static FGlobalChangeMaskAllocator s_DefaultChangeMaskAllocator;
+
+// Helper class to process all ReplicationInfos for a record
+struct TReplicationRecordHelper
+{
+	typedef FReplicationWriter::FReplicationInfo FReplicationInfo;
+	typedef FReplicationRecord::FRecordInfoList FRecordInfoList;
+	typedef FReplicationWriter::EReplicatedObjectState EReplicatedObjectState;
+
+	FReplicationInfo* ReplicationInfos;
+	FRecordInfoList* ReplicationInfosRecordInfoLists;
+	FReplicationRecord* ReplicationRecord;
+
+	TReplicationRecordHelper(FReplicationInfo* InReplicationInfos, FRecordInfoList* InReplicationInfosRecordInfoLists, FReplicationRecord* InReplicationRecordRecord)
+	: ReplicationInfos(InReplicationInfos)
+	, ReplicationInfosRecordInfoLists(InReplicationInfosRecordInfoLists)
+	, ReplicationRecord(InReplicationRecordRecord)
+	{
+	}
+	
+	template <typename T>
+	void Process(uint32 RecordInfoCount, T&& Functor)
+	{
+		for (uint32 It = 0; It < RecordInfoCount; ++It)
+		{
+ 			const FReplicationRecord::FRecordInfo& RecordInfo = ReplicationRecord->PeekInfo();
+			FNetObjectAttachmentsWriter::ReplicationRecord AttachmentRecord = (RecordInfo.HasAttachments ? ReplicationRecord->DequeueAttachmentRecord() : 0U);
+			FReplicationInfo& Info = ReplicationInfos[RecordInfo.Index];
+			FReplicationRecord::FRecordInfoList& RecordInfoList = ReplicationInfosRecordInfoLists[RecordInfo.Index];
+
+			// We Need to cache this as the the ReplicationInfo might be invalidated by the functor
+			const uint32 ChangeMaskBitCount = Info.ChangeMaskBitCount;
+
+			// Invoke function
+			Functor(RecordInfo, Info, AttachmentRecord);
+
+			// We must free any dynamic memory allocated in PushRecordInfo
+			if (RecordInfo.HasChangeMask)
+			{
+				FChangeMaskStorageOrPointer::Free(RecordInfo.ChangeMaskOrPtr, ChangeMaskBitCount, s_DefaultChangeMaskAllocator);
+			}
+
+			// We must remove the record and unlink it
+			// It is safe to call even if we have stopped replicating the object		
+			ReplicationRecord->PopInfoAndRemoveFromList(RecordInfoList);
+		}
+	}
+};
+
+#if UE_NET_VALIDATE_REPLICATION_RECORD
+
+static bool s_ValidateReplicationRecord(const FReplicationRecord* ReplicationRecord, uint32 MaxInternalIndexCount, bool bVerifyFirstRecord)
+{
+	if (ReplicationRecord->GetRecordCount() == 0U)
+	{
+		return true;
+	}
+
+	// validate count
+	{
+		uint32 TotalPushedInfos = 0U;
+		for (uint32 It = 0U, EndIt = ReplicationRecord->GetRecordCount(); It < EndIt; ++It)
+		{
+			TotalPushedInfos += ReplicationRecord->PeekRecordAtOffset(It);
+		}
+
+		if (TotalPushedInfos != ReplicationRecord->GetInfoCount())
+		{
+			ensure(false);
+			return false;
+		}
+	}
+	
+	// Verify last / first record
+	const uint16 RecordInfoCount = bVerifyFirstRecord ? ReplicationRecord->PeekRecordAtOffset(0) : (ReplicationRecord->PeekRecordAtOffset(ReplicationRecord->GetRecordCount() - 1));
+
+	// check for duplicates
+	{
+		FNetBitArray BitArray;
+		BitArray.Init(MaxInternalIndexCount);
+		
+		uint32 Offset = bVerifyFirstRecord ? 0U : ReplicationRecord->GetInfoCount() - RecordInfoCount;
+		for (uint32 It = 0U; It < RecordInfoCount; ++It)
+		{
+ 			const FReplicationRecord::FRecordInfo& RecordInfo = ReplicationRecord->PeekInfoAtOffset(It + Offset);
+			if (BitArray.GetBit(RecordInfo.Index))
+			{
+				ensure(false);
+				return false;
+			}
+			BitArray.SetBit(RecordInfo.Index);
+		}
+	}
+
+	return true;
+}
+
+#endif
+
+FReplicationWriter::~FReplicationWriter()
+{
+	DiscardAllRecords();
+	StopAllReplication();
+}
+
+void FReplicationWriter::SetReplicationEnabled(bool bInReplicationEnabled)
+{
+	bReplicationEnabled = bInReplicationEnabled;
+}
+
+bool FReplicationWriter::IsReplicationEnabled() const
+{
+	return bReplicationEnabled;
+}
+
+// $IRIS TODO : May need to introduce queue and send behaviors. For example one may want to send only with object.
+// One may not want to send unless the object is replicated very soon etc.
+void FReplicationWriter::QueueNetObjectAttachments(FInternalNetHandle OwnerInternalIndex, FInternalNetHandle SubObjectInternalIndex, TArrayView<const TRefCountPtr<FNetBlob>> InAttachments)
+{
+	if (InAttachments.Num() <= 0)
+	{
+		ensureMsgf(false, TEXT("%s"), TEXT("QueueNetObjectAttachments expects at least one attachment."));
+		return;
+	}
+
+	const bool bObjectInScope = ObjectsInScope.GetBit(OwnerInternalIndex);
+	if (!bObjectInScope && !Parameters.bAllowSendingAttachmentsToObjectsNotInScope)
+	{
+		UE_CLOG_REPLICATIONWRITER_WARNING(bWarnAboutDroppedAttachmentsToObjectsNotInScope, TEXT("Dropping %s attachment due to object ( InternalIndex: %u ) not in scope."), (EnumHasAnyFlags(InAttachments[0]->GetCreationInfo().Flags, ENetBlobFlags::Reliable) ? TEXT("reliable") : TEXT("unreliable")), OwnerInternalIndex);
+		return;
+	}
+
+	const uint32 TargetIndex = bObjectInScope ? (SubObjectInternalIndex != FNetHandleManager::InvalidInternalIndex ? SubObjectInternalIndex : OwnerInternalIndex) : ObjectIndexForOOBAttachment;
+	ENetObjectAttachmentType AttachmentType = (bObjectInScope ? ENetObjectAttachmentType::Normal : ENetObjectAttachmentType::OutOfBand);
+	if (!Attachments.Enqueue(AttachmentType, TargetIndex, InAttachments))
+	{
+		return;
+	}
+
+	// There's a special case for out of band attachments, we don't need to mark anything dirty.
+	if (IsObjectIndexForOOBAttachment(TargetIndex))
+	{
+		return;
+	}
+
+	ObjectsWithDirtyChanges.SetBit(OwnerInternalIndex);
+
+	FReplicationInfo& TargetInfo = GetReplicationInfo(TargetIndex);
+	TargetInfo.HasAttachments = 1;
+
+	if (OwnerInternalIndex != SubObjectInternalIndex)
+	{
+		FReplicationInfo& OwnerInfo = GetReplicationInfo(OwnerInternalIndex);
+		OwnerInfo.HasDirtySubObjects = 1;
+	}
+}
+
+void FReplicationWriter::SetState(uint32 InternalIndex, EReplicatedObjectState NewState)
+{
+	FReplicationInfo& Info = GetReplicationInfo(InternalIndex);
+
+	UE_LOG_REPLICATIONWRITER_CONN(TEXT("ReplicationWriter.SetState for ( InternalIndex: %u ) %s -> %s"), InternalIndex, LexToString(Info.GetState()), LexToString(NewState));
+	Info.SetState(NewState);
+}
+
+void FReplicationWriter::Init(const FReplicationParameters& InParameters)
+{
+	// Store copy of parameters
+	Parameters = InParameters;
+
+	// Cache internal systems
+	ReplicationSystemInternal = Parameters.ReplicationSystem->GetReplicationSystemInternal();
+	NetHandleManager = &ReplicationSystemInternal->GetNetHandleManager();
+	ReplicationBridge = Parameters.ReplicationSystem->GetReplicationBridge();
+	BaselineManager = &ReplicationSystemInternal->GetDeltaCompressionBaselineManager();
+	ObjectReferenceCache = &ReplicationSystemInternal->GetObjectReferenceCache();
+	ReplicationFiltering = &ReplicationSystemInternal->GetFiltering();
+	ReplicationConditionals = &ReplicationSystemInternal->GetConditionals();
+	const FNetBlobManager* NetBlobManager = &ReplicationSystemInternal->GetNetBlobManager();
+	PartialNetObjectAttachmentHandler = NetBlobManager->GetPartialNetObjectAttachmentHandler();
+	NetObjectBlobHandler = NetBlobManager->GetNetObjectBlobHandler();
+
+	// Init book keeping
+	ReplicatedObjects.SetNumZeroed(Parameters.MaxActiveReplicatedObjectCount);
+	ReplicatedObjectsRecordInfoLists.SetNumZeroed(Parameters.MaxActiveReplicatedObjectCount);
+	SchedulingPriorities.SetNumZeroed(Parameters.MaxActiveReplicatedObjectCount);
+
+	ObjectsPendingDestroy.Init(Parameters.MaxActiveReplicatedObjectCount);
+	ObjectsWithDirtyChanges.Init(Parameters.MaxActiveReplicatedObjectCount);
+	ObjectsInScope.Init(Parameters.MaxActiveReplicatedObjectCount);	
+	WriteContext.ObjectsWrittenThisTick.Init(Parameters.MaxActiveReplicatedObjectCount);
+
+	// Attachments init
+	SetupReplicationInfoForAttachmentsToObjectsNotInScope();
+
+	bReplicationEnabled = false;
+}
+
+void FReplicationWriter::GetInitialChangeMask(ChangeMaskStorageType* ChangeMaskData, const FReplicationProtocol* Protocol)
+{
+	FNetBitArrayView ChangeMask(ChangeMaskData, Protocol->ChangeMaskBitCount, FNetBitArrayView::NoResetNoValidate);
+
+	// Just fill with all dirty for now
+	ChangeMask.SetAllBits();
+}
+
+void FReplicationWriter::StartReplication(uint32 InternalIndex)
+{
+	FReplicationInfo& Info = GetReplicationInfo(InternalIndex);
+
+	check(ReplicatedObjects[InternalIndex].GetState() == EReplicatedObjectState::Invalid);
+
+	// Reset info
+	Info = FReplicationInfo();
+	Info.LastAckedBaselineIndex = FDeltaCompressionBaselineManager::InvalidBaselineIndex;
+	Info.PendingBaselineIndex = FDeltaCompressionBaselineManager::InvalidBaselineIndex;
+
+	const uint32 bIsDestructionInfo = NetHandleManager->GetIsDestroyedStartupObject(InternalIndex) ? 1U : 0U;	
+	if (bIsDestructionInfo)
+	{
+		// Check status of original object about to be destroyed, if it has been confirmed as created we do not replicate the destruction info object at all
+		if (const uint32 OriginalInternalIndex = NetHandleManager->GetOriginalDestroyedStartupObjectIndex(InternalIndex))
+		{
+			const FReplicationInfo& OriginalInfo = GetReplicationInfo(OriginalInternalIndex);
+			if (OriginalInfo.GetState() != EReplicatedObjectState::Invalid && OriginalInfo.IsCreationConfirmed)
+			{
+				// We do not need to send the destruction info so we mark it as PermanentlyDestroyed
+				SetState(InternalIndex, EReplicatedObjectState::PermanentlyDestroyed);
+
+				Info.IsDestructionInfo = bIsDestructionInfo;
+				Info.IsCreationConfirmed = 1U;
+
+				NetHandleManager->AddNetObjectRef(InternalIndex);
+
+				return;
+			}
+		}
+	}
+
+	// Pending create
+	SetState(InternalIndex, EReplicatedObjectState::PendingCreate);
+
+	const FNetHandleManager::FReplicatedObjectData& Data = NetHandleManager->GetReplicatedObjectDataNoCheck(InternalIndex);
+	NetHandleManager->AddNetObjectRef(InternalIndex);
+
+	Info.ChangeMaskBitCount = Data.Protocol->ChangeMaskBitCount;
+	Info.HasDirtySubObjects = 1U;
+	Info.IsSubObject = NetHandleManager->GetSubObjectInternalIndices().GetBit(InternalIndex);
+	Info.HasDirtyChangeMask = 1U;
+	Info.HasAttachments = 0U;
+	Info.HasChangemaskFilter = EnumHasAnyFlags(Data.Protocol->ProtocolTraits, EReplicationProtocolTraits::HasConditionalChangeMask);
+	Info.IsDestructionInfo = bIsDestructionInfo;
+	Info.IsCreationConfirmed = 0U;
+	Info.TearOff = Data.bTearOff;
+	Info.SubObjectPendingDestroy = 0U;
+	Info.IsDeltaCompressionEnabled = BaselineManager->GetDeltaCompressionStatus(InternalIndex) == ENetObjectDeltaCompressionStatus::Allow ? 1U : 0U;
+	Info.LastAckedBaselineIndex = FDeltaCompressionBaselineManager::InvalidBaselineIndex;
+	Info.PendingBaselineIndex = FDeltaCompressionBaselineManager::InvalidBaselineIndex;
+
+	// Allocate storage for changemask (if needed)
+	FChangeMaskStorageOrPointer::Alloc(Info.ChangeMaskOrPtr, Info.ChangeMaskBitCount, s_DefaultChangeMaskAllocator);
+
+	// Get Initial ChangeMask
+	GetInitialChangeMask(Info.GetChangeMaskStoragePointer(), Data.Protocol);
+
+	// Reset record for object
+	ReplicationRecord.ResetList(ReplicatedObjectsRecordInfoLists[InternalIndex]);
+
+	// Set initial priority
+	// Subobject are always set to have zero priority as they are replicated with owner
+	// Currently we also do this for dependent objects to support objects with zero priority that should only replicate with parents
+	SchedulingPriorities[InternalIndex] = (Data.IsDependentObject() || Info.IsSubObject) ? 0.f : CreatePriority;
+
+	UE_LOG_REPLICATIONWRITER_CONN(TEXT("ReplicationWriter.StartReplication for ( InternalIndex: %u ) %s"), InternalIndex, *Data.Handle.ToString());
+
+	ObjectsWithDirtyChanges.SetBit(InternalIndex);
+}
+
+void FReplicationWriter::StopReplication(uint32 InternalIndex)
+{
+	FReplicationInfo& Info = GetReplicationInfo(InternalIndex);
+
+	// Invalidate state
+	SetState(InternalIndex, EReplicatedObjectState::Invalid);
+
+	// Need to free allocated ChangeMask (if it is allocated)
+	FChangeMaskStorageOrPointer::Free(Info.ChangeMaskOrPtr, Info.ChangeMaskBitCount, s_DefaultChangeMaskAllocator);
+	
+	Info.IsCreationConfirmed = 0U;
+
+	// Remove from objects with dirty changes
+	ObjectsWithDirtyChanges.ClearBit(InternalIndex);
+
+	// Remove from pending destroy
+	ObjectsPendingDestroy.ClearBit(InternalIndex);
+
+	// Explicitly remove from objects in scope since we might call StopReplication from outside ScopeUpdate
+	ObjectsInScope.ClearBit(InternalIndex);
+
+	UE_LOG_REPLICATIONWRITER_CONN(TEXT("ReplicationWriter.StopReplication for ( InternalIndex: %u )"), InternalIndex);
+	NetHandleManager->ReleaseNetObjectRef(InternalIndex);
+
+	Attachments.DropAllAttachments(ENetObjectAttachmentType::Normal, InternalIndex);
+
+	// Release baselines
+	if (Info.PendingBaselineIndex != FDeltaCompressionBaselineManager::InvalidBaselineIndex)
+	{
+		BaselineManager->DestroyBaseline(Parameters.ConnectionId, InternalIndex, Info.PendingBaselineIndex);
+	}
+	if (Info.PendingBaselineIndex != Info.LastAckedBaselineIndex && Info.LastAckedBaselineIndex != FDeltaCompressionBaselineManager::InvalidBaselineIndex)
+	{
+		BaselineManager->DestroyBaseline(Parameters.ConnectionId, InternalIndex, Info.LastAckedBaselineIndex);
+	}
+
+	Info.PendingBaselineIndex = Info.LastAckedBaselineIndex = FDeltaCompressionBaselineManager::InvalidBaselineIndex;
+}
+
+FReplicationWriter::FReplicationInfo& FReplicationWriter::GetReplicationInfo(uint32 InternalIndex)
+{
+	return ReplicatedObjects[InternalIndex];
+}
+
+const FReplicationWriter::FReplicationInfo& FReplicationWriter::GetReplicationInfo(uint32 InternalIndex) const
+{
+	return ReplicatedObjects[InternalIndex];
+}
+
+void FReplicationWriter::WriteNetHandleId(FNetBitStreamWriter& Writer, FNetHandle Handle)
+{
+	Writer.WriteBits(Handle.GetId(), FNetHandle::IdBits);
+}
+
+void FReplicationWriter::UpdateScope(const FNetBitArrayView& UpdatedScope)
+{
+	IRIS_PROFILER_SCOPE(FReplicationWriter_ScopeUpdate);
+
+	auto&& NewObjectFunctor = [this](uint32 Index)
+	{
+		// We can only start replicating an object that is not currently replicated
+		// Later on we will support "cancelling" destroy objects that are flushing state data
+		FReplicationInfo& Info = GetReplicationInfo(Index);
+		const EReplicatedObjectState State = Info.GetState();
+
+		if (State == EReplicatedObjectState::Invalid)
+		{
+			StartReplication(Index);
+		}
+		else if (State == EReplicatedObjectState::WaitOnDestroyConfirmation || State == EReplicatedObjectState::CancelPendingDestroy)
+		{
+			// Need to clear the pending destroy bit or else the object will be masked out of ObjectsInScope.
+			// Keep the SubObjectPendingDestroy status as is until we know if the destroy packet was received or not.
+			ObjectsPendingDestroy.ClearBit(Index);
+			SetState(Index, EReplicatedObjectState::CancelPendingDestroy);
+		}
+		else if (State == EReplicatedObjectState::SubObjectPendingDestroy || State == EReplicatedObjectState::PendingDestroy)
+		{
+			// Object was waiting to be destroyed but should now resume replication.
+			// If the object has been created we can go back to Created state, otherwise we go back to WaitOnCreateConfirmation
+			SetState(Index, EReplicatedObjectState::WaitOnDestroyConfirmation);
+			SetState(Index, EReplicatedObjectState::CancelPendingDestroy);
+			SetState(Index, (Info.IsCreationConfirmed ? EReplicatedObjectState::Created : EReplicatedObjectState::WaitOnCreateConfirmation));
+
+			Info.SubObjectPendingDestroy = 0U;
+			Info.HasDirtyChangeMask |= FNetBitArrayView(Info.GetChangeMaskStoragePointer(), Info.ChangeMaskBitCount).IsAnyBitSet();
+			ObjectsWithDirtyChanges.SetBitValue(Index, Info.HasDirtyChangeMask);
+			ObjectsPendingDestroy.ClearBit(Index);
+
+			if (State == EReplicatedObjectState::SubObjectPendingDestroy)
+			{
+				// If owner is not pending destroy we mark it as dirty as appropriate
+				FNetHandleManager::FReplicatedObjectData& ObjectData = NetHandleManager->GetReplicatedObjectDataNoCheck(Index);
+				FReplicationInfo& OwnerInfo = ReplicatedObjects[ObjectData.SubObjectRootIndex];
+				if (OwnerInfo.GetState() < EReplicatedObjectState::PendingDestroy)
+				{
+					OwnerInfo.HasDirtySubObjects |= Info.HasDirtyChangeMask;
+					ObjectsWithDirtyChanges.SetBitValue(ObjectData.SubObjectRootIndex, ObjectsWithDirtyChanges.GetBit(ObjectData.SubObjectRootIndex) | Info.HasDirtyChangeMask);
+				}
+			}
+		}
+		else
+		{
+			UE_LOG_REPLICATIONWRITER_CONN(TEXT("New object added to scope, Waiting to start replication for ( InternalIndex: %u ) currently in State: %s "), Index, LexToString(State));
+
+			check(ObjectsWithDirtyChanges.GetBit(Index) == false);
+			check(Info.HasDirtyChangeMask == 0U);
+		}
+	};
+
+	auto&& DestroyedObjectFunctor = [this](uint32 Index) 
+	{
+		// Request object to be destroyed
+		FReplicationInfo& Info = GetReplicationInfo(Index);
+
+		// We handle objects marked for tear-off using the state update path.
+		if (Info.TearOff)
+		{
+			return;
+		}
+
+		const EReplicatedObjectState State = Info.GetState();
+		if (State < EReplicatedObjectState::PendingDestroy)
+		{
+			// We have not sent the object yet so we can just stop replication
+			if (State == EReplicatedObjectState::PendingCreate)
+			{
+				StopReplication(Index);
+			}
+			else if (State == EReplicatedObjectState::CancelPendingDestroy)
+			{
+				// If we wanted to cancel the pending destroy but now want to destroy the object again we can resume waiting for the destroy.
+				ObjectsPendingDestroy.SetBit(Index);
+				SetState(Index, EReplicatedObjectState::WaitOnDestroyConfirmation);
+			}
+			else if (State < EReplicatedObjectState::PendingDestroy)
+			{
+				// Handle special inlined destroy path for subobjects
+				if (Info.IsSubObject)
+				{
+					FNetHandleManager::FReplicatedObjectData& ObjectData = NetHandleManager->GetReplicatedObjectDataNoCheck(Index);
+					if (ObjectData.IsSubObject())
+					{
+						// If owner is not pending destroy we mark it as dirty so that we can replicate subobject destruction properly
+						// We might get away with not doing this if owner or subobject does not have any unconfirmed changes in flight.
+						FReplicationInfo& OwnerInfo = ReplicatedObjects[ObjectData.SubObjectRootIndex];
+						if (OwnerInfo.GetState() < EReplicatedObjectState::PendingDestroy)
+						{
+							ObjectsWithDirtyChanges.SetBit(ObjectData.SubObjectRootIndex);
+							OwnerInfo.HasDirtySubObjects = 1U;
+
+							SetState(Index, EReplicatedObjectState::SubObjectPendingDestroy);
+							ObjectsPendingDestroy.SetBit(Index);
+							Info.SubObjectPendingDestroy = 1U;			
+			
+							return;
+						}
+					}
+				}
+				// Subobject owner should try to fixup any objects in SubObjectPendingDestroyState
+				else if (Info.HasDirtySubObjects)
+				{
+					for (uint32 SubObjectIndex : NetHandleManager->GetSubObjects(Index))
+					{
+						FReplicationInfo& SubObjectInfo = GetReplicationInfo(SubObjectIndex);
+						if (SubObjectInfo.GetState() == EReplicatedObjectState::SubObjectPendingDestroy)
+						{
+							SubObjectInfo.SetState(EReplicatedObjectState::PendingDestroy);
+							SubObjectInfo.SubObjectPendingDestroy = 0U;
+						}
+					}
+				}
+				
+
+				SetState(Index, EReplicatedObjectState::PendingDestroy);
+				ObjectsPendingDestroy.SetBit(Index);
+				ObjectsWithDirtyChanges.ClearBit(Index);
+
+				Info.HasDirtyChangeMask = 0U;
+			}
+		}
+		else if (State == EReplicatedObjectState::PermanentlyDestroyed)
+		{
+			StopReplication(Index);
+		}
+	};
+
+	FNetBitArrayView CurrentScope = MakeNetBitArrayView(ObjectsInScope);
+	FNetBitArrayView::ForAllExclusiveBits(UpdatedScope, CurrentScope, NewObjectFunctor, DestroyedObjectFunctor);
+	CurrentScope.Copy(UpdatedScope);
+
+	// No objects marked for destroy can be in scope
+	ObjectsInScope.Combine(ObjectsPendingDestroy, FNetBitArrayBase::AndNotOp);
+}
+
+void FReplicationWriter::InternalUpdateDirtyChangeMasks(const FChangeMaskCache& CachedChangeMasks, uint32 bTearOff)
+{
+	IRIS_PROFILER_SCOPE(FReplicationWriter_UpdateDirtyChangeMasks);
+
+	const ChangeMaskStorageType* StoragePtr = CachedChangeMasks.Storage.GetData();
+	for (const auto& Entry : CachedChangeMasks.Indices)
+	{
+		if (!ObjectsInScope.GetBit(Entry.InternalIndex))
+		{
+			continue;
+		}
+
+		ObjectsWithDirtyChanges.SetBit(Entry.InternalIndex);
+		FReplicationInfo& Info = ReplicatedObjects[Entry.InternalIndex];
+
+		if (Entry.bMarkSubObjectOwnerDirty == 0U)
+		{		
+			// Mark object for TearOff, that is that we will stop replication as soon as the tear-off is acknowledged
+			Info.TearOff = bTearOff;
+
+			// Merge in dirty changes
+			if (Entry.bHasDirtyChangeMask)
+			{
+				const uint32 ChangeMaskBitCount = Info.ChangeMaskBitCount;
+
+				// Merge updated changes
+				FNetBitArrayView Changes(Info.GetChangeMaskStoragePointer(), ChangeMaskBitCount);
+
+				const FNetBitArrayView UpdatedChanges = MakeNetBitArrayView(StoragePtr + Entry.StorageOffset, ChangeMaskBitCount);
+				Changes.Combine(UpdatedChanges, FNetBitArrayView::OrOp);
+
+				// Mark changemask as dirty
+				Info.HasDirtyChangeMask = 1U;
+			}
+
+			// In order to support flush behavior of tear off objects we should make sure to include everything in the final state https://jira.it.epicgames.com/browse/UENET-1079
+			// We should walk the list of any in-flight changes and include them to the changemask to ensure that
+			// all states are included with the tear-off
+		}
+		else
+		{
+			Info.HasDirtySubObjects = 1U;
+		}
+	}
+
+	//UE_LOG_REPLICATIONWRITER(TEXT("FReplicationWriter::UpdateDirtyChangeMasks() Updated %u Objects for ConnectionId:%u, ReplicationSystemId: %u."), CachedChangeMasks.Indices.Num(), Parameters.ConnectionId, Parameters.ReplicationSystem->GetId());	
+}
+
+void FReplicationWriter::NotifyDestroyedObjectPendingTearOff(FInternalNetHandle ObjectInternalIndex)
+{
+	const FReplicationInfo& ReplicationInfo = GetReplicationInfo(ObjectInternalIndex);
+	if (ReplicationInfo.GetState() == EReplicatedObjectState::PendingCreate)
+	{
+		check(ReplicationInfo.TearOff == 1U);
+		StopReplication(ObjectInternalIndex);
+	}
+}
+
+const FNetBitArray& FReplicationWriter::GetObjectsRequiringPriorityUpdate() const
+{
+	return ObjectsWithDirtyChanges;
+}
+
+void FReplicationWriter::UpdatePriorities(const float* UpdatedPriorities)
+{
+	IRIS_PROFILER_SCOPE(FReplicationWriter_UpdatePriorities);
+	auto&& UpdatePriority = [LocalPriorities = SchedulingPriorities.GetData(), UpdatedPriorities](uint32 Index)
+	{
+		LocalPriorities[Index] += UpdatedPriorities[Index];
+	};
+
+	ObjectsWithDirtyChanges.ForAllSetBits(UpdatePriority);
+}
+
+uint32 FReplicationWriter::ScheduleObjects(FScheduleObjectInfo* OutScheduledObjectIndices)
+{
+	IRIS_PROFILER_SCOPE(FReplicationWriter_ScheduleObjects);
+
+	uint32 ScheduledObjectCount = 0;
+	float* LocalPriorities = SchedulingPriorities.GetData();
+
+	FScheduleObjectInfo* ScheduledObjectIndices = OutScheduledObjectIndices;
+
+	// Special index is handled later.
+	ObjectsWithDirtyChanges.ClearBit(ObjectIndexForOOBAttachment);
+
+	const FNetBitArray& UpdatedObjects = ObjectsWithDirtyChanges;
+	const FNetBitArray& SubObjects = NetHandleManager->GetSubObjectInternalIndices();
+
+	// Fill IndexList with all objects with dirty state that have positive priority excluding subobjects
+	auto&& FillIndexListFunc = [&LocalPriorities, &ScheduledObjectIndices, &ScheduledObjectCount](uint32 Index)
+	{
+		const float UpdatedPriority = LocalPriorities[Index];
+
+		FScheduleObjectInfo& ScheduledObjectInfo = ScheduledObjectIndices[ScheduledObjectCount];
+		ScheduledObjectInfo.Index = Index;
+		ScheduledObjectInfo.SortKey = UpdatedPriority;
+
+		ScheduledObjectCount += (UpdatedPriority >= FReplicationWriter::SchedulingThresholdPriority) ? 1.f : 0.f;
+	};
+
+	// Invoke functor for all updated objects that not are sub objects.
+	FNetBitArray::ForAllSetBits(UpdatedObjects, SubObjects, FNetBitArray::AndNotOp, FillIndexListFunc);
+
+	// we now have our list of objects to write.
+	return ScheduledObjectCount;
+}
+
+uint32 FReplicationWriter::SortScheduledObjects(FScheduleObjectInfo* ScheduledObjectIndices, uint32 ScheduledObjectCount, uint32 StartIndex)
+{
+	check(ScheduledObjectCount > 0 && StartIndex <= ScheduledObjectCount);
+
+	// Partial sort of scheduled objects
+	{
+		IRIS_PROFILER_SCOPE(FReplicationWriter_SortScheduledObjects);
+
+		// We only need a partial sort of the highest priority objects as we wont be able to fit that much data in a packet anyway
+		// $IRIS TODO: Implement and evalute partial sort algorithm, currently we simply use std::partial_sort https://jira.it.epicgames.com/browse/UE-123444
+		FScheduleObjectInfo* StartIt = ScheduledObjectIndices + StartIndex;
+		FScheduleObjectInfo* EndIt = ScheduledObjectIndices + ScheduledObjectCount;
+		FScheduleObjectInfo* SortIt = FMath::Min(StartIt + PartialSortObjectCount, EndIt);
+
+		std::partial_sort(StartIt, SortIt, EndIt, [](const FScheduleObjectInfo& EntryA, const FScheduleObjectInfo& EntryB) { return EntryA.SortKey > EntryB.SortKey; });
+	}
+
+	return FMath::Min(ScheduledObjectCount - StartIndex, PartialSortObjectCount);
+}
+
+void FReplicationWriter::HandleDeliveredRecord(const FReplicationRecord::FRecordInfo& RecordInfo, FReplicationInfo& Info, FNetObjectAttachmentsWriter::ReplicationRecord AttachmentRecord)
+{
+	EReplicatedObjectState DeliveredState = (EReplicatedObjectState)RecordInfo.ReplicatedObjectState;
+	EReplicatedObjectState CurrentState = Info.GetState();
+	const uint32 InternalIndex = RecordInfo.Index;
+	
+	checkf(CurrentState != EReplicatedObjectState::Invalid, TEXT("Object ( InternalIndex: %u ) has an invalid state."), InternalIndex);
+
+	// We confirmed a new baseline
+	if (RecordInfo.NewBaselineIndex != FDeltaCompressionBaselineManager::InvalidBaselineIndex)
+	{
+		check(RecordInfo.NewBaselineIndex == Info.PendingBaselineIndex);
+
+		// Destroy old baseline
+		if (Info.LastAckedBaselineIndex != FDeltaCompressionBaselineManager::InvalidBaselineIndex)
+		{			
+			BaselineManager->DestroyBaseline(Parameters.ConnectionId, InternalIndex, Info.LastAckedBaselineIndex);
+			UE_LOG_REPLICATIONWRITER_CONN(TEXT("Destroyed old baseline %u for ( InternalObjectIndex: %u )"), Info.LastAckedBaselineIndex, InternalIndex);			
+		}
+		Info.LastAckedBaselineIndex = RecordInfo.NewBaselineIndex;
+		Info.PendingBaselineIndex = FDeltaCompressionBaselineManager::InvalidBaselineIndex;
+
+		UE_LOG_REPLICATIONWRITER_CONN(TEXT("Acknowledged baseline %u for ( InternalObjectIndex: %u )"), RecordInfo.NewBaselineIndex, InternalIndex);
+	}
+
+	// Update state
+	switch (DeliveredState)
+	{
+		case EReplicatedObjectState::WaitOnCreateConfirmation:
+		{
+			// if we are still waiting for CreateConfirmation
+			if (CurrentState == EReplicatedObjectState::WaitOnCreateConfirmation)
+			{
+				// If this is a destruction info, just put it in the destroyed state
+				if (Info.IsDestructionInfo)
+				{
+					SetState(InternalIndex, EReplicatedObjectState::PermanentlyDestroyed);
+				}
+				// If this object was teared off, it can now be considered as destroyed
+				else if (RecordInfo.WroteTearOff)
+				{
+					SetState(InternalIndex, EReplicatedObjectState::PendingTearOff);
+					SetState(InternalIndex, EReplicatedObjectState::Destroyed);
+					StopReplication(InternalIndex);
+				}
+				else
+				{
+					SetState(InternalIndex, EReplicatedObjectState::Created);
+				}
+			}
+			Info.IsCreationConfirmed = 1U;
+		}
+		break;
+
+		case EReplicatedObjectState::WaitOnDestroyConfirmation:
+		{
+			SetState(InternalIndex, EReplicatedObjectState::Destroyed);
+
+			// It is now safe to stop tracking this object
+			StopReplication(InternalIndex);
+
+			if (CurrentState == EReplicatedObjectState::CancelPendingDestroy)
+			{
+				StartReplication(InternalIndex);
+				ObjectsInScope.SetBit(InternalIndex);
+			}
+		}
+		break;
+
+		case EReplicatedObjectState::AttachmentToObjectNotInScope:
+		{
+			check(IsObjectIndexForOOBAttachment(InternalIndex));
+			Attachments.OnPacketDelivered(ENetObjectAttachmentType::OutOfBand, ObjectIndexForOOBAttachment, AttachmentRecord);
+		}
+		return;
+
+		case EReplicatedObjectState::HugeObject:
+		{
+			check(IsObjectIndexForOOBAttachment(InternalIndex));
+			Attachments.OnPacketDelivered(ENetObjectAttachmentType::HugeObject, ObjectIndexForOOBAttachment, AttachmentRecord);
+
+			// If we've sent the entire state now we can clear the huge object state and proceed as normal.
+			if (Attachments.IsAllSentAndAcked(ENetObjectAttachmentType::HugeObject, ObjectIndexForOOBAttachment))
+			{
+				FReplicationInfo& ReplicationInfo = GetReplicationInfo(HugeObjectContext.InternalIndex);
+				for (const FObjectRecord& ObjectRecord : HugeObjectContext.BatchRecord.ObjectReplicationRecords)
+				{
+					FReplicationInfo& HugeObjectReplicationInfo =  GetReplicationInfo(ObjectRecord.Record.Index);
+					const uint32 ChangeMaskBitCount = HugeObjectReplicationInfo.ChangeMaskBitCount;
+					HandleDeliveredRecord(ObjectRecord.Record, HugeObjectReplicationInfo, ObjectRecord.AttachmentRecord);
+					if (ObjectRecord.Record.HasChangeMask)
+					{
+						FChangeMaskStorageOrPointer::Free(ObjectRecord.Record.ChangeMaskOrPtr, ChangeMaskBitCount, s_DefaultChangeMaskAllocator);
+					}
+				}
+
+				// We need to explicitly acknowledge exports made through the huge object batch
+				NetExports->AcknowledgeBatchExports(HugeObjectContext.BatchExports);
+				
+				ClearHugeObjectContext(HugeObjectContext);
+			}
+		}
+		return;
+
+		default:
+		break;
+	}
+
+	if (RecordInfo.HasAttachments)
+	{
+		Attachments.OnPacketDelivered(ENetObjectAttachmentType::Normal, InternalIndex, AttachmentRecord);
+	}
+}
+
+void FReplicationWriter::HandleDiscardedRecord(const FReplicationRecord::FRecordInfo& RecordInfo, FReplicationInfo& Info, FNetObjectAttachmentsWriter::ReplicationRecord AttachmentRecord)
+{
+	const EReplicatedObjectState DeliveredState = (EReplicatedObjectState)RecordInfo.ReplicatedObjectState;
+	const uint32 InternalIndex = RecordInfo.Index;
+
+	// There are a couple of special cases we need to handle. Regular attachments are ignored since they don't require special handling at the moment.
+	switch (DeliveredState)
+	{
+		// If we need to handle attachments this should return rather than fallback on some default path like HandleDeliveredRecord.
+		case EReplicatedObjectState::AttachmentToObjectNotInScope:
+		{
+		}
+		return;
+
+		case EReplicatedObjectState::HugeObject:
+		{
+			check(IsObjectIndexForOOBAttachment(InternalIndex));
+
+			// Deal with it similar to if the entire state has been sent as we need to go through all records.
+			{
+				FReplicationInfo& ReplicationInfo = GetReplicationInfo(HugeObjectContext.InternalIndex);
+				for (const FObjectRecord& ObjectRecord : HugeObjectContext.BatchRecord.ObjectReplicationRecords)
+				{
+					FReplicationInfo& HugeObjectReplicationInfo = GetReplicationInfo(ObjectRecord.Record.Index);
+					const uint32 ChangeMaskBitCount = HugeObjectReplicationInfo.ChangeMaskBitCount;
+					HandleDiscardedRecord(ObjectRecord.Record, HugeObjectReplicationInfo, ObjectRecord.AttachmentRecord);
+					if (ObjectRecord.Record.HasChangeMask)
+					{
+						FChangeMaskStorageOrPointer::Free(ObjectRecord.Record.ChangeMaskOrPtr, ChangeMaskBitCount, s_DefaultChangeMaskAllocator);
+					}
+				}
+
+				ClearHugeObjectContext(HugeObjectContext);
+			}
+		}
+		return;
+	}
+}
+
+template<>
+void FReplicationWriter::HandleDroppedRecord<FReplicationWriter::EReplicatedObjectState::WaitOnCreateConfirmation>(FReplicationWriter::EReplicatedObjectState CurrentState, const FReplicationRecord::FRecordInfo& RecordInfo, FReplicationInfo& Info, FNetObjectAttachmentsWriter::ReplicationRecord AttachmentRecord)
+{
+	const uint32 InternalIndex = RecordInfo.Index;
+
+	if (CurrentState < EReplicatedObjectState::Created)
+	{
+		// Resend creation data
+		SetState(InternalIndex, EReplicatedObjectState::PendingCreate);
+
+		// Must also restore changemask
+		FNetBitArrayView ChangeMask(Info.GetChangeMaskStoragePointer(), Info.ChangeMaskBitCount);
+		FNetBitArrayView LostChangeMask = FChangeMaskUtil::MakeChangeMask(RecordInfo.ChangeMaskOrPtr, Info.ChangeMaskBitCount);
+		ChangeMask.Combine(LostChangeMask, FNetBitArrayView::OrOp);
+
+		// Mark object as having dirty changes
+		ObjectsWithDirtyChanges.SetBit(InternalIndex);
+
+		// Mark changemask dirty
+		Info.HasDirtyChangeMask = 1U;
+
+		// Indicate that we have dirty subobjects
+		Info.HasDirtySubObjects = 1U;
+
+		// Mark attachments as dirty
+		Info.HasAttachments |= RecordInfo.HasAttachments;
+
+		// Bump prio
+		SchedulingPriorities[InternalIndex] += FReplicationWriter::CreatePriority;
+	}
+	else if (CurrentState == EReplicatedObjectState::SubObjectPendingDestroy)
+	{
+		// If SubObject has been destroyed while we where waiting for creation ack we can just stop replication
+		SetState(InternalIndex, EReplicatedObjectState::WaitOnDestroyConfirmation);
+		SetState(InternalIndex, EReplicatedObjectState::Destroyed);
+		StopReplication(InternalIndex);
+	}
+}
+
+template<>
+void FReplicationWriter::HandleDroppedRecord<FReplicationWriter::EReplicatedObjectState::Created>(FReplicationWriter::EReplicatedObjectState CurrentState, const FReplicationRecord::FRecordInfo& RecordInfo, FReplicationInfo& Info, FNetObjectAttachmentsWriter::ReplicationRecord AttachmentRecord)
+{
+	const uint32 InternalIndex = RecordInfo.Index;
+
+	// For now we do not support flush (https://jira.it.epicgames.com/browse/UENET-1079), so if we drop data in flight while we are pending destroy/tear-off we will ignore the lost data.
+	if (CurrentState < EReplicatedObjectState::PendingDestroy)
+	{
+		// Mask in any lost changes
+		bool bNeedToResendAttachments = RecordInfo.HasAttachments;
+		bool bNeedToResendState = false;
+
+		FNetBitArrayView LostChangeMask = FChangeMaskUtil::MakeChangeMask(RecordInfo.ChangeMaskOrPtr, (RecordInfo.HasChangeMask ? Info.ChangeMaskBitCount : 1U));
+		if (RecordInfo.HasChangeMask)
+		{
+			// Iterate over all data in flight for this object and mask away any already re-transmitted changes
+			// N.B. We don't check if this object is in huge object mode and check to see if any of these changes were part of that payload.
+			const FReplicationRecord::FRecordInfo* CurrentRecordInfo = ReplicationRecord.GetInfoForIndex(RecordInfo.NextIndex);
+			while (CurrentRecordInfo)
+			{
+				if (CurrentRecordInfo->HasChangeMask)
+				{
+					const FNetBitArrayView CurrentRecordInfoChangeMask(FChangeMaskUtil::MakeChangeMask(CurrentRecordInfo->ChangeMaskOrPtr, Info.ChangeMaskBitCount));
+					LostChangeMask.Combine(CurrentRecordInfoChangeMask, FNetBitArrayView::AndNotOp);
+				}
+
+				CurrentRecordInfo = ReplicationRecord.GetInfoForIndex(CurrentRecordInfo->NextIndex);
+			};
+
+			bNeedToResendState = LostChangeMask.IsAnyBitSet();
+		}
+
+		// if we lost changes that are not already retransmitted we update the changemask
+		if (bNeedToResendState | bNeedToResendAttachments)
+		{
+			if (bNeedToResendState)
+			{
+				FNetBitArrayView ChangeMask(Info.GetChangeMaskStoragePointer(), Info.ChangeMaskBitCount);
+				ChangeMask.Combine(LostChangeMask, FNetBitArrayView::OrOp);
+			}
+
+			// Mark object as having dirty changes
+			ObjectsWithDirtyChanges.SetBit(InternalIndex);
+
+			// Mark changemask as dirty
+			Info.HasDirtyChangeMask |= bNeedToResendState;
+
+			// Mark attachments as dirty
+			Info.HasAttachments |= bNeedToResendAttachments;
+
+			// Give slight priority bump
+			SchedulingPriorities[InternalIndex] += FReplicationWriter::LostStatePriorityBump;
+
+			// Check uncached flag in order to catch if object was added as a subobject or dependent object after data was sent
+			const bool bIsSubObjectOrDependentObject = NetHandleManager->GetSubObjectInternalIndices().GetBit(InternalIndex);
+			if (bIsSubObjectOrDependentObject)
+			{
+				// If this object is a subobject, mark owner dirty as well as subobjects always should be scheduled together with parent
+				const FNetHandleManager::FReplicatedObjectData& ObjectData = NetHandleManager->GetReplicatedObjectData(InternalIndex);
+				uint32 SubObjectOwnerInternalIndex = ObjectData.SubObjectRootIndex;
+
+				FReplicationInfo& SubObjectOwnerReplicationInfo = GetReplicationInfo(SubObjectOwnerInternalIndex);
+
+				if (ensure(SubObjectOwnerReplicationInfo.GetState() < EReplicatedObjectState::PendingDestroy))
+				{
+					// Mark owner as dirty
+					ObjectsWithDirtyChanges.SetBit(SubObjectOwnerInternalIndex);
+
+					// Indicate that we have dirty subobjects
+					SubObjectOwnerReplicationInfo.HasDirtySubObjects = 1U;
+
+					// Give slight priority bump to owner
+					SchedulingPriorities[SubObjectOwnerInternalIndex] += FReplicationWriter::LostStatePriorityBump;
+				}
+			}
+		}
+	}
+}
+
+template<>
+void FReplicationWriter::HandleDroppedRecord<FReplicationWriter::EReplicatedObjectState::WaitOnDestroyConfirmation>(FReplicationWriter::EReplicatedObjectState CurrentState, const FReplicationRecord::FRecordInfo& RecordInfo, FReplicationInfo& Info, FNetObjectAttachmentsWriter::ReplicationRecord AttachmentRecord)
+{
+	const uint32 InternalIndex = RecordInfo.Index;
+
+	check(CurrentState == EReplicatedObjectState::WaitOnDestroyConfirmation || CurrentState == EReplicatedObjectState::CancelPendingDestroy);
+
+	// If we want to cancel the destroy and lost the destroy packet we can resume normal replication.
+	if (CurrentState == EReplicatedObjectState::CancelPendingDestroy)
+	{
+		checkfSlow(!RecordInfo.WroteTearOff, TEXT("Torn off objects can't cancel destroy. ( InternalIndex: %u ) %s"), InternalIndex, ToCStr(NetHandleManager->GetReplicatedObjectData(InternalIndex).Handle.ToString()));
+
+		if (RecordInfo.WroteDestroySubObject && Info.SubObjectPendingDestroy)
+		{
+			// If the subobject owner still is replicated and valid
+			FNetHandleManager::FReplicatedObjectData& ObjectData = NetHandleManager->GetReplicatedObjectDataNoCheck(InternalIndex);
+			check(ObjectData.IsSubObject());
+
+			// If owner is not pending destroy we mark it as dirty as appropriate.
+			FNetBitArrayView ChangeMask(Info.GetChangeMaskStoragePointer(), Info.ChangeMaskBitCount);
+			Info.HasDirtyChangeMask |= ChangeMask.IsAnyBitSet();
+			Info.SubObjectPendingDestroy = 0U;
+			ObjectsPendingDestroy.ClearBit(InternalIndex);
+			SetState(InternalIndex, EReplicatedObjectState::Created);
+
+			FReplicationInfo& OwnerInfo = ReplicatedObjects[ObjectData.SubObjectRootIndex];
+			if (OwnerInfo.GetState() < EReplicatedObjectState::PendingDestroy)
+			{
+				OwnerInfo.HasDirtySubObjects |= Info.HasDirtyChangeMask;
+				ObjectsWithDirtyChanges.SetBitValue(ObjectData.SubObjectRootIndex, ObjectsWithDirtyChanges.GetBit(ObjectData.SubObjectRootIndex) | OwnerInfo.HasDirtySubObjects);
+			}
+		}
+		else
+		{
+			// Check whether there are any dirty changes and mark object as dirty as appropriate.
+			FNetBitArrayView ChangeMask(Info.GetChangeMaskStoragePointer(), Info.ChangeMaskBitCount);
+			Info.HasDirtyChangeMask |= ChangeMask.IsAnyBitSet();
+
+			ObjectsWithDirtyChanges.SetBitValue(InternalIndex, Info.HasDirtyChangeMask);
+
+			ObjectsPendingDestroy.ClearBit(InternalIndex);
+
+			SetState(InternalIndex, EReplicatedObjectState::Created);
+		}	
+	}
+	else
+	{
+		// We dropped a packet with tear-off data, that is a destroy with state data so we need to resend that state
+		if (RecordInfo.WroteTearOff)
+		{
+			check(Info.TearOff);
+
+			SetState(InternalIndex, EReplicatedObjectState::PendingTearOff);
+
+			if (RecordInfo.HasChangeMask)
+			{
+				// Must also restore changemask
+				FNetBitArrayView ChangeMask(Info.GetChangeMaskStoragePointer(), Info.ChangeMaskBitCount);
+				FNetBitArrayView LostChangeMask = FChangeMaskUtil::MakeChangeMask(RecordInfo.ChangeMaskOrPtr, Info.ChangeMaskBitCount);
+				ChangeMask.Combine(LostChangeMask, FNetBitArrayView::OrOp);
+
+				// Mark changemask dirty
+				Info.HasDirtyChangeMask = 1U;
+			}
+
+			// Indicate that we have dirty subobjects
+			Info.HasDirtySubObjects = 1U;
+
+			// Mark attachments as dirty
+			Info.HasAttachments |= RecordInfo.HasAttachments;
+
+			// Mark object as having dirty changes
+			ObjectsWithDirtyChanges.SetBit(InternalIndex);
+
+			// Bump prio
+			SchedulingPriorities[InternalIndex] += FReplicationWriter::TearOffPriority;
+		}
+		else if (RecordInfo.WroteDestroySubObject && Info.SubObjectPendingDestroy)
+		{
+			// If the subobject owner still is replicated and valid
+			FNetHandleManager::FReplicatedObjectData& ObjectData = NetHandleManager->GetReplicatedObjectDataNoCheck(InternalIndex);
+			check(ObjectData.IsSubObject());
+
+			// If owner is not pending destroy we mark it as dirty so that we can replicate subobject destruction properly
+			// We might get away with not doing this if owner or subobject does not have any unconfirmed changes in flight.
+			FReplicationInfo& OwnerInfo = ReplicatedObjects[ObjectData.SubObjectRootIndex];
+			if (OwnerInfo.GetState() < EReplicatedObjectState::PendingDestroy)
+			{
+				ObjectsWithDirtyChanges.SetBit(ObjectData.SubObjectRootIndex);
+				OwnerInfo.HasDirtySubObjects = 1U;
+
+				SetState(InternalIndex, EReplicatedObjectState::SubObjectPendingDestroy);
+				ObjectsPendingDestroy.SetBit(InternalIndex);
+			}
+		}
+		else
+		{
+			// Mark for resend of Destroy
+			ObjectsPendingDestroy.SetBit(InternalIndex);
+			ObjectsWithDirtyChanges.ClearBit(InternalIndex);
+			Info.HasDirtyChangeMask = 0U;
+
+			SetState(InternalIndex, EReplicatedObjectState::PendingDestroy);
+		}
+	}
+}
+
+void FReplicationWriter::HandleDroppedRecord(const FReplicationRecord::FRecordInfo& RecordInfo, FReplicationInfo& Info, FNetObjectAttachmentsWriter::ReplicationRecord AttachmentRecord)
+{
+	EReplicatedObjectState LostObjectState = (EReplicatedObjectState)RecordInfo.ReplicatedObjectState;
+	EReplicatedObjectState CurrentState = Info.GetState();
+	const uint32 InternalIndex = RecordInfo.Index;
+
+	check(CurrentState != EReplicatedObjectState::Invalid);
+
+	UE_LOG_REPLICATIONWRITER_CONN(TEXT("Handle dropped data for ( InternalIndex: %u ) %s, LostState %s, CurrentState is %s"), InternalIndex, *(NetHandleManager->GetReplicatedObjectData(InternalIndex).Handle.ToString()), LexToString(LostObjectState), LexToString(CurrentState));
+
+	// If we loose a baseline we must notify the BaselineManager and invalidate our pendingbaselineindex
+	if (RecordInfo.NewBaselineIndex != FDeltaCompressionBaselineManager::InvalidBaselineIndex)
+	{
+		check(RecordInfo.NewBaselineIndex == Info.PendingBaselineIndex);
+		UE_LOG_REPLICATIONWRITER_CONN(TEXT("Lost baseline %u for ( InternalObjectIndex: %u )"), RecordInfo.NewBaselineIndex, InternalIndex);
+		
+		BaselineManager->LostBaseline(Parameters.ConnectionId, InternalIndex, RecordInfo.NewBaselineIndex);
+		Info.PendingBaselineIndex = FDeltaCompressionBaselineManager::InvalidBaselineIndex;
+	}
+
+	switch (LostObjectState)
+	{
+		// We dropped creation state, restore state to PendingCreate and bump priority to make sure we sent it again
+		case EReplicatedObjectState::WaitOnCreateConfirmation:
+		{
+			HandleDroppedRecord<EReplicatedObjectState::WaitOnCreateConfirmation>(CurrentState, RecordInfo, Info, AttachmentRecord);
+		}
+		break;
+
+		// Object is created, update lost state data unless object are currently being flushed/teared-off or destroyed
+		case EReplicatedObjectState::Created:
+		case EReplicatedObjectState::PendingFlush:
+		case EReplicatedObjectState::WaitOnFlush:
+		{
+			HandleDroppedRecord<EReplicatedObjectState::Created>(CurrentState, RecordInfo, Info, AttachmentRecord);
+		}
+		break;
+
+		case EReplicatedObjectState::WaitOnDestroyConfirmation:
+		{
+			HandleDroppedRecord<EReplicatedObjectState::WaitOnDestroyConfirmation>(CurrentState, RecordInfo, Info, AttachmentRecord);
+		}
+		break;
+
+		case EReplicatedObjectState::CancelPendingDestroy:
+		{
+			checkf(false, TEXT("%s"), TEXT("CancelPendingDestroy is not a state that should be replicated."));
+		}
+		break;
+
+		case EReplicatedObjectState::AttachmentToObjectNotInScope:
+		{
+			check(IsObjectIndexForOOBAttachment(InternalIndex));
+			Attachments.OnPacketLost(ENetObjectAttachmentType::OutOfBand, ObjectIndexForOOBAttachment, AttachmentRecord);
+		}
+		return;
+
+		case EReplicatedObjectState::HugeObject:
+		{
+			check(IsObjectIndexForOOBAttachment(InternalIndex));
+			Attachments.OnPacketLost(ENetObjectAttachmentType::HugeObject, ObjectIndexForOOBAttachment, AttachmentRecord);
+		}
+		return;
+
+		
+		default:
+		break;
+	};
+
+	if (RecordInfo.HasAttachments)
+	{
+		Attachments.OnPacketLost(ENetObjectAttachmentType::Normal, InternalIndex, AttachmentRecord);
+	}
+}
+
+void FReplicationWriter::ProcessDeliveryNotification(EPacketDeliveryStatus PacketDeliveryStatus)
+{
+#if UE_NET_VALIDATE_REPLICATION_RECORD
+	check(s_ValidateReplicationRecord(&ReplicationRecord, Parameters.MaxActiveReplicatedObjectCount + 1U, true));
+#endif
+
+	const uint32 RecordCount = ReplicationRecord.PopRecord();
+
+	if (RecordCount > 0)
+	{
+		TReplicationRecordHelper Helper(ReplicatedObjects.GetData(), ReplicatedObjectsRecordInfoLists.GetData(), &ReplicationRecord);
+
+		if (PacketDeliveryStatus == EPacketDeliveryStatus::Delivered)
+		{
+			Helper.Process(RecordCount,
+				[this](const FReplicationRecord::FRecordInfo& RecordInfo, FReplicationInfo& Info, FNetObjectAttachmentsWriter::ReplicationRecord AttachmentRecord)
+				{ 
+					HandleDeliveredRecord(RecordInfo, Info, AttachmentRecord);
+				}
+			);
+		}
+		else if (PacketDeliveryStatus == EPacketDeliveryStatus::Lost)
+		{
+			Helper.Process(RecordCount,
+				[this](const FReplicationRecord::FRecordInfo& RecordInfo, FReplicationInfo& Info, FNetObjectAttachmentsWriter::ReplicationRecord AttachmentRecord)
+				{
+					HandleDroppedRecord(RecordInfo, Info, AttachmentRecord);
+				}
+			);
+		}
+		else if (PacketDeliveryStatus == EPacketDeliveryStatus::Discard)
+		{
+			Helper.Process(RecordCount,
+				[this](const FReplicationRecord::FRecordInfo& RecordInfo, FReplicationInfo& Info, FNetObjectAttachmentsWriter::ReplicationRecord AttachmentRecord)
+				{
+					HandleDiscardedRecord(RecordInfo, Info, AttachmentRecord);
+				}
+			);
+		}
+		else
+		{
+			checkf(false, TEXT("Unknown packet delivery status %u"), unsigned(PacketDeliveryStatus));
+		}
+	}
+}
+
+void FReplicationWriter::CreateObjectRecord(const FNetBitArrayView* ChangeMask, const FReplicationInfo& Info, const FBatchObjectInfo& ObjectInfo, FReplicationWriter::FObjectRecord& OutRecord)
+{
+	OutRecord.AttachmentRecord = ObjectInfo.AttachmentRecord;
+
+	FReplicationRecord::FRecordInfo& RecordInfo = OutRecord.Record;
+
+	RecordInfo.Index = ObjectInfo.InternalIndex;
+	RecordInfo.ReplicatedObjectState = ObjectInfo.AttachmentType == ENetObjectAttachmentType::HugeObject ? uint8(EReplicatedObjectState::HugeObject) : (uint8)Info.GetState();
+	RecordInfo.HasChangeMask = ChangeMask ? 1U : 0U;
+	RecordInfo.HasAttachments = (ObjectInfo.AttachmentRecord != 0 ? 1U : 0U);
+	RecordInfo.WroteTearOff = Info.TearOff;
+	RecordInfo.WroteDestroySubObject = Info.SubObjectPendingDestroy;
+	
+	// If we wrote a new baseline we need to store it in the record
+	RecordInfo.NewBaselineIndex = ObjectInfo.bSentState ? ObjectInfo.NewBaselineIndex : FDeltaCompressionBaselineManager::InvalidBaselineIndex;
+
+	if (ChangeMask)
+	{
+		// $IRIS: TODO: Implement other type of changemask allocator that utilizes the the fifo nature of the record
+		// https://jira.it.epicgames.com/browse/UE-127372
+		// Allocate and copy changemask
+		FChangeMaskStorageOrPointer::Alloc(RecordInfo.ChangeMaskOrPtr, ChangeMask->GetNumBits(), s_DefaultChangeMaskAllocator);
+		FChangeMaskUtil::CopyChangeMask(RecordInfo.ChangeMaskOrPtr, *ChangeMask);
+	}
+	else
+	{
+		// Clear change mask
+		RecordInfo.ChangeMaskOrPtr = FChangeMaskStorageOrPointer();
+	}
+}
+
+void FReplicationWriter::CommitObjectRecord(uint32 InternalObjectIndex, const FObjectRecord& ObjectRecord)
+{
+	// Push and link replication record to data already in-flight
+	ReplicationRecord.PushInfoAndAddToList(ReplicatedObjectsRecordInfoLists[InternalObjectIndex], ObjectRecord.Record, ObjectRecord.AttachmentRecord);
+}
+
+void FReplicationWriter::CommitBatchRecord(const FBatchRecord& BatchRecord)
+{
+	for (const FObjectRecord& ObjectRecord : BatchRecord.ObjectReplicationRecords)
+	{
+		 CommitObjectRecord(ObjectRecord.Record.Index, ObjectRecord);
+	}
+}
+
+uint32 FReplicationWriter::WriteObjectsPendingDestroy(FNetSerializationContext& Context)
+{
+	FNetBitStreamWriter& Writer = *Context.GetBitStreamWriter();
+
+	UE_NET_TRACE_SCOPE(ObjectsPendingDestroy, Writer, Context.GetTraceCollector(), ENetTraceVerbosity::Trace);
+
+	// Write bit indicating that we did write, we also need one trailing bit to indicate the end
+	uint16 WrittenCount = 0;
+
+	// Write how many destroyed objects we have
+	const uint32 HeaderPos = Writer.GetPosBits();
+	Writer.WriteBits(WrittenCount, 16);
+
+	if (Writer.IsOverflown())
+	{
+		return 0U;
+	}
+
+	bool bWroteAllDestroyedObjects = true;
+
+	for (uint32 InternalIndex = 0U; (InternalIndex = ObjectsPendingDestroy.FindFirstOne(InternalIndex + 1U)) != ~0U; )
+	{
+		FReplicationInfo& Info = GetReplicationInfo(InternalIndex);
+		const FNetHandleManager::FReplicatedObjectData& ObjectData = NetHandleManager->GetReplicatedObjectDataNoCheck(InternalIndex);
+
+		// Already waiting on destroy confirmation
+		if (Info.GetState() == EReplicatedObjectState::WaitOnDestroyConfirmation)
+		{
+			continue;
+		}
+
+		if (ObjectData.IsSubObject())
+		{
+			if (Info.SubObjectPendingDestroy)
+			{
+				// If the owner is pending destroy then we need to replicate the destruction for the subobject here.
+				if (ObjectsPendingDestroy.GetBit(ObjectData.SubObjectRootIndex))
+				{
+					//			UE_LOG(LogTemp, Warning, TEXT("Object %s was a subobject but owner was destroyed so we need to destroy object normally"), ToCStr(ObjectData.Handle.ToString()));
+					SetState(InternalIndex, EReplicatedObjectState::PendingDestroy);
+					Info.SubObjectPendingDestroy = 0U;
+				}
+				else
+				{
+					// Replicate subobject destruction via state replication.
+					continue;
+				}
+			}
+			else
+			{
+				// If the owner is valid and isn't pending destroy we need to replicate destruction via state replication.
+				if (!ObjectsPendingDestroy.GetBit(ObjectData.SubObjectRootIndex))
+				{
+					FReplicationInfo& OwnerInfo = GetReplicationInfo(ObjectData.SubObjectRootIndex);
+					if (OwnerInfo.GetState() != EReplicatedObjectState::Invalid && OwnerInfo.GetState() < EReplicatedObjectState::PendingDestroy)
+					{
+						//			UE_LOG(LogTemp, Warning, TEXT("SubObject %s Should be destroyed by replicating owner"), ToCStr(ObjectData.Handle.ToString()));
+
+						SetState(InternalIndex, EReplicatedObjectState::SubObjectPendingDestroy);
+						Info.SubObjectPendingDestroy = 1U;
+
+						OwnerInfo.HasDirtySubObjects = 1U;
+
+						ObjectsWithDirtyChanges.SetBit(ObjectData.SubObjectRootIndex);
+						continue;
+					}
+				}
+			}
+		}
+
+		check(Info.GetState() == EReplicatedObjectState::PendingDestroy);
+
+		// We do not support destroying an object that is currently being sent as a huge object.
+		if (IsDestroyObjectPartOfActiveHugeObject(InternalIndex, Info))
+		{
+			UE_LOG(LogIris, Warning, TEXT("Skipping writing destroy for object ( InternalIndex: %u ) which is part of active huge object."), InternalIndex);
+			bWroteAllDestroyedObjects = false;
+			continue;
+		}
+
+		UE_NET_TRACE_OBJECT_SCOPE(ObjectData.Handle,  Writer, Context.GetTraceCollector(), ENetTraceVerbosity::Trace);
+
+		FNetBitStreamRollbackScope RollbackScope(Writer);
+
+		// Write NetObjectHandle with the needed bitCount
+		WriteNetHandleId(Writer, ObjectData.Handle);
+
+		// Write bit indicating if the static instance should be destroyed or not (could skip the bit for dynamic objects)
+		const bool bShouldDestroyInstance = ObjectData.Handle.IsDynamic() || NetHandleManager->GetIsDestroyedStartupObject(InternalIndex);
+		Writer.WriteBool(bShouldDestroyInstance);
+
+		if (!Writer.IsOverflown())
+		{
+			// Must update state before pushing record
+			SetState(InternalIndex, EReplicatedObjectState::WaitOnDestroyConfirmation);
+			Info.HasDirtyChangeMask = 0U;
+
+			// update transmission record
+			FBatchObjectInfo ObjectInfo = {};
+			ObjectInfo.InternalIndex = InternalIndex;
+			FObjectRecord ObjectRecord;
+			CreateObjectRecord(nullptr, Info, ObjectInfo, ObjectRecord);
+			CommitObjectRecord(InternalIndex, ObjectRecord);
+	
+			++WrittenCount;
+		}
+		else
+		{
+			bWroteAllDestroyedObjects = false;
+			break;
+		}
+	}
+
+	// Write Header
+	{
+		FNetBitStreamWriteScope WriteScope(Writer, HeaderPos);
+		Writer.WriteBits(WrittenCount, 16);
+	}
+
+	WriteContext.bHasDestroyedObjectsToSend = !bWroteAllDestroyedObjects;
+
+	return WrittenCount;
+}
+
+bool FReplicationWriter::CanSendObject(uint32 InternalIndex) const
+{
+	const FReplicationInfo& Info = GetReplicationInfo(InternalIndex);
+	const EReplicatedObjectState State = Info.GetState();
+
+	const bool bIsInitialState = IsInitialState(State);
+
+	// Currently we do wait for CreateConfirmation before sending more data
+	// We might want to change this and allow "bombing" creation info until we get confirmation to minimize latency
+	// We also prevent objects from being transmitted if they are waiting on destroy/tear-off confirmation or cancelling destroy.
+	if (State == EReplicatedObjectState::WaitOnCreateConfirmation || State == EReplicatedObjectState::WaitOnDestroyConfirmation || State == EReplicatedObjectState::CancelPendingDestroy)
+	{
+		return false;
+	}
+
+	// Don't send more recent state that could arrive before the huge state. We only need to check the parent.
+	if (IsActiveHugeObject(InternalIndex))
+	{
+		const bool bIsSubObjectOrDependentObject = NetHandleManager->GetSubObjectInternalIndices().GetBit(InternalIndex);
+		if (!bIsSubObjectOrDependentObject)
+		{
+			return false;
+		}
+	}
+
+	if (Info.HasDirtySubObjects)
+	{
+		for (FInternalNetHandle SubObjectInternalIndex : NetHandleManager->GetSubObjects(InternalIndex))
+		{
+			if (!CanSendObject(SubObjectInternalIndex))
+			{
+				return false;
+			}
+		}
+	}
+
+	// Currently we need to enforce a strict dependency on the initial dependencies
+	for (FInternalNetHandle DependentInternalIndex : NetHandleManager->GetDependentObjects(InternalIndex))
+	{
+		const FReplicationInfo& DependentInfo = GetReplicationInfo(DependentInternalIndex);
+		const bool bIsDependentObjectInitialState = IsInitialState(DependentInfo.GetState());
+
+		if (bIsDependentObjectInitialState && !CanSendObject(DependentInternalIndex))
+		{
+			UE_LOG(LogIris, Log, TEXT("ReplicationWriter: Cannot send internal index (%u) due to waiting on init dependency internal index (%d)"), InternalIndex, DependentInternalIndex);
+			return false;
+		}
+	}
+
+	return true;
+}
+
+void FReplicationWriter::SerializeObjectStateDelta(FNetSerializationContext& Context, uint32 InternalIndex, const FReplicationInfo& Info, const FNetHandleManager::FReplicatedObjectData& ObjectData, const uint8* ReplicatedObjectStateBuffer, FDeltaCompressionBaseline& CurrentBaseline, uint32 CreatedBaselineIndex)
+{
+	FNetBitStreamWriter& Writer = *Context.GetBitStreamWriter();
+
+	// Write baseline info
+	Writer.WriteBits(Info.LastAckedBaselineIndex, FDeltaCompressionBaselineManager::BaselineIndexBitCount);
+	if (Info.LastAckedBaselineIndex != FDeltaCompressionBaselineManager::InvalidBaselineIndex)
+	{
+		// Verify assumptions made about new baselineindices
+		check(CurrentBaseline.IsValid());
+		check(CreatedBaselineIndex == FDeltaCompressionBaselineManager::InvalidBaselineIndex || CreatedBaselineIndex == (Info.LastAckedBaselineIndex + 1) % FDeltaCompressionBaselineManager::MaxBaselineCount);
+
+		// Do we want to store a new baseline?
+		Writer.WriteBool(CreatedBaselineIndex != FDeltaCompressionBaselineManager::InvalidBaselineIndex);
+
+		UE_NET_TRACE_SCOPE(DeltaCompressed, Writer, Context.GetTraceCollector(), ENetTraceVerbosity::Trace);
+
+		FReplicationProtocolOperations::SerializeWithMaskDelta(Context, Info.GetChangeMaskStoragePointer(), ReplicatedObjectStateBuffer, CurrentBaseline.StateBuffer, ObjectData.Protocol);
+	}
+	else
+	{
+		// if we do not have a valid LastAckedBaselineIndex we need to write the full CreatedBaselineIndex
+		Writer.WriteBits(CreatedBaselineIndex, FDeltaCompressionBaselineManager::BaselineIndexBitCount);
+
+		// $IRIS: $TODO: Consider Delta compressing against default state
+		// Write non delta compressed state
+		FReplicationProtocolOperations::SerializeWithMask(Context, Info.GetChangeMaskStoragePointer(), ReplicatedObjectStateBuffer, ObjectData.Protocol);	
+	}
+}
+
+FReplicationWriter::EWriteObjectStatus FReplicationWriter::WriteObjectInBatch(FNetSerializationContext& Context, uint32 InternalIndex, uint32 WriteObjectFlags, FBatchInfo& OutBatchInfo)
+{
+	IRIS_PROFILER_SCOPE(FReplicationWriter_WriteObjectInBatch);
+
+	FNetBitStreamWriter& Writer = *Context.GetBitStreamWriter();
+	
+	FReplicationInfo& Info = GetReplicationInfo(InternalIndex);
+	const EReplicatedObjectState State = Info.GetState();
+
+	// As an object might still have subobjects pending destroy in the list of subobjects 
+	if (State == EReplicatedObjectState::Invalid || !ensureAlwaysMsgf(State > EReplicatedObjectState::Invalid && State < EReplicatedObjectState::PendingDestroy, TEXT("Unsupported state %s"), ToCStr(LexToString(State))))
+	{
+		return EWriteObjectStatus::InvalidState;
+	}
+
+	const bool bIsInitialState = IsInitialState(State);
+
+#if UE_NET_REPLICATION_SUPPORT_SKIP_INITIAL_STATE
+	uint32 InitialStateHeaderPos = 0U;
+	const uint32 NumBitsUsedForInitialStateSize = Parameters.NumBitsUsedForInitialStateSize;
+#endif
+
+	// Create a temporary batch entry. We don't want to push it to the batch info unless we're successful.
+	FBatchObjectInfo BatchEntry = {};
+
+	uint32 CreatedBaselineIndex = FDeltaCompressionBaselineManager::InvalidBaselineIndex;
+	FDeltaCompressionBaseline CurrentBaseline;
+
+	// We need to release created baseline if we fail to commit anything to batchrecord
+	ON_SCOPE_EXIT
+	{
+		if (CreatedBaselineIndex != FDeltaCompressionBaselineManager::InvalidBaselineIndex)
+		{
+			UE_LOG_REPLICATIONWRITER_CONN(TEXT("Destroy cancelled baseline %u for ( InternalObjectIndex: %u )"), CreatedBaselineIndex, InternalIndex);
+			BaselineManager->DestroyBaseline(Parameters.ConnectionId, InternalIndex, CreatedBaselineIndex);
+		}
+	};
+
+	// Only write data for the object if we have data to write
+	const FNetHandleManager::FReplicatedObjectData& ObjectData = NetHandleManager->GetReplicatedObjectDataNoCheck(InternalIndex);
+	uint8* ReplicatedObjectStateBuffer = NetHandleManager->GetReplicatedObjectStateBufferNoCheck(InternalIndex);
+
+	// Filter out changemasks that are not supposed to be replicated to this connection
+	const bool bNeedToFilterChangeMask = (bIsInitialState | Info.HasDirtyChangeMask) & Info.HasChangemaskFilter;
+	if (bNeedToFilterChangeMask)
+	{
+		ApplyFilterToChangeMask(OutBatchInfo.ParentInternalIndex, InternalIndex, Info, ObjectData.Protocol, ReplicatedObjectStateBuffer);
+	}
+
+	const bool bIsObjectIndexForAttachment = IsObjectIndexForOOBAttachment(InternalIndex);
+	const bool bHasState = (bIsInitialState | Info.HasDirtyChangeMask) & !!(WriteObjectFlags & EWriteObjectFlag::WriteObjectFlag_State);
+	const bool bHasAttachments = (Info.HasAttachments | bIsObjectIndexForAttachment);
+	const bool bWriteAttachments = bHasAttachments & !!(WriteObjectFlags & EWriteObjectFlag::WriteObjectFlag_Attachments);
+	BatchEntry.bHasUnsentAttachments = bHasAttachments;
+
+	Context.SetIsInitState(bIsInitialState);
+
+	// Update flag since dependent objects might change their status after creation
+	Info.IsSubObject = ObjectData.IsSubObject();
+
+	if (bHasState | bWriteAttachments | Info.TearOff | Info.SubObjectPendingDestroy)
+	{
+		const FNetHandle NetHandle = ObjectData.Handle;
+#if UE_NET_TRACE_ENABLED
+		const FNetHandle NetHandleForTraceScope = (WriteObjectFlags & EWriteObjectFlag::WriteObjectFlag_HugeObject ? NetHandleManager->GetReplicatedObjectDataNoCheck(HugeObjectContext.InternalIndex).Handle : NetHandle);
+		UE_NET_TRACE_OBJECT_SCOPE(NetHandleForTraceScope, Writer, Context.GetTraceCollector(), ENetTraceVerbosity::Trace);
+#endif
+		FNetBitStreamRollbackScope ObjectRollbackScope(Writer);
+
+		const bool bIsDestructionInfo = false;
+		Writer.WriteBool(bIsDestructionInfo);
+
+#if UE_NET_USE_READER_WRITER_SENTINEL
+		WriteSentinelBits(&Writer, 8);
+#endif
+
+		// We send the Index of the handle to the remote end
+		// $IRIS: $TODO: consider sending the internal index instead to save bits and only send handle when we create the object, https://jira.it.epicgames.com/browse/UE-127373
+		{
+			UE_NET_TRACE_SCOPE(Handle, Writer, Context.GetTraceCollector(), ENetTraceVerbosity::Trace);
+			WriteNetHandleId(Writer, NetHandle);
+		}
+
+		// TearOff
+		Writer.WriteBool(Info.TearOff);
+
+		// SubObject pending destroy
+		if (Writer.WriteBool(Info.SubObjectPendingDestroy))
+		{
+			UE_NET_TRACE_SCOPE(SubObjectPendingDestroy, Writer, Context.GetTraceCollector(), ENetTraceVerbosity::Trace);
+			// Write bit indicating if the static instance should be destroyed or not (could skip the bit for dynamic objects)
+			const bool bShouldDestroyInstance = ObjectData.Handle.IsDynamic() || NetHandleManager->GetIsDestroyedStartupObject(InternalIndex);
+			Writer.WriteBool(bShouldDestroyInstance);
+		}
+
+		if (Writer.WriteBool(bHasState))
+		{
+			BatchEntry.bSentState = 1;
+
+			// If the last transmitted baseline is acknowledged we can request a new baseline to be stored for the current state, we cannot compress against it until it has been acknowledged
+			if (Info.IsDeltaCompressionEnabled)
+			{
+				// Lookup current baseline that we should compress against
+				if (Info.LastAckedBaselineIndex != FDeltaCompressionBaselineManager::InvalidBaselineIndex)
+				{
+					CurrentBaseline = BaselineManager->GetBaseline(Parameters.ConnectionId, InternalIndex, Info.LastAckedBaselineIndex);
+
+					// If we cannot find the baseline it has become invalidated, if that is the case we must invalidate all tracking and request a new baseline to be created
+					if (!CurrentBaseline.IsValid())
+					{
+						InvalidateBaseline(InternalIndex, Info);
+					}
+				}
+
+				if (Info.PendingBaselineIndex == FDeltaCompressionBaselineManager::InvalidBaselineIndex)
+				{
+					// For new objects we start with baselineindex 0
+					const uint32 NextBaselineIndex = bIsInitialState ? 0U : (Info.LastAckedBaselineIndex + 1U) % FDeltaCompressionBaselineManager::MaxBaselineCount;
+					FDeltaCompressionBaseline NewBaseline = BaselineManager->CreateBaseline(Parameters.ConnectionId, InternalIndex, NextBaselineIndex);				
+					if (NewBaseline.IsValid())
+					{
+						CreatedBaselineIndex = NextBaselineIndex;
+
+						// $IRIS: $TODO: Consider conditionals, how to ensure that baseline stored on remote is complete if state of conditional changes?
+						// $IRIS: $TODO: Currently due to how repnotifies are implemented we might have to write an extra changemask when sending a new baseline to avoid extra calls to repnotifies
+
+						// Modify changemask to include any data we have in flight to ensure baseline integrity on receiving end
+						if (PatchupObjectChangeMaskWithInflightChanges(InternalIndex, Info))
+						{
+							// Mask off changemasks that may have been disabled due to conditionals.
+							ApplyFilterToChangeMask(OutBatchInfo.ParentInternalIndex, InternalIndex, Info, ObjectData.Protocol, ReplicatedObjectStateBuffer);
+						}
+
+						UE_LOG_REPLICATIONWRITER_CONN(TEXT("Created new baseline %u for ( InternalObjectIndex: %u )"), CreatedBaselineIndex, InternalIndex);
+					}
+				}
+			}
+
+			// $TODO: Consider rewriting the FReplicationProtocolOperations::SerializeWithMask() methods to accept the changemask passed in the Context rather setting it up again later
+			FNetBitArrayView ChangeMask = MakeNetBitArrayView(Info.GetChangeMaskStoragePointer(), Info.ChangeMaskBitCount);
+			Context.SetChangeMask(&ChangeMask);
+
+			// Export references so that we can share the serialized state with other connection
+			if (!CollectAndWriteExports(Context, ReplicatedObjectStateBuffer, ObjectData.Protocol))
+			{
+				return EWriteObjectStatus::BitStreamOverflow;
+			}
+
+#if UE_NET_USE_READER_WRITER_SENTINEL
+			WriteSentinelBits(&Writer, 8);
+#endif
+			
+			// Role downgrade
+			{
+				FInternalNetSerializationContext* InternalContext = Context.GetInternalContext();
+				InternalContext->bDowngradeAutonomousProxyRole = (Context.GetLocalConnectionId() != ReplicationFiltering->GetOwningConnection(InternalIndex));
+			}
+
+			if (Writer.WriteBool(bIsInitialState))
+			{
+				// Creation Info
+				{
+					UE_NET_TRACE_SCOPE(CreationInfo, Writer, Context.GetTraceCollector(), ENetTraceVerbosity::Trace);
+
+#if UE_NET_REPLICATION_SUPPORT_SKIP_INITIAL_STATE
+					InitialStateHeaderPos = Writer.GetPosBits();
+					Writer.WriteBits(0U, NumBitsUsedForInitialStateSize);
+#endif
+
+					// Warn if we cannot replicate this object
+					if (!ObjectData.InstanceProtocol)
+					{
+						UE_LOG_REPLICATIONWRITER_WARNING(TEXT("Failed to replicate ( InternalIndex: %u ) %s, ProtocolName: %s, Currently we do not support creating a remote instance when the instance has been detached."), InternalIndex, *NetHandle.ToString(), ToCStr(ObjectData.Protocol->DebugName));
+						return EWriteObjectStatus::NoInstanceProtocol;
+					}
+
+					// Real SubObjects needs to be tracked on client as well
+					if (Writer.WriteBool(ObjectData.IsSubObject()))
+					{
+						const FNetHandle& OwnerHandle = NetHandleManager->GetReplicatedObjectDataNoCheck(ObjectData.SubObjectRootIndex).Handle;
+						WriteNetHandleId(Writer, OwnerHandle);
+						if (!OwnerHandle.IsValid() || Writer.IsOverflown())
+						{
+							return OwnerHandle.IsValid() ? EWriteObjectStatus::BitStreamOverflow : EWriteObjectStatus::InvalidOwner;
+						}
+					}
+
+					if (Writer.WriteBool(Info.IsDeltaCompressionEnabled))
+					{
+						// As we might fail to create a baseline for initial state we need to include it here.
+						Writer.WriteBits(CreatedBaselineIndex, FDeltaCompressionBaselineManager::BaselineIndexBitCount);
+					}
+
+					FReplicationBridgeSerializationContext BridgeContext(Context, Parameters.ConnectionId, Info.IsDestructionInfo == 1U);
+
+					// We need to send creation info, if we fail, we skip this object for now
+					if (!ReplicationBridge->CallWriteNetHandleCreationInfo(BridgeContext, NetHandle))
+					{
+						return BridgeContext.SerializationContext.HasError() ? EWriteObjectStatus::Error : EWriteObjectStatus::BitStreamOverflow;
+					}
+				}
+				// Serialize initial state data for this object using delta compression against default state
+				FReplicationProtocolOperations::SerializeInitialStateWithMask(Context, Info.GetChangeMaskStoragePointer(), ReplicatedObjectStateBuffer, ObjectData.Protocol);
+			}
+			else
+			{
+				if (Info.IsDeltaCompressionEnabled)
+				{
+					SerializeObjectStateDelta(Context, InternalIndex, Info, ObjectData, ReplicatedObjectStateBuffer, CurrentBaseline, CreatedBaselineIndex);
+				}
+				else
+				{
+					// Serialize state data for this object
+					FReplicationProtocolOperations::SerializeWithMask(Context, Info.GetChangeMaskStoragePointer(), ReplicatedObjectStateBuffer, ObjectData.Protocol);
+				}
+			}
+		}
+
+		{
+			const uint32 HasAttachmentsWritePos = Writer.GetPosBits();
+			Writer.WriteBool(bWriteAttachments);
+			if (Writer.IsOverflown())
+			{
+				return EWriteObjectStatus::BitStreamOverflow;
+			}
+
+			if (bWriteAttachments)
+			{
+				FNetBitStreamWriter AttachmentWriter = Writer.CreateSubstream();
+				FNetSerializationContext AttachmentContext = Context.MakeSubContext(&AttachmentWriter);
+				BatchEntry.AttachmentType = ENetObjectAttachmentType::Normal;
+				if (bIsObjectIndexForAttachment)
+				{
+					BatchEntry.AttachmentType = (WriteObjectFlags & EWriteObjectFlag::WriteObjectFlag_HugeObject ? ENetObjectAttachmentType::HugeObject : ENetObjectAttachmentType::OutOfBand);
+					AttachmentWriter.WriteBool(BatchEntry.AttachmentType == ENetObjectAttachmentType::HugeObject);
+				}
+
+				Attachments.Serialize(AttachmentContext, BatchEntry.AttachmentType, InternalIndex, NetHandle, BatchEntry.AttachmentRecord, BatchEntry.bHasUnsentAttachments);
+
+				// If we didn't manage to fit any attachments then clear the HasAttachments bool in the packet
+				if (AttachmentWriter.GetPosBits() == 0 || AttachmentWriter.IsOverflown())
+				{
+					BatchEntry.bSentAttachments = 0;
+
+					const uint32 BitsThatWasAvailableForAttachements = Writer.GetBitsLeft();
+
+					Writer.DiscardSubstream(AttachmentWriter);
+					{
+						FNetBitStreamWriteScope HasAttachmentsWriteScope(Writer, HasAttachmentsWritePos);
+						Writer.WriteBool(false);
+					}
+
+					Info.HasAttachments = BatchEntry.bHasUnsentAttachments;
+					
+					// If we should have had enough space to write a attachment, the attachement + exports must be huge and we need to fall back on using the huge object path
+					const uint32 SplitThreshold = PartialNetObjectAttachmentHandler->GetConfig()->GetBitCountSplitThreshold() * 2;
+					const bool bFallbackToHugeObjectPath = BatchEntry.AttachmentType != ENetObjectAttachmentType::HugeObject && BatchEntry.bHasUnsentAttachments && (BitsThatWasAvailableForAttachements >= SplitThreshold);
+					if (bFallbackToHugeObjectPath)
+					{
+						UE_LOG(LogIris, Verbose, TEXT("Failed to write huge attachment for object %s ( InternalIndex: %u ), forcing fallback on hugeobject for attachments"), *NetHandle.ToString(), InternalIndex);
+						Writer.DoOverflow();
+					}
+					else if (!BatchEntry.bSentState)
+					{
+						// If we didn't send state and didn't send any attachments let's rollback
+						ObjectRollbackScope.Rollback();
+					}
+				}
+				else
+				{
+					BatchEntry.bSentAttachments = 1;
+
+					Writer.CommitSubstream(AttachmentWriter);
+
+					// Update the HasAttachments info based on this object batch failing. If the batch is a success we update again.
+					Info.HasAttachments = Attachments.HasUnsentAttachments(BatchEntry.AttachmentType, InternalIndex);
+				}
+			}
+		}
+
+		if (Writer.IsOverflown())
+		{
+			return EWriteObjectStatus::BitStreamOverflow;
+		}
+	}
+
+	// Success so far. Fill in batch entry. Keep index to update info later as the array can resize.
+	const int ParentBatchEntryIndex = OutBatchInfo.ObjectInfos.Num();
+	{
+		FBatchObjectInfo& FinalBatchEntry = OutBatchInfo.ObjectInfos.Emplace_GetRef();
+		FinalBatchEntry = MoveTemp(BatchEntry);
+
+		FinalBatchEntry.bIsInitialState = bIsInitialState;
+		FinalBatchEntry.InternalIndex = InternalIndex;
+		FinalBatchEntry.bHasDirtySubObjects = false;
+		FinalBatchEntry.bSentTearOff = Info.TearOff;
+		FinalBatchEntry.bSentDestroySubObject = Info.SubObjectPendingDestroy;
+		FinalBatchEntry.NewBaselineIndex = CreatedBaselineIndex;
+	}
+
+	// Reset CreatedBaselineIndex to avoid it being released on scope exit
+	CreatedBaselineIndex = FDeltaCompressionBaselineManager::InvalidBaselineIndex;
+
+	// Write dirty sub objects
+	if (Info.HasDirtySubObjects && !Info.IsSubObject)
+	{
+		bool bHasDirtySubObjects = false;
+		
+		FReplicationConditionals::FSubObjectsToReplicateArray SubObjectsToReplicate;
+		ReplicationConditionals->GetSubObjectsToReplicate(Parameters.ConnectionId, InternalIndex, SubObjectsToReplicate);		
+
+		for (FInternalNetHandle SubObjectInternalIndex : SubObjectsToReplicate)
+		{
+			const int BatchObjectInfoCount = OutBatchInfo.ObjectInfos.Num();
+			EWriteObjectStatus SubObjectWriteStatus = WriteObjectInBatch(Context, SubObjectInternalIndex, WriteObjectFlags, OutBatchInfo);
+			if (!IsWriteObjectSuccess(SubObjectWriteStatus))
+			{
+				// Need to remove the batch entry from BatchInfo.
+				OutBatchInfo.ObjectInfos.RemoveAt(OutBatchInfo.ObjectInfos.Num() - 1);
+				return SubObjectWriteStatus;
+			}
+
+			// There are success statuses where no object info is added. In such case we shouldn't read from it.
+			if (OutBatchInfo.ObjectInfos.Num() > BatchObjectInfoCount)
+			{
+				const FBatchObjectInfo& SubObjectEntry = OutBatchInfo.ObjectInfos.Last();
+				bHasDirtySubObjects = SubObjectEntry.bHasDirtySubObjects | SubObjectEntry.bHasUnsentAttachments;
+			}
+		}
+
+		// Update parent batch info
+		{
+			FBatchObjectInfo& ParentBatchEntry = OutBatchInfo.ObjectInfos[ParentBatchEntryIndex];
+			ParentBatchEntry.bHasDirtySubObjects |= bHasDirtySubObjects;
+		}
+	}
+
+	// If we have any dependent objects pending creation we currently include them in the batch in order to guarantee that they get sent in the same packet
+	// When we implement support for proper dependencies this can be changed.
+	for (FInternalNetHandle DependentInternalIndex : NetHandleManager->GetDependentObjects(InternalIndex))
+	{
+		const bool bIsDependentInitialState = IsInitialState(GetReplicationInfo(DependentInternalIndex).GetState());
+		if (bIsDependentInitialState && !WriteContext.ObjectsWrittenThisTick.GetBit(DependentInternalIndex))
+		{
+			UE_NET_TRACE_SCOPE(DependentObjectData, *Context.GetBitStreamWriter(), Context.GetTraceCollector(), ENetTraceVerbosity::VeryVerbose);
+			EWriteObjectStatus DependentObjectWriteStatus = WriteObjectInBatch(Context, DependentInternalIndex, WriteObjectFlags, OutBatchInfo);
+			if (!IsWriteObjectSuccess(DependentObjectWriteStatus))
+			{
+				// Need to remove the batch entry from BatchInfo.
+				OutBatchInfo.ObjectInfos.RemoveAt(OutBatchInfo.ObjectInfos.Num() - 1);
+				return DependentObjectWriteStatus;
+			}
+		}
+	}
+
+	// We include the size of the data written so we can skip it if needed.
+#if UE_NET_REPLICATION_SUPPORT_SKIP_INITIAL_STATE
+	if (bIsInitialState)
+	{
+		const uint32 InitialStateWrittenBits = (Writer.GetPosBits() - InitialStateHeaderPos) - NumBitsUsedForInitialStateSize;
+		FNetBitStreamWriteScope SizeScope(Writer, InitialStateHeaderPos);
+		Writer.WriteBits(InitialStateWrittenBits, NumBitsUsedForInitialStateSize);
+	}
+#endif
+
+	return EWriteObjectStatus::Success;
+}
+
+int FReplicationWriter::PrepareAndSendHugeObjectPayload(FNetSerializationContext& Context, uint32 InternalIndex)
+{
+	IRIS_PROFILER_SCOPE(FReplicationWriter_PrepareAndSendHugeObjectPayload);
+	typedef uint32 HugeObjectStorageType;
+	const uint32 BitsPerStorageWord = sizeof(HugeObjectStorageType) * 8;
+
+	TArray<HugeObjectStorageType> HugeObjectPayload;
+	HugeObjectPayload.AddUninitialized((PartialNetObjectAttachmentHandler->GetConfig()->GetTotalMaxPayloadBitCount() + (BitsPerStorageWord - 1U))/BitsPerStorageWord);
+
+	// Setup a special context for the huge object serialization.
+	FNetBitStreamWriter HugeObjectWriter;
+	const uint32 MaxHugeObjectPayLoadBytes = HugeObjectPayload.Num() * sizeof(HugeObjectStorageType);
+	HugeObjectWriter.InitBytes(HugeObjectPayload.GetData(), MaxHugeObjectPayLoadBytes);
+	FNetSerializationContext HugeObjectSerializationContext = Context.MakeSubContext(&HugeObjectWriter);
+	HugeObjectSerializationContext.SetNetTraceCollector(nullptr);
+
+	// Huge object header needed for the receiving side to be able to process this correctly.
+	FNetObjectBlob::FHeader HugeObjectHeader = {};
+	const uint32 HeaderPos = HugeObjectWriter.GetPosBits();
+	FNetObjectBlob::SerializeHeader(HugeObjectSerializationContext, HugeObjectHeader);
+	const uint32 PastHeaderPos = HugeObjectWriter.GetPosBits();
+
+	FBatchInfo BatchInfo;
+	BatchInfo.ParentInternalIndex = InternalIndex;
+	uint32 WriteObjectFlags = WriteObjectFlag_State;
+	// Get the creation going as quickly as possible.
+	if (!Context.IsInitState())
+	{
+		WriteObjectFlags |= WriteObjectFlag_Attachments;
+	}
+	
+	// Push new ExportContext for the hugeobject-batch as we cannot share exports with an OOB object
+	HugeObjectContext.BatchExports.Reset();
+	FNetExports::FExportScope ExportScope = NetExports->MakeExportScope(HugeObjectSerializationContext, HugeObjectContext.BatchExports);
+
+	// We can encounter other errors than bitstream overflow now that we've got a really large buffer to write to.
+	const EWriteObjectStatus WriteHugeObjectStatus = WriteObjectInBatch(HugeObjectSerializationContext, InternalIndex, WriteObjectFlags, BatchInfo);
+
+	// If we cannot fit the object in the largest supported buffer then we will never fit the object.
+	if (WriteHugeObjectStatus == EWriteObjectStatus::BitStreamOverflow)
+	{
+		UE_LOG(LogIris, Error, TEXT("Unable to fit object ( InternalIndex: %u ) in maximum combined payload of %u bytes. Connection %u will be disconnected."), InternalIndex, MaxHugeObjectPayLoadBytes, Context.GetLocalConnectionId());
+		Context.SetError(NetError_ObjectStateTooLarge);
+		return -1;
+	}
+
+	// If we encounter some other error we can try sending a smaller object in the meantime.
+	if (!IsWriteObjectSuccess(WriteHugeObjectStatus))
+	{
+		// Need to call this to cleanup data from batch
+		HandleObjectBatchFailure(WriteHugeObjectStatus, BatchInfo, WriteBitStreamInfo);
+
+		UE_LOG(LogIris, Verbose, TEXT("Problem writing huge object ( InternalIndex: %u ). WriteObjectStatus: %u. Trying smaller object."), InternalIndex, unsigned(WriteHugeObjectStatus));
+		return 0;
+	}
+
+	HugeObjectContext.InternalIndex = InternalIndex;
+	HugeObjectContext.SendStatus = EHugeObjectSendStatus::Sending;
+	// Store batch record for later processing once the whole state is acked.
+	HugeObjectHeader.ObjectCount = HandleObjectBatchSuccess(BatchInfo, HugeObjectContext.BatchRecord);
+
+	// Write huge object header
+	{
+		FNetBitStreamWriteScope WriteScope(HugeObjectWriter, HeaderPos);
+		FNetObjectBlob::SerializeHeader(HugeObjectSerializationContext, HugeObjectHeader);
+	}
+
+	HugeObjectWriter.CommitWrites();
+
+	// Create a NetObjectBlob from the temporary buffer and split it into multiple smaller pieces.
+	const uint32 PayLoadBitCount = HugeObjectWriter.GetPosBits();
+	const uint32 StorageWordsWritten = (PayLoadBitCount + (BitsPerStorageWord - 1)) / BitsPerStorageWord;
+
+	check(StorageWordsWritten <= (uint32)HugeObjectPayload.Num());
+
+	TArrayView<HugeObjectStorageType> PayloadView(HugeObjectPayload.GetData(), StorageWordsWritten);
+	TRefCountPtr<FNetObjectBlob> NetObjectBlob = NetObjectBlobHandler->CreateNetObjectBlob(PayloadView, PayLoadBitCount);
+	TArray<TRefCountPtr<FNetBlob>> PartialNetBlobs;
+	const bool bSplitSuccess = PartialNetObjectAttachmentHandler->SplitRawDataNetBlob(TRefCountPtr<FRawDataNetBlob>(NetObjectBlob.GetReference()), PartialNetBlobs, HugeObjectContext.DebugName);
+	if (!bSplitSuccess)
+	{
+		UE_LOG(LogIris, Error, TEXT("Unable to split huge object ( InternalIndex: %u ) payload. Connection %u will be disconnected."), InternalIndex, Context.GetLocalConnectionId());
+		Context.SetError(NetError_ObjectStateTooLarge);
+		return -1;
+	}
+
+	// Enqueue attachments
+	const bool bEnqueueSuccess = Attachments.Enqueue(ENetObjectAttachmentType::HugeObject, ObjectIndexForOOBAttachment, MakeArrayView(PartialNetBlobs.GetData(), PartialNetBlobs.Num()));
+	check(bEnqueueSuccess);
+
+	// Write huge object attachment(s)
+	{
+		FNetBitStreamWriter& Writer = *Context.GetBitStreamWriter();
+		UE_NET_TRACE_SCOPE(Batch, Writer, Context.GetTraceCollector(), ENetTraceVerbosity::Trace);
+		FNetBitStreamRollbackScope RollbackScope(Writer);
+		
+		FBatchInfo HugeObjectBatchInfo;
+		HugeObjectBatchInfo.ParentInternalIndex = FNetHandleManager::InvalidInternalIndex;
+		const uint32 WriteHugeObjectFlags = WriteObjectFlag_Attachments | WriteObjectFlag_HugeObject;
+		const EWriteObjectStatus HugeObjectStatus = WriteObjectInBatch(Context, ObjectIndexForOOBAttachment, WriteHugeObjectFlags, HugeObjectBatchInfo);
+		if (!IsWriteObjectSuccess(HugeObjectStatus))
+		{
+			// Need to call this in order to cleanup data associated with batch
+			HandleObjectBatchFailure(HugeObjectStatus, HugeObjectBatchInfo, WriteBitStreamInfo);
+
+			ensureAlwaysMsgf(HugeObjectStatus == EWriteObjectStatus::BitStreamOverflow, TEXT("Expected split payload to not be able to generate other errors than overflow. Got %u"), unsigned(HugeObjectStatus));
+			// It's unexpected, but not a critical error, if no part of the payload could be sent.
+			// We do expect a smaller object to be sent though so that's why 0 is returned.
+
+			// Mark the context so that we can try to send the huge object in the next packet if we are allowed
+			WriteContext.bHasHugeObjectToSend = 1;
+
+			// Try to fit a smaller object.
+			return 0;
+		}
+
+		FBatchRecord BatchRecord;
+		HandleObjectBatchSuccess(HugeObjectBatchInfo, BatchRecord);
+		CommitBatchRecord(BatchRecord);
+
+		// If all chunks did not make it into the packet (expected) mark the context so that we can try to send the huge object in the next packet if we are allowed
+		WriteContext.bHasHugeObjectToSend = Attachments.HasUnsentAttachments(ENetObjectAttachmentType::HugeObject, ObjectIndexForOOBAttachment);
+
+		return 1;
+	}
+
+	// Should not get here.
+	return -1;
+}
+
+int FReplicationWriter::WriteObjectBatch(FNetSerializationContext& Context, uint32 InternalIndex, uint32 WriteObjectFlags)
+{
+	IRIS_PROFILER_SCOPE(FReplicationWriter_WriteObjectBatch);
+
+	// If this is a destruction info we treat it differently and just write the information required to destruct the object
+	if (GetReplicationInfo(InternalIndex).IsDestructionInfo)
+	{
+		return WriteDestructionInfo(Context, InternalIndex);
+	}
+
+	FNetBitStreamWriter& Writer = *Context.GetBitStreamWriter();
+
+	// Batch successful writes and commit them as a atomic batch.
+	// It is a fail if we fail to write any subobject with dirty state.
+	// It is also not ok to skip over creation header- if we do then the entire batch needs to be delayed.
+
+	// Write object and subobjects. Try #1 - send state and attachments.
+	{
+		UE_NET_TRACE_SCOPE(Batch, Writer, Context.GetTraceCollector(), ENetTraceVerbosity::Trace);
+		FNetBitStreamRollbackScope RollbackScope(Writer);
+		FNetExportRollbackScope ExportRollbackScope(Context);
+
+		WriteBitStreamInfo.BatchStartPos = Writer.GetPosBits();
+		FBatchInfo BatchInfo;
+		BatchInfo.ParentInternalIndex = InternalIndex;
+		const EWriteObjectStatus WriteObjectStatus = WriteObjectInBatch(Context, InternalIndex, WriteObjectFlags, BatchInfo);
+		if (IsWriteObjectSuccess(WriteObjectStatus))
+		{
+			FBatchRecord BatchRecord;
+			const int WrittenObjectCount = HandleObjectBatchSuccess(BatchInfo, BatchRecord);
+
+			CommitBatchRecord(BatchRecord);
+			return WrittenObjectCount;
+		}
+
+		const EWriteObjectRetryMode WriteRetryMode = HandleObjectBatchFailure(WriteObjectStatus, BatchInfo, WriteBitStreamInfo);
+
+		// Regardless of the reason for fail we should rollback anything written.
+		RollbackScope.Rollback();
+
+		// Rollback exported references that was exported as part of the batch we just rolled back
+		ExportRollbackScope.Rollback();
+
+		switch (WriteRetryMode)
+		{
+			case EWriteObjectRetryMode::Abort:
+			{
+				return -1;
+			}
+
+			case EWriteObjectRetryMode::TrySmallObject:
+			{
+				++WriteContext.FailedToWriteSmallObjectCount;
+				return 0;
+			}
+
+			case EWriteObjectRetryMode::SplitHugeObject:
+			{
+				break;
+			}
+
+			default:
+			{
+				check(false);
+			}
+		}
+	}
+
+	// Try #2 - Object will be serialized to a temporary buffer of maximum supported size and split into multiple chunks.
+	{
+		const int SendHugeObjectStatus = PrepareAndSendHugeObjectPayload(Context, InternalIndex);
+		return SendHugeObjectStatus;
+	}
+}
+
+int FReplicationWriter::WriteDestructionInfo(FNetSerializationContext& Context, uint32 InternalIndex)
+{
+	IRIS_PROFILER_SCOPE(FReplicationWriter_WriteDestructionInfo);
+
+	FNetBitStreamWriter& Writer = *Context.GetBitStreamWriter();
+	UE_NET_TRACE_SCOPE(Batch, Writer, Context.GetTraceCollector(), ENetTraceVerbosity::Trace);
+
+	// Rollback for entire batch
+	FNetBitStreamRollbackScope Rollback(Writer);
+	FNetExportRollbackScope ExportRollbackScope(Context);
+
+	// Only write data for the object if we have data to write
+	const FNetHandleManager::FReplicatedObjectData& ObjectData = NetHandleManager->GetReplicatedObjectDataNoCheck(InternalIndex);
+
+	// Special case for static objects that should be destroyed on the client but we have not replicated
+	Writer.WriteBool(true);
+
+	FReplicationBridgeSerializationContext BridgeContext(Context, Parameters.ConnectionId, true);
+
+	if (!ReplicationBridge->CallWriteNetHandleCreationInfo(BridgeContext, ObjectData.Handle))
+	{
+		// Trigger Rollback
+		Writer.DoOverflow();
+
+		return 0;
+	}
+
+#if UE_NET_USE_READER_WRITER_SENTINEL
+	WriteSentinelBits(&Writer, 8);
+#endif
+
+	// Push record
+	if (!Writer.IsOverflown())
+	{
+		// We did write the initial state, change the state to WaitOnCreateConfirmation
+		SetState(InternalIndex, EReplicatedObjectState::WaitOnCreateConfirmation);
+
+		FReplicationInfo& Info = GetReplicationInfo(InternalIndex);
+
+		FBatchObjectInfo ObjectInfo = {};
+		ObjectInfo.InternalIndex = InternalIndex;
+		FObjectRecord ObjectRecord;
+		CreateObjectRecord(nullptr, Info, ObjectInfo, ObjectRecord);
+		CommitObjectRecord(InternalIndex, ObjectRecord);
+
+		Info.HasDirtyChangeMask = 0U;
+		Info.HasDirtySubObjects = 0U;
+		Info.HasAttachments = 0U;
+
+		ObjectsWithDirtyChanges.ClearBit(InternalIndex);
+
+		// Reset scheduling priority
+		SchedulingPriorities[InternalIndex] = 0.0f;
+	}
+
+	return Writer.IsOverflown() ? -1 :  1;
+}
+
+uint32 FReplicationWriter::WriteOOBAttachments(FNetSerializationContext& Context)
+{
+	uint32 WrittenObjectCount = 0U;	
+	if (WriteContext.bHasHugeObjectToSend)
+	{
+		IRIS_PROFILER_SCOPE(FReplicationWriter_WriteHugeObjectAttachments);
+		const int32 Result = WriteObjectBatch(Context, ObjectIndexForOOBAttachment, WriteObjectFlag_Attachments | WriteObjectFlag_HugeObject);
+		if (Result == -1)
+		{
+			return WrittenObjectCount;
+		}
+
+		WriteContext.bHasHugeObjectToSend = Attachments.HasUnsentAttachments(ENetObjectAttachmentType::HugeObject, ObjectIndexForOOBAttachment);
+		WrittenObjectCount += Result;
+	}
+
+	if (WriteContext.bHasOOBAttachmentsToSend && CanSendObject(ObjectIndexForOOBAttachment))
+	{
+		IRIS_PROFILER_SCOPE(FReplicationWriter_WriteOOBAttachments);
+		const int32 Result = WriteObjectBatch(Context, ObjectIndexForOOBAttachment, WriteObjectFlag_Attachments);
+		if (Result == -1)
+		{
+			return WrittenObjectCount;
+		}
+
+		WriteContext.bHasOOBAttachmentsToSend = Attachments.HasUnsentAttachments(ENetObjectAttachmentType::OutOfBand, ObjectIndexForOOBAttachment);
+		WrittenObjectCount += Result;
+	}
+
+	return WrittenObjectCount;
+}
+
+uint32 FReplicationWriter::WriteObjects(FNetSerializationContext& Context)
+{
+	uint32 WrittenObjectCount = 0;
+	
+	const uint32 ObjectCount = WriteContext.ScheduledObjectCount;
+	FScheduleObjectInfo* ObjectList = WriteContext.ScheduledObjectInfos;
+
+	uint32 ObjectListIt = WriteContext.CurrentIndex;
+	uint32 SortedCount = WriteContext.SortedObjectCount;
+
+	bool bContinue = WriteContext.bHasUpdatedObjectsToSend;
+
+	auto&& SendObjectFunction = [this, &Context, &WrittenObjectCount](FInternalNetHandle InternalIndex)
+	{
+		if (!this->WriteContext.ObjectsWrittenThisTick.GetBit(InternalIndex) && this->CanSendObject(InternalIndex))
+		{
+			const int32 Result = this->WriteObjectBatch(Context, InternalIndex, WriteObjectFlag_State | WriteObjectFlag_Attachments);
+			if (Result >= 0)
+			{
+				WrittenObjectCount += Result;
+			}
+			else
+			{
+				return false;
+			}
+		}
+		return true;
+	};
+	
+
+	while (bContinue && ObjectListIt < ObjectCount)
+	{
+		// Partial sort next batch
+		if (ObjectListIt >= SortedCount)
+		{
+			SortedCount += SortScheduledObjects(ObjectList, ObjectCount, ObjectListIt);
+		}
+
+		for (;;)
+		{
+			// Try to send dependent objects
+			while (WriteContext.DependentObjectsPendingSend.Num() > 0)
+			{
+				UE_NET_TRACE_SCOPE(DependentObjectData, *Context.GetBitStreamWriter(), Context.GetTraceCollector(), ENetTraceVerbosity::VeryVerbose);
+
+				const uint32 InternalIndex = WriteContext.DependentObjectsPendingSend.Pop();
+				checkSlow(InternalIndex != ObjectIndexForOOBAttachment);
+				ensureAlwaysMsgf(GetReplicationInfo(InternalIndex).GetState() != EReplicatedObjectState::Invalid, TEXT("DependentObject with internalIndex %u is not in scope"), InternalIndex);
+				if (!SendObjectFunction(InternalIndex))
+				{
+					// If we fail, we put the object back on the pending send stack and try again in the next packet of the batch
+					// The reason for pop before using the index is that the SendObjectFunction will push new dependent objects on the stack
+					WriteContext.DependentObjectsPendingSend.Push(InternalIndex);				
+					bContinue = false;
+					break;
+				}
+			}
+
+			// Normal send, dependent objects will be pushed as needed
+			if (ObjectListIt < SortedCount)
+			{
+				const uint32 InternalIndex = ObjectList[ObjectListIt].Index;
+
+				checkSlow(InternalIndex != ObjectIndexForOOBAttachment);
+				if (!SendObjectFunction(InternalIndex))
+				{
+					bContinue = false;
+					break;
+				}
+
+				++ObjectListIt;
+			}
+			else
+			{
+				break;
+			}
+		}
+	}
+
+	// if we have more data to write, store state so that we can continue if we are allowed to write more data
+	WriteContext.bHasUpdatedObjectsToSend = (ObjectListIt != ObjectCount) || (WriteContext.DependentObjectsPendingSend.Num() > 0U);
+	WriteContext.CurrentIndex = ObjectListIt;
+	WriteContext.SortedObjectCount = SortedCount;
+
+	return WrittenObjectCount;
+}
+
+int FReplicationWriter::HandleObjectBatchSuccess(const FBatchInfo& BatchInfo, FReplicationWriter::FBatchRecord& OutRecord)
+{
+	int WrittenObjectCount = 0;
+
+	OutRecord.ObjectReplicationRecords.Reserve(BatchInfo.ObjectInfos.Num());
+	for (const FBatchObjectInfo& BatchObjectInfo : BatchInfo.ObjectInfos)
+	{
+		//UE_LOG_REPLICATIONWRITER(TEXT("FReplicationWriter::Wrote Object with %s InitialState: %u"), *NetHandleManager->NetHandle.ToString(), bIsInitialState ? 1u : 0u);
+		FReplicationInfo& Info = GetReplicationInfo(BatchObjectInfo.InternalIndex);
+
+		// We did write the initial state, change the state to WaitOnCreateConfirmation
+		if (BatchObjectInfo.bIsInitialState)
+		{
+			SetState(BatchObjectInfo.InternalIndex, EReplicatedObjectState::WaitOnCreateConfirmation);
+		}
+		else if (BatchObjectInfo.bSentTearOff)
+		{
+			SetState(BatchObjectInfo.InternalIndex, EReplicatedObjectState::PendingTearOff);
+			SetState(BatchObjectInfo.InternalIndex, EReplicatedObjectState::WaitOnDestroyConfirmation);
+		}
+		else if (BatchObjectInfo.bSentDestroySubObject)
+		{			
+			SetState(BatchObjectInfo.InternalIndex, EReplicatedObjectState::WaitOnDestroyConfirmation);
+		}
+
+		// We're now committing to what we wrote so inform the attachments writer.
+		if (BatchObjectInfo.AttachmentRecord)
+		{
+			Attachments.CommitReplicationRecord(BatchObjectInfo.AttachmentType, BatchObjectInfo.InternalIndex, BatchObjectInfo.AttachmentRecord);
+		}
+
+		// Update transmission record.
+		if (BatchObjectInfo.bSentState)
+		{
+			FNetBitArrayView ChangeMask(Info.GetChangeMaskStoragePointer(), Info.ChangeMaskBitCount);
+			FObjectRecord& ObjectRecord = OutRecord.ObjectReplicationRecords.AddDefaulted_GetRef();
+			CreateObjectRecord(&ChangeMask, Info, BatchObjectInfo, ObjectRecord);
+
+			// The object no longer has any dirty state, but may still have attachments that didn't fit
+			ChangeMask.Reset();
+		}
+		else if (BatchObjectInfo.AttachmentRecord != 0 || BatchObjectInfo.bSentTearOff || BatchObjectInfo.bSentDestroySubObject)
+		{
+			FObjectRecord& ObjectRecord = OutRecord.ObjectReplicationRecords.AddDefaulted_GetRef();
+			CreateObjectRecord(nullptr, Info, BatchObjectInfo, ObjectRecord);
+		}
+
+		// Schedule rest of dependent objects for replication, note there is no guarantee that they will replicate in same packet
+		for (FInternalNetHandle DependentInternalIndex : NetHandleManager->GetDependentObjects(BatchObjectInfo.InternalIndex))
+		{
+			if (ObjectsWithDirtyChanges.GetBit(DependentInternalIndex) && !WriteContext.ObjectsWrittenThisTick.GetBit(DependentInternalIndex))
+			{
+				WriteContext.DependentObjectsPendingSend.Push(DependentInternalIndex);
+				// Bumping the scheduling priority here will make sure that we will be scheduled the next update if we are not allowed to replicate this frame
+				SchedulingPriorities[DependentInternalIndex] = FMath::Max(SchedulingPriorities[BatchObjectInfo.InternalIndex], SchedulingPriorities[DependentInternalIndex]);
+			}
+		}			
+
+		WrittenObjectCount += (BatchObjectInfo.bSentState | BatchObjectInfo.bSentAttachments | BatchObjectInfo.bSentTearOff | BatchObjectInfo.bSentDestroySubObject) ? 1 : 0;
+
+		Info.HasDirtyChangeMask = 0U;
+		Info.HasDirtySubObjects = BatchObjectInfo.bHasDirtySubObjects;
+		Info.HasAttachments = BatchObjectInfo.bHasUnsentAttachments;
+
+		// Indicate that we are now waiting for a new baseline to be acknowledged
+		if (BatchObjectInfo.bSentState && BatchObjectInfo.NewBaselineIndex != FDeltaCompressionBaselineManager::InvalidBaselineIndex)
+		{
+			Info.PendingBaselineIndex = BatchObjectInfo.NewBaselineIndex;
+		}
+		
+		const bool bObjectIsStillDirty = BatchObjectInfo.bHasUnsentAttachments | BatchObjectInfo.bHasDirtySubObjects;
+		ObjectsWithDirtyChanges.SetBitValue(BatchObjectInfo.InternalIndex, bObjectIsStillDirty);
+
+		// Reset scheduling priority if everything was replicated
+		if (!bObjectIsStillDirty)
+		{
+			SchedulingPriorities[BatchObjectInfo.InternalIndex] = 0.0f;
+		}
+
+		if (BatchObjectInfo.InternalIndex != ObjectIndexForOOBAttachment)
+		{
+			// Mark this object as written this tick to avoid sending it multiple times
+			WriteContext.ObjectsWrittenThisTick.SetBit(BatchObjectInfo.InternalIndex);
+		}
+	}
+
+	return WrittenObjectCount;
+}
+
+FReplicationWriter::EWriteObjectRetryMode FReplicationWriter::HandleObjectBatchFailure(FReplicationWriter::EWriteObjectStatus WriteObjectStatus, const FBatchInfo& BatchInfo, const FReplicationWriter::FBitStreamInfo& BatchBitStreamInfo) const
+{
+	// Cleanup data stored in BatchInfo
+	for (const FBatchObjectInfo& BatchObjectInfo : BatchInfo.ObjectInfos)
+	{
+		// If we did not end up using the baseline we need to release it
+		if (BatchObjectInfo.bSentState && BatchObjectInfo.NewBaselineIndex != FDeltaCompressionBaselineManager::InvalidBaselineIndex)
+		{
+			BaselineManager->LostBaseline(Parameters.ConnectionId, BatchObjectInfo.InternalIndex, BatchObjectInfo.NewBaselineIndex);
+		}
+	}
+
+	if (WriteObjectStatus == EWriteObjectStatus::NoInstanceProtocol || WriteObjectStatus == EWriteObjectStatus::InvalidOwner)
+	{
+		return EWriteObjectRetryMode::TrySmallObject;
+	}
+
+	const uint32 BitsLeft = BatchBitStreamInfo.ReplicationCapacity - (BatchBitStreamInfo.BatchStartPos - BatchBitStreamInfo.ReplicationStartPos);
+	if (BitsLeft < Parameters.SmallObjectBitThreshold)
+	{
+		return EWriteObjectRetryMode::Abort;
+	}
+
+	// If there are more bits left than the split threshold we treat it as a huge object and proceed with splitting.
+	// We expect at least one part of the payload to be sendable if there are more bits left than the split threshold.
+	if (HugeObjectContext.SendStatus == EHugeObjectSendStatus::Idle && PartialNetObjectAttachmentHandler != nullptr)
+	{
+		const uint32 SplitThreshold = PartialNetObjectAttachmentHandler->GetConfig()->GetBitCountSplitThreshold();
+		if (BitsLeft > SplitThreshold)
+		{
+			//UE_LOG_REPLICATIONWRITER(TEXT("FReplicationWriter::HandleObjectBatchFailure Failed to write object with ParentInternalIndex: %u EWriteObjectRetryMode::SplitHugeObject"), BatchInfo.ParentInternalIndex);
+			return EWriteObjectRetryMode::SplitHugeObject;
+		}
+	}
+
+	if (WriteContext.FailedToWriteSmallObjectCount >= Parameters.MaxFailedSmallObjectCount)	
+	{
+		return EWriteObjectRetryMode::Abort;
+	}
+
+	// Default- try some more, hopefully smaller state, objects.
+	//UE_LOG_REPLICATIONWRITER(TEXT("FReplicationWriter::HandleObjectBatchFailure Failed to write object with ParentInternalIndex: %u EWriteObjectRetryMode::TrySmallObject"), BatchInfo.ParentInternalIndex);
+	return EWriteObjectRetryMode::TrySmallObject;
+}
+
+UDataStream::EWriteResult FReplicationWriter::BeginWrite()
+{
+	IRIS_PROFILER_SCOPE(FReplicationWriter_PrepareWrite);
+
+	// For now we do not support partial writes
+	check(!WriteContext.bIsValid);
+
+	if (!bReplicationEnabled)
+	{
+		return UDataStream::EWriteResult::NoData;
+	}
+
+	// See if we have any work to do
+	const bool bHasUpdatedObjectsToSend = ObjectsWithDirtyChanges.IsAnyBitSet();
+	const bool bHasDestroyedObjectsToSend = ObjectsPendingDestroy.IsAnyBitSet();
+	const bool bHasUnsentOOBAttachments = Attachments.HasUnsentAttachments(ENetObjectAttachmentType::OutOfBand, ObjectIndexForOOBAttachment);
+	const bool bHasUnsentHugeObject = Attachments.HasUnsentAttachments(ENetObjectAttachmentType::HugeObject, ObjectIndexForOOBAttachment);
+
+	// Nothing to send
+	if (!(bHasUpdatedObjectsToSend | bHasDestroyedObjectsToSend | bHasUnsentOOBAttachments | bHasUnsentHugeObject))
+	{
+		return UDataStream::EWriteResult::NoData;
+	}
+
+	// Initialize context which can be used over multiple calls to WriteData
+	WriteContext.bHasUpdatedObjectsToSend  = bHasUpdatedObjectsToSend | bHasUnsentOOBAttachments | bHasUnsentHugeObject;
+	WriteContext.bHasDestroyedObjectsToSend = bHasDestroyedObjectsToSend;
+	WriteContext.bHasHugeObjectToSend = bHasUnsentHugeObject;
+	WriteContext.bHasOOBAttachmentsToSend = bHasUnsentOOBAttachments;
+	WriteContext.CurrentIndex = 0u;
+	WriteContext.FailedToWriteSmallObjectCount = 0U;
+	WriteContext.SortedObjectCount = 0u;
+
+	// Reset dependent object array
+	WriteContext.DependentObjectsPendingSend.Reset();
+
+	// Reset objects written this tick
+	WriteContext.ObjectsWrittenThisTick.Reset();
+
+	// $IRIS TODO: LinearAllocator/ScratchPad?
+	// Allocate space for indices to send
+	// This should be allocated from frame temp allocator and be cleaned up end of frame, we might want this data to persist over multiple write calls but not over multiple frames 
+	// https://jira.it.epicgames.com/browse/UE-127374	
+	WriteContext.ScheduledObjectInfos = reinterpret_cast<FScheduleObjectInfo*>(FMemory::Malloc(sizeof(FScheduleObjectInfo) * Parameters.MaxActiveReplicatedObjectCount));
+	WriteContext.ScheduledObjectCount = ScheduleObjects(WriteContext.ScheduledObjectInfos);
+
+	WriteContext.bIsValid = true;
+
+	return UDataStream::EWriteResult::HasMoreData;
+}
+
+void FReplicationWriter::EndWrite()
+{
+	IRIS_PROFILER_SCOPE(FReplicationWriter_FinishWrite);
+
+	if (WriteContext.bIsValid)
+	{
+ #if UE_NET_ENABLE_REPLICATIONWRITER_LOG
+		// See if we failed to write any objects
+		if (const int32 NumPendingDependentObjects = WriteContext.DependentObjectsPendingSend.Num())
+		{
+			UE_LOG_REPLICATIONWRITER_WARNING(TEXT("FReplicationWriter::EndWrite() Has %d more dependent objects to write"), NumPendingDependentObjects);
+		}
+		
+		if (const int32 NumPendingObjectsToWrite = WriteContext.ScheduledObjectCount - WriteContext.CurrentIndex)
+		{
+			UE_LOG_REPLICATIONWRITER_WARNING(TEXT("FReplicationWriter::EndWrite() Has %d more objects with dirty data"), NumPendingObjectsToWrite);
+		}
+#endif
+
+		FMemory::Free(WriteContext.ScheduledObjectInfos);
+		WriteContext.ScheduledObjectInfos = nullptr;
+		WriteContext.bIsValid = false;
+	}
+}
+
+bool FReplicationWriter::HasDataToSend(const FWriteContext& Context) const
+{
+	return WriteContext.bIsValid & (Context.bHasDestroyedObjectsToSend | Context.bHasUpdatedObjectsToSend | Context.bHasHugeObjectToSend | Context.bHasOOBAttachmentsToSend);
+}
+
+UDataStream::EWriteResult FReplicationWriter::Write(FNetSerializationContext& Context)
+{
+	IRIS_PROFILER_SCOPE(FReplicationWriter_Write);
+
+	if (!HasDataToSend(WriteContext))
+	{
+		return UDataStream::EWriteResult::NoData;
+	}
+
+	FNetBitStreamWriter& Writer = *Context.GetBitStreamWriter();
+
+	// Setup internal context
+	FInternalNetSerializationContext InternalContext(Parameters.ReplicationSystem);
+	Context.SetLocalConnectionId(Parameters.ConnectionId);
+	Context.SetInternalContext(&InternalContext);	
+
+	// Give some info for the case when we consider splitting a huge object.
+	WriteBitStreamInfo.ReplicationStartPos = Writer.GetPosBits();
+	WriteBitStreamInfo.ReplicationCapacity = Writer.GetBitsLeft();
+
+	UE_NET_TRACE_SCOPE(ReplicationData, Writer, Context.GetTraceCollector(), ENetTraceVerbosity::Trace);
+	
+	FNetBitStreamRollbackScope Rollback(Writer);
+	const uint32 HeaderPos = Writer.GetPosBits();
+
+	uint16 WrittenObjectCount = 0;
+	const uint32 OldReplicationInfoCount = ReplicationRecord.GetInfoCount();
+
+	Writer.WriteBits(WrittenObjectCount, 16);
+
+	// Write timestamps etc? Or do we do this in header.
+	// WriteReplicationFrameData();
+
+	WrittenObjectCount += WriteObjectsPendingDestroy(Context);
+
+	// Only reason for overflow here is if we did not fit header
+	if (Writer.IsOverflown())
+	{
+		return UDataStream::EWriteResult::NoData;
+	}
+
+	WrittenObjectCount += WriteOOBAttachments(Context);
+
+	WrittenObjectCount += WriteObjects(Context);
+
+	UDataStream::EWriteResult WriteResult = UDataStream::EWriteResult::Ok;
+
+	// If we have more data to write, request more updates
+	// $IRIS $TODO: When we have better control over bandwidth usage, introduce setting to only allow over-commit if we have a huge object or split rpc to send
+	// https://jira.it.epicgames.com/browse/UE-127371
+	if (HasDataToSend(WriteContext))
+	{
+		if (WriteContext.bHasHugeObjectToSend)
+		{
+			WriteResult = UDataStream::EWriteResult::HasMoreData;
+		}
+		else
+		{
+			WriteResult = UDataStream::EWriteResult::Ok;
+		}
+	}
+	
+	if (!Writer.IsOverflown() && WrittenObjectCount > 0)
+	{
+		{
+			// Seek back to HeaderPos and update the header
+			FNetBitStreamWriteScope WriteScope(Writer, HeaderPos);
+			Writer.WriteBits(WrittenObjectCount, 16);
+		}
+
+		//UE_LOG_REPLICATIONWRITER(TEXT("FReplicationWriter::Write() Wrote %u Objects for ConnectionId:%u, ReplicationSystemId: %u."), WrittenObjectCount, Parameters.ConnectionId, Parameters.ReplicationSystem->GetId());	
+	
+		// Push record
+		const uint32 ReplicationInfoCount = ReplicationRecord.GetInfoCount() - OldReplicationInfoCount;
+		ReplicationRecord.PushRecord(ReplicationInfoCount);
+
+#if UE_NET_VALIDATE_REPLICATION_RECORD
+		check(s_ValidateReplicationRecord(&ReplicationRecord, Parameters.MaxActiveReplicatedObjectCount + 1U, false));
+#endif
+
+	}
+	else 
+	{
+		// Trigger rollback as we did not write any data
+		Writer.DoOverflow();
+		WriteResult = UDataStream::EWriteResult::NoData;
+	}
+
+	// Report packet stats
+	UE_NET_TRACE_PACKET_STATSCOUNTER(Parameters.ReplicationSystem->GetId(), Parameters.ConnectionId, ReplicationWriter.WrittenObjectCount, WrittenObjectCount, ENetTraceVerbosity::Trace);
+	UE_NET_TRACE_PACKET_STATSCOUNTER(Parameters.ReplicationSystem->GetId(), Parameters.ConnectionId, ReplicationWriter.RemainingObjectsPendingWriteCount, WriteContext.ScheduledObjectCount - WriteContext.CurrentIndex, ENetTraceVerbosity::Trace);
+	UE_NET_TRACE_PACKET_STATSCOUNTER(Parameters.ReplicationSystem->GetId(), Parameters.ConnectionId, ReplicationWriter.ScheduledObjectCount, WriteContext.ScheduledObjectCount, ENetTraceVerbosity::Trace);
+
+	return WriteResult;
+}
+
+void FReplicationWriter::SetupReplicationInfoForAttachmentsToObjectsNotInScope()
+{
+	FReplicationInfo& Info = GetReplicationInfo(ObjectIndexForOOBAttachment);
+	Info = FReplicationInfo();
+	Info.State = static_cast<uint32>(EReplicatedObjectState::AttachmentToObjectNotInScope);
+	ReplicationRecord.ResetList(ReplicatedObjectsRecordInfoLists[ObjectIndexForOOBAttachment]);
+}
+
+void FReplicationWriter::ApplyFilterToChangeMask(uint32 ParentInternalIndex, uint32 InternalIndex, FReplicationInfo& Info, const FReplicationProtocol* Protocol, const uint8* InternalStateBuffer)
+{
+	const uint32* ConditionalChangeMaskPointer = (EnumHasAnyFlags(Protocol->ProtocolTraits, EReplicationProtocolTraits::HasConditionalChangeMask) ? reinterpret_cast<const uint32*>(InternalStateBuffer + Protocol->GetConditionalChangeMaskOffset()) : static_cast<const uint32*>(nullptr));
+	const bool bChangeMaskWasModified = ReplicationConditionals->ApplyConditionalsToChangeMask(Parameters.ConnectionId, ParentInternalIndex, InternalIndex, Info.GetChangeMaskStoragePointer(), ConditionalChangeMaskPointer, Protocol);
+	if (bChangeMaskWasModified && !MakeNetBitArrayView(Info.GetChangeMaskStoragePointer(), Info.ChangeMaskBitCount).IsAnyBitSet())
+	{
+		Info.HasDirtyChangeMask = 0;
+	}
+}
+
+void FReplicationWriter::InvalidateBaseline(uint32 InternalIndex, FReplicationInfo& Info)
+{
+	FReplicationRecord::FRecordInfoList& RecordInfoList = ReplicatedObjectsRecordInfoLists[InternalIndex];
+
+	// Iterate over all data in flight for this object and mark any new baselines as invalid to avoid acking or nacking an invalidated baseline
+	FReplicationRecord::FRecordInfo* CurrentRecordInfo = ReplicationRecord.GetInfoForIndex(RecordInfoList.FirstRecordIndex);
+	while (CurrentRecordInfo)
+	{
+		CurrentRecordInfo->NewBaselineIndex = FDeltaCompressionBaselineManager::InvalidBaselineIndex;
+		CurrentRecordInfo = ReplicationRecord.GetInfoForIndex(CurrentRecordInfo->NextIndex);
+	};
+
+	Info.LastAckedBaselineIndex = FDeltaCompressionBaselineManager::InvalidBaselineIndex;
+	Info.PendingBaselineIndex = FDeltaCompressionBaselineManager::InvalidBaselineIndex;
+}
+
+bool FReplicationWriter::PatchupObjectChangeMaskWithInflightChanges(uint32 InternalIndex, FReplicationInfo& Info)
+{
+	bool bInFlightChangesAdded = false;
+
+	const FReplicationRecord::FRecordInfoList& RecordInfoList = ReplicatedObjectsRecordInfoLists[InternalIndex];
+
+	FNetBitArrayView ChangeMask = FChangeMaskUtil::MakeChangeMask(Info.ChangeMaskOrPtr, Info.ChangeMaskBitCount);
+
+	// Iterate over all data in flight for this object and include any changes in flight to ensure atomicity when received
+	// N.B. We don't check if this object is in huge object mode and check to see if any of these changes were part of that payload.
+	const FReplicationRecord::FRecordInfo* CurrentRecordInfo = ReplicationRecord.GetInfoForIndex(RecordInfoList.FirstRecordIndex);
+	while (CurrentRecordInfo)
+	{
+		if (CurrentRecordInfo->HasChangeMask)
+		{
+			bInFlightChangesAdded = true;
+			const FNetBitArrayView CurrentRecordInfoChangeMask(FChangeMaskUtil::MakeChangeMask(CurrentRecordInfo->ChangeMaskOrPtr, Info.ChangeMaskBitCount));
+			ChangeMask.Combine(CurrentRecordInfoChangeMask, FNetBitArrayView::OrOp);
+		}
+
+		CurrentRecordInfo = ReplicationRecord.GetInfoForIndex(CurrentRecordInfo->NextIndex);
+	};
+
+	return bInFlightChangesAdded;
+}
+
+void FReplicationWriter::SetRemoteNetTokenStoreState(FNetTokenStoreState* InRemoteTokenStoreState)
+{
+	ResolveContext.RemoteNetTokenStoreState = InRemoteTokenStoreState;
+}
+
+void FReplicationWriter::SetNetExports(FNetExports& InNetExports)
+{
+	NetExports = &InNetExports;
+}
+
+bool FReplicationWriter::IsDestroyObjectPartOfActiveHugeObject(uint32 InternalIndex, const FReplicationInfo& Info) const
+{
+	if (HugeObjectContext.SendStatus == EHugeObjectSendStatus::Idle)
+	{
+		return false;
+	}
+
+	if (InternalIndex == HugeObjectContext.InternalIndex)
+	{
+		return true;
+	}
+
+	if (Info.IsSubObject)
+	{
+		// At this state any information regarding the parent has been cleared. Check if the index is part of the payload.
+		const FBatchRecord& BatchRecord = HugeObjectContext.BatchRecord;
+		for (const FObjectRecord& ObjectRecord : MakeArrayView(BatchRecord.ObjectReplicationRecords.GetData(), BatchRecord.ObjectReplicationRecords.Num()))
+		{
+			if (InternalIndex == ObjectRecord.Record.Index)
+			{
+				return true;
+			}
+		}
+	}
+
+	return false;
+}
+
+void FReplicationWriter::ClearHugeObjectContext(FHugeObjectContext& Context) const
+{
+	Context.InternalIndex = 0;
+	Context.SendStatus = EHugeObjectSendStatus::Idle;
+	Context.BatchRecord = FBatchRecord();
+	Context.BatchExports.Reset();
+}
+
+bool FReplicationWriter::CollectAndWriteExports(FNetSerializationContext& Context, uint8* RESTRICT InternalBuffer, const FReplicationProtocol* Protocol) const
+{
+	FNetReferenceCollector Collector(ENetReferenceCollectorTraits::OnlyCollectReferencesThatCanBeExported);
+
+	FReplicationProtocolOperationsInternal::CollectReferences(Context, Collector, InternalBuffer, Protocol);
+
+	if (Collector.GetCollectedReferences().Num())
+	{
+		TArray<FNetObjectReference, TInlineAllocator<32>> Exports;
+		for (const FNetReferenceCollector::FReferenceInfo& Info : MakeArrayView(Collector.GetCollectedReferences()))
+		{
+			Exports.AddUnique(Info.Reference);
+		}
+		return ObjectReferenceCache->WriteExports(Context, MakeArrayView(Exports));
+	}
+	else
+	{
+		return ObjectReferenceCache->WriteExports(Context, MakeArrayView<const FNetObjectReference>(nullptr, 0));
+	}
+}
+
+bool FReplicationWriter::IsWriteObjectSuccess(EWriteObjectStatus Status) const
+{
+	return (Status == EWriteObjectStatus::Success) | (Status == EWriteObjectStatus::InvalidState);
+}
+
+void FReplicationWriter::DiscardAllRecords()
+{
+	TReplicationRecordHelper Helper(ReplicatedObjects.GetData(), ReplicatedObjectsRecordInfoLists.GetData(), &ReplicationRecord);
+
+	const uint32 RecordCount = ReplicationRecord.GetRecordCount();
+	for (uint32 RecordIt = 0, RecordEndIt = RecordCount; RecordIt != RecordEndIt; ++RecordIt)
+	{
+		if (const uint32 RecordInfoCount = ReplicationRecord.PopRecord())
+		{
+			Helper.Process(RecordInfoCount,
+				[this](const FReplicationRecord::FRecordInfo& RecordInfo, FReplicationInfo& Info, FNetObjectAttachmentsWriter::ReplicationRecord AttachmentRecord)
+				{
+					HandleDiscardedRecord(RecordInfo, Info, AttachmentRecord);
+				}
+			);
+		}
+	}
+}
+
+// Do minimal work to free references and resources. We assume connection removal handling will be dealt with by respective subsystems,
+// such as the DeltaCompressionBaselinemanager releasing baselines. 
+void FReplicationWriter::StopAllReplication()
+{
+	if (ReplicatedObjects.Num() == 0)
+	{
+		return;
+	}
+
+	// Don't process special index.
+	ReplicatedObjects[ObjectIndexForOOBAttachment].State = (uint8)EReplicatedObjectState::Invalid;
+
+	// We cannot tell for sure which objects need processing so we check them all.
+	const FReplicationInfo* FirstInfo = ReplicatedObjects.GetData();
+	for (const FReplicationInfo& Info : ReplicatedObjects)
+	{
+		if (Info.GetState() == EReplicatedObjectState::Invalid)
+		{
+			continue;
+		}
+
+		// Free allocated ChangeMask (if it is allocated)
+		FChangeMaskStorageOrPointer::Free(Info.ChangeMaskOrPtr, Info.ChangeMaskBitCount, s_DefaultChangeMaskAllocator);
+
+		// Release object reference
+		const uint32 InternalIndex = &Info - FirstInfo;
+		NetHandleManager->ReleaseNetObjectRef(InternalIndex);
+	}
+}
+
+}
