Iris experimental replication system.

diff --git a/Engine/Source/Runtime/Experimental/Iris/Core/Private/Iris/ReplicationSystem/ReplicationWriter.h b/Engine/Source/Runtime/Experimental/Iris/Core/Private/Iris/ReplicationSystem/ReplicationWriter.h
new file mode 100644
--- /dev/null
+++ b/Engine/Source/Runtime/Experimental/Iris/Core/Private/Iris/ReplicationSystem/ReplicationWriter.h
@@ -0,0 +1,479 @@
+// Copyright Epic Games, Inc. All Rights Reserved.
+
+#pragma once
+
+// This class will never be included from public headers
+#include "Net/Core/NetBitArray.h"
+#include "Iris/ReplicationSystem/NetHandleManager.h"
+#include "Iris/ReplicationSystem/ReplicationTypes.h"
+#include "Iris/ReplicationSystem/ReplicationRecord.h"
+#include "Iris/DataStream/DataStream.h"
+#include "Iris/ReplicationSystem/AttachmentReplication.h"
+#include "Iris/ReplicationSystem/NetBlob/NetBlob.h"
+#include "Iris/ReplicationSystem/NetHandle.h"
+#include "Iris/ReplicationSystem/NetExports.h"
+#include "Iris/ReplicationSystem/ObjectReferenceCache.h"
+#include "Containers/Array.h"
+
+// Forward declaration
+class UReplicationSystem;
+class UNetObjectBlobHandler;
+class UPartialNetObjectAttachmentHandler;
+class UReplicationBridge;
+namespace UE::Net
+{
+	class FNetBitStreamReader;
+	class FNetBitStreamWriter;
+	class FNetObjectAttachment;
+	class FNetSerializationContext;
+	struct FReplicationProtocol;
+	namespace Private
+	{
+		struct FChangeMaskCache;
+		class FNetHandleManager;
+		class FReliableNetBlobQueue;
+		class FReplicationConditionals;
+		class FReplicationFiltering;
+		class FReplicationSystemInternal;
+		class FDeltaCompressionBaselineManager;
+		class FDeltaCompressionBaseline;
+	}
+}
+
+namespace UE::Net::Private
+{
+
+class FReplicationWriter
+{
+public:
+	// Scheduling constants
+	static constexpr float CreatePriority = 1.f;
+	static constexpr float TearOffPriority = 1.f;
+	static constexpr float LostStatePriorityBump = 1.f;
+	static constexpr float SchedulingThresholdPriority = 1.f;
+
+	// When scheduling there is no point in scheduling more objects than we can fit in a packet
+	static const uint32 PartialSortObjectCount = 128u;
+
+public:
+
+	// State
+	enum class EReplicatedObjectState : uint8
+	{
+		Invalid = 0,
+
+		// Special states
+		AttachmentToObjectNotInScope,	// Special state for object index used for sending attachments to remote owned objects
+		HugeObject,						// Special state for object index used for sending parts of huge object payloads
+
+		// 
+		PendingCreate,					// Not yet created, no data in flight
+		WaitOnCreateConfirmation,		// Waiting for confirmation from remote, we can send state data, but if we do we must also include creation info until the object is Created
+		Created,						// Confirmed by remote
+		PendingFlush,					// This object should be flushed
+		WaitOnFlush,					// Pending flush, we are waiting for all in-flight state data to be acknowledged
+		Flushed,						// No more data in flight we can now hibernate or destroy
+		Hibernating,					// Hibernating, if an object is flushed it can be set to hibernate/be dormant
+		PendingTearOff,					// TearOff should be sent
+		SubObjectPendingDestroy,		// SubObject destroy should be sent
+		CancelPendingDestroy,			// Destroy was sent but object wants to start replicating again
+		PendingDestroy,					// Object is set to be destroyed
+		WaitOnDestroyConfirmation,		// Destroy has been sent, waiting for response from client
+		Destroyed,						// Confirmed as Destroyed,
+		PermanentlyDestroyed,			// DestructionInfo has been confirmed as received
+
+		Count
+	};
+
+	static_assert((uint8)(EReplicatedObjectState::Count) <= 32, "EReplicatedObjectState must fit in 5 bits. See FReplicationInfo::State and FReplicationRecord::FRecordInfo::ReplicatedObjectState members.");
+
+	const TCHAR* LexToString(const EReplicatedObjectState State);
+
+	// Bookkeeping info required for a replicated object
+	// Keep as small as possible since there is one per replicated object per connection
+	// Changemask can and will most likely be replaced by a smaller index into a pool to reduce overhead
+	struct FReplicationInfo
+	{
+		inline FReplicationInfo();
+
+		FChangeMaskStorageOrPointer ChangeMaskOrPtr;			// Changemask storage or pointer to storage	
+		union 
+		{
+			uint64 Value;
+			struct 
+			{
+				uint64 ChangeMaskBitCount : 16;							// This is cached to avoid having to look it up in the protocol
+				uint64 State : 5;										// Current state
+				uint64 HasDirtySubObjects : 1;							// Set if this object has dirty subobjects
+				uint64 IsSubObject : 1;									// Set if this object is a subobject
+				uint64 HasDirtyChangeMask : 1;							// Set if the ChangeMask might be dirty, if not set the changemask should be zero!
+				uint64 HasAttachments : 1;								// Set if there are attachments, such as RPCs, waiting to be sent
+				uint64 HasChangemaskFilter : 1;							// Do we need to filter our changemask or not
+				uint64 IsDestructionInfo : 1;							// If this is a destruction info
+				uint64 IsCreationConfirmed : 1;							// We know that this object has been created on the receiving end
+				uint64 TearOff : 1;										// This object should be torn off
+				uint64 SubObjectPendingDestroy : 1;						// This object is a subobject that should be destroyed when we replicate owner
+				uint64 IsDeltaCompressionEnabled : 1;
+				uint64 LastAckedBaselineIndex : 2;
+				uint64 PendingBaselineIndex : 2;
+				uint64 Padding : 29;
+			};
+		};
+
+		static const uint32 LocalChangeMaskMaxBitCount = 64u;
+
+		void SetState(EReplicatedObjectState NewState);
+		EReplicatedObjectState GetState() const { return (EReplicatedObjectState)State; }
+
+		ChangeMaskStorageType* GetChangeMaskStoragePointer() { return ChangeMaskOrPtr.GetPointer(ChangeMaskBitCount); }
+		const ChangeMaskStorageType* GetChangeMaskStoragePointer() const { return ChangeMaskOrPtr.GetPointer(ChangeMaskBitCount); }
+	};
+
+	static_assert(sizeof(FReplicationInfo) == 16, "Expected sizeof FReplicationInfo to be 16 bytes");
+
+public:
+	~FReplicationWriter();
+
+	// Init
+	void Init(const FReplicationParameters& InParameters);
+
+	// Update new or existing/destroyed 
+	void UpdateScope(const FNetBitArrayView& ScopedObjects);
+
+	// Mark objects as tearoff and propagate dirty changemasks
+	void TearOffAndUpdateDirtyChangeMasks(const FChangeMaskCache& CachedChangeMasks) { InternalUpdateDirtyChangeMasks(CachedChangeMasks, 1U); }
+
+	void NotifyDestroyedObjectPendingTearOff(FInternalNetHandle ObjectInternalIndex);
+
+	// Propagate dirty changemasks
+	void UpdateDirtyChangeMasks(const FChangeMaskCache& CachedChangeMasks) { InternalUpdateDirtyChangeMasks(CachedChangeMasks, 0U); }
+
+	// Returns objects that are in need of a priority update. It could be dirty, new or objects in need of resending.
+	const FNetBitArray& GetObjectsRequiringPriorityUpdate() const;
+
+	// UpdatedPriorities contains priorities for all objects. Objects in need of a priority update should use the newly calculated priorities.
+	void UpdatePriorities(const float* UpdatedPriorities);
+
+	UDataStream::EWriteResult BeginWrite();
+
+	// WriteData to Packet, returns true for now if data was written
+	UDataStream::EWriteResult Write(FNetSerializationContext& Context);
+
+	void EndWrite();
+
+	// Deal with processing of lost and delivered data.
+	void ProcessDeliveryNotification(EPacketDeliveryStatus PacketDeliveryStatus);
+
+	void SetReplicationEnabled(bool bInReplicationEnabled);
+	bool IsReplicationEnabled() const;
+
+	void SetRemoteNetTokenStoreState(FNetTokenStoreState* RemoteTokenStoreState);
+	void SetNetExports(FNetExports& InNetExports);
+
+	// Attachments
+	void QueueNetObjectAttachments(FInternalNetHandle OwnerInternalIndex, FInternalNetHandle SubObjectInternalIndex, TArrayView<const TRefCountPtr<FNetBlob>> Attachments);
+
+private:
+	// Various types
+
+	enum : uint32
+	{
+		ObjectIndexForOOBAttachment = 0U,
+	};
+
+	// Propagate dirty changemasks
+	void InternalUpdateDirtyChangeMasks(const FChangeMaskCache& CachedChangeMasks, uint32 bTearOff);
+
+	struct FScheduleObjectInfo
+	{
+		uint32 Index;
+		float SortKey;
+	};
+
+	// We persist some information during write so that we can support writing multiple packets with data without re-doing scheduling work.
+	struct FWriteContext
+	{
+		FWriteContext() : bIsValid(0) {}
+
+		// Objects we have written in this packet batch to avoid sending same object multiple times
+		FNetBitArray ObjectsWrittenThisTick;
+		// DependentObjets that we should try to write this packet batch, aka. allow overcommit if we have pending DependentObjects when the packet is full
+		TArray<uint32, TInlineAllocator<32>> DependentObjectsPendingSend;
+		FScheduleObjectInfo* ScheduledObjectInfos;				// Scheduled objects
+		uint32 ScheduledObjectCount;
+
+		uint32 SortedObjectCount;								// We do partial sorting so we need to track how many object we have sorted
+		uint32 CurrentIndex;									// If we allow multiple packets to be written during the same update this is where we should continue writing
+
+		uint32 FailedToWriteSmallObjectCount;
+
+		uint32 bHasDestroyedObjectsToSend : 1;
+		uint32 bHasUpdatedObjectsToSend : 1;
+		uint32 bHasHugeObjectToSend : 1;
+		uint32 bHasOOBAttachmentsToSend : 1;
+		uint32 bIsValid : 1;
+	};
+
+	struct FBatchObjectInfo
+	{
+		FNetHandle Handle;
+		uint32 InternalIndex;
+		FNetObjectAttachmentsWriter::ReplicationRecord AttachmentRecord;
+		ENetObjectAttachmentType AttachmentType;
+		bool bHasUnsentAttachments;
+		uint32 NewBaselineIndex : 2;
+		uint32 bIsInitialState : 1;
+		uint32 bSentState : 1;
+		uint32 bSentAttachments : 1;
+		uint32 bHasDirtySubObjects : 1;
+		uint32 bSentTearOff : 1;
+		uint32 bSentDestroySubObject : 1;
+	};
+
+	struct FBatchInfo
+	{
+		TArray<FBatchObjectInfo, TInlineAllocator<16>> ObjectInfos;
+		uint32 ParentInternalIndex;
+	};
+
+	struct FObjectRecord
+	{
+		FReplicationRecord::FRecordInfo Record;
+		FNetObjectAttachmentsWriter::ReplicationRecord AttachmentRecord;
+	};
+
+	struct FBatchRecord
+	{
+		TArray<FObjectRecord, TInlineAllocator<16>> ObjectReplicationRecords;
+	};
+
+	struct FBitStreamInfo
+	{
+		uint32 ReplicationStartPos;
+		uint32 BatchStartPos;
+		uint32 ReplicationCapacity;
+	};
+
+	enum class EHugeObjectSendStatus : uint32
+	{
+		Idle,
+		Sending,
+	};
+
+	struct FHugeObjectContext
+	{
+		FHugeObjectContext();
+
+		EHugeObjectSendStatus SendStatus;
+		uint32 InternalIndex;
+		FBatchRecord BatchRecord;
+		FNetExportContext::FBatchExports BatchExports;
+		const FNetDebugName* DebugName;
+	};
+
+	enum EWriteObjectFlag : unsigned
+	{
+		WriteObjectFlag_State = 1U,
+		WriteObjectFlag_Attachments = WriteObjectFlag_State << 1U,
+		WriteObjectFlag_HugeObject = WriteObjectFlag_Attachments << 1U,
+	};
+
+	enum class EWriteObjectRetryMode : unsigned
+	{
+		// Stop trying to write more object this frame.
+		Abort,
+		// Continue with something smaller, it might succeed.
+		TrySmallObject,
+		// The object is probably huge. Enter special mode for huge objects.
+		SplitHugeObject,
+	};
+
+	enum class EWriteObjectStatus : unsigned
+	{
+		Success,
+
+		// The object is in an invalid state and won't be written. This is not considered a failure.
+		InvalidState,
+
+		// BitStream overflow.
+		BitStreamOverflow,
+
+		// A detached instance, which no longer has an instance protocol. This object cannot be replicated.
+		NoInstanceProtocol,
+
+		// A subobject with an invalid owner.
+		InvalidOwner,
+
+		// Some error occurred while serializing the object.
+		Error,
+
+	};
+
+private:
+
+	bool IsObjectIndexForOOBAttachment(uint32 InternalIndex) const { return InternalIndex == ObjectIndexForOOBAttachment; }
+
+	void GetInitialChangeMask(ChangeMaskStorageType* ChangeMaskStorage, const FReplicationProtocol* Protocol);
+
+	// Start replication of new object with the specified InternalIndex
+	void StartReplication(uint32 InternalIndex);
+
+	// Stop replication object with the specified InternalIndex
+	void StopReplication(uint32 InternalIndex);
+
+	// Get ReplicationInfo for specified InternalIndex, prefer to use this method over direct access.
+	FReplicationInfo& GetReplicationInfo(uint32 InternalIndex);
+	const FReplicationInfo& GetReplicationInfo(uint32 InternalIndex) const;
+
+	// Set the state of a ReplicatedObject, prefer this method to enable logging
+	void SetState(uint32 InternalIndex, EReplicatedObjectState NewState);
+
+	// Write index part of handle
+	void WriteNetHandleId(FNetBitStreamWriter& Writer, FNetHandle Handle);
+		
+	// 
+	void CreateObjectRecord(const FNetBitArrayView* ChangeMask, const FReplicationInfo& Info, const FBatchObjectInfo& ObjectInfo, FObjectRecord& OutRecord);
+	
+	// Push and link info for written object to ReplicationRecord
+	void CommitObjectRecord(uint32 InternalObjectIndex, const FObjectRecord& Record);
+
+	void CommitBatchRecord(const FBatchRecord& BatchRecord);
+
+	uint32 ScheduleObjects(FScheduleObjectInfo* ScheduledObjectIndices);
+	
+	// Partial sort of OutScheduledObjectIndices, will sort at most PartialSortObjectCount objects
+	uint32 SortScheduledObjects(FScheduleObjectInfo* ScheduledObjectIndices, uint32 ScheduledObjectCount, uint32 StartIndex);
+
+	// Write all objects pending destroy (or as many as we fit in the current packet)
+	uint32 WriteObjectsPendingDestroy(FNetSerializationContext& Context);
+
+	// Write objects recursive
+	EWriteObjectStatus WriteObjectInBatch(FNetSerializationContext& Context, uint32 InternalIndex, uint32 WriteObjectFlags, FBatchInfo& OutBatchInfo);
+
+	// Write destruction info for an object that should be destroyed
+	// returns > 1 if Objects was written, 0 if the objects was skipped (failed dependency or waiting for creation confirmation) -1 if the stream is full and we should stop
+	int WriteDestructionInfo(FNetSerializationContext& Context, uint32 InternalIndex);
+
+	// Write Object and any subobject(s) to stream as an atomic batch
+	// returns > 1 if Objects was written, 0 if the objects was skipped (failed dependency or waiting for creation confirmation) -1 if the stream is full and we should stop
+	int WriteObjectBatch(FNetSerializationContext& Context, uint32 InternalIndex, uint32 WriteObjectFlags);
+
+	int PrepareAndSendHugeObjectPayload(FNetSerializationContext& Context, uint32 InternalIndex);
+
+	// Write OOBAttachments
+	uint32 WriteOOBAttachments(FNetSerializationContext& Context);
+
+	// Write as many scheduled objects to stream as we can fit.
+	uint32 WriteObjects(FNetSerializationContext& Context);
+
+	// Updates ReplicationInfos, pushes ReplicationRecords etc after a successful call to WriteObjectInBatch() on a top level object
+	int HandleObjectBatchSuccess(const FBatchInfo& BatchInfo, FBatchRecord& OutRecord);
+
+	// Determines the best course of action after a WriteObjectBatch() call failed.
+	EWriteObjectRetryMode HandleObjectBatchFailure(EWriteObjectStatus WriteObjectStatus, const FBatchInfo& BatchInfo, const FBitStreamInfo& BatchBitStreamInfo) const;
+
+	// Update logic for dropped RecordInfo
+	void HandleDroppedRecord(const FReplicationRecord::FRecordInfo& RecordInfo, FReplicationInfo& Info, FNetObjectAttachmentsWriter::ReplicationRecord AttachmentRecord);
+	template<EReplicatedObjectState LostState> void HandleDroppedRecord(EReplicatedObjectState CurrentState, const FReplicationRecord::FRecordInfo& RecordInfo, FReplicationInfo& Info, FNetObjectAttachmentsWriter::ReplicationRecord AttachmentRecord);
+
+	// Update logic for delivered RecordInfo
+	void HandleDeliveredRecord(const FReplicationRecord::FRecordInfo& RecordInfo, FReplicationInfo& Info, FNetObjectAttachmentsWriter::ReplicationRecord AttachmentRecord);
+
+	// Update logic for discarded RecordInfo, for preventing memory leaks on disconnect and shutdown.
+	void HandleDiscardedRecord(const FReplicationRecord::FRecordInfo& RecordInfo, FReplicationInfo& Info, FNetObjectAttachmentsWriter::ReplicationRecord AttachmentRecord);
+
+	// Setup replication info to be able to send attachments to objects not in scope
+	void SetupReplicationInfoForAttachmentsToObjectsNotInScope();
+
+	void ApplyFilterToChangeMask(uint32 ParentInternalIndex, uint32 InternalIndex, FReplicationInfo& Info, const FReplicationProtocol* Protocol, const uint8* InternalStateBuffer);
+
+	// Patchup changemask to include any in-flight changes. Returns true if in-flight changes were added.
+	bool PatchupObjectChangeMaskWithInflightChanges(uint32 InternalIndex, FReplicationInfo& Info);
+
+	// Invalidate all inflight baseline information
+	void InvalidateBaseline(uint32 InternalIndex, FReplicationInfo& Info);
+
+	// Returns true object and subobjects can be created on remote
+	bool CanSendObject(uint32 InternalIndex) const;
+
+	inline bool IsInitialState(const EReplicatedObjectState State) const { return State == EReplicatedObjectState::PendingCreate || (bHighPrioCreate && State == EReplicatedObjectState::WaitOnCreateConfirmation); }
+
+	bool IsActiveHugeObject(uint32 InternalIndex) const { return HugeObjectContext.InternalIndex == InternalIndex && HugeObjectContext.SendStatus != EHugeObjectSendStatus::Idle; }
+	bool IsDestroyObjectPartOfActiveHugeObject(uint32 InternalIndex, const FReplicationInfo& Info) const;
+
+	void ClearHugeObjectContext(FHugeObjectContext& Context) const;
+
+	bool HasDataToSend(const FWriteContext& Context) const;
+
+	bool CollectAndWriteExports(FNetSerializationContext& Context, uint8* RESTRICT InternalBuffer, const FReplicationProtocol* Protocol) const;
+
+	bool IsWriteObjectSuccess(EWriteObjectStatus Status) const;
+
+	void SerializeObjectStateDelta(FNetSerializationContext& Context, uint32 InternalIndex, const FReplicationInfo& Info, const FNetHandleManager::FReplicatedObjectData& ObjectData, const uint8* RESTRICT ReplicatedObjectStateBuffer, FDeltaCompressionBaseline& CurrentBaseline, uint32 CreatedBaselineIndex);
+
+	void DiscardAllRecords();
+	void StopAllReplication();
+
+private:
+	// Replication parameters
+	FReplicationParameters Parameters;
+
+	// Record of all in-flight data
+	FReplicationRecord ReplicationRecord;
+
+	// Tracking information for the state of all objects
+	TArray<FReplicationInfo> ReplicatedObjects;
+
+	// Tracking information linking all in-flight data per object
+	TArray<FReplicationRecord::FRecordInfoList> ReplicatedObjectsRecordInfoLists;
+
+	// Each replicated object has a scheduling priority that is bumped every time we have a chance to send and zeroed out every time the object is successfully sent
+	TArray<float> SchedulingPriorities;
+
+	// Track Objects Pending Destroy?
+	FNetBitArray ObjectsPendingDestroy;
+
+	// Objects in this bitArray with dirty change masks
+	FNetBitArray ObjectsWithDirtyChanges;
+
+	// Track Objects That is in scope for this connection
+	FNetBitArray ObjectsInScope;
+	
+	// Handles logic for all attachments to objects.
+	FNetObjectAttachmentsWriter Attachments;
+
+	// Cached internal systems
+	FReplicationSystemInternal* ReplicationSystemInternal = nullptr;
+	FNetHandleManager* NetHandleManager = nullptr;
+	UReplicationBridge* ReplicationBridge = nullptr;
+	FDeltaCompressionBaselineManager* BaselineManager = nullptr;
+	FNetObjectResolveContext ResolveContext;
+	FObjectReferenceCache* ObjectReferenceCache = nullptr;
+	const FReplicationFiltering* ReplicationFiltering = nullptr;
+	FReplicationConditionals* ReplicationConditionals = nullptr;
+	const UPartialNetObjectAttachmentHandler* PartialNetObjectAttachmentHandler = nullptr;
+	const UNetObjectBlobHandler* NetObjectBlobHandler = nullptr;
+	FNetExports* NetExports = nullptr;
+
+	FWriteContext WriteContext;
+	FBitStreamInfo WriteBitStreamInfo;
+	FHugeObjectContext HugeObjectContext;
+
+	// Is replication enabled?
+	bool bReplicationEnabled = false;
+	
+	// Should we use high prio create?
+	const bool bHighPrioCreate = false;
+};
+
+inline FReplicationWriter::FReplicationInfo::FReplicationInfo()
+: Value(0U)
+{
+}
+
+template<FReplicationWriter::EReplicatedObjectState LostState> void FReplicationWriter::HandleDroppedRecord(FReplicationWriter::EReplicatedObjectState CurrentState, const FReplicationRecord::FRecordInfo& RecordInfo, FReplicationWriter::FReplicationInfo& Info, FNetObjectAttachmentsWriter::ReplicationRecord AttachmentRecord)
+{
+	//static_assert(false, "Expected specialization to exist.");
+}
+
+}
diff --git a/Engine/Source/Runtime/Experimental/Iris/Core/Private/Iris/ReplicationSystem/StringTokenStore.cpp b/Engine/Source/Runtime/Experimental/Iris/Core/Private/Iris/ReplicationSystem/StringTokenStore.cpp
new file mode 100644
--- /dev/null
+++ b/Engine/Source/Runtime/Experimental/Iris/Core/Private/Iris/ReplicationSystem/StringTokenStore.cpp
@@ -0,0 +1,141 @@
+// Copyright Epic Games, Inc. All Rights Reserved.
+
+#include "Iris/ReplicationSystem/StringTokenStore.h"
+#include "Iris/Core/IrisLog.h"
+#include "Iris/Serialization/NetBitStreamReader.h"
+#include "Iris/Serialization/NetBitStreamWriter.h"
+#include "Iris/Serialization/NetBitStreamUtil.h"
+#include "Iris/Serialization/NetSerializationContext.h"
+#include "Hash/CityHash.h"
+
+#if (UE_BUILD_SHIPPING || UE_BUILD_TEST)
+#	define UE_NET_ENABLE_STRING_TOKEN_LOG 0
+#else
+#	define UE_NET_ENABLE_STRING_TOKEN_LOG 1
+#endif 
+
+#if UE_NET_ENABLE_STRING_TOKEN_LOG
+#	define UE_LOG_STRINGTOKEN(Format, ...)  UE_LOG(LogIris, Verbose, Format, ##__VA_ARGS__)
+#else
+#	define UE_LOG_STRINGTOKEN(...)
+#endif
+
+#define UE_LOG_STRINGTOKEN_WARNING(Format, ...)  UE_LOG(LogIris, Warning, Format, ##__VA_ARGS__)
+
+namespace UE::Net
+{
+
+FNetToken FStringTokenStore::GetOrCreateToken(const TCHAR* Name, uint32 Length)
+{
+	FNetTokenStoreKey Key = GetOrCreatePersistentString(Name, Length);
+	if (Key.IsValid())
+	{
+		const uint32 StringIndex = Key.GetKeyValue();
+		if (StoredTokens[StringIndex].IsValid())
+		{
+			return StoredTokens[StringIndex];
+		}
+		else
+		{
+			const FNetToken NewToken = CreateToken(MakeNetTokenStoreKey(GetTypeId(), StringIndex), *TokenStore.GetLocalNetTokenStoreState());
+			StoredTokens[StringIndex] = NewToken;
+
+			UE_LOG_STRINGTOKEN(TEXT("FStringTokenStore::GetOrCreateToken - Created new token %u for %s"), NewToken.GetIndex(), Name);
+
+			return NewToken;
+		}
+	}
+
+	return FNetToken();
+}
+
+FNetToken FStringTokenStore::GetOrCreateToken(const FString& String)
+{
+	return GetOrCreateToken(ToCStr(String), String.Len());
+}
+
+FNetTokenStoreKey FStringTokenStore::GetOrCreatePersistentString(const TCHAR* Name, uint32 Length)
+{
+	// Hash name
+	const uint32 NameSize = Length * sizeof(TCHAR);
+	uint64 HashedName = CityHash64((const char*)Name, NameSize);
+	
+	// Lock if we have to.. make that a policy
+	//FScopeLock Lock(&CriticalSection);
+	if (const FNetTokenStoreKey* ExistingKey = HashToKey.Find(HashedName))
+	{
+		return *ExistingKey;
+	}
+	else if (StoredTokens.Num() < FNetToken::MaxNetTokenCount)
+	{
+		// Allocate memory and copy persistent string
+		TCHAR* PersistentString = (TCHAR*)Allocator.Alloc(NameSize + sizeof(TCHAR), alignof(TCHAR));
+		FCString::Strncpy((TCHAR*)PersistentString, Name, Length + 1U);
+
+		// Store bookkeeping data
+		const FNetTokenStoreKey NewKey = MakeNetTokenStoreKey(GetTypeId(), StoredStrings.Num());
+
+		HashToKey.Add(HashedName, NewKey);
+		StoredStrings.Add(PersistentString);
+		StoredTokens.Add(FNetToken());
+
+		return NewKey;
+	}
+
+	return FNetTokenStoreKey();
+}
+
+FStringTokenStore::FStringTokenStore(FNetTokenStore& InTokenStore)
+: TokenStore(InTokenStore)
+, Allocator()
+{
+	// Reserve 0
+	StoredStrings.Add(nullptr);
+	StoredTokens.Add(FNetToken());
+
+	// Register
+	TokenStore.RegisterDataStore(this);
+}
+
+const TCHAR* FStringTokenStore::ResolveToken(FNetToken Token) const
+{
+	return ResolveRemoteToken(Token, *TokenStore.GetLocalNetTokenStoreState());
+}
+
+const TCHAR* FStringTokenStore::ResolveRemoteToken(FNetToken Token, const FNetTokenStoreState& NetTokenStoreState) const
+{
+	const FNetTokenStoreKey StoreKey = GetTokenKey(Token, NetTokenStoreState);
+
+	if (StoreKey.IsValid() && StoreKey.GetKeyValue() < (uint32)StoredStrings.Num())
+	{
+		return StoredStrings[StoreKey.GetKeyValue()];
+	}
+
+	return nullptr;
+}
+
+void FStringTokenStore::WriteTokenData(FNetSerializationContext& Context, FNetTokenStoreKey TokenStoreKey) const
+{
+	// $TODO: $IRIS: Do not calculate the length of the string to write the data.
+	WriteString(Context.GetBitStreamWriter(), FStringView(StoredStrings[TokenStoreKey.GetKeyValue()]));
+}
+
+FNetTokenStoreKey FStringTokenStore::ReadTokenData(FNetSerializationContext& Context)
+{
+	FNetBitStreamReader* Reader = Context.GetBitStreamReader();
+
+	// Read the token data and add it to the string store without assigning LocalToken
+	FString Temp;
+	ReadString(Reader, Temp);
+
+	if (!Reader->IsOverflown())
+	{
+		return GetOrCreatePersistentString(*Temp, Temp.Len());
+	}
+	else
+	{
+		return FNetTokenStoreKey();
+	}
+}
+
+}
diff --git a/Engine/Source/Runtime/Experimental/Iris/Core/Private/Iris/ReplicationSystem/WorldLocations.cpp b/Engine/Source/Runtime/Experimental/Iris/Core/Private/Iris/ReplicationSystem/WorldLocations.cpp
new file mode 100644
--- /dev/null
+++ b/Engine/Source/Runtime/Experimental/Iris/Core/Private/Iris/ReplicationSystem/WorldLocations.cpp
@@ -0,0 +1,34 @@
+// Copyright Epic Games, Inc. All Rights Reserved.
+
+#include "Iris/ReplicationSystem/WorldLocations.h"
+#include "HAL/LowLevelMemTracker.h"
+
+namespace UE::Net
+{
+
+void FWorldLocations::Init(const FWorldLocationsInitParams& InitParams)
+{
+	ValidWorldLocations.Init(InitParams.MaxObjectCount);
+}
+
+void FWorldLocations::SetHasWorldLocation(uint32 ObjectIndex, bool bHasWorldLocation)
+{
+	ValidWorldLocations.SetBitValue(ObjectIndex, bHasWorldLocation);
+	if (bHasWorldLocation)
+	{
+		if (ObjectIndex >= uint32(WorldLocations.Num()))
+		{
+			LLM_SCOPE(ELLMTag::Iris);
+			WorldLocations.Add(ObjectIndex + 1U - WorldLocations.Num());
+		}
+		WorldLocations[ObjectIndex] = FVector::Zero();
+	}
+}
+
+void FWorldLocations::SetWorldLocation(uint32 ObjectIndex, const FVector& WorldLocation)
+{
+	checkSlow(ValidWorldLocations.GetBit(ObjectIndex));
+	WorldLocations[ObjectIndex] = WorldLocation;
+}
+
+}
diff --git a/Engine/Source/Runtime/Experimental/Iris/Core/Private/Iris/Serialization/InternalEnumNetSerializers.h b/Engine/Source/Runtime/Experimental/Iris/Core/Private/Iris/Serialization/InternalEnumNetSerializers.h
new file mode 100644
--- /dev/null
+++ b/Engine/Source/Runtime/Experimental/Iris/Core/Private/Iris/Serialization/InternalEnumNetSerializers.h
@@ -0,0 +1,20 @@
+// Copyright Epic Games, Inc. All Rights Reserved.
+
+#pragma once
+
+#include "Iris/Serialization/EnumNetSerializers.h"
+
+namespace UE::Net::Private
+{
+
+bool InitEnumNetSerializerConfig(FEnumInt8NetSerializerConfig& OutConfig, const UEnum* Enum);
+bool InitEnumNetSerializerConfig(FEnumInt16NetSerializerConfig& OutConfig, const UEnum* Enum);
+bool InitEnumNetSerializerConfig(FEnumInt32NetSerializerConfig& OutConfig, const UEnum* Enum);
+bool InitEnumNetSerializerConfig(FEnumInt64NetSerializerConfig& OutConfig, const UEnum* Enum);
+
+bool InitEnumNetSerializerConfig(FEnumUint8NetSerializerConfig& OutConfig, const UEnum* Enum);
+bool InitEnumNetSerializerConfig(FEnumUint16NetSerializerConfig& OutConfig, const UEnum* Enum);
+bool InitEnumNetSerializerConfig(FEnumUint32NetSerializerConfig& OutConfig, const UEnum* Enum);
+bool InitEnumNetSerializerConfig(FEnumUint64NetSerializerConfig& OutConfig, const UEnum* Enum);
+
+}
diff --git a/Engine/Source/Runtime/Experimental/Iris/Core/Private/Iris/Serialization/InternalNetSerializationContext.cpp b/Engine/Source/Runtime/Experimental/Iris/Core/Private/Iris/Serialization/InternalNetSerializationContext.cpp
new file mode 100644
--- /dev/null
+++ b/Engine/Source/Runtime/Experimental/Iris/Core/Private/Iris/Serialization/InternalNetSerializationContext.cpp
@@ -0,0 +1,35 @@
+// Copyright Epic Games, Inc. All Rights Reserved.
+
+#include "Iris/Serialization/InternalNetSerializationContext.h"
+#include "Iris/ReplicationSystem/ReplicationSystem.h"
+#include "Iris/ReplicationSystem/ReplicationSystemInternal.h"
+#include "HAL/MemoryBase.h"
+
+namespace UE::Net::Private
+{
+
+void* FInternalNetSerializationContext::Alloc(SIZE_T Size, SIZE_T Alignment)
+{
+	return GMalloc->Malloc(Size, Alignment);
+}
+
+void FInternalNetSerializationContext::Free(void* Ptr)
+{
+	return GMalloc->Free(Ptr);
+}
+
+void* FInternalNetSerializationContext::Realloc(void* PrevAddress, SIZE_T NewSize, uint32 Alignment)
+{
+	return GMalloc->Realloc(PrevAddress, NewSize, Alignment);
+}
+
+FInternalNetSerializationContext::FInternalNetSerializationContext(UReplicationSystem* ReplicationSystem, FNetTokenStoreState* RemoteTokenStoreState)
+: ReplicationSystem(ReplicationSystem)
+, ObjectReferenceCache(&ReplicationSystem->GetReplicationSystemInternal()->GetObjectReferenceCache())
+, bDowngradeAutonomousProxyRole(0)
+, bInlineObjectReferenceExports(0)
+{
+	ResolveContext.RemoteNetTokenStoreState = RemoteTokenStoreState;
+}
+
+}
diff --git a/Engine/Source/Runtime/Experimental/Iris/Core/Private/Iris/Serialization/IrisObjectReferencePackageMap.cpp b/Engine/Source/Runtime/Experimental/Iris/Core/Private/Iris/Serialization/IrisObjectReferencePackageMap.cpp
new file mode 100644
--- /dev/null
+++ b/Engine/Source/Runtime/Experimental/Iris/Core/Private/Iris/Serialization/IrisObjectReferencePackageMap.cpp
@@ -0,0 +1,59 @@
+// Copyright Epic Games, Inc. All Rights Reserved.
+
+#include "Iris/Serialization/IrisObjectReferencePackageMap.h"
+
+bool UIrisObjectReferencePackageMap::SerializeObject(FArchive& Ar, UClass* InClass, UObject*& Obj, FNetworkGUID* OutNetGUID)
+{
+	if (!References)
+	{
+		return false;
+	}
+
+	if (Ar.IsSaving())
+	{
+		int32 Index;
+		if (!References->Find(Obj, Index))
+		{
+			Index = References->Add(Obj);
+		}
+		uint8 IndexByte = Index;
+		if (ensureAlwaysMsgf(IndexByte < References->Num(), TEXT("UIrisObjectReferencePackageMap::SerializeObject, failed to serialize object reference. A Maximum of 256 references are currently supported by this PackageMap")))
+		{
+			Ar << IndexByte;
+		}
+		else
+		{
+			return false;
+		}
+	}
+	else
+	{
+		uint8 IndexByte = 255U;
+		Ar << IndexByte;
+		if (ensureAlwaysMsgf(IndexByte < References->Num(), TEXT("UIrisObjectReferencePackageMap::SerializeObject, failed to read object reference index %u is out of bounds"), IndexByte))
+		{
+			Obj = (*References)[IndexByte];
+		}
+		else
+		{
+			return false;
+		}
+	}
+
+	return true;
+}
+
+void UIrisObjectReferencePackageMap::InitForRead(const FObjectReferenceArray* InReferences)
+{ 
+	References = const_cast<FObjectReferenceArray*>(InReferences);
+}
+
+void UIrisObjectReferencePackageMap::InitForWrite(FObjectReferenceArray* InReferences)
+{
+	References = InReferences; 
+	if (References)
+	{
+		References->Reset();
+	}
+}
+
diff --git a/Engine/Source/Runtime/Experimental/Iris/Core/Private/Iris/Serialization/PolymorphicNetSerializer.cpp b/Engine/Source/Runtime/Experimental/Iris/Core/Private/Iris/Serialization/PolymorphicNetSerializer.cpp
new file mode 100644
--- /dev/null
+++ b/Engine/Source/Runtime/Experimental/Iris/Core/Private/Iris/Serialization/PolymorphicNetSerializer.cpp
@@ -0,0 +1,84 @@
+// Copyright Epic Games, Inc. All Rights Reserved.
+
+#include "Iris/Serialization/PolymorphicNetSerializer.h"
+#include "Iris/Serialization/PolymorphicNetSerializerImpl.h"
+#include "Iris/Core/IrisProfiler.h"
+#include "Iris/Serialization/InternalNetSerializationContext.h"
+#include "Iris/ReplicationState/ReplicationStateDescriptorBuilder.h"
+#include "Iris/ReplicationSystem/ReplicationSystem.h"
+#include "Iris/ReplicationSystem/ReplicationSystemInternal.h"
+#include "Iris/ReplicationSystem/ReplicationOperationsInternal.h"
+#include "UObject/UObjectIterator.h"
+
+namespace UE::Net
+{
+
+const FName NetError_PolymorphicStructNetSerializer_InvalidStructType(TEXT("Invalid struct type"));
+
+void FPolymorphicNetSerializerScriptStructCache::InitForType(const UScriptStruct* InScriptStruct)
+{
+	IRIS_PROFILER_SCOPE(FPolymorphicNetSerializerScriptStructCache_InitForType);
+
+	RegisteredTypes.Reset();
+
+	// Find all script structs of this type and add them to the list and build descriptor
+	// (not sure of a better way to do this but it should only happen once at startup)
+	for (TObjectIterator<UScriptStruct> It; It; ++It)
+	{
+		if (It->IsChildOf(InScriptStruct))
+		{
+			FTypeInfo Entry;
+			Entry.ScriptStruct = *It;
+
+			// Get or create descriptor
+			FReplicationStateDescriptorBuilder::FParameters Params;
+			Entry.Descriptor = FReplicationStateDescriptorBuilder::CreateDescriptorForStruct(*It, Params);
+			
+			if (Entry.Descriptor.IsValid())
+			{
+				RegisteredTypes.Add(Entry);
+
+				CommonTraits |= (Entry.Descriptor->Traits & EReplicationStateTraits::HasObjectReference);
+			}
+			else
+			{
+				UE_LOG(LogIris, Error, TEXT("FPolymorphicNetSerializerScriptStructCache::InitForType Failed to create descriptor for type %s when building cache for base %s"), ToCStr((*It)->GetName()), ToCStr(InScriptStruct->GetName()));
+			}
+		}
+	}
+	
+	RegisteredTypes.Sort([](const FTypeInfo& A, const FTypeInfo& B) { return A.ScriptStruct->GetName().ToLower() > B.ScriptStruct->GetName().ToLower(); });
+
+	if ((uint32)RegisteredTypes.Num() > MaxRegisteredTypeCount)
+	{
+		UE_LOG(LogIris, Error, TEXT("FPolymorphicNetSerializerScriptStructCache::InitForType Too many types (%u of %u) of base %s"), RegisteredTypes.Num(), MaxRegisteredTypeCount,  ToCStr(InScriptStruct->GetName()));
+		check((uint32)RegisteredTypes.Num() <= MaxRegisteredTypeCount);
+	}
+}
+
+}
+
+namespace UE::Net::Private
+{
+
+void* FPolymorphicStructNetSerializerInternal::Alloc(FNetSerializationContext& Context, SIZE_T Size, SIZE_T Alignment)
+{
+	return Context.GetInternalContext()->Alloc(Size, Alignment);
+}
+
+void FPolymorphicStructNetSerializerInternal::Free(FNetSerializationContext& Context, void* Ptr)
+{
+	return Context.GetInternalContext()->Free(Ptr);
+}
+
+void FPolymorphicStructNetSerializerInternal::CollectReferences(FNetSerializationContext& Context, FNetReferenceCollector& Collector, const FNetSerializerChangeMaskParam& OuterChangeMaskInfo, const uint8* RESTRICT SrcInternalBuffer,  const FReplicationStateDescriptor* Descriptor)
+{
+	return FReplicationStateOperationsInternal::CollectReferences(Context, Collector, OuterChangeMaskInfo, SrcInternalBuffer, Descriptor);
+}
+
+void FPolymorphicStructNetSerializerInternal::CloneQuantizedState(FNetSerializationContext& Context, uint8* RESTRICT DstInternalBuffer, const uint8* RESTRICT SrcInternalBuffer, const FReplicationStateDescriptor* Descriptor)
+{
+	return FReplicationStateOperationsInternal::CloneQuantizedState(Context, DstInternalBuffer, SrcInternalBuffer, Descriptor);
+}
+
+}
diff --git a/Engine/Source/Runtime/Experimental/Iris/Core/Private/Iris/Serialization/RotatorNetSerializers.cpp b/Engine/Source/Runtime/Experimental/Iris/Core/Private/Iris/Serialization/RotatorNetSerializers.cpp
new file mode 100644
--- /dev/null
+++ b/Engine/Source/Runtime/Experimental/Iris/Core/Private/Iris/Serialization/RotatorNetSerializers.cpp
@@ -0,0 +1,426 @@
+// Copyright Epic Games, Inc. All Rights Reserved.
+
+#include "Iris/Serialization/RotatorNetSerializers.h"
+#include "Iris/Serialization/NetBitStreamReader.h"
+#include "Iris/Serialization/NetBitStreamWriter.h"
+#include "Math/Rotator.h"
+
+namespace UE::Net
+{
+
+struct FRotatorAsShortNetSerializer
+{
+	static constexpr uint32 Version = 0;
+
+	struct FQuantizedType
+	{
+		// Using three bits to indicate whether the components are zero or not.
+		uint16 XYZIsNotZero;
+		uint16 X;
+		uint16 Y;
+		uint16 Z;
+	};
+
+	typedef FRotator SourceType;
+	typedef FQuantizedType QuantizedType;
+	typedef FRotatorAsShortNetSerializerConfig ConfigType;
+
+	inline static const ConfigType DefaultConfig;
+
+	static void Serialize(FNetSerializationContext&, const FNetSerializeArgs& Args);
+	static void Deserialize(FNetSerializationContext&, const FNetDeserializeArgs& Args);
+
+	static void SerializeDelta(FNetSerializationContext&, const FNetSerializeDeltaArgs& Args);
+	static void DeserializeDelta(FNetSerializationContext&, const FNetDeserializeDeltaArgs& Args);
+
+	static void Quantize(FNetSerializationContext&, const FNetQuantizeArgs& Args);
+	static void Dequantize(FNetSerializationContext&, const FNetDequantizeArgs& Args);
+
+	static bool IsEqual(FNetSerializationContext&, const FNetIsEqualArgs& Args);
+	static bool Validate(FNetSerializationContext&, const FNetValidateArgs& Args);
+
+protected:
+	using ScalarType = decltype(SourceType::Pitch);
+
+	enum Constants : uint16
+	{
+		XDiffersMask = 1U,
+		YDiffersMask = 2U,
+		ZDiffersMask = 4U,
+
+		XYZDiffersMask = 7U,
+	};
+
+	static constexpr ScalarType Scale = ScalarType(65536)/ScalarType(360);
+	static constexpr ScalarType Bias = ScalarType(0.5f);
+	static constexpr ScalarType InvScale = ScalarType(360)/ScalarType(65536);
+};
+UE_NET_IMPLEMENT_SERIALIZER(FRotatorAsShortNetSerializer);
+
+struct FRotatorNetSerializer : public FRotatorAsShortNetSerializer
+{
+	static constexpr uint32 Version = 0;
+
+	typedef FRotatorNetSerializerConfig ConfigType;
+
+	inline static const ConfigType DefaultConfig;
+};
+UE_NET_IMPLEMENT_SERIALIZER(FRotatorNetSerializer);
+
+struct FRotatorAsByteNetSerializer
+{
+	static constexpr uint32 Version = 0;
+
+	struct FQuantizedType
+	{
+		// Using three bits to indicate whether the components are zero or not.
+		uint8 XYZIsNotZero;
+		uint8 X;
+		uint8 Y;
+		uint8 Z;
+	};
+
+	typedef FRotator SourceType;
+	typedef FQuantizedType QuantizedType;
+	typedef FRotatorAsByteNetSerializerConfig ConfigType;
+
+	inline static const ConfigType DefaultConfig;
+
+	static void Serialize(FNetSerializationContext&, const FNetSerializeArgs& Args);
+	static void Deserialize(FNetSerializationContext&, const FNetDeserializeArgs& Args);
+
+	static void SerializeDelta(FNetSerializationContext&, const FNetSerializeDeltaArgs& Args);
+	static void DeserializeDelta(FNetSerializationContext&, const FNetDeserializeDeltaArgs& Args);
+
+	static void Quantize(FNetSerializationContext&, const FNetQuantizeArgs& Args);
+	static void Dequantize(FNetSerializationContext&, const FNetDequantizeArgs& Args);
+
+	static bool IsEqual(FNetSerializationContext&, const FNetIsEqualArgs& Args);
+	static bool Validate(FNetSerializationContext&, const FNetValidateArgs& Args);
+
+protected:
+	using ScalarType = decltype(SourceType::Pitch);
+
+	enum Constants : uint16
+	{
+		XDiffersMask = 1U,
+		YDiffersMask = 2U,
+		ZDiffersMask = 4U,
+
+		XYZDiffersMask = 7U,
+	};
+
+	static constexpr ScalarType Scale = ScalarType(256)/ScalarType(360);
+	static constexpr ScalarType Bias = ScalarType(0.5f);
+	static constexpr ScalarType InvScale = ScalarType(360)/ScalarType(256);
+};
+UE_NET_IMPLEMENT_SERIALIZER(FRotatorAsByteNetSerializer);
+
+// FRotatorAsShortNetSerializer implementation
+void FRotatorAsShortNetSerializer::Serialize(FNetSerializationContext& Context, const FNetSerializeArgs& Args)
+{
+	const QuantizedType& Value = *reinterpret_cast<const QuantizedType*>(Args.Source);
+
+	FNetBitStreamWriter* Writer = Context.GetBitStreamWriter();
+	Writer->WriteBits(Value.XYZIsNotZero, 3U);
+	if (Value.XYZIsNotZero & XDiffersMask)
+	{
+		Writer->WriteBits(Value.X, 16U);
+	}
+	if (Value.XYZIsNotZero & YDiffersMask)
+	{
+		Writer->WriteBits(Value.Y, 16U);
+	}
+	if (Value.XYZIsNotZero & ZDiffersMask)
+	{
+		Writer->WriteBits(Value.Z, 16U);
+	}
+}
+
+void FRotatorAsShortNetSerializer::Deserialize(FNetSerializationContext& Context, const FNetDeserializeArgs& Args)
+{
+	FNetBitStreamReader* Reader = Context.GetBitStreamReader();
+
+	QuantizedType Value;
+	const uint32 XYZIsNotZero = Reader->ReadBits(3U);
+	Value.XYZIsNotZero = XYZIsNotZero;
+	Value.X = XYZIsNotZero & XDiffersMask ? Reader->ReadBits(16U) : 0U;
+	Value.Y = XYZIsNotZero & YDiffersMask ? Reader->ReadBits(16U) : 0U;
+	Value.Z = XYZIsNotZero & ZDiffersMask ? Reader->ReadBits(16U) : 0U;
+
+	QuantizedType& Target = *reinterpret_cast<QuantizedType*>(Args.Target);
+	Target = Value;
+}
+
+void FRotatorAsShortNetSerializer::SerializeDelta(FNetSerializationContext& Context, const FNetSerializeDeltaArgs& Args)
+{
+	// Per component equality. With the current scaling a 10 degree change would require 12 bits to replicate
+	// due to 11 bits for the value and 1 for the sign. Add a small lookup table index for that and you're
+	// up to 14 bits instead of 1+16 bits.
+
+	const QuantizedType& Value = *reinterpret_cast<QuantizedType*>(Args.Source);
+	const QuantizedType& PrevValue = *reinterpret_cast<QuantizedType*>(Args.Prev);
+
+	const uint16 DX = Value.X - PrevValue.X;
+	const uint16 DY = Value.Y - PrevValue.Y;
+	const uint16 DZ = Value.Z - PrevValue.Z;
+
+	uint32 XYZDiffers = 0;
+	XYZDiffers |= (DX != 0) ? uint32(XDiffersMask) : uint32(0);
+	XYZDiffers |= (DY != 0) ? uint32(YDiffersMask) : uint32(0);
+	XYZDiffers |= (DZ != 0) ? uint32(ZDiffersMask) : uint32(0);
+
+	FNetBitStreamWriter* Writer = Context.GetBitStreamWriter();
+	Writer->WriteBits(XYZDiffers, 3U);
+	if (XYZDiffers & XDiffersMask)
+	{
+		Writer->WriteBits(DX, 16U);
+	}
+	if (XYZDiffers & YDiffersMask)
+	{
+		Writer->WriteBits(DY, 16U);
+	}
+	if (XYZDiffers & ZDiffersMask)
+	{
+		Writer->WriteBits(DZ, 16U);
+	}
+}
+
+void FRotatorAsShortNetSerializer::DeserializeDelta(FNetSerializationContext& Context, const FNetDeserializeDeltaArgs& Args)
+{
+	FNetBitStreamReader* Reader = Context.GetBitStreamReader();
+
+	const QuantizedType& PrevValue = *reinterpret_cast<QuantizedType*>(Args.Prev);
+	QuantizedType TempValue = PrevValue;
+
+	const uint32 XYZDiffers = Reader->ReadBits(3U);
+	if (XYZDiffers & XDiffersMask)
+	{
+		const uint16 DX = Reader->ReadBits(16U);
+		TempValue.X += DX;
+	}
+	if (XYZDiffers & YDiffersMask)
+	{
+		const uint16 DY = Reader->ReadBits(16U);
+		TempValue.Y += DY;
+	}
+	if (XYZDiffers & ZDiffersMask)
+	{
+		const uint16 DZ = Reader->ReadBits(16U);
+		TempValue.Z += DZ;
+	}
+
+	QuantizedType& Target = *reinterpret_cast<QuantizedType*>(Args.Target);
+	Target = TempValue;
+}
+
+void FRotatorAsShortNetSerializer::Quantize(FNetSerializationContext&, const FNetQuantizeArgs& Args)
+{
+	const SourceType& Source = *reinterpret_cast<const SourceType*>(Args.Source);
+	QuantizedType& Target = *reinterpret_cast<QuantizedType*>(Args.Target);
+
+	QuantizedType TempValue = {};
+	TempValue.X = uint16(uint32(int32(Source.Pitch*Scale + Bias)));
+	TempValue.Y = uint16(uint32(int32(Source.Yaw*Scale + Bias)));
+	TempValue.Z = uint16(uint32(int32(Source.Roll*Scale + Bias)));
+	TempValue.XYZIsNotZero |= (TempValue.X != 0) ? XDiffersMask : uint16(0);
+	TempValue.XYZIsNotZero |= (TempValue.Y != 0) ? YDiffersMask : uint16(0);
+	TempValue.XYZIsNotZero |= (TempValue.Z != 0) ? ZDiffersMask : uint16(0);
+
+	Target = TempValue;
+}
+
+void FRotatorAsShortNetSerializer::Dequantize(FNetSerializationContext&, const FNetDequantizeArgs& Args)
+{
+	const QuantizedType& Source = *reinterpret_cast<const QuantizedType*>(Args.Source);
+	SourceType& Target = *reinterpret_cast<SourceType*>(Args.Target);
+
+	SourceType TempValue;
+	TempValue.Pitch = Source.X*InvScale;
+	TempValue.Yaw = Source.Y*InvScale;
+	TempValue.Roll = Source.Z*InvScale;
+
+	Target = TempValue;
+}
+
+bool FRotatorAsShortNetSerializer::IsEqual(FNetSerializationContext& Context, const FNetIsEqualArgs& Args)
+{
+	if (Args.bStateIsQuantized)
+	{
+		const QuantizedType& SourceValue0 = *reinterpret_cast<const QuantizedType*>(Args.Source0);
+		const QuantizedType& SourceValue1 = *reinterpret_cast<const QuantizedType*>(Args.Source1);
+
+		const bool bXIsEqual = (SourceValue0.X == SourceValue1.X);
+		const bool bYIsEqual = (SourceValue0.Y == SourceValue1.Y);
+		const bool bZIsEqual = (SourceValue0.Z == SourceValue1.Z);
+		// There's no need to check the flags as they are derived from the X,Y and Z values.
+		return bXIsEqual & bYIsEqual & bZIsEqual;
+	}
+	else
+	{
+		const SourceType& SourceValue0 = *reinterpret_cast<const SourceType*>(Args.Source0);
+		const SourceType& SourceValue1 = *reinterpret_cast<const SourceType*>(Args.Source1);
+
+		const bool bPitchIsEqual = (SourceValue0.Pitch == SourceValue1.Pitch);
+		const bool bYawIsEqual = (SourceValue0.Yaw == SourceValue1.Yaw);
+		const bool bRollIsEqual = (SourceValue0.Roll == SourceValue1.Roll);
+		return bPitchIsEqual & bYawIsEqual & bRollIsEqual;
+	}
+}
+
+bool FRotatorAsShortNetSerializer::Validate(FNetSerializationContext& Context, const FNetValidateArgs& Args)
+{
+	const SourceType& Value = *reinterpret_cast<SourceType*>(Args.Source);
+	// Make sure values are valid degree angles. This should catch NaNs as well.
+	return (Value.Pitch >= 0.0f && Value.Pitch < 360.0f) && (Value.Yaw >= 0.0f && Value.Yaw < 360.0f) & (Value.Roll >= 0.0f && Value.Roll < 360.0f); 
+}
+
+// FRotatorAsByteNetSerializer implementation
+void FRotatorAsByteNetSerializer::Serialize(FNetSerializationContext& Context, const FNetSerializeArgs& Args)
+{
+	const QuantizedType& Value = *reinterpret_cast<const QuantizedType*>(Args.Source);
+
+	FNetBitStreamWriter* Writer = Context.GetBitStreamWriter();
+	Writer->WriteBits(Value.XYZIsNotZero, 3U);
+	if (Value.XYZIsNotZero & XDiffersMask)
+	{
+		Writer->WriteBits(Value.X, 8U);
+	}
+	if (Value.XYZIsNotZero & YDiffersMask)
+	{
+		Writer->WriteBits(Value.Y, 8U);
+	}
+	if (Value.XYZIsNotZero & ZDiffersMask)
+	{
+		Writer->WriteBits(Value.Z, 8U);
+	}
+}
+
+void FRotatorAsByteNetSerializer::Deserialize(FNetSerializationContext& Context, const FNetDeserializeArgs& Args)
+{
+	FNetBitStreamReader* Reader = Context.GetBitStreamReader();
+
+	QuantizedType Value;
+	const uint32 XYZIsNotZero = Reader->ReadBits(3U);
+	Value.XYZIsNotZero = XYZIsNotZero;
+	Value.X = XYZIsNotZero & XDiffersMask ? Reader->ReadBits(8U) : 0U;
+	Value.Y = XYZIsNotZero & YDiffersMask ? Reader->ReadBits(8U) : 0U;
+	Value.Z = XYZIsNotZero & ZDiffersMask ? Reader->ReadBits(8U) : 0U;
+
+	QuantizedType& Target = *reinterpret_cast<QuantizedType*>(Args.Target);
+	Target = Value;
+}
+
+void FRotatorAsByteNetSerializer::SerializeDelta(FNetSerializationContext& Context, const FNetSerializeDeltaArgs& Args)
+{
+	const QuantizedType& Value = *reinterpret_cast<QuantizedType*>(Args.Source);
+	const QuantizedType& PrevValue = *reinterpret_cast<QuantizedType*>(Args.Prev);
+
+	uint32 XYZDiffers = 0;
+	XYZDiffers |= (Value.X != PrevValue.X) ? uint32(XDiffersMask) : uint32(0);
+	XYZDiffers |= (Value.Y != PrevValue.Y) ? uint32(YDiffersMask) : uint32(0);
+	XYZDiffers |= (Value.Z != PrevValue.Z) ? uint32(ZDiffersMask) : uint32(0);
+
+	FNetBitStreamWriter* Writer = Context.GetBitStreamWriter();
+	Writer->WriteBits(XYZDiffers, 3U);
+	if (XYZDiffers & XDiffersMask)
+	{
+		Writer->WriteBits(Value.X, 8U);
+	}
+	if (XYZDiffers & YDiffersMask)
+	{
+		Writer->WriteBits(Value.Y, 8U);
+	}
+	if (XYZDiffers & ZDiffersMask)
+	{
+		Writer->WriteBits(Value.Z, 8U);
+	}
+}
+
+void FRotatorAsByteNetSerializer::DeserializeDelta(FNetSerializationContext& Context, const FNetDeserializeDeltaArgs& Args)
+{
+	FNetBitStreamReader* Reader = Context.GetBitStreamReader();
+
+	const QuantizedType& PrevValue = *reinterpret_cast<QuantizedType*>(Args.Prev);
+	QuantizedType TempValue = PrevValue;
+
+	const uint32 XYZDiffers = Reader->ReadBits(3U);
+	if (XYZDiffers & XDiffersMask)
+	{
+		TempValue.X = Reader->ReadBits(8U);
+	}
+	if (XYZDiffers & YDiffersMask)
+	{
+		TempValue.Y = Reader->ReadBits(8U);
+	}
+	if (XYZDiffers & ZDiffersMask)
+	{
+		TempValue.Z = Reader->ReadBits(8U);
+	}
+
+	QuantizedType& Target = *reinterpret_cast<QuantizedType*>(Args.Target);
+	Target = TempValue;
+}
+
+void FRotatorAsByteNetSerializer::Quantize(FNetSerializationContext&, const FNetQuantizeArgs& Args)
+{
+	const SourceType& Source = *reinterpret_cast<const SourceType*>(Args.Source);
+	QuantizedType& Target = *reinterpret_cast<QuantizedType*>(Args.Target);
+
+	QuantizedType TempValue = {};
+	TempValue.X = uint8(uint32(int32(Source.Pitch*Scale + Bias)));
+	TempValue.Y = uint8(uint32(int32(Source.Yaw*Scale + Bias)));
+	TempValue.Z = uint8(uint32(int32(Source.Roll*Scale + Bias)));
+	TempValue.XYZIsNotZero |= (TempValue.X != 0) ? XDiffersMask : uint8(0);
+	TempValue.XYZIsNotZero |= (TempValue.Y != 0) ? YDiffersMask : uint8(0);
+	TempValue.XYZIsNotZero |= (TempValue.Z != 0) ? ZDiffersMask : uint8(0);
+
+	Target = TempValue;
+}
+
+void FRotatorAsByteNetSerializer::Dequantize(FNetSerializationContext&, const FNetDequantizeArgs& Args)
+{
+	const QuantizedType& Source = *reinterpret_cast<const QuantizedType*>(Args.Source);
+	SourceType& Target = *reinterpret_cast<SourceType*>(Args.Target);
+
+	SourceType TempValue;
+	TempValue.Pitch = Source.X*InvScale;
+	TempValue.Yaw = Source.Y*InvScale;
+	TempValue.Roll = Source.Z*InvScale;
+
+	Target = TempValue;
+}
+
+bool FRotatorAsByteNetSerializer::IsEqual(FNetSerializationContext& Context, const FNetIsEqualArgs& Args)
+{
+	if (Args.bStateIsQuantized)
+	{
+		const QuantizedType& SourceValue0 = *reinterpret_cast<const QuantizedType*>(Args.Source0);
+		const QuantizedType& SourceValue1 = *reinterpret_cast<const QuantizedType*>(Args.Source1);
+
+		const bool bXIsEqual = (SourceValue0.X == SourceValue1.X);
+		const bool bYIsEqual = (SourceValue0.Y == SourceValue1.Y);
+		const bool bZIsEqual = (SourceValue0.Z == SourceValue1.Z);
+		// There's no need to check the flags as they are derived from the X,Y and Z values.
+		return bXIsEqual & bYIsEqual & bZIsEqual;
+	}
+	else
+	{
+		const SourceType& SourceValue0 = *reinterpret_cast<const SourceType*>(Args.Source0);
+		const SourceType& SourceValue1 = *reinterpret_cast<const SourceType*>(Args.Source1);
+
+		const bool bPitchIsEqual = (SourceValue0.Pitch == SourceValue1.Pitch);
+		const bool bYawIsEqual = (SourceValue0.Yaw == SourceValue1.Yaw);
+		const bool bRollIsEqual = (SourceValue0.Roll == SourceValue1.Roll);
+		return bPitchIsEqual & bYawIsEqual & bRollIsEqual;
+	}
+}
+
+bool FRotatorAsByteNetSerializer::Validate(FNetSerializationContext& Context, const FNetValidateArgs& Args)
+{
+	const SourceType& Value = *reinterpret_cast<SourceType*>(Args.Source);
+	// Make sure values are valid degree angles. This should catch NaNs as well.
+	return (Value.Pitch >= 0.0f && Value.Pitch < 360.0f) && (Value.Yaw >= 0.0f && Value.Yaw < 360.0f) & (Value.Roll >= 0.0f && Value.Roll < 360.0f); 
+}
+
+}
diff --git a/Engine/Source/Runtime/Experimental/Iris/Core/Public/Iris/Core/IrisDebugging.h b/Engine/Source/Runtime/Experimental/Iris/Core/Public/Iris/Core/IrisDebugging.h
new file mode 100644
--- /dev/null
+++ b/Engine/Source/Runtime/Experimental/Iris/Core/Public/Iris/Core/IrisDebugging.h
@@ -0,0 +1,94 @@
+// Copyright Epic Games, Inc. All Rights Reserved.
+
+#pragma once
+
+#include "CoreMinimal.h"
+#include "Iris/ReplicationSystem/NetHandle.h"
+
+class UReplicationSystem;
+namespace UE::Net
+{
+struct FReplicationProtocol;
+struct FReplicationInstanceProtocol;
+}
+
+#if !UE_BUILD_SHIPPING
+
+namespace UE::Net::IrisDebugHelper
+{ 
+
+/** Dummy methods binding pointers to external methods to avoid them from being stripped by compiler */
+IRISCORE_API uint64 Init();
+
+/** Trigger a breakpoint and return true if the object contains the current debug name */
+bool BreakOnObjectName(UObject* Object);
+
+/** Trigger a breakpoint and return true if the NetHandle is the current debug nethandle */
+bool BreakOnNetHandle(FNetHandle NetHandle);
+
+/** Trigger a breakpoint and return true if the name contains the debug RPC string */
+bool BreakOnRPCName(FName RPCName);
+
+/** Output state data to StringBuilder for the specified Handle */
+void NetObjectStateToString(FStringBuilderBase& StringBuilder, FNetHandle NetHandle);
+
+/** Find all handles references registered for a protocol and output to StringBuilder */
+void NetObjectProtocolReferencesToString(FStringBuilderBase& StringBuilder, uint64 ProtocolId, uint32 ReplicationSystemId);
+
+// Helper functions for debugging state data exposed as extern "C" that are callable from immediate and watch window in the debugger
+
+/** Get the ReplicationSystem with the Id */
+extern "C" IRISCORE_API UReplicationSystem* GetReplicationSystemForDebug(uint32 Id);
+
+/** DebugOutputObject state for Handle specified by NetHandleId and ReplicationSystemId and */
+extern "C" IRISCORE_API void DebugOutputNetObjectState(uint32 NetHandleId, uint32 ReplicationSystemId);
+
+/**
+	Variant of NetObjectStateToString that can be used from breakpoints and in watch window to print the current state of a NetHandle
+	NOTE: Use only for debugging as this variant uses a static buffer which is not thread safe	
+*/
+extern "C" IRISCORE_API const TCHAR* DebugNetObjectStateToString(uint32 NetHandleId, uint32 ReplicationSystemId);
+
+/** Find all handles references registered for a protocol and output to DebugOutput in debugger */
+extern "C" IRISCORE_API void DebugOutputNetObjectProtocolReferences(uint64 ProtocolId, uint32 ReplicationSystemId);
+
+/** Get info about replicated object */
+struct FNetReplicatedObjectDebugInfo
+{
+	const FNetHandle* Handle;
+	uint32 InternalNetHandleIndex;
+	const UReplicationSystem* ReplicationSystem;
+	const FReplicationProtocol* Protocol;
+	const FReplicationInstanceProtocol* InstanceProtocol;
+};
+
+/** Look up replicated NetHandle from Instance pointer and return debug information, this variant searches all active replication systems */
+extern "C" IRISCORE_API FNetReplicatedObjectDebugInfo DebugNetObject(UObject* Instance);
+
+/** Look up replicated NetHandle from Instance pointer and return debug information, this variant only searches the replicatioSystem with the specified id */
+extern "C" IRISCORE_API FNetReplicatedObjectDebugInfo DebugNetObjectById(UObject* Instance, uint32 ReplicationSystemId);
+
+/** Look up replicated handle and return debug information */
+extern "C" IRISCORE_API FNetReplicatedObjectDebugInfo DebugNetHandle(FNetHandle Handle);
+
+/** Look up replicated handle specified by handle id and replicationsystem id and return debug information */
+extern "C" IRISCORE_API FNetReplicatedObjectDebugInfo DebugNetHandleById(uint32 NetHandleId, uint32 ReplicationSystemId);
+
+/** Variant of DebugOutputNetObjectProtocolReferences that can be used from breakpoints or in watch window to find all handles references registered for a protocol and output to DebugOutput in debugger
+ * NOTE: Use only for debugging as this variant uses a static buffer which is not thread safe	
+*/
+extern "C" IRISCORE_API const TCHAR* DebugNetObjectProtocolReferencesToString(uint64 ProtocolId, uint32 ReplicationSystemId);
+
+/** Set the Object Name to break on */
+extern "C" IRISCORE_API void SetIrisDebugObjectName(const ANSICHAR* NameBuffer);
+
+/** Set the NetHandle to break on */
+extern "C" IRISCORE_API void SetIrisDebugNetHandle(uint32 NetHandleId);
+
+/** Set the RPC Name to break on */
+extern "C" IRISCORE_API void SetIrisDebugRPCName(const ANSICHAR* NameBuffer);
+
+}; // end of namespaces
+
+#endif
+
diff --git a/Engine/Source/Runtime/Experimental/Iris/Core/Public/Iris/Core/IrisProfiler.h b/Engine/Source/Runtime/Experimental/Iris/Core/Public/Iris/Core/IrisProfiler.h
new file mode 100644
--- /dev/null
+++ b/Engine/Source/Runtime/Experimental/Iris/Core/Public/Iris/Core/IrisProfiler.h
@@ -0,0 +1,31 @@
+// Copyright Epic Games, Inc. All Rights Reserved.
+
+#pragma once
+
+#ifndef IRIS_PROFILER_ENABLE
+#	if (UE_BUILD_SHIPPING)
+#		define IRIS_PROFILER_ENABLE 0
+#	else
+#		define IRIS_PROFILER_ENABLE 1
+#	endif
+#endif
+
+//#define IRIS_USE_SUPERLUMINAL
+
+#if IRIS_PROFILER_ENABLE
+#	ifdef IRIS_USE_SUPERLUMINAL
+#		include "c:/Program Files/Superluminal/Performance/API/include/Superluminal/PerformanceAPI.h"
+#		include "HAL/PreprocessorHelpers.h"
+#		pragma comment (lib, "c:/Program Files/Superluminal/Performance/API/lib/x64/PerformanceAPI_MD.lib")
+#		define IRIS_PROFILER_SCOPE(x) PERFORMANCEAPI_INSTRUMENT(PREPROCESSOR_TO_STRING(x))
+#		define IRIS_PROFILER_SCOPE_TEXT(x) PERFORMANCEAPI_INSTRUMENT_DATA(PREPROCESSOR_JOIN(IrisProfilerScope, __LINE__), x)
+#	else
+#		include "ProfilingDebugging/CpuProfilerTrace.h"
+#		define IRIS_PROFILER_SCOPE(x) TRACE_CPUPROFILER_EVENT_SCOPE(x)
+#		define IRIS_PROFILER_SCOPE_TEXT(X) TRACE_CPUPROFILER_EVENT_SCOPE_TEXT(X)
+#	endif
+#else
+#	define PERFORMANCEAPI_ENABLED 0
+#	define IRIS_PROFILER_SCOPE(x)
+#	define IRIS_PROFILER_SCOPE_TEXT(X)
+#endif
diff --git a/Engine/Source/Runtime/Experimental/Iris/Core/Public/Iris/IrisConfig.h b/Engine/Source/Runtime/Experimental/Iris/Core/Public/Iris/IrisConfig.h
new file mode 100644
--- /dev/null
+++ b/Engine/Source/Runtime/Experimental/Iris/Core/Public/Iris/IrisConfig.h
@@ -0,0 +1,33 @@
+// Copyright Epic Games, Inc. All Rights Reserved.
+
+#pragma once
+
+#include "CoreTypes.h"
+
+namespace UE::Net
+{
+
+IRISCORE_API bool ShouldUseIrisReplication();
+IRISCORE_API void SetUseIrisReplication(bool EnableIrisReplication);
+
+}
+
+/** Allow multiple replication systems inside the same process. Ex: PIE support */
+#ifndef UE_NET_ALLOW_MULTIPLE_REPLICATION_SYSTEMS
+#if (WITH_EDITOR || WITH_DEV_AUTOMATION_TESTS || WITH_AUTOMATION_WORKER)
+#	define UE_NET_ALLOW_MULTIPLE_REPLICATION_SYSTEMS 1
+#else
+#	define UE_NET_ALLOW_MULTIPLE_REPLICATION_SYSTEMS 0
+#endif
+#endif
+
+/* NetBitStreamReader/Writer validation support */
+#ifndef UE_NETBITSTREAMWRITER_VALIDATE
+#if !(UE_BUILD_SHIPPING || UE_BUILD_TEST)
+#define UE_NETBITSTREAMWRITER_VALIDATE 1
+#else
+#define UE_NETBITSTREAMWRITER_VALIDATE 0
+#endif
+#endif
+
+
diff --git a/Engine/Source/Runtime/Experimental/Iris/Core/Public/Iris/ReplicationState/Private/IrisFastArraySerializerInternal.h b/Engine/Source/Runtime/Experimental/Iris/Core/Public/Iris/ReplicationState/Private/IrisFastArraySerializerInternal.h
new file mode 100644
--- /dev/null
+++ b/Engine/Source/Runtime/Experimental/Iris/Core/Public/Iris/ReplicationState/Private/IrisFastArraySerializerInternal.h
@@ -0,0 +1,276 @@
+// Copyright Epic Games, Inc. All Rights Reserved.
+#pragma once
+
+#include "Iris/ReplicationState/IrisFastArraySerializer.h"
+#include "Net/Core/NetBitArray.h"
+#include "Iris/ReplicationState/ReplicationStateUtil.h"
+
+namespace UE::Net 
+{
+
+namespace Private 
+{
+
+/*
+ * Internal access of FastArraySerializer, for internal use only
+*/
+struct FIrisFastArraySerializerPrivateAccessor
+{
+	static FNetBitArrayView GetChangeMask(FIrisFastArraySerializer& Array) { return MakeNetBitArrayView(&Array.ChangeMaskStorage[0], FIrisFastArraySerializer::IrisFastArrayChangeMaskBits + 1U); }
+	static FNetBitArrayView GetConditionalChangeMask(FIrisFastArraySerializer& Array) { return MakeNetBitArrayView(&Array.ChangeMaskStorage[1], FIrisFastArraySerializer::IrisFastArrayChangeMaskBits + 1U); };
+	static FReplicationStateHeader& GetReplicationStateHeader(FIrisFastArraySerializer& Array) { return Array.ReplicationStateHeader; }
+
+	/*
+	 * Mark array as dirty and notify DirtyObjectTracker that the object is dirty if it has not been done before
+	 */
+	static void MarkArrayDirty(FIrisFastArraySerializer& Array);
+
+	/*
+	 * Mark array and all bits in changemask as dirty and notify DirtyObjectTracker that the object is dirty if it has not been done before
+	 */
+	static void MarkAllArrayItemsDirty(FIrisFastArraySerializer& Array, uint32 StartingIndex = 0U);
+
+	/*
+	 * Mark array dirty and mark changemask bit used for array Item as dirty and notify DirtyObjectTracker that the object is dirty if it has not been done before
+     * Currently the same changemask bit might be used to indicate dirtiness for multiple array items
+	 */
+	static void MarkArrayItemDirty(FIrisFastArraySerializer& Array, int32 Index);
+};
+
+inline void FIrisFastArraySerializerPrivateAccessor::MarkAllArrayItemsDirty(FIrisFastArraySerializer& Array, uint32 StartingIndex)
+{
+	checkSlow(Array.ReplicationStateHeader.IsBound());
+
+	FNetBitArrayView MemberChangeMask = UE::Net::Private::FIrisFastArraySerializerPrivateAccessor::GetChangeMask(Array);
+	if (!MemberChangeMask.GetBit(0))
+	{
+		MarkNetObjectStateDirty(Array.ReplicationStateHeader);
+	}
+	if (StartingIndex == 0)
+	{
+		MemberChangeMask.SetAllBits();
+	}
+	else
+	{
+		MemberChangeMask.SetBit(FIrisFastArraySerializer::IrisFastArrayPropertyBitIndex);
+		MemberChangeMask.SetBits(FIrisFastArraySerializer::IrisFastArrayChangeMaskBitOffset + StartingIndex, FIrisFastArraySerializer::IrisFastArrayChangeMaskBits - StartingIndex);
+	}
+}
+
+inline void FIrisFastArraySerializerPrivateAccessor::MarkArrayDirty(FIrisFastArraySerializer& Array)
+{
+	checkSlow(Array.ReplicationStateHeader.IsBound());
+
+	FNetBitArrayView MemberChangeMask = UE::Net::Private::FIrisFastArraySerializerPrivateAccessor::GetChangeMask(Array);
+
+	// Dirty object unless already dirty, we only use the array bit for this purpose
+	if (!MemberChangeMask.GetBit(FIrisFastArraySerializer::IrisFastArrayPropertyBitIndex))
+	{
+		MarkNetObjectStateDirty(Array.ReplicationStateHeader);
+		MemberChangeMask.SetBit(FIrisFastArraySerializer::IrisFastArrayPropertyBitIndex);
+	}
+}
+
+inline void FIrisFastArraySerializerPrivateAccessor::MarkArrayItemDirty(FIrisFastArraySerializer& Array, int32 ItemIdx)
+{
+	checkSlow(Array.ReplicationStateHeader.IsBound());
+
+	// Mark changemask dirty for this item
+	// We are using a modulo scheme for dirtiness
+	FNetBitArrayView MemberChangeMask = UE::Net::Private::FIrisFastArraySerializerPrivateAccessor::GetChangeMask(Array);
+	MemberChangeMask.SetBit((ItemIdx % FIrisFastArraySerializer::IrisFastArrayChangeMaskBits) + FIrisFastArraySerializer::IrisFastArrayChangeMaskBitOffset);
+
+	// Dirty object unless already dirty, we only use the array bit for this purpose
+	if (!MemberChangeMask.GetBit(0))
+	{
+		MarkNetObjectStateDirty(Array.ReplicationStateHeader);
+		MemberChangeMask.SetBits(0, 1);
+	}
+}
+
+} // end namespace Private
+
+/**
+ * Experimental support for more explicit interface to edit FastArrays which can be used to avoid polling
+ * The idea is that the interface would implement a subset of the Array interface and be used instead of directly modifying the array
+ */
+template <typename FastArrayType>
+class TIrisFastArrayEditor
+{
+public:
+	typedef typename FastArrayType::ItemArrayType FastArrayItemArrayType;
+	typedef typename FastArrayItemArrayType::ElementType FastArrayItemType;
+
+	using FIrisFastArraySerializerPrivateAccessor = Private::FIrisFastArraySerializerPrivateAccessor;
+
+	TIrisFastArrayEditor(FastArrayType& InFastArray) : FastArray(InFastArray) {}
+
+	/**
+	 * Forwards MarkItemDirty call to FastArray and if the FastArray is bound it will also update DirtyState tracking
+	 */
+	void MarkItemDirty(FastArrayItemType& Item);
+	
+	/**
+	 * Forwards MarkArrayDirty call to FastArray and if the FastArray is bound it will also update DirtyState tracking
+	 */
+	void MarkArrayDirty() { FastArray.MarkArrayDirty(); };
+
+	/**
+	 * Local add which will add item to array without dirtying it
+	 */
+	void AddLocal(const FastArrayItemType& ItemEntry) { FastArray.GetItemArray().Add(ItemEntry); }
+
+	/**
+	 * Local edit which will modify local item without dirtying it
+	 */
+	FastArrayItemType& EditLocal(int32 ItemIdx) { return FastArray.GetItemArray()[ItemIdx]; }
+
+	/**
+	 * Add Item to array and call MarkItemDirty
+	 */
+	void Add(const FastArrayItemType& ItemEntry);
+
+	/**
+	 * Edit Item in array and call MarkItemDirty
+	 */
+	FastArrayItemType& Edit(int32 ItemIdx);
+
+	/**
+	 * Mutable access to Item, will call MarkItemDirty
+	 */
+	FastArrayItemType& operator[](int32 ItemIdx) { return Edit(ItemIdx); }
+
+	/**
+	 * Remove item at the specified index, will forward call to MarkArrayDirty, if bound will mark all potentially moved items as dirty
+	 */
+	void Remove(int32 ItemIdx);
+
+	/**
+	 * Remove item at the specified index, will only mark the affected item dirty
+	 */
+	void RemoveAtSwap(int32 ItemIdx);
+
+	int32 Num() const { return FastArray.GetItemArray().Num(); }
+
+	/**
+	 * Empty array and call MarkArrayDirty
+	 */
+	void Empty();
+
+private:
+	friend FastArrayType;
+	FastArrayType& FastArray;
+};
+
+template <typename FastArrayType>
+void TIrisFastArrayEditor<FastArrayType>::MarkItemDirty(FastArrayItemType& Item)
+{
+	const FReplicationStateHeader& Header = FIrisFastArraySerializerPrivateAccessor::GetReplicationStateHeader(FastArray);
+	if (!Header.IsBound())
+	{
+		FastArray.MarkItemDirty(Item);
+		return;
+	}
+
+	const FastArrayItemArrayType& Array = FastArray.GetItemArray();
+	const int32 Index = (&Item - Array.GetData());
+	if (Array.IsValidIndex(Index))
+	{
+		FNetBitArrayView MemberChangeMask = FIrisFastArraySerializerPrivateAccessor::GetChangeMask(FastArray);
+
+		// If this is a new element make sure it is at the end, otherwise we must dirty array
+		if ((Item.ReplicationID == INDEX_NONE) && (Index != (Array.Num() - 1)))
+		{
+			// Must mark all items that might have been shifted dirty, including the new one, this is to mimic behavior of current FastArraySerializer
+			for (FastArrayItemType& ArrayItem : MakeArrayView(&Item, Array.Num() - Index))
+			{
+				const bool bIsWritingOnClient = false;
+				if (!FastArray.template ShouldWriteFastArrayItem<FastArrayItemType, FastArrayType>(ArrayItem, bIsWritingOnClient))
+				{
+					if (ArrayItem.ReplicationID == INDEX_NONE)
+					{
+						FastArray.MarkItemDirty(ArrayItem);						
+					}
+					// Mark item dirty in changemask
+					FIrisFastArraySerializerPrivateAccessor::MarkArrayItemDirty(FastArray, (&ArrayItem - Array.GetData()));
+				}
+			}
+		}
+		else
+		{
+			// This is in order to execute old logic
+			FastArray.MarkItemDirty(Item);
+			// Mark item dirty in changemask
+			FIrisFastArraySerializerPrivateAccessor::MarkArrayItemDirty(FastArray, Index);
+		}
+	}
+};
+	
+
+// The idea is to provide an interface that is modeled after what we can do with a TArray but with dirty tracking per member
+// This way we can update changemask directly and does not need to poll for changes
+template <typename FastArrayType>
+void TIrisFastArrayEditor<FastArrayType>::Add(const FastArrayItemType& ItemEntry)
+{
+	FastArrayItemArrayType& ItemArray = FastArray.GetItemArray();
+	int32 Idx = ItemArray.Add(ItemEntry);
+	MarkItemDirty(ItemArray.GetData()[Idx]);
+}
+
+template <typename FastArrayType>
+typename TIrisFastArrayEditor<FastArrayType>::FastArrayItemType& TIrisFastArrayEditor<FastArrayType>::Edit(int32 ItemIdx)
+{
+	FastArrayItemType& Item = FastArray.GetItemArray()[ItemIdx];
+	MarkItemDirty(Item);
+	return Item;
+}
+
+template <typename FastArrayType>
+void TIrisFastArrayEditor<FastArrayType>::Remove(int32 ItemIdx)
+{
+	FastArrayItemArrayType& ItemArray = FastArray.GetItemArray();
+
+	if (ItemIdx >= 0 && ItemArray.Num() > ItemIdx)
+	{
+		// If the Item is at the end we only need to mark the array as dirty
+		ItemArray.RemoveAt(ItemIdx);
+		MarkArrayDirty();
+		if (ItemIdx < ItemArray.Num())
+		{
+			FIrisFastArraySerializerPrivateAccessor::MarkAllArrayItemsDirty(FastArray, ItemIdx);
+		}
+	}
+}
+
+template <typename FastArrayType>
+void TIrisFastArrayEditor<FastArrayType>::RemoveAtSwap(int32 ItemIdx)
+{
+	FastArrayItemArrayType& ItemArray = FastArray.GetItemArray();
+
+	if (ItemIdx >= 0 && ItemArray.Num() > ItemIdx)
+	{
+		ItemArray.RemoveAtSwap(ItemIdx);
+
+		// We just need to dirty the modified item
+		if (ItemArray.Num() != 0)
+		{
+			MarkItemDirty(ItemIdx);
+		}
+		else
+		{
+			MarkArrayDirty();
+		}
+	}
+}
+
+template <typename FastArrayType>
+void TIrisFastArrayEditor<FastArrayType>::Empty()
+{
+	FastArray.GetItemArray().Empty();
+	MarkArrayDirty();
+}
+
+} // end namespace UE::Net
+
+
+
diff --git a/Engine/Source/Runtime/Experimental/Iris/Core/Public/Iris/ReplicationState/ReplicationStateFwd.h b/Engine/Source/Runtime/Experimental/Iris/Core/Public/Iris/ReplicationState/ReplicationStateFwd.h
new file mode 100644
--- /dev/null
+++ b/Engine/Source/Runtime/Experimental/Iris/Core/Public/Iris/ReplicationState/ReplicationStateFwd.h
@@ -0,0 +1,55 @@
+// Copyright Epic Games, Inc. All Rights Reserved.
+
+#pragma once
+
+#include "CoreTypes.h"
+
+namespace UE {
+namespace Net {
+namespace Private {
+	struct FReplicationStateHeaderAccessor;
+}
+
+/** 
+ * A ReplicationState always contains a ReplicationStateHeader which we use to bind replication states for dirty tracking
+ */
+struct FReplicationStateHeader
+{
+	inline FReplicationStateHeader() : ReplicationIndex(0), ReplicationSystemId(0), bInitStateIsDirty(0), ReplicationFlags(0) {};
+
+	/** Returns true if the state is bound to the dirty tracking system */
+	bool IsBound() const { return ReplicationIndex != 0U; }
+
+private:
+	friend struct Private::FReplicationStateHeaderAccessor;
+
+	// All replication states that are bound by a instance protocol as assigned an index used for dirty state tracking
+	uint32 ReplicationIndex : 24;
+	// In order to support PIE we need to be able to manage multiple replication systems
+	uint32 ReplicationSystemId : 4;
+	// Init state doesn't use changemasks, instead we have a reserved bit here
+	uint32 bInitStateIsDirty : 1;
+	// Unused so far
+	uint32 ReplicationFlags : 3;
+};
+
+static_assert(sizeof(FReplicationStateHeader) == sizeof(uint32) && alignof(FReplicationStateHeader) == alignof(uint32), "FReplicationStateHeader must currently have the same size and alignment as uint32");
+
+namespace Private {
+
+/** 
+ * Internal access, should only be used by internal code
+ */
+struct FReplicationStateHeaderAccessor
+{
+	static uint32 GetReplicationIndex(const FReplicationStateHeader& Header) { return Header.ReplicationIndex; }
+	static uint32 GetReplicationSystemId(const FReplicationStateHeader& Header) { return Header.ReplicationSystemId; }
+	static bool GetIsInitStateDirty(const FReplicationStateHeader& Header) { return Header.bInitStateIsDirty; }
+	static uint32 GetReplicationFlags(const FReplicationStateHeader& Header) { return Header.ReplicationFlags; }
+
+	static void MarkInitStateDirty(FReplicationStateHeader& Header) { Header.bInitStateIsDirty = true; }
+	static void ClearInitStateIsDirty(FReplicationStateHeader& Header) { Header.bInitStateIsDirty = false; }
+	static void SetReplicationIndex(FReplicationStateHeader& Header, uint32 Index, uint32 ReplicationSystemId) { Header.ReplicationIndex = Index; Header.ReplicationSystemId = ReplicationSystemId; }
+};
+
+}}}
diff --git a/Engine/Source/Runtime/Experimental/Iris/Core/Public/Iris/ReplicationState/ReplicationStateUtil.h b/Engine/Source/Runtime/Experimental/Iris/Core/Public/Iris/ReplicationState/ReplicationStateUtil.h
new file mode 100644
--- /dev/null
+++ b/Engine/Source/Runtime/Experimental/Iris/Core/Public/Iris/ReplicationState/ReplicationStateUtil.h
@@ -0,0 +1,81 @@
+// Copyright Epic Games, Inc. All Rights Reserved.
+
+#pragma once
+
+#include "CoreTypes.h"
+#include "Net/Core/NetBitArray.h"
+#include "Iris/ReplicationState/ReplicationStateDescriptor.h"
+#include "Iris/ReplicationState/ReplicationStateFwd.h"
+
+namespace UE::Net::Private
+{
+	extern IRISCORE_API void MarkNetObjectStateDirty(uint32 ReplicationSystemId, uint32 NetObjectIndex);
+}
+
+namespace UE::Net
+{
+/**
+ * Mark the a NetObject/Handle as dirty which means it requires to be copied when we update the Replication system
+ */
+inline void MarkNetObjectStateDirty(const FReplicationStateHeader& Header) { Private::MarkNetObjectStateDirty(Private::FReplicationStateHeaderAccessor::GetReplicationSystemId(Header), Private::FReplicationStateHeaderAccessor::GetReplicationIndex(Header)); }
+}
+
+namespace UE::Net::Private
+{
+
+/**
+ * Get FReplicationStateHeader from a ReplicationState
+ */
+inline UE::Net::FReplicationStateHeader& GetReplicationStateHeader(void* StateBuffer, const FReplicationStateDescriptor* Descriptor)
+{
+	// $IRIS: $TODO: Currently we do not store the offset to the internal replication state
+	if (!EnumHasAnyFlags(Descriptor->Traits, EReplicationStateTraits::IsNativeFastArrayReplicationState))
+	{
+		return *reinterpret_cast<UE::Net::FReplicationStateHeader*>(StateBuffer);
+	}
+	else
+	{
+		const SIZE_T ReplicationStateHeaderOffset = Descriptor->ChangeMasksExternalOffset + 8U;
+		return *reinterpret_cast<UE::Net::FReplicationStateHeader*>(reinterpret_cast<uint8*>(StateBuffer) + ReplicationStateHeaderOffset);
+	}
+}
+
+inline bool IsReplicationStateBound(void* StateBuffer, const FReplicationStateDescriptor* Descriptor)
+{
+	return GetReplicationStateHeader(StateBuffer, Descriptor).IsBound();
+}
+
+/**
+ * Get MemberChangeMask
+ */
+inline FNetBitArrayView GetMemberChangeMask(uint8* StateBuffer, const FReplicationStateDescriptor* Descriptor)
+{
+	return FNetBitArrayView((FNetBitArrayView::StorageWordType*)(StateBuffer + Descriptor->GetChangeMaskOffset()), Descriptor->ChangeMaskBitCount, FNetBitArrayView::NoResetNoValidate);
+}
+
+/**
+ * Get MemberConditionalChangeMask. Valid if state has lifetime conditionals.
+ */
+inline FNetBitArrayView GetMemberConditionalChangeMask(uint8* StateBuffer, const FReplicationStateDescriptor* Descriptor)
+{
+	return FNetBitArrayView((FNetBitArrayView::StorageWordType*)(StateBuffer + Descriptor->GetConditionalChangeMaskOffset()), Descriptor->ChangeMaskBitCount, FNetBitArrayView::NoResetNoValidate);
+}
+
+/**
+ * Mark specific member dirty, if this is the first bit marked as dirty in the local MemberChangeMask and the state is bound, mark owning object as dirty as well.
+ * Note that all bits described by the ChangeMaskInfo will be dirtied by the call. 
+ */
+inline void MarkDirty(const UE::Net::FReplicationStateHeader& InternalState, FNetBitArrayView& MemberChangeMask, const FReplicationStateMemberChangeMaskDescriptor& ChangeMaskInfo)
+{
+	// If this state is bound to a replicated object, notify the replication system that we have data to copy
+	// Note that we only check the first bit of the changemask, as this is used to indicate whether the property is dirty or not
+	// some properties might use additional bits but will only be treated as dirty if the master bit is set.
+	if (InternalState.IsBound() && !MemberChangeMask.GetBit(ChangeMaskInfo.BitOffset))
+	{
+		MarkNetObjectStateDirty(InternalState);
+	}
+
+	MemberChangeMask.SetBits(ChangeMaskInfo.BitOffset, ChangeMaskInfo.BitCount);
+}
+
+}
diff --git a/Engine/Source/Runtime/Experimental/Iris/Core/Public/Iris/ReplicationSystem/Filtering/FilterOutNetObjectFilter.h b/Engine/Source/Runtime/Experimental/Iris/Core/Public/Iris/ReplicationSystem/Filtering/FilterOutNetObjectFilter.h
new file mode 100644
--- /dev/null
+++ b/Engine/Source/Runtime/Experimental/Iris/Core/Public/Iris/ReplicationSystem/Filtering/FilterOutNetObjectFilter.h
@@ -0,0 +1,26 @@
+// Copyright Epic Games, Inc. All Rights Reserved.
+
+#pragma once
+
+#include "Iris/ReplicationSystem/Filtering/NetObjectFilter.h"
+#include "FilterOutNetObjectFilter.generated.h"
+
+UCLASS(transient, MinimalAPI)
+class UFilterOutNetObjectFilterConfig final : public UNetObjectFilterConfig
+{
+	GENERATED_BODY()
+};
+
+UCLASS()
+class UFilterOutNetObjectFilter final : public UNetObjectFilter
+{
+	GENERATED_BODY()
+
+protected:
+	// UNetObjectFilter interface
+	IRISCORE_API virtual void Init(FNetObjectFilterInitParams&) override;
+	IRISCORE_API virtual bool AddObject(uint32 ObjectIndex, FNetObjectFilterAddObjectParams&) override;
+	IRISCORE_API virtual void RemoveObject(uint32 ObjectIndex, const FNetObjectFilteringInfo&) override;
+	IRISCORE_API virtual void UpdateObjects(FNetObjectFilterUpdateParams&) override;
+	IRISCORE_API virtual void Filter(FNetObjectFilteringParams&) override;
+};
diff --git a/Engine/Source/Runtime/Experimental/Iris/Core/Public/Iris/ReplicationSystem/Filtering/NetObjectFilter.h b/Engine/Source/Runtime/Experimental/Iris/Core/Public/Iris/ReplicationSystem/Filtering/NetObjectFilter.h
new file mode 100644
--- /dev/null
+++ b/Engine/Source/Runtime/Experimental/Iris/Core/Public/Iris/ReplicationSystem/Filtering/NetObjectFilter.h
@@ -0,0 +1,207 @@
+// Copyright Epic Games, Inc. All Rights Reserved.
+
+#pragma once
+
+#include "Net/Core/NetBitArray.h"
+#include "Iris/ReplicationSystem/ReplicationView.h"
+#include "UObject/ObjectMacros.h"
+#include "NetObjectFilter.generated.h"
+
+struct FNetObjectFilteringInfo;
+class UReplicationSystem;
+namespace UE::Net
+{
+	typedef uint32 FNetObjectFilterHandle;
+	typedef uint16 FNetObjectGroupHandle;
+	struct FReplicationInstanceProtocol;
+	struct FReplicationProtocol;
+}
+
+namespace UE::Net
+{
+
+constexpr FNetObjectFilterHandle InvalidNetObjectFilterHandle = FNetObjectFilterHandle(0);
+constexpr FNetObjectFilterHandle ToOwnerFilterHandle = FNetObjectFilterHandle(1);
+/** ConnectionFilterHandle is for internal use only. */
+constexpr FNetObjectFilterHandle ConnectionFilterHandle = FNetObjectFilterHandle(2);
+
+/** Invalid group handle */
+constexpr FNetObjectGroupHandle InvalidNetObjectGroupHandle = FNetObjectGroupHandle(0);
+
+/** Special group, NetHandles assigned to this group will be filtered out for all connections */
+constexpr FNetObjectGroupHandle NotReplicatedNetObjectGroupHandle = FNetObjectGroupHandle(1);
+
+/** Special group, SubObjects assigned to this group will replicate to owner of RootParent */
+constexpr FNetObjectGroupHandle NetGroupOwnerNetObjectGroupHandle = FNetObjectGroupHandle(2);
+
+/** Special group, NetHandles assigned to this group will Replicate if replay netconditions is met  */
+constexpr FNetObjectGroupHandle NetGroupReplayNetObjectGroupHandle = FNetObjectGroupHandle(3);
+
+/** Returns true of the provided GroupHandle is a reserved NetObjectGroupHandle */
+static constexpr bool IsReservedNetObjectGroupHandle(FNetObjectGroupHandle GroupHandle) { return GroupHandle >= NotReplicatedNetObjectGroupHandle && GroupHandle <= NetGroupReplayNetObjectGroupHandle; }
+
+/** Used to control whether an object is allowed to be replicated or not. */
+enum class ENetFilterStatus : uint32
+{
+	/** Do not allow replication. */
+	Disallow,
+	/** Allow replication. */
+	Allow,
+};
+
+}
+
+/**
+ * Parameters passed to UNetObjectFilter::Filter.
+ */
+struct FNetObjectFilteringParams
+{
+	FNetObjectFilteringParams(const UE::Net::FNetBitArrayView InFilteredObjects);
+
+	/** The indices of the objects that have this filter set. The indices of set bits correspond to the object indices. */
+	const UE::Net::FNetBitArrayView FilteredObjects;
+
+	/**
+	 * The contents of OutAllowedObjects is undefined when passed to Filter(). The filter is responsible
+	 * for setting and clearing bits for objects that have this filter set, which is provided in the
+	 * FilteredObjects member. It's safe to set or clear all bits in the bitarray as the callee will
+	 * only care about bits which the filter is responsible for.
+	 */
+	UE::Net::FNetBitArrayView OutAllowedObjects;
+
+	/** FilteringInfos for all objects. Index using the set bit indices in FilteredObjects. */
+	const FNetObjectFilteringInfo* FilteringInfos;
+
+	/** State buffers for all objects. Index using the set bit indices in FilteredObjects. */
+	uint8 *const* StateBuffers;
+
+	/** ID of the connection that the filtering applies to. */
+	uint32 ConnectionId;
+
+	/** The view associated with the connection and its sub-sconnections that objects are filtered for. */
+	UE::Net::FReplicationView View;
+};
+
+/**
+ * Parameters passed to UNetObjectFilter::PreFilter.
+ */
+struct FNetObjectPreFilteringParams
+{
+};
+
+/**
+ * Parameters passed to UNetObjectFilter::PostFilter.
+ */
+struct FNetObjectPostFilteringParams
+{
+};
+
+/**
+ * Filter specific data stored per object, such as offsets to tags.
+ * The data is initialized to zero by default.
+ */
+struct alignas(8) FNetObjectFilteringInfo
+{
+	uint16 Data[4];
+};
+
+/**
+ * Base class for filter specific configuration.
+ * @see FNetObjectFilterDefinition
+ */
+UCLASS(Transient, MinimalAPI)
+class UNetObjectFilterConfig : public UObject
+{
+	GENERATED_BODY()
+};
+
+/** Parameters passed to the filter's Init() call. */
+struct FNetObjectFilterInitParams
+{
+	TObjectPtr<const UReplicationSystem> ReplicationSystem;
+	/** Optional config as set in the FNetObjectFilterDefinition. */
+	UNetObjectFilterConfig* Config = nullptr;
+	/** The maximum number of objects in the system. */
+	uint32 MaxObjectCount = 0;
+	/** The maximum number of connections in the system. */
+	uint32 MaxConnectionCount = 0;
+};
+
+struct FNetObjectFilterAddObjectParams
+{
+	/** The info is zeroed before the AddObject() call. Fill in with filter specifics, like offsets to tags. */
+	FNetObjectFilteringInfo& OutInfo;
+
+	/** The FReplicationInstanceProtocol which describes the source state data. */
+	const UE::Net::FReplicationInstanceProtocol* InstanceProtocol;
+
+	/** The FReplicationProtocol which describes the internal state data. */
+	const UE::Net::FReplicationProtocol* Protocol;
+
+	/**
+	 * One can retrieve relevant information from the object state buffer using the FReplicationProtocol.
+	 * Note that this is the internal network representation of the data which is stored in quantized form.
+	 * NetSerializers can dequantize the data to the original source data form.
+	 */
+	const uint8* StateBuffer;
+};
+
+/** Parameters passed to the filter's UpdateObjects() call. */
+struct FNetObjectFilterUpdateParams
+{
+	/** Indices of the updated objects. */
+	const uint32* ObjectIndices;
+	/** The number of objects that have been updated. */
+	uint32 ObjectCount;
+
+	/** InstanceProtocols for updated objects. Index using 0..ObjectCount-1. */
+	UE::Net::FReplicationInstanceProtocol const*const* InstanceProtocols;
+
+	/** State buffers for all objects. Index using ObjectIndices[0..ObjectCount-1]. */
+	uint8 *const* StateBuffers;
+
+	/** Infos for all objects. Index using ObjectIndices[0..ObjectCount-1]. */
+	FNetObjectFilteringInfo* FilteringInfos;
+};
+
+UCLASS()
+class UNetObjectFilter : public UObject
+{
+	GENERATED_BODY()
+
+public:
+	/** Called right after constructor for enabled filters. Must be overriden. */ 
+	IRISCORE_API virtual void Init(FNetObjectFilterInitParams&) PURE_VIRTUAL(Init,);
+
+	/** A new connection has been added. An opportunity for the filter to allocate per connection info. */
+	IRISCORE_API virtual void AddConnection(uint32 ConnectionId);
+
+	/** A new connection has been added. An opportunity for the filter to deallocate per connection info. */
+	IRISCORE_API virtual void RemoveConnection(uint32 ConnectionId);
+
+	/** A new object want to use this filter. Opportunity to cache some information for it. The info struct passed has been zeroed. Must be overriden. */
+	IRISCORE_API virtual bool AddObject(uint32 ObjectIndex, FNetObjectFilterAddObjectParams&) PURE_VIRTUAL(AddObject, return false;)
+
+	/** An object do no longer want to use this filter. */
+	IRISCORE_API virtual void RemoveObject(uint32 ObjectIndex, const FNetObjectFilteringInfo&) PURE_VIRTUAL(RemoveObject,)
+
+	/** A set of objects ususing this filter has been updated. An opportunity for the filter to update cached data. */
+	IRISCORE_API virtual void UpdateObjects(FNetObjectFilterUpdateParams&) PURE_VIRTUAL(UpdateObjects,)
+
+	/**
+	 * If there are any connections being replicated and there's a chance Filter() will be called then PreFilter()
+	 * will be called exactly once before all calls to Filter().
+	 */
+	IRISCORE_API virtual void PreFilter(FNetObjectPreFilteringParams&);
+
+	/** Filter a batch of objects. There may be multiple calls to this function even for the same connection. Must be overriden. */
+	IRISCORE_API virtual void Filter(FNetObjectFilteringParams&);
+
+	/**
+	 * If PreFilter() was called then PostFilter() will be called exactly once after all Filter() calls.
+	 */
+	IRISCORE_API virtual void PostFilter(FNetObjectPostFilteringParams&);
+
+protected:
+	IRISCORE_API UNetObjectFilter();
+};
diff --git a/Engine/Source/Runtime/Experimental/Iris/Core/Public/Iris/ReplicationSystem/Filtering/NetObjectGridFilter.h b/Engine/Source/Runtime/Experimental/Iris/Core/Public/Iris/ReplicationSystem/Filtering/NetObjectGridFilter.h
new file mode 100644
--- /dev/null
+++ b/Engine/Source/Runtime/Experimental/Iris/Core/Public/Iris/ReplicationSystem/Filtering/NetObjectGridFilter.h
@@ -0,0 +1,188 @@
+// Copyright Epic Games, Inc. All Rights Reserved.
+
+#pragma once
+
+#include "Iris/ReplicationSystem/Filtering/NetObjectFilter.h"
+#include "Net/Core/NetBitArray.h"
+#include "Containers/ChunkedArray.h"
+#include "Math/Vector.h"
+#include "UObject/StrongObjectPtr.h"
+#include "NetObjectGridFilter.generated.h"
+
+namespace UE::Net
+{
+	class FWorldLocations;
+}
+
+UCLASS(transient, config=Engine, MinimalAPI)
+class UNetObjectGridFilterConfig : public UNetObjectFilterConfig
+{
+	GENERATED_BODY()
+
+public:
+	/** How many frames a view position should be considered relevant. To avoid culling issues when player borders cells. */
+	UPROPERTY(Config)
+	uint32 ViewPosRelevancyFrameCount = 2;
+
+	UPROPERTY(Config)
+	float CellSizeX = 20000.0f;
+
+	UPROPERTY(Config)
+	float CellSizeY = 20000.0f;
+
+	/** Objects with larger sqrt(NetCullDistanceSqr) will be rejected. */
+	UPROPERTY(Config)
+	float MaxCullDistance = 20000.0f;
+
+	/** Objects without a NetCullDistanceSquared property will assume to have this value but squared. */
+	UPROPERTY(Config)
+	float DefaultCullDistance = 15000.0f;
+
+	/** Coordinates will be clamped to MinPos and MaxPos. */
+	UPROPERTY(Config)
+	FVector MinPos = {-0.5f*2097152.0f, -0.5f*2097152.0f, -0.5f*2097152.0f};
+
+	/** Coordinates will be clamped to MinPos and MaxPos. */
+	UPROPERTY(Config)
+	FVector MaxPos = {+0.5f*2097152.0f, +0.5f*2097152.0f, +0.5f*2097152.0f};
+};
+
+UCLASS()
+class UNetObjectGridFilter : public UNetObjectFilter
+{
+	GENERATED_BODY()
+
+protected:
+	// UNetObjectFilter interface
+	IRISCORE_API virtual void Init(FNetObjectFilterInitParams&) override;
+	IRISCORE_API virtual void AddConnection(uint32 ConnectionId) override;
+	IRISCORE_API virtual void RemoveConnection(uint32 ConnectionId) override;
+	IRISCORE_API virtual bool AddObject(uint32 ObjectIndex, FNetObjectFilterAddObjectParams&) override;
+	IRISCORE_API virtual void RemoveObject(uint32 ObjectIndex, const FNetObjectFilteringInfo&) override;
+	IRISCORE_API virtual void UpdateObjects(FNetObjectFilterUpdateParams&) override;
+	IRISCORE_API virtual void PreFilter(FNetObjectPreFilteringParams&) override;
+	IRISCORE_API virtual void Filter(FNetObjectFilteringParams&) override;
+
+protected:
+	struct FObjectLocationInfo : public FNetObjectFilteringInfo
+	{
+		bool IsUsingWorldLocations() const { return GetLocationStateIndex() == InvalidStateIndex; }
+		bool IsUsingLocationInState() const { return GetLocationStateIndex() != InvalidStateIndex; }
+
+		void SetLocationStateOffset(uint16 Offset) { Data[0] = Offset; }
+		uint16 GetLocationStateOffset() const { return Data[0]; }
+
+		void SetLocationStateIndex(uint16 Index) { Data[1] = Index; }
+		uint16 GetLocationStateIndex() const { return Data[1]; }
+
+		void SetInfoIndex(uint32 Index) { Data[2] = Index & 65535U; Data[3] = Index >> 16U; }
+		uint32 GetInfoIndex() const { return (uint32(Data[3]) << 16U) | uint32(Data[2]); }
+	};
+
+	IRISCORE_API void AddCellInfoForObject(const FObjectLocationInfo& ObjectInfo, const UE::Net::FReplicationInstanceProtocol* InstanceProtocol);
+	IRISCORE_API void RemoveCellInfoForObject(const FObjectLocationInfo& ObjectInfo);
+	IRISCORE_API void UpdateCellInfoForObject(const FObjectLocationInfo& ObjectInfo, const UE::Net::FReplicationInstanceProtocol* InstanceProtocol);
+
+private:
+	enum : unsigned
+	{
+		ObjectInfosChunkSize = 64 * 1024,
+		InvalidStateIndex = 65535U,
+		InvalidStateOffset = 65535U,
+	};
+
+	struct FCellBox
+	{
+		int32 MinX;
+		int32 MaxX;
+		int32 MinY;
+		int32 MaxY;
+
+		bool operator==(const FCellBox&) const;
+		bool operator!=(const FCellBox&) const;
+	};
+
+	// We can't fit all info we need in 4x16bits.
+	struct FPerObjectInfo
+	{
+		FVector Position = {0,0,0};
+		FCellBox CellBox = {};
+		float CullDistance = 0;
+		uint32 ObjectIndex = 0;
+		uint16 CullDistanceSqrStateIndex = InvalidStateIndex;
+		uint16 CullDistanceSqrStateOffset = InvalidStateOffset;
+	};
+
+	struct FCellCoord
+	{
+		int32 X;
+		int32 Y;
+
+	private:
+		friend bool operator==(const FCellCoord& A, const FCellCoord& B)
+		{
+			return (A.X == B.X) & (A.Y == B.Y);
+		}
+
+		friend uint32 GetTypeHash(const FCellCoord& Coords)
+		{
+			return ::GetTypeHash((uint64(uint32(Coords.X)) << 32U) | uint64(uint32(Coords.Y)));
+		}
+	};
+
+	struct FCellObjects
+	{
+		TSet<uint32> ObjectIndices;
+	};
+
+	struct FCellAndTimestamp
+	{
+		FCellCoord Cell;
+		uint32 Timestamp;
+
+	private:
+		friend bool operator==(const FCellAndTimestamp& A, const FCellAndTimestamp& B)
+		{
+			return (A.Cell == B.Cell) & (A.Timestamp == B.Timestamp);
+		}
+	};
+
+	struct FPerConnectionInfo
+	{
+		// We don't expect a lot of view positions from a single connection
+		TArray<FCellAndTimestamp, TInlineAllocator<32>> RecentCells;
+	};
+
+private:
+	uint32 AllocObjectInfo();
+	void FreeObjectInfo(uint32 Index);
+
+	void UpdatePositionAndCullDistance(const FObjectLocationInfo& ObjectLocationInfo, FPerObjectInfo& PerObjectInfo, const UE::Net::FReplicationInstanceProtocol* InstanceProtocol);
+	void CalculateCellBox(const FPerObjectInfo& PerObjectInfo, FCellBox& OutCellBox);
+	void CalculateCellCoord(FCellCoord& OutCoord, const FVector& Pos);
+
+	static bool AreCellsDisjoint(const FCellBox& A, const FCellBox& B);
+	static bool DoesCellContainCoord(const FCellBox& Cell, const FCellCoord& Coord);
+
+	TStrongObjectPtr<UNetObjectGridFilterConfig> Config;
+	TArray<FPerConnectionInfo> PerConnectionInfos;
+	TChunkedArray<FPerObjectInfo, ObjectInfosChunkSize> ObjectInfos;
+	UE::Net::FNetBitArray AssignedObjectInfoIndices;
+
+	TMap<FCellCoord, FCellObjects> Cells;
+	uint32 FrameIndex = 0;
+
+	const UE::Net::FWorldLocations* WorldLocations = nullptr;
+};
+
+
+//
+inline bool UNetObjectGridFilter::FCellBox::operator==(const UNetObjectGridFilter::FCellBox& Other) const
+{
+	return (MinX == Other.MinX) & (MinY == Other.MinY) & (MaxX == Other.MaxX) & (MaxY == Other.MaxY);
+}
+
+inline bool UNetObjectGridFilter::FCellBox::operator!=(const UNetObjectGridFilter::FCellBox& Other) const
+{
+	return (MinX != Other.MinX) | (MinY != Other.MinY) | (MaxX != Other.MaxX) | (MaxY != Other.MaxY);
+}
diff --git a/Engine/Source/Runtime/Experimental/Iris/Core/Public/Iris/ReplicationSystem/Filtering/NopNetObjectFilter.h b/Engine/Source/Runtime/Experimental/Iris/Core/Public/Iris/ReplicationSystem/Filtering/NopNetObjectFilter.h
new file mode 100644
--- /dev/null
+++ b/Engine/Source/Runtime/Experimental/Iris/Core/Public/Iris/ReplicationSystem/Filtering/NopNetObjectFilter.h
@@ -0,0 +1,26 @@
+// Copyright Epic Games, Inc. All Rights Reserved.
+
+#pragma once
+
+#include "Iris/ReplicationSystem/Filtering/NetObjectFilter.h"
+#include "NopNetObjectFilter.generated.h"
+
+UCLASS(transient, MinimalAPI)
+class UNopNetObjectFilterConfig final : public UNetObjectFilterConfig
+{
+	GENERATED_BODY()
+};
+
+UCLASS()
+class UNopNetObjectFilter final : public UNetObjectFilter
+{
+	GENERATED_BODY()
+
+protected:
+	// UNetObjectFilter interface
+	IRISCORE_API virtual void Init(FNetObjectFilterInitParams&) override;
+	IRISCORE_API virtual bool AddObject(uint32 ObjectIndex, FNetObjectFilterAddObjectParams&) override;
+	IRISCORE_API virtual void RemoveObject(uint32 ObjectIndex, const FNetObjectFilteringInfo&) override;
+	IRISCORE_API virtual void UpdateObjects(FNetObjectFilterUpdateParams&) override;
+	IRISCORE_API virtual void Filter(FNetObjectFilteringParams&) override;
+};
diff --git a/Engine/Source/Runtime/Experimental/Iris/Core/Public/Iris/ReplicationSystem/NetBlob/ShrinkWrapNetBlob.h b/Engine/Source/Runtime/Experimental/Iris/Core/Public/Iris/ReplicationSystem/NetBlob/ShrinkWrapNetBlob.h
new file mode 100644
--- /dev/null
+++ b/Engine/Source/Runtime/Experimental/Iris/Core/Public/Iris/ReplicationSystem/NetBlob/ShrinkWrapNetBlob.h
@@ -0,0 +1,39 @@
+// Copyright Epic Games, Inc. All Rights Reserved.
+
+#pragma once
+#include "CoreTypes.h"
+#include "Iris/ReplicationSystem/NetBlob/NetBlob.h"
+
+namespace UE::Net
+{
+
+/**
+ * A ShrinkWrapNetBlob is typically used on the sending side for data with multiple destinations.
+ * In that case the contents of the original blob can be serialized, once, to a buffer and then wrapped
+ * in an instance of this class. The serialization of the buffer is likely faster than the original
+ * serialization as no particular logic needs to be performed and serializing a buffer is an optimized path.
+ * @note If tracing is enabled the OriginalBlob will be serialized instead of the already serialized buffer.
+ *       This is for debugging purposes.
+ * @note Deserialization will always be performed by the original blob type.
+ */
+class FShrinkWrapNetBlob final : public FNetBlob
+{
+public:
+	IRISCORE_API FShrinkWrapNetBlob(const TRefCountPtr<FNetBlob>& OriginalBlob, TArray<uint32>&& Payload, uint32 PayloadBitCount);
+
+private:
+	virtual TArrayView<const FNetObjectReference> GetExports() const override final;
+	virtual void SerializeWithObject(FNetSerializationContext& Context, FNetHandle NetHandle) const override final;
+	virtual void DeserializeWithObject(FNetSerializationContext& Context, FNetHandle NetHandle) override final;
+
+	virtual void Serialize(FNetSerializationContext& Context) const override final;
+	virtual void Deserialize(FNetSerializationContext& Context) override final;
+
+	void InternalSerialize(FNetSerializationContext& Context) const;
+
+	TRefCountPtr<FNetBlob> OriginalBlob;
+	TArray<uint32> SerializedBlob;
+	uint32 SerializedBlobBitCount;
+};
+
+}
diff --git a/Engine/Source/Runtime/Experimental/Iris/Core/Public/Iris/ReplicationSystem/NetHandle.h b/Engine/Source/Runtime/Experimental/Iris/Core/Public/Iris/ReplicationSystem/NetHandle.h
new file mode 100644
--- /dev/null
+++ b/Engine/Source/Runtime/Experimental/Iris/Core/Public/Iris/ReplicationSystem/NetHandle.h
@@ -0,0 +1,123 @@
+// Copyright Epic Games, Inc. All Rights Reserved.
+
+#pragma once
+
+#include "CoreTypes.h"
+#include "Containers/UnrealString.h"
+#include "HAL/Platform.h"
+#include "Iris/IrisConfig.h"
+#include "Misc/StringBuilder.h"
+#include "Templates/TypeHash.h"
+
+// Forward declarations
+namespace UE::Net::Private
+{
+	class FNetHandleManager;
+}
+
+namespace UE::Net
+{
+
+/**
+* FNetHandle
+*/
+
+class FNetHandle
+{
+public:
+	enum { Invalid = 0 };
+	enum { IdBits = 32 };
+	enum { ReplicationSystemIdBits = 4 };
+
+public:
+	FNetHandle() : Value(Invalid) {}
+
+	uint32 GetId() const { return Id; }
+	uint32 GetReplicationSystemId() const { check(ReplicationSystemId != 0U); return ReplicationSystemId - 1U; }
+	bool IsValid() const { return Value != Invalid; }
+
+	bool IsCompleteHandle() const { return Value != Invalid && ReplicationSystemId != 0U; }
+
+	bool IsStatic() const { return Id & StaticIdMask; }
+	bool IsDynamic() const { return IsValid() && !IsStatic(); }
+
+	bool operator==(const FNetHandle& Other)const { return Id == Other.Id; }
+	bool operator<(const FNetHandle& Other)const { return Id < Other.Id; }
+	bool operator!=(const FNetHandle& Other)const { return Id != Other.Id; }
+
+	FString ToString() const;
+
+	static bool FullCompare(FNetHandle A, FNetHandle B) { return A.Value == B.Value; }
+	
+private:
+	static uint32 MakeNetHandleId(uint32 Seed, bool bIsStatic);
+	static FNetHandle MakeNetHandle(uint32 Id, uint32 ReplicationSystemId);
+	static FNetHandle MakeNetHandleFromId(uint32 Id);
+
+	friend uint32 GetTypeHash(const FNetHandle& Handle);
+	friend Private::FNetHandleManager;
+
+	static constexpr uint32 StaticIdMask = 1U;
+	static constexpr uint32 IdMask = ~0U;
+	static constexpr uint32 MaxReplicationSystemId = (1 << ReplicationSystemIdBits) - 1;
+
+	union 
+	{
+		struct
+		{
+			uint32 Id : IdBits;										// Id, lowest bit indicates if the handle is static or dynamic
+			uint32 ReplicationSystemId : ReplicationSystemIdBits;	// ReplicationSystemId, when running in pie, we track the owning instance
+		};
+		uint64 Value;
+	};
+};
+
+FORCEINLINE uint32 GetTypeHash(const FNetHandle& Handle)
+{
+	// Id and type is relevant
+	return ::GetTypeHash(Handle.GetId());
+}
+
+FORCEINLINE uint32 GetObjectIdForNetTrace(const FNetHandle& Handle)
+{
+	return Handle.GetId();
+}
+
+inline FString FNetHandle::ToString() const
+{
+	FString Result;
+#if UE_NET_ALLOW_MULTIPLE_REPLICATION_SYSTEMS
+	const uint32 ReplicationSystemIdToDisplay = ReplicationSystemId - 1U;
+	Result = FString::Printf(TEXT("NetHandle (Id=%u):(RepSystemId=%u)"), GetId(), ReplicationSystemIdToDisplay);
+#else
+	Result = FString::Printf(TEXT("NetHandle (Id=%u)"), GetId());
+#endif
+	return Result;
+}
+
+}
+
+inline FStringBuilderBase& operator<<(FStringBuilderBase& Builder, const UE::Net::FNetHandle& NetHandle) 
+{ 	
+#if UE_NET_ALLOW_MULTIPLE_REPLICATION_SYSTEMS
+	const uint32 ReplicationSystemIdToDisplay = NetHandle.GetReplicationSystemId() - 1U;
+	Builder.Appendf(TEXT("NetHandle (Id=%u):(RepSystemId=%u)"), NetHandle.GetId(), ReplicationSystemIdToDisplay);
+#else
+	Builder.Appendf(TEXT("NetHandle (Id=%u)"), NetHandle.GetId());
+#endif
+	return Builder;
+}
+
+inline FAnsiStringBuilderBase& operator<<(FAnsiStringBuilderBase& Builder, const UE::Net::FNetHandle& NetHandle)
+{
+	#if UE_NET_ALLOW_MULTIPLE_REPLICATION_SYSTEMS
+		const uint32 ReplicationSystemIdToDisplay = NetHandle.GetReplicationSystemId() - 1U;
+		Builder.Appendf("NetHandle (Id=%u):(RepSystemId=%u)", NetHandle.GetId(), ReplicationSystemIdToDisplay);
+	#else
+		Builder.Appendf("NetHandle (Id=%u)", NetHandle.GetId());
+	#endif
+		return Builder;
+}
+
+
+
diff --git a/Engine/Source/Runtime/Experimental/Iris/Core/Public/Iris/ReplicationSystem/NetToken.h b/Engine/Source/Runtime/Experimental/Iris/Core/Public/Iris/ReplicationSystem/NetToken.h
new file mode 100644
--- /dev/null
+++ b/Engine/Source/Runtime/Experimental/Iris/Core/Public/Iris/ReplicationSystem/NetToken.h
@@ -0,0 +1,92 @@
+// Copyright Epic Games, Inc. All Rights Reserved.
+
+#pragma once
+
+#include "CoreTypes.h"
+#include "Containers/UnrealString.h"
+#include "Templates/TypeHash.h"
+
+namespace UE::Net
+{
+	class FNetBitStreamReader;
+	class FNetBitStreamWriter;
+}
+
+namespace UE::Net
+{
+
+class FNetToken
+{
+public:
+	typedef uint32 FTypeId;
+
+	enum : uint32 { InvalidTokenTypeId = ~FTypeId(0) };
+	enum : uint32 { InvalidTokenIndex = 0U };
+	enum : uint32 { TokenTypeIdBits = 2U };
+	enum : uint32 { TokenBits = 20U };
+	enum : uint32 { MaxTypeIdCount = 1U << TokenTypeIdBits };
+	enum : uint32 { MaxNetTokenCount = 1U << TokenBits };
+
+public:	
+	FNetToken() : Index(InvalidTokenIndex) {}
+	inline bool IsValid() const { return Index != InvalidTokenIndex; }
+	uint32 GetIndex() const { return Index; }
+	bool operator==(const FNetToken& Other) const { return Index == Other.Index; }
+	FString ToString() const;
+	static FNetToken MakeNetToken(uint32 Index) { check(Index < MaxNetTokenCount); return FNetToken(Index); }
+
+private:
+	explicit FNetToken(uint32 InIndex) : Index(InIndex) {}
+
+private:
+	friend IRISCORE_API FNetToken ReadNetToken(FNetBitStreamReader*);
+
+	uint32 Index;
+};
+
+class FNetTokenStoreKey
+{
+public:
+	enum { Invalid = 0U };
+
+	FNetTokenStoreKey() : Value(Invalid) {}
+	bool IsValid() const { return Value != Invalid; }
+	uint32 GetKeyValue() const { return Key; }
+	FNetToken::FTypeId GetTypeId() const { return TypeId; }
+	bool operator==(const FNetTokenStoreKey& Other) const { return Value == Other.Value; }
+
+private:
+	friend class FNetTokenStore;
+	friend class FNetTokenDataStore;
+
+private:
+	union
+	{
+		struct
+		{
+			uint32 TypeId : FNetToken::TokenTypeIdBits;
+			uint32 Key : FNetToken::TokenBits;
+		};
+		uint32 Value;
+	};
+};
+
+static_assert(sizeof(FNetTokenStoreKey) == sizeof(uint32), "FNetTokenKey should fit in a uint32");
+
+FORCEINLINE uint32 GetTypeHash(const FNetToken& Token)
+{
+	return ::GetTypeHash(Token.GetIndex());
+}
+
+inline FString FNetToken::ToString() const
+{
+	FString Result;
+	Result = FString::Printf(TEXT("NetToken (Index=%u)"), Index);
+	return Result;
+}
+
+// Read and write tokens
+IRISCORE_API FNetToken ReadNetToken(FNetBitStreamReader* Reader);
+IRISCORE_API void WriteNetToken(FNetBitStreamWriter* Writer, FNetToken Token);
+
+}
