Iris experimental replication system.

diff --git a/Engine/Source/Runtime/Experimental/Iris/Core/Public/Iris/ReplicationSystem/NetTokenStore.h b/Engine/Source/Runtime/Experimental/Iris/Core/Public/Iris/ReplicationSystem/NetTokenStore.h
new file mode 100644
--- /dev/null
+++ b/Engine/Source/Runtime/Experimental/Iris/Core/Public/Iris/ReplicationSystem/NetTokenStore.h
@@ -0,0 +1,98 @@
+// Copyright Epic Games, Inc. All Rights Reserved.
+
+#pragma once
+
+#include "CoreTypes.h"
+#include "NetToken.h"
+
+class UNetTokenDataStream;
+namespace UE::Net
+{
+	class FNetTokenDataStore;
+	class FNetTokenStore;
+	class FNetTokenStoreState;
+	class FNetSerializationContext;
+}
+
+namespace UE::Net
+{
+
+// virtual interface for NetTokenDataStores
+class FNetTokenDataStore
+{
+public:
+	
+	virtual ~FNetTokenDataStore();
+
+protected:
+	FNetTokenDataStore();
+
+	virtual void WriteTokenData(FNetSerializationContext& Context, FNetTokenStoreKey Key) const = 0;
+	virtual FNetTokenStoreKey ReadTokenData(FNetSerializationContext& Context) = 0;
+
+	FNetTokenStoreKey GetTokenKey(FNetToken Token, const FNetTokenStoreState& TokenStoreState) const;
+	inline FNetToken::FTypeId GetTypeId() const { return TypeId; }
+
+	// Create new NetToken
+	FNetToken CreateToken(FNetTokenStoreKey Key, FNetTokenStoreState& TokenStoreState);
+
+	// Make NetTokenStoreKey
+	static FNetTokenStoreKey MakeNetTokenStoreKey(FNetToken::FTypeId TokenTypeId, uint32 TokenKeyValue);
+
+private:
+	friend FNetTokenStore;
+	FNetToken::FTypeId TypeId;
+};
+
+// This is the token store, we have one per game instance
+class FNetTokenStore
+{
+	UE_NONCOPYABLE(FNetTokenStore);
+public:
+	using FNetTokenStoreState = UE::Net::FNetTokenStoreState;
+
+	FNetTokenStore();
+	~FNetTokenStore();
+
+	// Register DataStore and return true if it was registered.
+	bool RegisterDataStore(FNetTokenDataStore* DataStore);
+
+	const FNetTokenStoreState* GetLocalNetTokenStoreState() const { return LocalNetTokenStoreState; }
+	FNetTokenStoreState* GetLocalNetTokenStoreState() { return LocalNetTokenStoreState; }
+
+	// Write data associated with the NetToken
+	void WriteTokenData(FNetSerializationContext& Context, const FNetToken& NetToken) const;
+
+	// Read data associated with the NetToken
+	void ReadTokenData(FNetSerializationContext& Context, const FNetToken& NetToken, FNetTokenStoreState& RemoteNetTokenStoreState);
+
+private:
+
+	// Write token data for the TokenIndex
+	void WriteTokenDataForIndex(FNetSerializationContext& Context, uint32 TokenIndex) const;
+
+	// Read token data for the TokenIndex
+	void ReadTokenDataForIndex(FNetSerializationContext& Context, uint32 TokenIndex, FNetTokenStoreState& Resolver);
+
+private:
+	friend UNetTokenDataStream;
+	friend FNetTokenDataStore;
+
+	FNetTokenStoreState* LocalNetTokenStoreState;
+	TArray<FNetTokenDataStore*> TokenDataStores;
+};
+
+inline FNetTokenStoreKey FNetTokenDataStore::MakeNetTokenStoreKey(FNetToken::FTypeId TokenTypeId, uint32 TokenKeyValue)
+{
+	check(TokenTypeId < FNetToken::MaxTypeIdCount);
+	check(TokenKeyValue < FNetToken::MaxNetTokenCount)
+
+	FNetTokenStoreKey TokenKey;
+	TokenKey.TypeId = TokenTypeId;
+	TokenKey.Key = TokenKeyValue;
+
+	return TokenKey;
+}
+
+
+}
diff --git a/Engine/Source/Runtime/Experimental/Iris/Core/Public/Iris/ReplicationSystem/ObjectReplicationBridge.h b/Engine/Source/Runtime/Experimental/Iris/Core/Public/Iris/ReplicationSystem/ObjectReplicationBridge.h
new file mode 100644
--- /dev/null
+++ b/Engine/Source/Runtime/Experimental/Iris/Core/Public/Iris/ReplicationSystem/ObjectReplicationBridge.h
@@ -0,0 +1,293 @@
+// Copyright Epic Games, Inc. All Rights Reserved.
+
+#pragma once
+
+#include "CoreTypes.h"
+#include "Net/Core/NetBitArray.h"
+#include "Iris/Core/NetObjectReference.h"
+#include "Iris/ReplicationSystem/ReplicationBridge.h"
+#include "Iris/ReplicationState/ReplicationStateDescriptor.h"
+#include "Iris/ReplicationSystem/ReplicationFragment.h"
+
+#include "ObjectReplicationBridge.generated.h"
+
+namespace UE::Net
+{
+	struct FNetObjectResolveContext;
+	typedef uint32 FNetObjectFilterHandle;
+	typedef uint16 FNetObjectGroupHandle;
+	typedef uint32 FNetObjectPrioritizerHandle;
+	class FNetObjectReference;
+	namespace Private
+	{
+		typedef uint32 FInternalNetHandle;
+		class FObjectPollFrequencyLimiter;
+	}
+}
+
+/*
+* Partial implementation of ReplicationBridge that can be used as a foundation for 
+* implementing support for replicating objects derived from UObject
+*/
+UCLASS(Transient, MinimalApi)
+class UObjectReplicationBridge : public UReplicationBridge
+{
+	GENERATED_BODY()
+
+public:
+	using UReplicationBridge::EndReplication;
+
+	struct FCreateNetHandleParams
+	{
+		uint32 bCanReceive : 1;
+		uint32 bNeedsPreUpdate : 1;
+		uint32 bNeedsWorldLocationUpdate : 1;
+		uint32 bAllowDynamicFilter : 1;
+
+		/**
+		  * If StaticPriority is > 0 the ReplicationSystem will use that as priority when scheduling objects.
+		  * If it's <= 0.0f one will look for a world location support and then use the default spatial prioritizer.
+		  */
+		float StaticPriority;
+
+		/**
+		  * How often the object should be polled for dirtiness, including calling the InstancePreUpdate function.
+		  * The period is zero based so 0 means poll every frame and 255 means poll every 256th frame.
+		  */
+		uint8 PollFramePeriod;
+	};
+	IRISCORE_API static FCreateNetHandleParams DefaultCreateNetHandleParams;
+
+	IRISCORE_API UObjectReplicationBridge();
+
+	/** Get the Object from a replicated handle, if the handle is invalid or not is a replicated handle the function will return nullptr */
+	IRISCORE_API UObject* GetReplicatedObject(FNetHandle Handle) const;
+
+	/** Get NetHandle from a Replicated UObject */
+	IRISCORE_API FNetHandle GetReplicatedHandle(const UObject* Object) const;
+
+	/** Try to resolve UObject from NetObjectReference, this function tries to resolve the object by loading if necessary.*/
+	IRISCORE_API UObject* ResolveObjectReference(const UE::Net::FNetObjectReference& ObjectRef, const UE::Net::FNetObjectResolveContext& ResolveContext);
+
+	/** Get or create NetObjectReference for object instance. */
+	IRISCORE_API UE::Net::FNetObjectReference GetOrCreateObjectReference(const UObject* Instance) const;
+
+	/** Get or create NetObjectReference for object identified by path relative to outer. */
+	IRISCORE_API UE::Net::FNetObjectReference GetOrCreateObjectReference(const FString& Path, const UObject* Outer) const;
+
+	/** Create NetHandle and start replicating the Instance */
+	IRISCORE_API FNetHandle BeginReplication(UObject* Instance, const FCreateNetHandleParams& Params = DefaultCreateNetHandleParams);
+
+	/**
+	 * Create NetHandle and start replicating the Instance as a SubObject of the OwnerHandle. If InsertRelativeSubObjectHandle is valid
+	 * the new subobject will be inserted in the subobject replication list next to the specified handle and the wanted insertion order.
+	 * Default behavior is to always add new subobjects at the end of the list.
+	 */
+	IRISCORE_API FNetHandle BeginReplication(FNetHandle OwnerHandle, UObject* Instance, FNetHandle InsertRelativeToSubObjectHandle, ESubObjectInsertionOrder InsertionOrder = UReplicationBridge::ESubObjectInsertionOrder::None, const FCreateNetHandleParams& Params = DefaultCreateNetHandleParams);
+
+	/** Create NetHandle and start replicating the Instance as a SubObject of the OwnerHandle */
+	FNetHandle BeginReplication(FNetHandle OwnerHandle, UObject* Instance, const FCreateNetHandleParams& Params = DefaultCreateNetHandleParams) { return BeginReplication(OwnerHandle, Instance, FNetHandle(), ESubObjectInsertionOrder::None, Params); };
+
+	/** 
+	 * Set NetCondition for a subobject, the condition is used to determine if the SubObject should replicated or not.
+	 * @note: As the filtering is done at the serialization level it is typically more efficient to use a separate NetObject for connection 
+	 * specific data as filtering can then be done at a higher level.
+	 */
+	IRISCORE_API void SetSubObjectNetCondition(FNetHandle SubObjectHandle, ELifetimeCondition Condition);
+
+	/** Stop replicating the object. */
+	IRISCORE_API void EndReplication(UObject* Instance, EEndReplicationFlags EndReplicationFlags = EEndReplicationFlags::Destroy, FEndReplicationParameters* Parameters = nullptr);
+
+	/** Add static destruction info, this is used when stably named objects are destroyed prior to starting replication */
+	IRISCORE_API void AddStaticDestructionInfo(const FString& ObjectPath, const UObject* Outer, const FEndReplicationParameters& Parameters);
+	
+	/** PIE package name remapping support. */
+	virtual bool RemapPathForPIE(FString& Path, bool bReading) const { return false; }
+
+	/** Returns true of the level that the Object belongs to has finished loading. */
+	virtual bool ObjectLevelHasFinishedLoading(UObject* Object) const { return true; }
+
+	struct FCreationHeader
+	{
+		virtual ~FCreationHeader() {};
+	};
+
+	// Dependent objects supprt 
+
+	enum class EDependentWarnFlags : uint32
+	{
+		None = 1 << 0U,
+		WarnAlreadyDependent = 1 << None, 
+	};
+	
+	/**
+	 * Adds a dependent object. A dependent object can replicate separately or if a parent replicates.
+	 * Dependent objects cannot be filtered out by dynamic filtering unless the parent is also filtered out.
+	 * @note: There is no guarantee that the data will end up in the same packet so it is a very loose form of dependency.
+	 */
+	IRISCORE_API void AddDependentObject(FNetHandle Parent, FNetHandle DependentObject, EDependentWarnFlags WarnFlags = EDependentWarnFlags::WarnAlreadyDependent);
+
+	/** Remove dependent object from parent. The dependent object will function as a standard standalone replicated object. */
+	IRISCORE_API void RemoveDependentObject(FNetHandle Parent, FNetHandle DependentObject);
+
+	// Dormancy support
+
+	/** Set whether object should go dormant. If dormancy is enabled any dirty state will be replicated first. */
+	IRISCORE_API void SetObjectWantsToBeDormant(FNetHandle Handle, bool bWantsToBeDormant);
+
+	/** Returns whether the object wants to be dormant. */
+	IRISCORE_API bool GetObjectWantsToBeDormant(FNetHandle Handle) const;
+
+	/** Trigger replication of dirty state for object wanting to be dormant. */
+	IRISCORE_API void ForceUpdateWantsToBeDormantObject(FNetHandle Handle);	
+
+protected:
+	IRISCORE_API virtual ~UObjectReplicationBridge();
+
+	// UReplicationBridge
+
+	IRISCORE_API virtual void Initialize(UReplicationSystem* InReplicationSystem) override;
+	IRISCORE_API virtual void Deinitialize() override;
+	IRISCORE_API virtual void PreSendUpdateSingleHandle(FNetHandle Handle) override;	
+	IRISCORE_API virtual void PreSendUpdate() override;	
+	IRISCORE_API virtual void UpdateInstancesWorldLocation() override;
+	IRISCORE_API virtual void PruneStaleObjects() override;	
+	IRISCORE_API virtual bool WriteNetHandleCreationInfo(FReplicationBridgeSerializationContext& Context, FNetHandle Handle) override;
+	IRISCORE_API virtual FNetHandle CreateNetHandleFromRemote(FNetHandle SubObjectOwnerNetHandle, FNetHandle WantedNetHandle, FReplicationBridgeSerializationContext& Context) override;
+	IRISCORE_API virtual void PostApplyInitialState(FNetHandle Handle) override;
+	IRISCORE_API virtual void DetachInstanceFromRemote(FNetHandle Handle, bool bTearOff, bool bShouldDestroyInstance) override;
+	IRISCORE_API virtual void DetachInstance(FNetHandle Handle) override;
+
+protected:
+	/**
+	* $IRIS TODO:
+	* Consider moving these methods to separate interface to allow for a single ObjectReplicationBridge to support different type specific methods to its own interface and allow for a single bridge to support multiple subtypes? 
+	* We will see what we need when we start to Serialize data we might not even want to expose virtual methods for doing this, and instead provide support for specifying a "creation state" that works exactly as any other state
+	* but only ever is replicated when we first instantiate an object of the specific type 
+	*/
+
+	/** Write data required to instantiate the Handle. */
+	virtual bool WriteCreationHeader(UE::Net::FNetSerializationContext& Context, FNetHandle Handle) { return false; };
+
+	/** Create Header and read data required to instantiate the Handle. */
+	virtual FCreationHeader* ReadCreationHeader(UE::Net::FNetSerializationContext& Context) { return nullptr; };
+
+	/** Called when we instantiate/find object instance requested by remote. */
+	virtual UObject* BeginInstantiateFromRemote(FNetHandle SubObjectOwnerHandle, const UE::Net::FNetObjectResolveContext& ResolveContext, const FCreationHeader* Header) { return nullptr; };
+
+	/** Invoked before we start applying state data to instance on remote end. */
+	virtual bool OnInstantiatedFromRemote(UObject* Instance, const FCreationHeader* InHeader, uint32 ConnectionId) const { return true; }
+
+	/** Invoked after remote NetHandle has been created and initial state is applied. */
+	virtual void EndInstantiateFromRemote(FNetHandle Handle) {};
+
+	/** Destroy or tear-off the game instance on request from remote. */
+	virtual void DestroyInstanceFromRemote(UObject* Instance, bool bTearOff) {};
+
+protected:
+	/** Lookup the UObject associated with the provided Handle. This function will not try to resolve the reference. */
+	IRISCORE_API UObject* GetObjectFromReferenceHandle(FNetHandle NetHandle) const;
+
+	/** Helper method that calls provided PreUpdateFunction and polls state data for all replicated instances with the NeedsPoll trait. */
+	using FInstancePreUpdateFunction = void(*)(FNetHandle, UObject*, const UReplicationBridge*);
+
+	/** Set the function that we should call before copying state data. */
+	IRISCORE_API void SetInstancePreUpdateFunction(FInstancePreUpdateFunction InPreUpdateFunction);
+	
+	/** Helper method to get the world location for replicated instances with the HasWorldLocation trait. */
+	using FInstanceGetWorldLocationFunction = FVector(*)(FNetHandle, const UObject*);
+
+	/** Set the function that we should call to get the world location of an object. */
+	IRISCORE_API void SetInstanceGetWorldLocationFunction(FInstanceGetWorldLocationFunction InGetWorldLocationFunction);
+	
+	// Poll frequency support
+
+	/** Set poll frame period for a specific object. Zero based period where 0 means every frame and 255 means every 256th frame. */
+	IRISCORE_API void SetPollFramePeriod(UE::Net::Private::FInternalNetHandle InternalReplicationIndex, uint8 FramePeriod);
+
+	/** Force polling of Object when ObjectToPoll is polled. */
+	IRISCORE_API void SetPollWithObject(FNetHandle ObjectToPollWith, FNetHandle Object);
+
+	/** Re-initialize config-driven parameters found in ObjectReplicationBridgeConfig. */
+	IRISCORE_API void LoadConfig();
+
+	/**
+	 * Set the function used to determine whether an object should use the default spatial filter
+	 * unless the class is configured to use some other filter.
+	 */
+	IRISCORE_API void SetShouldUseDefaultSpatialFilterFunction(TFunction<bool(const UClass*)>);
+
+	/** Set the function used to determine whether two classes are to be considered equal when it comes to filtering. Used on subclasses. */
+	IRISCORE_API void SetShouldSubclassUseSameFilterFunction(TFunction<bool(const UClass* Class, const UClass* Subclass)>);
+
+private:
+
+	void PreUpdateAndPollImpl(FNetHandle Handle);
+
+	void RegisterRemoteInstance(FNetHandle Handle, UObject* InstancePtr, const UE::Net::FReplicationProtocol* Protocol, UE::Net::FReplicationInstanceProtocol* InstanceProtocol, const FCreationHeader* Header, uint32 ConnectionId);
+
+	void UnregisterRemoteInstance(FNetHandle Handle, bool bTearOff, bool bShouldDestroyInstance);
+
+	/** Finds a poll period override if it exists. OutPollPeriod will only be modified if this method returns true. */
+	bool FindPollInfo(const UClass* Class, uint8& OutPollPeriod);
+
+	/** Tries to load the classes used in poll period overrides. */
+	void FindClassesInPollPeriodOverrides();
+
+	/** Retrieves the dynamic filter to set for the given class. Will return an invalid handle if no dynamic filter should be set. */
+	UE::Net::FNetObjectFilterHandle GetDynamicFilter(const UClass* Class);
+
+	/** Retrieves the prioritizer to set for the given class. Will return an invalid handle if no prioritizer should be set. */
+	UE::Net::FNetObjectPrioritizerHandle GetPrioritizer(const UClass* Class);
+
+	/** Returns true if instances of this class should be delta compressed */
+	bool ShouldClassBeDeltaCompressed(const UClass* Class);
+
+	FInstancePreUpdateFunction PreUpdateInstanceFunction;
+	FInstanceGetWorldLocationFunction GetInstanceWorldLocationFunction;
+
+	FName GetConfigClassPathName(const UClass* Class);
+
+private:
+
+	TMap<const UClass*, FName> ConfigClassPathNameCache;
+
+	// Polling
+	struct FPollInfo
+	{
+		uint8 PollFramePeriod = 0;
+		TWeakObjectPtr<const UClass> Class;
+	};
+	UE::Net::Private::FObjectPollFrequencyLimiter* PollFrequencyLimiter;
+
+	// Class hierarchies with poll period overrides
+	TMap<FName, FPollInfo> ClassHierarchyPollPeriodOverrides;
+	// Exact classes with poll period overrides
+	TMap<FName, FPollInfo> ClassesWithPollPeriodOverride;
+	// Exact classes without poll period override
+	TSet<FName> ClassesWithoutPollPeriodOverride;
+
+	// Filter mapping
+	TMap<FName, UE::Net::FNetObjectFilterHandle> ClassesWithDynamicFilter;
+	TFunction<bool(const UClass*)> ShouldUseDefaultSpatialFilterFunction;
+	TFunction<bool(const UClass*,const UClass*)> ShouldSubclassUseSameFilterFunction;
+
+	// Prioritizer mapping
+	TMap<FName, UE::Net::FNetObjectPrioritizerHandle> ClassesWithPrioritizer;
+
+	// Delta compression
+	TMap<FName, bool> ClassesWithDeltaCompression;
+
+	// Array of dormant objects that has requested a flush
+	TArray<FNetHandle> DormantHandlesPendingFlush;
+
+	// Objects which has object references and could be affected by garbage collection.
+	UE::Net::FNetBitArray ObjectsWithObjectReferences;
+	// Objects that needs to update their object references due to garbage collection.
+	UE::Net::FNetBitArray GarbageCollectionAffectedObjects;
+
+	UE::Net::FNetObjectFilterHandle DefaultSpatialFilterHandle;
+
+	bool bHasPollOverrides = false;
+	bool bHasDirtyClassesInPollPeriodOverrides = false;
+};
diff --git a/Engine/Source/Runtime/Experimental/Iris/Core/Public/Iris/ReplicationSystem/Private/FastArrayReplicationFragmentInternal.h b/Engine/Source/Runtime/Experimental/Iris/Core/Public/Iris/ReplicationSystem/Private/FastArrayReplicationFragmentInternal.h
new file mode 100644
--- /dev/null
+++ b/Engine/Source/Runtime/Experimental/Iris/Core/Public/Iris/ReplicationSystem/Private/FastArrayReplicationFragmentInternal.h
@@ -0,0 +1,280 @@
+// Copyright Epic Games, Inc. All Rights Reserved.
+
+#pragma once
+
+#include "Iris/ReplicationSystem/ReplicationFragment.h"
+#include "Iris/Core/IrisProfiler.h"
+#include "Iris/ReplicationState/IrisFastArraySerializer.h"
+#include "Iris/ReplicationState/Private/IrisFastArraySerializerInternal.h"
+#include "Iris/ReplicationState/PropertyReplicationState.h"
+#include "Iris/ReplicationState/ReplicationStateDescriptor.h"
+
+namespace UE {
+namespace Net {
+
+namespace FastArrayPollingPolicies {
+
+struct FPollingState
+{
+	int32 ArrayReplicationKey = -1;
+	struct FEntry
+	{
+		int32 ReplicationKey = -1;
+		int32 ReplicationID = -1;
+	};
+	TArray<FEntry>  ItemPollData;
+};
+
+/*
+ * FNeedPollingPolicy, used for Native Iris FastArrays that does not need polling
+ */
+class FNeedPollingPolicy
+{
+public:
+	FPollingState* GetPollingState() { return &PollingState; }	
+private:
+	FPollingState PollingState;
+};
+
+/*
+ * FNeedPollingPolicy, used for Native Iris FastArrays that do not require polling
+ */
+class FNoPollingPolicy
+{
+public:
+	FPollingState* GetPollingState() { return nullptr; }
+};
+
+} // End of namespace FastArrayPollingPolicies
+
+namespace Private {
+
+/** Utility methods to behave similar to FastArraySerializer */
+struct FFastArrayReplicationFragmentHelper
+{
+	/** Create descriptor for the FastArray property if it does not already exist */
+	IRISCORE_API static TRefCountPtr<const FReplicationStateDescriptor> GetOrCreateDescriptorForFastArrayProperty(UObject* Object, FFragmentRegistrationContext& Context, int32 RepIndex);
+
+	/** Rebuild IndexMap for FastArrraySerializer */
+	template <typename FastArrayType, typename ItemArrayType>
+	static void ConditionalRebuildItemMap(FastArrayType& ArraySerializer, const ItemArrayType& Items);
+
+	/** Apply received state and try to behave like current FastArrays */
+	template <typename FastArrayType, typename ItemArrayType>
+	static void ApplyReplicatedState(FastArrayType* DstFastArray, ItemArrayType* DstWrappedArray, FastArrayType* SrcFastArray, const ItemArrayType* SrcWrappedArray, const FReplicationStateDescriptor* ArrayElementDescriptor, FReplicationStateApplyContext& Context);
+
+	/** Copy array element, only replicated items will be copied */
+	IRISCORE_API static void InternalCopyArrayElement(const FReplicationStateDescriptor* ArrayElementDescriptor, void* RESTRICT Dst, const void* RESTRICT Src);
+
+	/** Compare array element, only replicated items will be compared */
+	IRISCORE_API static bool InternalCompareArrayElement(const FReplicationStateDescriptor* ArrayElementDescriptor, void* RESTRICT Dst, const void* RESTRICT Src);
+
+	/**
+	 * Conditionally invoke PostReplicatedReceive method depending on if it is defined or not
+	 * We only want to do this for FastArrays that define PostReplicatedReceive since it might require extra work to calculate the required parameters
+	 */
+	template<typename FastArrayType>
+	static inline typename TEnableIf<TModels<FFastArraySerializer::CPostReplicatedReceiveFuncable, FastArrayType, const FFastArraySerializer::FPostReplicatedReceiveParameters&>::Value, void>::Type CallPostReplicatedReceiveOrNot(FastArrayType& ArraySerializer, bool bHasUnresolvedReferences)
+	{
+		FFastArraySerializer::FPostReplicatedReceiveParameters PostReceivedParameters;
+		PostReceivedParameters.bHasMoreUnmappedReferences = bHasUnresolvedReferences;
+		ArraySerializer.PostReplicatedReceive(PostReceivedParameters);
+	}
+
+	template<typename FastArrayType>
+	static inline typename TEnableIf<!TModels<FFastArraySerializer::CPostReplicatedReceiveFuncable, FastArrayType, const FFastArraySerializer::FPostReplicatedReceiveParameters&>::Value, void>::Type CallPostReplicatedReceiveOrNot(FastArrayType& ArraySerializer, bool bHasUnresolvedReferences) {}
+};
+
+class FFastArrayReplicationFragmentBase : public FReplicationFragment
+{
+public:
+	IRISCORE_API void Register(FFragmentRegistrationContext& Context, EReplicationFragmentTraits InTraits);
+
+protected:
+	IRISCORE_API FFastArrayReplicationFragmentBase(EReplicationFragmentTraits InTraits, UObject* InOwner, const FReplicationStateDescriptor* InDescriptor);
+
+	IRISCORE_API const FReplicationStateDescriptor* GetFastArraySerializerPropertyDescriptor() const;
+	IRISCORE_API const FReplicationStateDescriptor* GetArrayElementDescriptor() const;
+
+	IRISCORE_API virtual void CollectOwner(FReplicationStateOwnerCollector* Owners) const override;
+	IRISCORE_API virtual void CallRepNotifies(FReplicationStateApplyContext& Context) override;	
+
+	IRISCORE_API virtual void ReplicatedStateToString(FStringBuilderBase& StringBuilder, FReplicationStateApplyContext& Context, EReplicationStateToStringFlags Flags) const override;
+
+	IRISCORE_API static void InternalCopyArrayElement(const FReplicationStateDescriptor* ArrayElementDescriptor, void* RESTRICT Dst, const void* RESTRICT Src);
+	IRISCORE_API static bool InternalCompareArrayElement(const FReplicationStateDescriptor* ArrayElementDescriptor, void* RESTRICT Dst, const void* RESTRICT Src);
+
+protected:
+	// Replication descriptor built for the specific property
+	TRefCountPtr<const FReplicationStateDescriptor> ReplicationStateDescriptor;
+
+	// This is the source state from which we source our state data
+	TUniquePtr<FPropertyReplicationState> SrcReplicationState;
+
+	// Owner
+	UObject* Owner;
+
+	// This allows us to quickly find the wrapped array relative to the owner
+	SIZE_T WrappedArrayOffsetRelativeFastArraySerializerProperty;
+};
+
+// Native FastArray
+class FNativeFastArrayReplicationFragmentBase : public FReplicationFragment
+{
+protected:
+	IRISCORE_API FNativeFastArrayReplicationFragmentBase(EReplicationFragmentTraits InTraits, UObject* InOwner, const FReplicationStateDescriptor* InDescriptor);
+
+	IRISCORE_API virtual void CollectOwner(FReplicationStateOwnerCollector* Owners) const override;
+
+	IRISCORE_API const FReplicationStateDescriptor* GetFastArraySerializerPropertyDescriptor() const;
+	IRISCORE_API const FReplicationStateDescriptor* GetArrayElementDescriptor() const;
+
+	// Dequantize state into DstExternalBuffer, note it is expected to be initialized
+	IRISCORE_API static void InternalDequantizeFastArray(FNetSerializationContext& Context, uint8* RESTRICT DstExternalBuffer, const uint8* RESTRICT SrcInternalBuffer, const FReplicationStateDescriptor* FastArrayPropertyDescriptor);
+
+	IRISCORE_API static void ToString(FStringBuilderBase& StringBuilder, const uint8* ExternalStateBuffer, const FReplicationStateDescriptor* FastArrayPropertyDescriptor);
+
+protected:
+	// Replication descriptor built for the specific property
+	TRefCountPtr<const FReplicationStateDescriptor> ReplicationStateDescriptor;
+
+	// Owner
+	UObject* Owner;
+
+	// This allows us to quickly find the wrapped array relative to the owner
+	SIZE_T WrappedArrayOffsetRelativeFastArraySerializerProperty;
+};
+
+template <typename FastArrayType, typename ItemArrayType>
+void FFastArrayReplicationFragmentHelper::ConditionalRebuildItemMap(FastArrayType& ArraySerializer, const ItemArrayType& Items)
+{
+	typedef typename ItemArrayType::ElementType ItemType;
+
+	if (ArraySerializer.ItemMap.Num() != Items.Num())
+	{
+		ArraySerializer.ItemMap.Reset();
+			
+		const ItemType* SrcItems = Items.GetData();
+		for (int32 It = 0, EndIt = Items.Num(); It != EndIt; ++It)
+		{
+			const ItemType& Item = SrcItems[It];
+			if (Item.ReplicationID == INDEX_NONE)
+			{
+				continue;
+			}
+
+			ArraySerializer.ItemMap.Add(Item.ReplicationID, It);
+		}
+	}
+}
+
+template <typename FastArrayType, typename ItemArrayType>
+void FFastArrayReplicationFragmentHelper::ApplyReplicatedState(FastArrayType* DstArraySerializer, ItemArrayType* DstWrappedArray, FastArrayType* SrcArraySerializer, const ItemArrayType* SrcWrappedArray, const FReplicationStateDescriptor* ArrayElementDescriptor, FReplicationStateApplyContext& Context)
+{
+	typedef typename ItemArrayType::ElementType ItemType;
+
+	UE_LOG(LogNetFastTArray, Log, TEXT("FFastArrayReplicationFragmentHelper::ApplyReplicatedState for %s"), Context.Descriptor->DebugName->Name);
+
+	// We need to rebuild our maps for both target array and incoming data
+	// Can optimize this later
+	ConditionalRebuildItemMap(*DstArraySerializer, *DstWrappedArray);
+	ConditionalRebuildItemMap(*SrcArraySerializer, *SrcWrappedArray);
+
+	// Find removed elements in received data, that is elements that exist in old map but not in new map
+	TArray<int32> RemovedIndices;
+	{
+		RemovedIndices.Reserve(DstWrappedArray->Num());
+		ItemType* DstItems = DstWrappedArray->GetData();
+		for (int32 It=0, EndIt=DstWrappedArray->Num(); It != EndIt; ++It)
+		{
+			const int32 ReplicationID = DstItems[It].ReplicationID;
+			if (ReplicationID != -1 && !SrcArraySerializer->ItemMap.Contains(ReplicationID))
+			{
+				UE_LOG(LogNetFastTArray, Log, TEXT("   Removed ID: %d local Idx: %d"), ReplicationID, It);
+				RemovedIndices.Add(It);
+
+				// Remove callback
+				DstItems[It].PreReplicatedRemove(*DstArraySerializer);
+			}
+		}
+	}
+	
+	const int32 PreRemoveSize = DstWrappedArray->Num();
+	const int32 FinalSize = PreRemoveSize - RemovedIndices.Num();
+
+	// Remove callback to FastArraySerializer
+	DstArraySerializer->PreReplicatedRemove(MakeArrayView(RemovedIndices), FinalSize);
+
+	// Find new and modified elements in received data, That is elements that do not exist in old map
+	TArray<int32> AddedIndices;
+	TArray<int32> ModifiedIndices;
+	{
+		AddedIndices.Reserve(SrcWrappedArray->Num());
+		ModifiedIndices.Reserve(SrcWrappedArray->Num());
+		const ItemType* SrcItems = SrcWrappedArray->GetData();
+		for (int32 It=0, EndIt=SrcWrappedArray->Num(); It != EndIt; ++It)
+		{			
+			if (int32* ExistingIndex = DstArraySerializer->ItemMap.Find(SrcItems[It].ReplicationID))
+			{
+				// As we currently always send the full array. To be correct, we must compare the element as well before issuing the callback
+				if (!InternalCompareArrayElement(ArrayElementDescriptor, &(*DstWrappedArray)[*ExistingIndex], &SrcItems[It]))
+				{
+					UE_LOG(LogNetFastTArray, Log, TEXT("   Changed. ID: %d -> Idx: %d"), SrcItems[It].ReplicationID, *ExistingIndex);
+
+					ModifiedIndices.Add(*ExistingIndex);
+
+					// We use per element copy since we do not want to overwrite data that is not replicated
+					InternalCopyArrayElement(ArrayElementDescriptor, &(*DstWrappedArray)[*ExistingIndex], &SrcItems[It]);
+
+					// Change callback
+					(*DstWrappedArray)[*ExistingIndex].PostReplicatedChange(*DstArraySerializer);
+				}
+			}
+			else
+			{
+				int32 AddedIndex = DstWrappedArray->Add(SrcItems[It]);
+
+				UE_LOG(LogNetFastTArray, Log, TEXT("   New. ID: %d. New Element! local Idx: %d"), SrcItems[It].ReplicationID, AddedIndex);
+
+				// We need to propagate the ReplicationID in order to find our object
+				(*DstWrappedArray)[AddedIndex].ReplicationID = SrcItems[It].ReplicationID;
+
+				// Add callback
+				(*DstWrappedArray)[AddedIndex].PostReplicatedAdd(*DstArraySerializer);
+
+				// should we store ids or indices?
+				AddedIndices.Add(AddedIndex);
+			}
+		}	
+
+		// Added and changed callbacks to FastArraySerializer
+		DstArraySerializer->PostReplicatedAdd(MakeArrayView(AddedIndices), FinalSize);
+		DstArraySerializer->PostReplicatedChange(MakeArrayView(ModifiedIndices), FinalSize);
+	}
+
+	// Remove indices
+	if (RemovedIndices.Num() > 0)
+	{
+		RemovedIndices.Sort();
+		for (int32 i = RemovedIndices.Num() - 1; i >= 0; --i)
+		{
+			int32 DeleteIndex = RemovedIndices[i];
+			if (DstWrappedArray->IsValidIndex(DeleteIndex))
+			{
+				DstWrappedArray->RemoveAtSwap(DeleteIndex, 1, false);
+			}
+		}
+
+		// Clear the map now that the indices are all shifted around. This kind of sucks, we could use slightly better data structures here I think.
+		// This will force the ItemMap to be rebuilt for the current Items array
+		DstArraySerializer->ItemMap.Empty();
+	}
+
+	// Increment keys so that a client can re-serialize the array if needed, such as for client replay recording.
+	DstArraySerializer->IncrementArrayReplicationKey();
+
+	// Invoke PostReplicatedReceive if is defined by the serializer
+	CallPostReplicatedReceiveOrNot(*DstArraySerializer, Context.bHasUnresolvableReferences);
+}
+
+}}} // End of namespaces
diff --git a/Engine/Source/Runtime/Experimental/Iris/Core/Public/Iris/ReplicationSystem/ReplicationBridge.h b/Engine/Source/Runtime/Experimental/Iris/Core/Public/Iris/ReplicationSystem/ReplicationBridge.h
new file mode 100644
--- /dev/null
+++ b/Engine/Source/Runtime/Experimental/Iris/Core/Public/Iris/ReplicationSystem/ReplicationBridge.h
@@ -0,0 +1,303 @@
+// Copyright Epic Games, Inc. All Rights Reserved.
+
+#pragma once
+
+#include "CoreTypes.h"
+#include "Containers/ContainersFwd.h"
+#include "Containers/Map.h"
+#include "Iris/Core/NetObjectReference.h"
+#include "Iris/ReplicationSystem/NetHandle.h"
+#include "UObject/ObjectKey.h"
+#include "UObject/ObjectMacros.h"
+
+#include "ReplicationBridge.generated.h"
+
+class UObjectReplicationBridge;
+class UReplicationSystem;
+class UNetDriver;
+
+namespace UE::Net
+{
+	struct FNetDependencyInfo;
+	typedef uint16 FNetObjectGroupHandle;
+	class FNetTokenStoreState;
+	class FReplicationFragment;
+	struct FReplicationInstanceProtocol;
+	struct FReplicationProtocol;
+	class FNetSerializationContext;
+	namespace Private
+	{
+		typedef uint32 FInternalNetHandle;
+		class FNetHandleManager;
+		class FNetObjectGroups;
+		struct FNetPushObjectHandle;
+		class FObjectReferenceCache;
+		class FReplicationProtocolManager;
+		class FReplicationReader;
+		class FReplicationStateDescriptorRegistry;
+		class FReplicationSystemImpl;
+		class FReplicationSystemInternal;
+		class FReplicationWriter;
+	}
+
+	typedef TArray<FNetDependencyInfo, TInlineAllocator<32> > FNetDependencyInfoArray;
+}
+
+#if UE_BUILD_SHIPPING
+	IRISCORE_API DECLARE_LOG_CATEGORY_EXTERN(LogIrisBridge, All, Display);
+#else
+	IRISCORE_API DECLARE_LOG_CATEGORY_EXTERN(LogIrisBridge, Log, All);
+#endif
+
+struct FReplicationBridgeSerializationContext
+{
+	FReplicationBridgeSerializationContext(UE::Net::FNetSerializationContext& InSerialiazationContext, uint32 InConnectionId, bool bInIsDestructionInfo = false);
+
+	UE::Net::FNetSerializationContext& SerializationContext;
+	uint32 ConnectionId;
+	bool bIsDestructionInfo;
+};
+
+enum class EEndReplicationFlags : uint32
+{
+	None								= 0U,
+	/** Destroy remote instance. Default for dynamic objects unless they have TearOff flag set. */
+	Destroy								= 1U,				
+	/** Stop replication object without destroying instance on the remote end. */
+	TearOff								= Destroy << 1U,
+	/** Complete replication of pending state to all clients before ending replication. */
+	Flush								= TearOff << 1U,
+};
+ENUM_CLASS_FLAGS(EEndReplicationFlags);
+
+UCLASS(Transient, MinimalAPI)
+class UReplicationBridge : public UObject
+{
+	GENERATED_BODY()
+
+protected:
+	using FNetHandle = UE::Net::FNetHandle;
+	using FNetDependencyInfoArray = UE::Net::FNetDependencyInfoArray;
+
+public:
+	IRISCORE_API UReplicationBridge();
+	IRISCORE_API virtual ~UReplicationBridge();
+
+	// Local interface
+	
+	struct FEndReplicationParameters
+	{
+		/** The location of the object. Used for distance based prioritization. */
+		FVector Location;
+		/** The level the object is placed in. */
+		const UObject* Level = nullptr;
+		/** Whether to use distance based priority for the destruction of the object. */
+		bool bUseDistanceBasedPrioritization = false;
+	};
+
+	enum class ESubObjectInsertionOrder : uint8
+	{
+		None,
+		ReplicateWith,
+	};
+
+	/**
+	 * Stop replicating the NetObject associated with the handle and mark the handle to be destroyed.
+	 * If EEndReplication::TearOff is set the remote instance will be Torn-off rather than being destroyed on the receiving end, after the call, any state changes will not be replicated
+     * If EEndReplication::Flush is set all pending states will be delivered before the remote instance is destroyed, final state will be immediately copied so it is safe to remove the object after this call
+	 * If EEndReplication::Destroy is set the remote instance will be destroyed, if this is set for a static instance and the EndReplicationParameters are set a permanent destruction info will be added
+	 * Dynamic instances are always destroyed unless the TearOff flag is set.
+	 */
+	IRISCORE_API void EndReplication(FNetHandle Handle, EEndReplicationFlags EndReplicationFlags = EEndReplicationFlags::Destroy, FEndReplicationParameters* Parameters = nullptr);
+
+	/** Returns true if the handle is replicated. */
+	IRISCORE_API bool IsReplicatedHandle(FNetHandle Handle) const;
+
+	/** Set the NetDriver used by the bridge. Called during creation if the NetDriver is recreated. */
+	IRISCORE_API virtual void SetNetDriver(UNetDriver* NetDriver);
+
+	/** Get the group associated with the level in order to control connection filtering for it. */
+	IRISCORE_API UE::Net::FNetObjectGroupHandle GetLevelGroup(const UObject* Level) const;
+
+protected:
+	/** Initializes the bridge. Is called during ReplicationSystem initialization. */
+	IRISCORE_API virtual void Initialize(UReplicationSystem* InReplicationSystem);
+
+	/** Deinitializes the bridge. Is called during ReplicationSystem deinitialization. */
+	IRISCORE_API virtual void Deinitialize();
+
+	/** Invoked before ReplicationSystem copies dirty state data. */
+	IRISCORE_API virtual void PreSendUpdate();
+
+	/** Invoked before ReplicationSystem copies dirty state data for a single replicated object. */
+	IRISCORE_API virtual void PreSendUpdateSingleHandle(FNetHandle Handle);
+
+	/** Update world locations in FWorldLocations for objects that support it. */
+	IRISCORE_API virtual void UpdateInstancesWorldLocation();
+
+	// Remote interface, invoked from Replication code during serialization
+	
+	/** Write data required to instantiate NetObject remotely to bitstream. */
+	IRISCORE_API virtual bool WriteNetHandleCreationInfo(FReplicationBridgeSerializationContext& Context, FNetHandle Handle);
+
+	/** Read data required to instantiate NetObject from bitstream. */
+	IRISCORE_API virtual FNetHandle CreateNetHandleFromRemote(FNetHandle SubObjectOwnerNetHandle, FNetHandle WantedNetHandle, FReplicationBridgeSerializationContext& Context);
+
+	/** Invoke after we have applied the initial state for an object.*/
+	IRISCORE_API virtual void PostApplyInitialState(FNetHandle Handle);
+
+	/**
+	 * Called when the instance is detached from the protocol on request by the remote. 
+	 * @param Handle The handle of the object to destroy or tear off.
+	 * @param bTearOff Whether the object should be torn off, i.e. not destroyed. 
+	 * @param bShouldDestroyInstance Whether the object should be destroyed.
+	 */
+	IRISCORE_API virtual void DetachInstanceFromRemote(FNetHandle Handle, bool bTearOff, bool bShouldDestroyInstance);
+
+	/** Called when we detach instance protocol from the local instance */
+	IRISCORE_API virtual void DetachInstance(FNetHandle Handle);
+
+	/** Invoked post garbage collect to allow us to detect stale objects */
+	IRISCORE_API virtual void PruneStaleObjects();
+
+	/** Invoked when we start to replicate an object for a specific connection to fill in any initial dependencies */
+	IRISCORE_API virtual void GetInitialDependencies(FNetHandle Handle, FNetDependencyInfoArray& OutDependencies) const;
+
+protected:
+	// Forward calls to internal operations that we allow replication bridges to access
+
+	/** Create a local NetHandle / NetObject using the ReplicationProtocol. */
+	IRISCORE_API FNetHandle InternalCreateNetObject(FNetHandle AllocatedHandle, const UE::Net::FReplicationProtocol* ReplicationProtocol);
+
+	/** Create a NetHandle / NetObject on request from the authoritative end. */
+	IRISCORE_API FNetHandle InternalCreateNetObjectFromRemote(FNetHandle WantedNetHandle, const UE::Net::FReplicationProtocol* ReplicationProtocol);
+
+	/** Attach instance to NetHandle. */
+	IRISCORE_API void InternalAttachInstanceToNetHandle(FNetHandle Handle, bool bBindInstanceProtocol, UE::Net::FReplicationInstanceProtocol* InstanceProtocol, UObject* Instance);
+
+	// Detach instance from NetObject/Handle and destroy the instance protocol
+	IRISCORE_API void InternalDetachInstanceFromNetHandle(FNetHandle Handle);
+
+	/** Destroy the handle and all internal book keeping associated with it. */
+	IRISCORE_API void InternalDestroyNetObject(FNetHandle Handle);
+	
+	/** Get the owner handle of a subobject handle. */
+	IRISCORE_API FNetHandle InternalGetSubObjectOwner(FNetHandle SubObjectHandle) const;
+
+	/** Add SubObjectHandle as SubObject to OwnerHandle. */
+	IRISCORE_API void InternalAddSubObject(FNetHandle OwnerHandle, FNetHandle SubObjectHandle, FNetHandle InsertRelativeToSubObjectHandle, ESubObjectInsertionOrder InsertionOrder);
+
+	/** Add destruction info for the referenced object. */
+	FNetHandle InternalAddDestructionInfo(FNetHandle Handle, const FEndReplicationParameters& Parameters);
+
+	inline UE::Net::Private::FReplicationProtocolManager* GetReplicationProtocolManager() const { return ReplicationProtocolManager; }
+	inline UReplicationSystem* GetReplicationSystem() const { return ReplicationSystem; }
+	inline UE::Net::Private::FReplicationStateDescriptorRegistry* GetReplicationStateDescriptorRegistry() const { return ReplicationStateDescriptorRegistry; }
+	inline UE::Net::Private::FObjectReferenceCache* GetObjectReferenceCache() const { return ObjectReferenceCache; }
+
+	/** Creates a group for a level for object filtering purposes. */
+	IRISCORE_API UE::Net::FNetObjectGroupHandle CreateLevelGroup(const UObject* Level);
+
+	/** Destroys the group associated with the level. */
+	IRISCORE_API void DestroyLevelGroup(const UObject* Level);
+
+private:
+
+	// Internal operations invoked by ReplicationSystem/ReplicationWriter
+	void ReadAndExecuteDestructionInfoFromRemote(FReplicationBridgeSerializationContext& Context);
+	void DetachSubObjectInstancesFromRemote(FNetHandle Handle, bool bTearOff, bool bShouldDestroyInstance);
+	void DestroyNetObjectFromRemote(FNetHandle Handle, bool bTearOff, bool bDestroyInstance);
+
+	// Adds the Handle to the list of handles pending tear-off, if bIsImmediate is true the object will be destroyed after the next update, otherwise
+	// it will be kept around until EndReplication is called.
+	void TearOff(FNetHandle Handle, bool bIsImmediate);
+
+	FNetHandle CallCreateNetHandleFromRemote(FNetHandle SubObjectOwnerHandle, FNetHandle WantedNetHandle, FReplicationBridgeSerializationContext& Context);
+	void CallPreSendUpdate(float DeltaSeconds);	
+	void CallPreSendUpdateSingleHandle(FNetHandle Handle);
+	void CallUpdateInstancesWorldLocation();
+	bool CallWriteNetHandleCreationInfo(FReplicationBridgeSerializationContext& Context, FNetHandle Handle);
+	void CallPostApplyInitialState(FNetHandle Handle);
+	void CallPruneStaleObjects();
+	void CallGetInitialDependencies(FNetHandle Handle, FNetDependencyInfoArray& OutDependencies) const;
+	void CallDetachInstance(FNetHandle Handle);
+	void CallDetachInstanceFromRemote(FNetHandle Handle, bool bTearOff, bool bShouldDestroyInstance);
+
+private:
+	// Internal method to copy state data for Handle and any SubObjects and mark them as being torn-off
+	void InternalTearOff(FNetHandle OwnerHandle);
+
+	// Destroy all SubObjects owned by provided handle
+	void InternalDestroySubObjects(FNetHandle OwnerHandle);
+
+	/**
+	 * Called from ReplicationSystem when a streaming level is about to unload.
+	 * Will remove the group associated with the level and remove destruction infos.
+	 */
+	void NotifyStreamingLevelUnload(const UObject* Level);
+
+	/**
+	 * Remove destruction infos associated with group
+	 * Passing in an invalid group handle indicates that we should remove all destruction infos
+	 */
+	void RemoveDestructionInfosForGroup(UE::Net::FNetObjectGroupHandle GroupHandle);
+
+	void DestroyLocalNetHandle(FNetHandle Handle);
+
+	// Tear-off all handles in the PendingTearOff list that has not yet been torn-off
+	void TearOffHandlesPendingTearOff();
+
+	// Update all the handles pending tear-off
+	void UpdateHandlesPendingTearOff();
+
+	void SetNetPushIdOnFragments(const TArrayView<const UE::Net::FReplicationFragment*const>& Fragments, const UE::Net::Private::FNetPushObjectHandle& PushHandle);
+	void ClearNetPushIdOnFragments(const TArrayView<const UE::Net::FReplicationFragment*const>& Fragments);
+
+	friend UReplicationSystem;
+	friend UE::Net::Private::FReplicationSystemImpl;
+	friend UE::Net::Private::FReplicationSystemInternal;
+	friend UE::Net::Private::FReplicationWriter;
+	friend UE::Net::Private::FReplicationReader;
+	friend UObjectReplicationBridge;
+
+	UReplicationSystem* ReplicationSystem;
+	UE::Net::Private::FReplicationProtocolManager* ReplicationProtocolManager;
+	UE::Net::Private::FReplicationStateDescriptorRegistry* ReplicationStateDescriptorRegistry;
+	UE::Net::Private::FNetHandleManager* NetHandleManager;
+	UE::Net::Private::FObjectReferenceCache* ObjectReferenceCache;
+	UE::Net::Private::FNetObjectGroups* Groups;
+
+	TMap<FObjectKey, UE::Net::FNetObjectGroupHandle> LevelGroups;
+
+private:
+
+	// DestructionInfos
+	struct FDestructionInfo
+	{
+		UE::Net::FNetObjectReference StaticRef;
+		UE::Net::FNetObjectGroupHandle LevelGroupHandle;
+		UE::Net::Private::FInternalNetHandle InternalReplicationIndex;
+	};
+
+	const UE::Net::FReplicationProtocol* DestructionInfoProtocol;
+	
+	// Need to track the objects with destruction infos so that we can clean them up properly
+	// We use this to be able ask remote to destroy static objects
+	TMap<FNetHandle, FDestructionInfo> StaticObjectsPendingDestroy;
+
+	struct FTearOffInfo
+	{
+		FTearOffInfo(FNetHandle InHandle, bool bInIsImmediate) : Handle(InHandle), bIsImmediate(bInIsImmediate) {}
+
+		FNetHandle Handle;
+		bool	bIsImmediate;
+	};
+	TArray<FTearOffInfo> HandlesPendingTearOff;
+};
+
+inline FReplicationBridgeSerializationContext::FReplicationBridgeSerializationContext(UE::Net::FNetSerializationContext& InSerialiazationContext, uint32 InConnectionId, bool bInIsDestructionInfo)
+: SerializationContext(InSerialiazationContext)
+, ConnectionId(InConnectionId)
+, bIsDestructionInfo(bInIsDestructionInfo)
+{
+}
diff --git a/Engine/Source/Runtime/Experimental/Iris/Core/Public/Iris/ReplicationSystem/ReplicationView.h b/Engine/Source/Runtime/Experimental/Iris/Core/Public/Iris/ReplicationSystem/ReplicationView.h
new file mode 100644
--- /dev/null
+++ b/Engine/Source/Runtime/Experimental/Iris/Core/Public/Iris/ReplicationSystem/ReplicationView.h
@@ -0,0 +1,24 @@
+// Copyright Epic Games, Inc. All Rights Reserved.
+
+#pragma once
+
+#include "CoreTypes.h"
+#include "Containers/Array.h"
+#include "Math/Vector.h"
+
+namespace UE::Net
+{
+
+struct FReplicationView
+{
+	struct FView
+	{
+		FVector Pos;
+		FVector Dir;
+		float FoVRadians;
+	};
+
+	TArray<FView, TInlineAllocator<4>> Views;
+};
+
+}
diff --git a/Engine/Source/Runtime/Experimental/Iris/Core/Public/Iris/ReplicationSystem/WorldLocations.h b/Engine/Source/Runtime/Experimental/Iris/Core/Public/Iris/ReplicationSystem/WorldLocations.h
new file mode 100644
--- /dev/null
+++ b/Engine/Source/Runtime/Experimental/Iris/Core/Public/Iris/ReplicationSystem/WorldLocations.h
@@ -0,0 +1,62 @@
+// Copyright Epic Games, Inc. All Rights Reserved.
+
+#pragma once
+
+#include "CoreTypes.h"
+#include "Containers/ChunkedArray.h"
+#include "Net/Core/NetBitArray.h"
+#include "Math/Vector.h"
+
+namespace UE::Net
+{
+
+struct FWorldLocationsInitParams
+{
+	uint32 MaxObjectCount = 0;
+};
+
+class FWorldLocations
+{
+public:
+	void Init(const FWorldLocationsInitParams& InitParams);
+
+	/** Returns whether the object has a valid world location or not. */
+	bool HasWorldLocation(uint32 ObjectIndex) const;
+
+	/** Returns the object's world location if it's valid or a zero vector if it's not. */
+	FVector GetWorldLocation(uint32 ObjectIndex) const;
+
+public:
+	// Internal API
+	void SetHasWorldLocation(uint32 ObjectIndex, bool bHasWorldLocation);
+	void SetWorldLocation(uint32 ObjectIndex, const FVector& WorldLocation);
+
+private:
+	enum : uint32
+	{
+		BytesPerLocationChunk = 65536U,
+	};
+
+	FNetBitArray ValidWorldLocations;
+	TChunkedArray<FVector, BytesPerLocationChunk> WorldLocations;
+};
+
+inline bool FWorldLocations::HasWorldLocation(uint32 ObjectIndex) const
+{
+	return ValidWorldLocations.GetBit(ObjectIndex);
+}
+
+inline FVector FWorldLocations::GetWorldLocation(uint32 ObjectIndex) const
+{
+	if (ValidWorldLocations.GetBit(ObjectIndex))
+	{
+		return WorldLocations[ObjectIndex];
+	}
+	else
+	{
+		return FVector::Zero();
+	}
+}
+
+}
+
diff --git a/Engine/Source/Runtime/Experimental/Iris/Core/Public/Iris/Serialization/IrisObjectReferencePackageMap.h b/Engine/Source/Runtime/Experimental/Iris/Core/Public/Iris/Serialization/IrisObjectReferencePackageMap.h
new file mode 100644
--- /dev/null
+++ b/Engine/Source/Runtime/Experimental/Iris/Core/Public/Iris/Serialization/IrisObjectReferencePackageMap.h
@@ -0,0 +1,35 @@
+// Copyright Epic Games, Inc. All Rights Reserved.
+
+#pragma once
+
+#include "UObject/CoreNet.h"
+#include "IrisObjectReferencePackageMap.generated.h"
+
+class FNetworkGUID;
+
+/**
+ * Custom packagemap implementation used to be able to capture UObject* references from external serialization.
+ * Any object references written when using this packagemap will be added to the References array and serialized as an index.
+ * When reading using this packagemap references will be read as an index and resolved by picking the corresponding entry from the provided array containing the references.
+ */
+UCLASS(transient, MinimalAPI)
+class UIrisObjectReferencePackageMap : public UPackageMap
+{
+public:
+	GENERATED_BODY()
+
+	typedef TArray<TObjectPtr<UObject>, TInlineAllocator<4>> FObjectReferenceArray;
+
+	// We override SerialzierObject in order to be able to capture object references
+	virtual bool SerializeObject(FArchive& Ar, UClass* InClass, UObject*& Obj, FNetworkGUID* OutNetGUID) override;
+
+	// Init for read, we need a reference array to be able to resolve references
+	IRISCORE_API void InitForRead(const FObjectReferenceArray* InReferences);
+
+	// Init for write, all written references will be added to the references array which will be resetted by this call
+	IRISCORE_API void InitForWrite(FObjectReferenceArray* InReferences);
+
+protected:
+	FObjectReferenceArray* References = nullptr;
+};
+
diff --git a/Engine/Source/Runtime/Experimental/Iris/Core/Public/Iris/Serialization/NetReferenceCollector.h b/Engine/Source/Runtime/Experimental/Iris/Core/Public/Iris/Serialization/NetReferenceCollector.h
new file mode 100644
--- /dev/null
+++ b/Engine/Source/Runtime/Experimental/Iris/Core/Public/Iris/Serialization/NetReferenceCollector.h
@@ -0,0 +1,69 @@
+// Copyright Epic Games, Inc. All Rights Reserved.
+
+#pragma once
+
+#include "CoreTypes.h"
+#include "Containers/Array.h"
+#include "Iris/Core/NetObjectReference.h"
+#include "Iris/Serialization/NetSerializer.h"
+#include "Iris/ReplicationState/ReplicationStateDescriptor.h"
+
+namespace UE::Net
+{
+
+enum class ENetReferenceCollectorTraits : uint32
+{
+	None								= 0U,
+	IncludeInvalidReferences			= 1U,
+	OnlyCollectReferencesThatCanBeExported = IncludeInvalidReferences << 1U,
+};
+ENUM_CLASS_FLAGS(ENetReferenceCollectorTraits);
+
+class FNetReferenceCollector
+{
+public:
+	struct FReferenceInfo
+	{
+		FNetReferenceInfo Info;
+		FNetObjectReference Reference;
+		FNetSerializerChangeMaskParam ChangeMaskInfo;
+	};
+
+	typedef TArray<FReferenceInfo, TInlineAllocator<32>> FReferenceInfoArray;
+
+public:
+	FNetReferenceCollector() : Traits(ENetReferenceCollectorTraits::None) {}
+	explicit FNetReferenceCollector(ENetReferenceCollectorTraits InTraits) : Traits(InTraits) {}
+	
+	inline void Add(const FNetReferenceInfo& ReferenceInfo, const FNetObjectReference& Reference, const FNetSerializerChangeMaskParam& ChangeMaskInfo);
+
+	const FReferenceInfoArray& GetCollectedReferences() const { return ReferenceInfos; }
+
+	void Reset() { ReferenceInfos.Reset(); }
+
+private:
+	FReferenceInfoArray ReferenceInfos;	
+	const ENetReferenceCollectorTraits Traits;
+};
+
+void FNetReferenceCollector::Add(const FNetReferenceInfo& ReferenceInfo, const FNetObjectReference& Reference, const FNetSerializerChangeMaskParam& ChangeMaskInfo)
+{
+	if (!Reference.IsValid() && !EnumHasAnyFlags(Traits, ENetReferenceCollectorTraits::IncludeInvalidReferences))
+	{
+		return;
+	}
+
+	if (!Reference.CanBeExported() && EnumHasAnyFlags(Traits, ENetReferenceCollectorTraits::OnlyCollectReferencesThatCanBeExported))
+	{
+		return;
+	}
+
+	FReferenceInfo RefInfo;
+	RefInfo.Info = ReferenceInfo;
+	RefInfo.Reference = Reference;
+	RefInfo.ChangeMaskInfo = ChangeMaskInfo;
+
+	ReferenceInfos.Add(MoveTemp(RefInfo));
+}
+
+}
diff --git a/Engine/Source/Runtime/Experimental/Iris/Core/Public/Iris/Serialization/ObjectNetSerializer.h b/Engine/Source/Runtime/Experimental/Iris/Core/Public/Iris/Serialization/ObjectNetSerializer.h
new file mode 100644
--- /dev/null
+++ b/Engine/Source/Runtime/Experimental/Iris/Core/Public/Iris/Serialization/ObjectNetSerializer.h
@@ -0,0 +1,51 @@
+// Copyright Epic Games, Inc. All Rights Reserved.
+
+#pragma once
+
+#include "NetSerializer.h"
+#include "Iris/ReplicationSystem//NetHandle.h"
+#include "UObject/ObjectMacros.h"
+
+#include "ObjectNetSerializer.generated.h"
+
+namespace UE::Net
+{
+	class FNetObjectReference;
+}
+
+USTRUCT()
+struct FObjectNetSerializerConfig : public FNetSerializerConfig
+{
+	GENERATED_BODY()
+};
+
+USTRUCT()
+struct FWeakObjectNetSerializerConfig : public FNetSerializerConfig
+{
+	GENERATED_BODY()
+};
+
+USTRUCT()
+struct FScriptInterfaceNetSerializerConfig : public FNetSerializerConfig
+{
+	GENERATED_BODY()
+
+	UPROPERTY()
+	TObjectPtr<UClass> InterfaceClass = nullptr;
+};
+
+namespace UE::Net
+{
+
+IRISCORE_API FNetHandle ReadNetHandle(FNetSerializationContext& Context);
+IRISCORE_API void WriteNetHandle(FNetSerializationContext& Context, const FNetHandle Handle);
+
+IRISCORE_API void ReadFullNetObjectReference(FNetSerializationContext& Context, FNetObjectReference& Reference);
+IRISCORE_API void WriteFullNetObjectReference(FNetSerializationContext& Context, const FNetObjectReference& Reference);
+
+UE_NET_DECLARE_SERIALIZER(FObjectNetSerializer, IRISCORE_API);
+UE_NET_DECLARE_SERIALIZER(FWeakObjectNetSerializer, IRISCORE_API);
+UE_NET_DECLARE_SERIALIZER(FScriptInterfaceNetSerializer, IRISCORE_API);
+
+}
+
diff --git a/Engine/Source/Runtime/Experimental/Iris/Core/Public/Iris/Serialization/PolymorphicNetSerializer.h b/Engine/Source/Runtime/Experimental/Iris/Core/Public/Iris/Serialization/PolymorphicNetSerializer.h
new file mode 100644
--- /dev/null
+++ b/Engine/Source/Runtime/Experimental/Iris/Core/Public/Iris/Serialization/PolymorphicNetSerializer.h
@@ -0,0 +1,88 @@
+// Copyright Epic Games, Inc. All Rights Reserved.
+
+#pragma once
+
+#include "NetSerializer.h"
+#include "UObject/ObjectMacros.h"
+#include "Iris/ReplicationState/ReplicationStateDescriptor.h"
+#include "Containers/Array.h"
+
+#include "PolymorphicNetSerializer.generated.h"
+
+class UScriptStruct;
+namespace UE::Net
+{
+	struct FReplicationStateDescriptor;
+}
+
+namespace UE::Net
+{
+
+struct FPolymorphicNetSerializerScriptStructCache
+{
+	enum : uint32 { RegisteredTypeBits = 5U };
+	enum : uint32 { MaxRegisteredTypeCount = (1U << RegisteredTypeBits) - 1U };
+	enum : uint32 { InvalidTypeIndex = 0U };
+
+	struct FTypeInfo
+	{
+		const UScriptStruct* ScriptStruct;
+		TRefCountPtr<const FReplicationStateDescriptor> Descriptor;
+	};
+
+	IRISCORE_API void InitForType(const UScriptStruct* InScriptStruct);
+	inline const uint32 GetTypeIndex(const UScriptStruct* ScriptStruct) const;
+	inline const FTypeInfo* GetTypeInfo(uint32 TypeIndex) const;
+	inline bool CanHaveNetReferences() const { return EnumHasAnyFlags(CommonTraits, EReplicationStateTraits::HasObjectReference); }
+
+private:
+	TArray<FTypeInfo> RegisteredTypes;
+	EReplicationStateTraits CommonTraits = EReplicationStateTraits::None;
+};
+
+}
+
+USTRUCT()
+struct FPolymorphicStructNetSerializerConfig : public FNetSerializerConfig
+{
+	GENERATED_BODY()
+
+	UE::Net::FPolymorphicNetSerializerScriptStructCache RegisteredTypes;
+};
+
+USTRUCT()
+struct FPolymorphicArrayStructNetSerializerConfig : public FNetSerializerConfig
+{
+	GENERATED_BODY()
+
+	UE::Net::FPolymorphicNetSerializerScriptStructCache RegisteredTypes;
+};
+
+namespace UE::Net
+{
+
+const FPolymorphicNetSerializerScriptStructCache::FTypeInfo* FPolymorphicNetSerializerScriptStructCache::GetTypeInfo(uint32 TypeIndex) const
+{
+	static_assert(InvalidTypeIndex == 0U, "Expected InvalidTypeIndex to be 0");
+	if ((TypeIndex - 1U) > static_cast<uint32>(RegisteredTypes.Num()))
+	{
+		return nullptr;
+	}
+	return &RegisteredTypes.GetData()[TypeIndex - 1U];
+}
+
+const uint32 FPolymorphicNetSerializerScriptStructCache::GetTypeIndex(const UScriptStruct* ScriptStruct) const
+{
+	if (ScriptStruct)
+	{
+		const int32 ArrayIndex =  RegisteredTypes.IndexOfByPredicate([&ScriptStruct](const FTypeInfo& TypeInfo) { return TypeInfo.ScriptStruct == ScriptStruct; } );
+		if (ArrayIndex != INDEX_NONE)
+		{
+			return static_cast<uint32>(ArrayIndex + 1);
+		}
+	}
+
+	return InvalidTypeIndex;
+}
+
+}
diff --git a/Engine/Source/Runtime/Experimental/Iris/Stub/IrisStub.Build.cs b/Engine/Source/Runtime/Experimental/Iris/Stub/IrisStub.Build.cs
new file mode 100644
--- /dev/null
+++ b/Engine/Source/Runtime/Experimental/Iris/Stub/IrisStub.Build.cs
@@ -0,0 +1,36 @@
+// Copyright Epic Games, Inc. All Rights Reserved.
+
+namespace UnrealBuildTool.Rules
+{
+	public class IrisStub : ModuleRules
+	{
+        public IrisStub(ReadOnlyTargetRules Target) : base(Target)
+		{
+            PublicIncludePaths.Add(ModuleDirectory + "/Public");
+
+			PublicDependencyModuleNames.AddRange(
+				new string[]
+				{
+					"Core",
+					"CoreUObject",
+					"NetCore",
+                }
+            );
+
+			if (Target.bUseIris == true)
+			{
+				PublicDefinitions.Add("UE_WITH_IRIS=1");
+				// This module should not be precompiled if Iris is enabled
+				if (Target.bBuildAllModules == true)
+				{
+					PrecompileForTargets = PrecompileTargetsType.None;
+				}
+			}
+			else
+			{
+				PublicDefinitions.Add("UE_WITH_IRIS=0");
+			}
+
+		}
+	}
+}
diff --git a/Engine/Source/Runtime/Experimental/Iris/Stub/Private/Iris/IrisStubModule.cpp b/Engine/Source/Runtime/Experimental/Iris/Stub/Private/Iris/IrisStubModule.cpp
new file mode 100644
--- /dev/null
+++ b/Engine/Source/Runtime/Experimental/Iris/Stub/Private/Iris/IrisStubModule.cpp
@@ -0,0 +1,10 @@
+// Copyright Epic Games, Inc. All Rights Reserved.
+
+#include "Modules/ModuleManager.h"
+
+class FIrisStubModule : public IModuleInterface
+{
+};
+
+IMPLEMENT_MODULE(FIrisStubModule, IrisStub);
+
diff --git a/Engine/Source/Runtime/Experimental/Iris/Stub/Public/Iris/ReplicationSystem/Filtering/NetObjectFilter.h b/Engine/Source/Runtime/Experimental/Iris/Stub/Public/Iris/ReplicationSystem/Filtering/NetObjectFilter.h
new file mode 100644
--- /dev/null
+++ b/Engine/Source/Runtime/Experimental/Iris/Stub/Public/Iris/ReplicationSystem/Filtering/NetObjectFilter.h
@@ -0,0 +1,21 @@
+// Copyright Epic Games, Inc. All Rights Reserved.
+
+#pragma once
+
+#include "UObject/ObjectMacros.h"
+#include "NetObjectFilter.generated.h"
+
+// Stub used when we are not compiling with iris to workaround UHT dependencies
+static_assert(UE_WITH_IRIS == 0, "IrisStub module should not be used when iris is enabled");
+
+UCLASS(Transient, MinimalAPI)
+class UNetObjectFilterConfig : public UObject
+{
+	GENERATED_BODY()
+};
+
+UCLASS(Transient, MinimalAPI)
+class UNetObjectFilter : public UObject
+{
+	GENERATED_BODY()
+};
diff --git a/Engine/Source/Runtime/Experimental/Iris/Stub/Public/Iris/ReplicationSystem/NetBlob/NetBlobHandler.h b/Engine/Source/Runtime/Experimental/Iris/Stub/Public/Iris/ReplicationSystem/NetBlob/NetBlobHandler.h
new file mode 100644
--- /dev/null
+++ b/Engine/Source/Runtime/Experimental/Iris/Stub/Public/Iris/ReplicationSystem/NetBlob/NetBlobHandler.h
@@ -0,0 +1,15 @@
+// Copyright Epic Games, Inc. All Rights Reserved.
+
+#pragma once
+
+#include "UObject/ObjectMacros.h"
+#include "NetBlobHandler.generated.h"
+
+// Stub used when we are not compiling with iris to workaround UHT dependencies
+static_assert(UE_WITH_IRIS == 0, "IrisStub module should not be used when iris is enabled");
+
+UCLASS(transient, MinimalApi)
+class UNetBlobHandler : public UObject
+{
+	GENERATED_BODY()
+};
diff --git a/Engine/Source/Runtime/Experimental/Iris/Stub/Public/Iris/ReplicationSystem/ObjectReplicationBridge.h b/Engine/Source/Runtime/Experimental/Iris/Stub/Public/Iris/ReplicationSystem/ObjectReplicationBridge.h
new file mode 100644
--- /dev/null
+++ b/Engine/Source/Runtime/Experimental/Iris/Stub/Public/Iris/ReplicationSystem/ObjectReplicationBridge.h
@@ -0,0 +1,15 @@
+// Copyright Epic Games, Inc. All Rights Reserved.
+
+#pragma once
+
+#include "ReplicationBridge.h"
+#include "ObjectReplicationBridge.generated.h"
+
+// Stub used when we are not compiling with iris to workaround UHT dependencies
+static_assert(UE_WITH_IRIS == 0, "IrisStub module should not be used when iris is enabled");
+
+UCLASS(Transient, MinimalApi)
+class UObjectReplicationBridge : public UReplicationBridge
+{
+	GENERATED_BODY()
+};
diff --git a/Engine/Source/Runtime/Experimental/Iris/Stub/Public/Iris/ReplicationSystem/Prioritization/NetObjectCountLimiter.h b/Engine/Source/Runtime/Experimental/Iris/Stub/Public/Iris/ReplicationSystem/Prioritization/NetObjectCountLimiter.h
new file mode 100644
--- /dev/null
+++ b/Engine/Source/Runtime/Experimental/Iris/Stub/Public/Iris/ReplicationSystem/Prioritization/NetObjectCountLimiter.h
@@ -0,0 +1,15 @@
+// Copyright Epic Games, Inc. All Rights Reserved.
+
+#pragma once
+
+#include "Iris/ReplicationSystem/Prioritization/NetObjectPrioritizer.h"
+#include "NetObjectCountLimiter.generated.h"
+
+// Stub used when we are not compiling with iris to workaround UHT dependencies
+static_assert(UE_WITH_IRIS == 0, "IrisStub module should not be used when iris is enabled");
+	
+UCLASS(transient, MinimalAPI)
+class UNetObjectCountLimiterConfig : public UNetObjectPrioritizerConfig
+{
+	GENERATED_BODY()
+};
diff --git a/Engine/Source/Runtime/Experimental/Iris/Stub/Public/Iris/ReplicationSystem/Prioritization/NetObjectPrioritizer.h b/Engine/Source/Runtime/Experimental/Iris/Stub/Public/Iris/ReplicationSystem/Prioritization/NetObjectPrioritizer.h
new file mode 100644
--- /dev/null
+++ b/Engine/Source/Runtime/Experimental/Iris/Stub/Public/Iris/ReplicationSystem/Prioritization/NetObjectPrioritizer.h
@@ -0,0 +1,21 @@
+// Copyright Epic Games, Inc. All Rights Reserved.
+
+#pragma once
+
+#include "UObject/ObjectMacros.h"
+#include "NetObjectPrioritizer.generated.h"
+
+// Stub used when we are not compiling with iris to workaround UHT dependencies
+static_assert(UE_WITH_IRIS == 0, "IrisStub module should not be used when iris is enabled");
+	
+UCLASS(Transient, MinimalAPI)
+class UNetObjectPrioritizerConfig : public UObject
+{
+	GENERATED_BODY()
+};
+
+UCLASS(Transient, MinimalAPI)
+class UNetObjectPrioritizer : public UObject
+{
+	GENERATED_BODY()
+};
diff --git a/Engine/Source/Runtime/Experimental/Iris/Stub/Public/Iris/ReplicationSystem/Prioritization/SphereWithOwnerBoostNetObjectPrioritizer.h b/Engine/Source/Runtime/Experimental/Iris/Stub/Public/Iris/ReplicationSystem/Prioritization/SphereWithOwnerBoostNetObjectPrioritizer.h
new file mode 100644
--- /dev/null
+++ b/Engine/Source/Runtime/Experimental/Iris/Stub/Public/Iris/ReplicationSystem/Prioritization/SphereWithOwnerBoostNetObjectPrioritizer.h
@@ -0,0 +1,15 @@
+// Copyright Epic Games, Inc. All Rights Reserved.
+
+#pragma once
+
+#include "Iris/ReplicationSystem/Prioritization/NetObjectPrioritizer.h"
+#include "SphereWithOwnerBoostNetObjectPrioritizer.generated.h"
+
+// Stub used when we are not compiling with iris to workaround UHT dependencies
+static_assert(UE_WITH_IRIS == 0, "IrisStub module should not be used when iris is enabled");
+	
+UCLASS(transient, MinimalAPI)
+class USphereWithOwnerBoostNetObjectPrioritizerConfig : public UNetObjectPrioritizerConfig
+{
+	GENERATED_BODY()
+};
diff --git a/Engine/Source/Runtime/Experimental/Iris/Stub/Public/Iris/ReplicationSystem/ReplicationBridge.h b/Engine/Source/Runtime/Experimental/Iris/Stub/Public/Iris/ReplicationSystem/ReplicationBridge.h
new file mode 100644
--- /dev/null
+++ b/Engine/Source/Runtime/Experimental/Iris/Stub/Public/Iris/ReplicationSystem/ReplicationBridge.h
@@ -0,0 +1,17 @@
+// Copyright Epic Games, Inc. All Rights Reserved.
+
+#pragma once
+
+#include "CoreTypes.h"
+#include "UObject/ObjectMacros.h"
+#include "ReplicationBridge.generated.h"
+
+// Stub used when we are not compiling with iris to workaround UHT dependencies
+static_assert(UE_WITH_IRIS == 0, "IrisStub module should not be used when iris is enabled");
+	
+UCLASS(Transient, MinimalAPI)
+class UReplicationBridge : public UObject
+{
+	GENERATED_BODY()
+};
+
diff --git a/Engine/Source/Runtime/Experimental/Iris/Stub/Public/Iris/Serialization/IrisObjectReferencePackageMap.h b/Engine/Source/Runtime/Experimental/Iris/Stub/Public/Iris/Serialization/IrisObjectReferencePackageMap.h
new file mode 100644
--- /dev/null
+++ b/Engine/Source/Runtime/Experimental/Iris/Stub/Public/Iris/Serialization/IrisObjectReferencePackageMap.h
@@ -0,0 +1,16 @@
+// Copyright Epic Games, Inc. All Rights Reserved.
+
+#pragma once
+
+#include "UObject/ObjectMacros.h"
+#include "IrisObjectReferencePackageMap.generated.h"
+
+// Stub used when we are not compiling with iris to workaround UHT dependencies
+static_assert(UE_WITH_IRIS == 0, "IrisStub module should not be used when iris is enabled");
+
+UCLASS(transient)
+class UIrisObjectReferencePackageMap : public UObject
+{
+public:
+	GENERATED_BODY()
+};
diff --git a/Engine/Source/Runtime/Experimental/Iris/Stub/Public/Iris/Serialization/NetSerializer.h b/Engine/Source/Runtime/Experimental/Iris/Stub/Public/Iris/Serialization/NetSerializer.h
new file mode 100644
--- /dev/null
+++ b/Engine/Source/Runtime/Experimental/Iris/Stub/Public/Iris/Serialization/NetSerializer.h
@@ -0,0 +1,18 @@
+// Copyright Epic Games, Inc. All Rights Reserved.
+
+#pragma once
+
+#include "UObject/ObjectMacros.h"
+#include "NetSerializer.generated.h"
+
+// Stub used when we are not compiling with iris to workaround UHT dependencies
+static_assert(UE_WITH_IRIS == 0, "IrisStub module should not be used when iris is enabled");
+
+USTRUCT()
+struct FNetSerializerConfig
+{
+	GENERATED_BODY()
+};
+
+// For convenience to avoid changing code when we compile without iris
+#define UE_NET_DECLARE_SERIALIZER(...)
diff --git a/Engine/Source/Runtime/Experimental/Iris/Stub/Public/Iris/Serialization/PolymorphicNetSerializer.h b/Engine/Source/Runtime/Experimental/Iris/Stub/Public/Iris/Serialization/PolymorphicNetSerializer.h
new file mode 100644
--- /dev/null
+++ b/Engine/Source/Runtime/Experimental/Iris/Stub/Public/Iris/Serialization/PolymorphicNetSerializer.h
@@ -0,0 +1,30 @@
+// Copyright Epic Games, Inc. All Rights Reserved.
+
+#pragma once
+
+#include "Iris/Serialization/NetSerializer.h"
+#include "PolymorphicNetSerializer.generated.h"
+
+// Stub used when we are not compiling with iris to workaround UHT dependencies
+static_assert(UE_WITH_IRIS == 0, "IrisStub module should not be used when iris is enabled");
+
+class UScriptStruct;
+struct FReplicationStateDescriptor;
+
+struct FPolymorphicNetSerializerScriptStructCache
+{
+};
+
+USTRUCT()
+struct FPolymorphicStructNetSerializerConfig : public FNetSerializerConfig
+{
+	GENERATED_BODY()
+};
+
+USTRUCT()
+struct FPolymorphicArrayStructNetSerializerConfig : public FNetSerializerConfig
+{
+	GENERATED_BODY()
+
+	FPolymorphicNetSerializerScriptStructCache RegisteredTypes;
+};
diff --git a/Engine/Source/Runtime/Experimental/Iris/Stub/Readme.txt b/Engine/Source/Runtime/Experimental/Iris/Stub/Readme.txt
new file mode 100644
--- /dev/null
+++ b/Engine/Source/Runtime/Experimental/Iris/Stub/Readme.txt
@@ -0,0 +1,3 @@
+This module is an internal R&D effort to explore how we can implement a high performant replication system
+This module is used when we are not compiling with Iris support
+
diff --git a/Engine/Source/Runtime/GameplayTags/GameplayTags.Build.cs b/Engine/Source/Runtime/GameplayTags/GameplayTags.Build.cs
--- a/Engine/Source/Runtime/GameplayTags/GameplayTags.Build.cs
+++ b/Engine/Source/Runtime/GameplayTags/GameplayTags.Build.cs
@@ -39,6 +39,8 @@ namespace UnrealBuildTool.Rules
 				}
                 );
             }
-        }
+
+			SetupIrisSupport(Target);
+		}
 	}
 }
diff --git a/Engine/Source/Runtime/Net/Core/Classes/Net/Serialization/FastArraySerializer.h b/Engine/Source/Runtime/Net/Core/Classes/Net/Serialization/FastArraySerializer.h
--- a/Engine/Source/Runtime/Net/Core/Classes/Net/Serialization/FastArraySerializer.h
+++ b/Engine/Source/Runtime/Net/Core/Classes/Net/Serialization/FastArraySerializer.h
@@ -211,6 +211,41 @@ struct TStructOpsTypeTraits< FExampleArray > : public TStructOpsTypeTraitsBase2<
  *		There is nothing special about them other than being unique.
  */
 
+/**
+* Helper to get auto-deduced FastArrayItemType from FastArraySerializer
+*/
+template <typename FastArrayType>
+class TFastArrayTypeHelper
+{
+private:
+	struct CGetFastArrayItemTypeFuncable
+	{
+		template <typename InFastArrayType, typename...>
+		auto Requires(InFastArrayType* FastArray) -> decltype(FastArray->GetFastArrayItemTypePtr());
+	};
+
+	// Helper to always return a Type even if GetFastArrayItemTypePtr is not defined
+	static constexpr auto FastArrayTypePtr = []
+	{
+		if constexpr (TModels<CGetFastArrayItemTypeFuncable, FastArrayType>::Value)
+		{
+			return FastArrayType::GetFastArrayItemTypePtr();
+		}
+		else
+		{
+			return static_cast<int*>(nullptr);
+		}
+	}();
+
+public:
+	using FastArrayItemType = typename TRemovePointer<typename std::decay<decltype(FastArrayTypePtr)>::type>::Type;
+
+	/**
+	 * Returns true if auto detected FastArrayItemType is a valid FastArrayItemType
+	 */
+	static constexpr bool HasValidFastArrayItemType() { return TIsDerivedFrom<FastArrayItemType, FFastArraySerializerItem>::IsDerived; }
+};
+
 /** Custom INetDeltaBaseState used by Fast Array Serialization */
 class FNetFastTArrayBaseState : public INetDeltaBaseState
 {
@@ -633,6 +668,13 @@ private:
 			TArray<FFastArraySerializer_FastArrayDeltaSerialize_FIdxIDPair, TInlineAllocator<8>>& ChangedElements,
 			TArray<int32, TInlineAllocator<8>>& DeletedElements);
 
+		/**
+		 * Variant of BuildChangedAndDeletedBuffersFromDefault used when we initialize from the default state which we are not suppose to modify
+		 */
+		void BuildChangedAndDeletedBuffersFromDefault(
+			TMap<int32, int32>& NewIDToKeyMap,
+			TArray<FFastArraySerializer_FastArrayDeltaSerialize_FIdxIDPair, TInlineAllocator<8>>& ChangedElements);
+
 		/** Writes out a FFastArraySerializerHeader */
 		void WriteDeltaHeader(FFastArraySerializerHeader& Header) const;
 
@@ -662,6 +704,9 @@ private:
 
 		template<typename FastArrayType = SerializerType>
 		inline typename TEnableIf<!TModels<CPostReplicatedReceiveFuncable, FastArrayType, const FFastArraySerializer::FPostReplicatedReceiveParameters>::Value, void>::Type CallPostReplicatedReceiveOrNot() {}
+
+		// Validate that deduced FastArrayItemType is valid and that it is the same as the specified one		
+		static_assert(TIsSame<typename TFastArrayTypeHelper<SerializerType>::FastArrayItemType, Type>::Value, "Auto deduced FastArrayItemType is invalid or differs from the specified type. Make sure that the FastArraySerializer has a single replicated array property.");
 	};
 
 	/**
@@ -804,6 +849,47 @@ bool FFastArraySerializer::TFastArraySerializeHelper<Type, SerializerType>::Cond
 	return true;
 }
 
+template<typename Type, typename SerializerType>
+void FFastArraySerializer::TFastArraySerializeHelper<Type, SerializerType>::BuildChangedAndDeletedBuffersFromDefault(
+	TMap<int32, int32>& NewIDToKeyMap,
+	TArray<FFastArraySerializer_FastArrayDeltaSerialize_FIdxIDPair, TInlineAllocator<8>>& ChangedElements)
+{
+	const int32 NumConsideredItems = CalcNumItemsForConsideration();
+
+	// Verify assumptions, we never expect the default state to have replicated and thus IDCounter and ArrayReplicationKey should be at the defaults
+	check(Parms.bIsInitializingBaseFromDefault);
+	check(ArraySerializer.IDCounter == 0 && ArraySerializer.ArrayReplicationKey == 0);
+
+	// We fake assignment of ID to avoid modifying the default state
+	int32 FakeIDCounter = ArraySerializer.IDCounter;
+	const int32 FakeReplicationKey = 0;
+
+	//-----------------------------------------------------------------
+	// When initializing from default we assume that all items are new
+	//-----------------------------------------------------------------
+	for (int32 i = 0; i < Items.Num(); ++i)
+	{
+		Type& Item = Items[i];
+
+		UE_LOG(LogNetFastTArray, Log, TEXT("    Array[%d] - ID %d. CL %d."), i, Item.ReplicationID, Item.ReplicationKey);
+		if (!ArraySerializer.template ShouldWriteFastArrayItem<Type, SerializerType>(Item, Parms.bIsWritingOnClient))
+		{
+			// On clients, this will skip items that were added predictively.
+			continue;
+		}
+
+		// Item should not have been touched before
+		check(Item.ReplicationID == INDEX_NONE && Item.ReplicationKey == INDEX_NONE);
+
+		const int32 FakeID = FakeIDCounter + i;
+
+		NewIDToKeyMap.Add(FakeID, FakeReplicationKey);
+
+		UE_LOG(LogNetFastTArray, Log, TEXT("       New! Element ID: %d. %s"), FakeID, *Item.GetDebugString());
+		ChangedElements.Add(FFastArraySerializer_FastArrayDeltaSerialize_FIdxIDPair(i, FakeID));
+	}
+}
+
 template<typename Type, typename SerializerType>
 void FFastArraySerializer::TFastArraySerializeHelper<Type, SerializerType>::BuildChangedAndDeletedBuffers(
 	TMap<int32, int32>& NewIDToKeyMap,
@@ -1323,7 +1409,15 @@ bool FFastArraySerializer::FastArrayDeltaSerialize(TArray<Type> &Items, FNetDelt
 
 		TArray<FFastArraySerializer_FastArrayDeltaSerialize_FIdxIDPair, TInlineAllocator<8>> ChangedElements;
 
-		Helper.BuildChangedAndDeletedBuffers(NewMap, OldMap, ChangedElements, Header.DeletedIndices);
+		// When we are initializing from default we use a simplified variant of BuildChangedAndDeletedBuffers that does not modify the source state
+		if (Parms.bIsInitializingBaseFromDefault)
+		{
+			Helper.BuildChangedAndDeletedBuffersFromDefault(NewMap, ChangedElements);
+		}
+		else
+		{
+			Helper.BuildChangedAndDeletedBuffers(NewMap, OldMap, ChangedElements, Header.DeletedIndices);
+		}
 		
 		// Note: we used to early return false here if nothing had changed, but we still need to send
 		// a bunch with the array key / base key, so that clients can look for implicit deletes.
@@ -1647,7 +1741,16 @@ bool FFastArraySerializer::FastArrayDeltaSerialize_DeltaSerializeStructs(TArray<
 		};
 
 		TArray<FFastArraySerializer_FastArrayDeltaSerialize_FIdxIDPair, TInlineAllocator<8>> ChangedElements;
-		Helper.BuildChangedAndDeletedBuffers(NewItemMap, OldItemMap, ChangedElements, Header.DeletedIndices);
+
+		// When we are initializing from default we use a simplified variant of BuildChangedAndDeletedBuffers that does not modify the source state
+		if (Parms.bIsInitializingBaseFromDefault)
+		{
+			Helper.BuildChangedAndDeletedBuffersFromDefault(NewItemMap, ChangedElements);
+		}
+		else
+		{
+			Helper.BuildChangedAndDeletedBuffers(NewItemMap, OldItemMap, ChangedElements, Header.DeletedIndices);
+		}
 
 		// Note: we used to early return false here if nothing had changed, but we still need to send
 		// a bunch with the array key / base key, so that clients can look for implicit deletes.
@@ -1710,3 +1813,16 @@ bool FFastArraySerializer::FastArrayDeltaSerialize_DeltaSerializeStructs(TArray<
 
 	return true;
 }
+
+
+/**
+ * Macro injected from UHT to facilitate automatic registration of FastArraySerializers when using iris replication
+ */
+#if UE_WITH_IRIS
+#define UE_NET_DECLARE_FASTARRAY(FastArrayType, FastArrayItemArrayMemberName, Api) \
+	static constexpr auto GetFastArrayItemTypePtr() { return static_cast<decltype(FastArrayType::FastArrayItemArrayMemberName)::ElementType*>(nullptr); }; \
+	Api static UE::Net::CreateAndRegisterReplicationFragmentFunc GetFastArrayCreateReplicationFragmentFunction();
+#else
+#define UE_NET_DECLARE_FASTARRAY(FastArrayType, FastArrayItemArrayMemberName, Api) 	static constexpr auto GetFastArrayItemTypePtr() { return static_cast<decltype(FastArrayType::FastArrayItemArrayMemberName)::ElementType*>(nullptr); };;
+#endif
+
diff --git a/Engine/Source/Runtime/Net/Core/Private/Net/Core/Misc/ResizableCircularQueueTests.cpp b/Engine/Source/Runtime/Net/Core/Private/Net/Core/Misc/ResizableCircularQueueTests.cpp
--- a/Engine/Source/Runtime/Net/Core/Private/Net/Core/Misc/ResizableCircularQueueTests.cpp
+++ b/Engine/Source/Runtime/Net/Core/Private/Net/Core/Misc/ResizableCircularQueueTests.cpp
@@ -5,7 +5,7 @@
 
 #if WITH_DEV_AUTOMATION_TESTS
 
-IMPLEMENT_SIMPLE_AUTOMATION_TEST(FResizableCircularQueueTest, "Network.ResizableCircularQueueTest", EAutomationTestFlags::ApplicationContextMask | EAutomationTestFlags::EngineFilter)
+IMPLEMENT_SIMPLE_AUTOMATION_TEST(FResizableCircularQueueTest, "Net.ResizableCircularQueueTest", EAutomationTestFlags::ApplicationContextMask | EAutomationTestFlags::EngineFilter)
 
 struct FResizableCircularQueueTestUtil
 {
@@ -19,7 +19,7 @@ struct FResizableCircularQueueTestUtil
 		uint32 ExpectedValue = ExpectedValueAtFront;
 
 		// Peek elements in queue at given offset, peek from back to front
-		for (SIZE_T It = 0; It < Queue.Count(); ++It)
+		for (SIZE_T It = 0, EndIt = Queue.Count(); It < EndIt; ++It)
 		{
 			bSuccess = bSuccess && (ExpectedValue == Queue.PeekAtOffset(It));
 			ExpectedValue += Increment;
@@ -64,18 +64,21 @@ bool FResizableCircularQueueTest::RunTest(const FString& Parameters)
 
 	// Test Push over Capacity
 	{
-		const SIZE_T ElementsToPush = 32;
+		const SIZE_T InitialQueueCapacity = 32;
 
-		FResizableCircularQueueTestUtil::QueueT Q(ElementsToPush);
+		FResizableCircularQueueTestUtil::QueueT Q(InitialQueueCapacity);
 
-		for (SIZE_T It=0; It < ElementsToPush; ++It)
+		for (SIZE_T It=0; It < InitialQueueCapacity; ++It)
 		{
 			Q.Enqueue(It);
 		}
+		Q.Pop();
+		Q.Enqueue(InitialQueueCapacity);
+		Q.Enqueue(InitialQueueCapacity + 1);
 		
-		TestEqual(TEXT("Test Push over Capacity - Size"), Q.Count(), ElementsToPush);
-		TestEqual(TEXT("Test Push over Capacity - Capacity"), Q.AllocatedCapacity(), ElementsToPush);
- 		TestTrue(TEXT("Test Push over Capacity - Expected"), FResizableCircularQueueTestUtil::VerifyQueueIntegrity(Q, 0, 1));
+		TestEqual(TEXT("Test Push over Capacity - Size"), Q.Count(), InitialQueueCapacity+1);
+		TestTrue(TEXT("Test Push over Capacity - Capacity"), Q.AllocatedCapacity() >= InitialQueueCapacity+1);
+ 		TestTrue(TEXT("Test Push over Capacity - Expected"), FResizableCircularQueueTestUtil::VerifyQueueIntegrity(Q, 1, 1));
 	}
 
 	// Test Push and Pop
@@ -135,6 +138,47 @@ bool FResizableCircularQueueTest::RunTest(const FString& Parameters)
 		TestEqual(TEXT("Test Push and pop all - Capacity after"), Q.AllocatedCapacity(), ElementsToPush);
 	}
 
+	// Test multi pop of element with non-trivial destructor
+	{
+		struct FStructWithDestructorForResizableCiruclarQueuePopTest
+		{
+			~FStructWithDestructorForResizableCiruclarQueuePopTest()
+			{
+				--Count;
+				check(Count == 0);
+			}
+
+			int Count = 1;
+		};
+
+		using FMyTestQueue = TResizableCircularQueue<FStructWithDestructorForResizableCiruclarQueuePopTest>;
+
+		// Pop zero
+		{
+			FMyTestQueue Q(0);
+			Q.Enqueue(FStructWithDestructorForResizableCiruclarQueuePopTest());
+			Q.Pop(0);
+			TestEqual(TEXT("Pop zero elements leaves one element in the queue."), Q.Count(), SIZE_T(1));
+		}
+
+		// Pop all
+		{
+			FMyTestQueue Q(32);
+			for (; Q.Count() < Q.AllocatedCapacity();)
+			{
+				Q.Enqueue(FStructWithDestructorForResizableCiruclarQueuePopTest());
+			}
+
+			// Want to setup the queue so the head and tail aren't at the "beginning" of the queue
+			Q.Pop(1);
+			TestEqual(TEXT("Pop one element leaves queue at capacity minus one."), Q.Count(), Q.AllocatedCapacity() - SIZE_T(1));
+
+			Q.Enqueue(FStructWithDestructorForResizableCiruclarQueuePopTest());
+			Q.Pop(Q.Count());
+			TestTrue(TEXT("Pop of all elements leaves queue empty."), Q.IsEmpty());
+		}
+	}
+
 	// Test index wrap
 	{
 		const SIZE_T ElementsToPush = 256;
@@ -209,6 +253,43 @@ bool FResizableCircularQueueTest::RunTest(const FString& Parameters)
 		TestEqual(TEXT("Test trim empty - Capacity"), Q.AllocatedCapacity(), SIZE_T(0));
 	}
 
+	// Test non-trivial type
+	{
+		struct FNonTrivialStruct
+		{
+			FNonTrivialStruct() : Value(10) {}
+
+			int Value;
+		};
+
+		using FQueue = TResizableCircularQueue<TArray<FNonTrivialStruct>>;
+
+		{
+			FQueue Q;
+			FQueue::ElementT& Array = Q.Enqueue();
+			TestEqual(TEXT("Test enqueue non-trivial element"), Array.Num(), 0);
+		}
+
+		{
+			FQueue Q;
+			FQueue::ElementT Array;
+			Array.SetNum(2);
+			Q.Enqueue(Array);
+			
+			const FQueue::ElementT& ArrayInQ = Q.Peek();
+			TestEqual(TEXT("Test enqueue array with size 2"), ArrayInQ.Num(), 2);
+			TestNotEqual(TEXT("Test enqueue array with size 2 has different allocation"), ArrayInQ.GetData(), static_cast<const FQueue::ElementT&>(Array).GetData());
+		}
+	}
+
+	// Test trivial type with element construction
+	{
+		using FQueue = TResizableCircularQueue<int>;
+
+		FQueue Q(8);
+		int& Value = Q.EnqueueDefaulted_GetRef();
+		TestEqual(TEXT("Test enqueue of initialized primitive type"), Value, 0);
+	}
 
 	check(!HasAnyErrors());
 
diff --git a/Engine/Source/Runtime/Net/Core/Private/Net/Core/PropertyConditions/RepChangedPropertyTracker.cpp b/Engine/Source/Runtime/Net/Core/Private/Net/Core/PropertyConditions/RepChangedPropertyTracker.cpp
--- a/Engine/Source/Runtime/Net/Core/Private/Net/Core/PropertyConditions/RepChangedPropertyTracker.cpp
+++ b/Engine/Source/Runtime/Net/Core/Private/Net/Core/PropertyConditions/RepChangedPropertyTracker.cpp
@@ -15,6 +15,18 @@ FRepChangedPropertyTracker::FRepChangedPropertyTracker(const bool InbIsReplay, c
 {}
 PRAGMA_ENABLE_DEPRECATION_WARNINGS
 
+namespace UE::Net::Private
+{
+#if UE_WITH_IRIS
+static FIrisSetPropertyCustomCondition IrisSetPropertyCustomConditionDelegate;
+
+void SetIrisSetPropertyCustomConditionDelegate(const FIrisSetPropertyCustomCondition& Delegate)
+{
+	IrisSetPropertyCustomConditionDelegate = Delegate;
+}
+#endif // UE_WITH_IRIS
+}
+
 void FRepChangedPropertyTracker::SetCustomIsActiveOverride(UObject* OwningObject, const uint16 RepIndex, const bool bIsActive)
 {
 	const bool bOldActive = ActiveState.GetActiveState(RepIndex);
@@ -25,7 +37,14 @@ void FRepChangedPropertyTracker::SetCustomIsActiveOverride(UObject* OwningObject
 	{
 		MARK_PROPERTY_DIRTY_UNSAFE(OwningObject, RepIndex);
 	}
-#endif	
+#endif
+
+#if UE_WITH_IRIS
+	if (bOldActive != bIsActive && UE::Net::Private::IrisSetPropertyCustomConditionDelegate.IsBound())
+	{
+		UE::Net::Private::IrisSetPropertyCustomConditionDelegate.Execute(OwningObject, RepIndex, bIsActive);
+	}
+#endif // UE_WITH_IRIS
 }
 
 PRAGMA_DISABLE_DEPRECATION_WARNINGS
@@ -49,4 +68,3 @@ void FRepChangedPropertyTracker::CountBytes(FArchive& Ar) const
 	ExternalData.CountBytes(Ar);
 	PRAGMA_ENABLE_DEPRECATION_WARNINGS
 }
-
diff --git a/Engine/Source/Runtime/Net/Core/Private/Net/Core/PushModel/PushModel.cpp b/Engine/Source/Runtime/Net/Core/Private/Net/Core/PushModel/PushModel.cpp
--- a/Engine/Source/Runtime/Net/Core/Private/Net/Core/PushModel/PushModel.cpp
+++ b/Engine/Source/Runtime/Net/Core/Private/Net/Core/PushModel/PushModel.cpp
@@ -23,7 +23,7 @@ struct FNetObjectManagerPushIdHelper
 	{
 		if (UObject* Object = InObjectKey.ResolveObjectPtr())
 		{
-			FObjectNetPushIdHelper::SetNetPushIdDynamic(Object, INDEX_NONE);
+			FObjectNetPushIdHelper::SetNetPushIdDynamic(Object, UEPushModelPrivate::FNetPushObjectId().GetValue());
 		}
 	}
 };
@@ -215,7 +215,7 @@ namespace UEPushModelPrivate
 			FCoreUObjectDelegates::GetPostGarbageCollect().Remove(PostGarbageCollectHandle);
 		}
 
-		void MarkPropertyDirty(const FNetPushObjectId ObjectId, const int32 RepIndex)
+		void MarkPropertyDirty(const FNetLegacyPushObjectId ObjectId, const int32 RepIndex)
 		{
 			const int32 ObjectIndex = ObjectId;
 			if (LIKELY(PerObjectStates.IsValidIndex(ObjectIndex)))
@@ -225,7 +225,7 @@ namespace UEPushModelPrivate
 			}
 		}
 
-		void MarkPropertyDirty(const FNetPushObjectId ObjectId, const int32 StartRepIndex, const int32 EndRepIndex)
+		void MarkPropertyDirty(const FNetLegacyPushObjectId ObjectId, const int32 StartRepIndex, const int32 EndRepIndex)
 		{
 			const int32 ObjectIndex = ObjectId;
 			if (LIKELY(PerObjectStates.IsValidIndex(ObjectIndex)))
@@ -240,7 +240,7 @@ namespace UEPushModelPrivate
 
 		const FPushModelPerNetDriverHandle AddNetworkObject(const FObjectKey ObjectKey, const uint16 NumReplicatedProperties)
 		{
-			FNetPushObjectId& InternalPushId = ObjectKeyToInternalId.FindOrAdd(ObjectKey, INDEX_NONE);
+			FNetLegacyPushObjectId& InternalPushId = ObjectKeyToInternalId.FindOrAdd(ObjectKey, INDEX_NONE);
 			if (INDEX_NONE == InternalPushId)
 			{
 				InternalPushId = PerObjectStates.EmplaceAtLowestFreeIndex(NewObjectLookupPosition, ObjectKey, NumReplicatedProperties);
@@ -334,7 +334,7 @@ namespace UEPushModelPrivate
 			return false;
 		}
 
-		bool ValidateObjectIdReassignment(FNetPushObjectId CurrentId, FNetPushObjectId NewId)
+		bool ValidateObjectIdReassignment(FNetLegacyPushObjectId CurrentId, FNetLegacyPushObjectId NewId)
 		{
 			if (!PerObjectStates.IsValidIndex(CurrentId))
 			{
@@ -399,7 +399,7 @@ namespace UEPushModelPrivate
 
 	private:
 		int32 NewObjectLookupPosition = 0;
-		TMap<FObjectKey, FNetPushObjectId> ObjectKeyToInternalId;
+		TMap<FObjectKey, FNetLegacyPushObjectId> ObjectKeyToInternalId;
 		TSparseArray<FPushModelPerObjectState> PerObjectStates;
 		FDelegateHandle PostGarbageCollectHandle;
 
@@ -436,14 +436,48 @@ namespace UEPushModelPrivate
 		TEXT("Whether or not properties declared in Blueprints will be forced to used Push Model")
 	);
 
-	void MarkPropertyDirty(const FNetPushObjectId ObjectId, const int32 RepIndex)
+#if UE_WITH_IRIS
+	static FIrisMarkPropertyDirty IrisMarkPropertyDirtyDelegate;
+	static FIrisMarkPropertiesDirty IrisMarkPropertiesDirtyDelegate;
+
+	void SetIrisMarkPropertyDirtyDelegate(const FIrisMarkPropertyDirty& Delegate)
+	{
+		IrisMarkPropertyDirtyDelegate = Delegate;
+	}
+
+	void SetIrisMarkPropertiesDirtyDelegate(const FIrisMarkPropertiesDirty& Delegate)
 	{
-		PushObjectManager.MarkPropertyDirty(ObjectId, RepIndex);
+		IrisMarkPropertiesDirtyDelegate = Delegate;
 	}
+#endif // UE_WITH_IRIS
 
-	void MarkPropertyDirty(const FNetPushObjectId ObjectId, const int32 StartRepIndex, const int32 EndRepIndex)
+	void MarkPropertyDirty(const UObject* Object, const FNetPushObjectId ObjectId, const int32 RepIndex)
 	{
-		PushObjectManager.MarkPropertyDirty(ObjectId, StartRepIndex, EndRepIndex);
+#if UE_WITH_IRIS
+		if (ObjectId.IsIrisId())
+		{
+			// Assume the delegate is bound. If it isn't we have an invalid ID, which is a bug that needs to be tracked down.
+			IrisMarkPropertyDirtyDelegate.Execute(Object, ObjectId.GetIrisPushObjectId(), RepIndex);
+		}
+		else
+#endif // UE_WITH_IRIS
+		{
+			PushObjectManager.MarkPropertyDirty(ObjectId.GetLegacyPushObjectId(), RepIndex);
+		}
+	}
+
+	void MarkPropertyDirty(const UObject* Object, const FNetPushObjectId ObjectId, const int32 StartRepIndex, const int32 EndRepIndex)
+	{
+#if UE_WITH_IRIS
+		if (ObjectId.IsIrisId())
+		{
+			IrisMarkPropertiesDirtyDelegate.Execute(Object, ObjectId.GetIrisPushObjectId(), StartRepIndex, EndRepIndex);
+		}
+		else
+#endif // UE_WITH_IRIS
+		{
+			PushObjectManager.MarkPropertyDirty(ObjectId.GetLegacyPushObjectId(), StartRepIndex, EndRepIndex);
+		}
 	}
 
 	/**
@@ -496,7 +530,7 @@ namespace UEPushModelPrivate
 		return PushObjectManager.DoesHaveDirtyPropertiesOrRecentlyCollectedGarbage(Handle);
 	}
 
-	bool ValidateObjectIdReassignment(FNetPushObjectId CurrentId, FNetPushObjectId NewId)
+	bool ValidateObjectIdReassignment(FNetLegacyPushObjectId CurrentId, FNetLegacyPushObjectId NewId)
 	{
 		return PushObjectManager.ValidateObjectIdReassignment(CurrentId, NewId);
 	}
diff --git a/Engine/Source/Runtime/Net/Core/Private/Net/Core/Trace/NetTraceInternal.cpp b/Engine/Source/Runtime/Net/Core/Private/Net/Core/Trace/NetTraceInternal.cpp
--- a/Engine/Source/Runtime/Net/Core/Private/Net/Core/Trace/NetTraceInternal.cpp
+++ b/Engine/Source/Runtime/Net/Core/Private/Net/Core/Trace/NetTraceInternal.cpp
@@ -9,11 +9,18 @@
 #include "Hash/CityHash.h"
 #include "Trace/Trace.h"
 #include "UObject/NameTypes.h"
+#include <atomic>
 
 #define UE_NET_TRACE_VALIDATE 1
 
 struct FNetTraceReporter;
-typedef uint64 FReplicationProtocolIdentifier;
+namespace UE::Net
+{
+	namespace Private
+	{
+		typedef uint64 FReplicationProtocolIdentifier;
+	}
+}
 
 #if UE_NET_TRACE_ENABLED
 
@@ -33,10 +40,10 @@ struct FNetTraceInternal
 	struct FThreadBuffer : public FTlsAutoCleanup
 	{
 		// Map FName to NameId, used to know what FNames we have traced
-		TMap<FName, FNetDebugNameId> DynamicFNameToNameIdMap;
+		TMap<FName, UE::Net::FNetDebugNameId> DynamicFNameToNameIdMap;
 
 		// Map hashed dynamic TCHAR* strings to NameId
-		TMap<uint64, FNetDebugNameId> DynamicNameHashToNameIdMap;
+		TMap<uint64, UE::Net::FNetDebugNameId> DynamicNameHashToNameIdMap;
 	};
 
 	// Get next NameId used to track what we already have traced
@@ -61,7 +68,7 @@ FNetTraceInternal::FThreadBuffer* FNetTraceInternal::CreateThreadBuffer()
 
 uint32 FNetTraceInternal::GetNextNameId()
 {
-	static TAtomic<FNetDebugNameId> NextNameId(1);
+	static std::atomic<UE::Net::FNetDebugNameId> NextNameId(1);
 
 	return NextNameId++;
 }
@@ -130,7 +137,7 @@ void FNetTrace::DestroyTraceCollector(FNetTraceCollector* Collector)
 	}
 }
 
-void FNetTrace::FoldTraceCollector(FNetTraceCollector* DstCollector, const FNetTraceCollector* SrcCollector, bool bIsBunch)
+void FNetTrace::FoldTraceCollector(FNetTraceCollector* DstCollector, const FNetTraceCollector* SrcCollector, uint32 Offset)
 {
 	if (DstCollector && SrcCollector && DstCollector != SrcCollector)
 	{
@@ -140,9 +147,6 @@ void FNetTrace::FoldTraceCollector(FNetTraceCollector* DstCollector, const FNetT
 		// When we fold non-bunch events we inject them at the current level
 		const uint32 Level = DstCollector->CurrentNestingLevel;		
 		
-		const uint32 LastEventStorageIndex = DstCollector->EventCount - 1U;
-		const uint32 Offset = (DstCollector->EventCount && !bIsBunch) ? FMath::Max(DstCollector->Events[LastEventStorageIndex].StartPos, DstCollector->Events[LastEventStorageIndex].EndPos) : 0U;
-
 		// Make sure that the events fit
 		if (SrcCollector->EventCount + DstCollector->EventCount > (uint32)DstCollector->Events.Num())
 		{
@@ -171,7 +175,7 @@ void FNetTrace::FoldTraceCollector(FNetTraceCollector* DstCollector, const FNetT
 
 void FNetTrace::PushStreamOffset(FNetTraceCollector* Collector, uint32 Offset)
 {
-	if (ensure(Collector->OffsetStackLevel < FNetTraceCollector::MaxNestingLevel - 1U))
+	if (ensure(Collector->OffsetStackLevel < (FNetTraceCollector::MaxNestingLevel - 1U)))
 	{
 		const uint32 OffsetStackLevel = Collector->OffsetStackLevel;
 
@@ -192,7 +196,7 @@ void FNetTrace::PopStreamOffset(FNetTraceCollector* Collector)
 uint32 FNetTrace::BeginPacketContentEvent(FNetTraceCollector& Collector, ENetTracePacketContentEventType EventType, uint32 Pos)
 {
 #if UE_NET_TRACE_VALIDATE
-	check(Collector.CurrentNestingLevel < FNetTraceCollector::MaxNestingLevel);
+	check(Collector.CurrentNestingLevel < (FNetTraceCollector::MaxNestingLevel - 1U));
 #endif
 
 	const uint32 EventCount = Collector.EventCount;
@@ -246,7 +250,7 @@ void FNetTrace::EndPacketContentEvent(FNetTraceCollector& Collector, uint32 Even
 	--Collector.CurrentNestingLevel;
 }
 
-void FNetTrace::TracePacketContentEvent(FNetTraceCollector& Collector, FNetDebugNameId InNetTraceNameId, uint32 StartPos, uint32 EndPos, uint32 Verbosity)
+void FNetTrace::TracePacketContentEvent(FNetTraceCollector& Collector, UE::Net::FNetDebugNameId InNetTraceNameId, uint32 StartPos, uint32 EndPos, uint32 Verbosity)
 {
 	if (FNetTrace::GetTraceVerbosity() >= Verbosity)
 	{
@@ -283,7 +287,7 @@ void FNetTrace::DiscardBunch(FNetTraceCollector& Collector)
 	Collector.PendingBunchEventIndex = ~0U;
 }
 
-void FNetTrace::EndBunch(FNetTraceCollector& DstCollector, FNetDebugNameId BunchName, uint32 StartPos, uint32 HeaderBits, uint32 BunchBits, const FNetTraceBunchInfo& BunchInfo)
+void FNetTrace::EndBunch(FNetTraceCollector& DstCollector, UE::Net::FNetDebugNameId BunchName, uint32 StartPos, uint32 HeaderBits, uint32 BunchBits, const FNetTraceBunchInfo& BunchInfo)
 {
 #if UE_NET_TRACE_VALIDATE
 	check(DstCollector.PendingBunchEventIndex != ~0U);
@@ -343,7 +347,7 @@ void FNetTrace::TraceBunch(FNetTraceCollector& DstCollector, const FNetTraceBunc
 	if (&DstCollector != BunchCollector)
 	{
 		FNetTrace::BeginBunch(DstCollector);
-		FNetTrace::FoldTraceCollector(&DstCollector, BunchCollector, true);
+		FNetTrace::FoldTraceCollector(&DstCollector, BunchCollector, 0U);
 	}
 		
 	FNetTrace::EndBunch(DstCollector, TraceName(BunchName), StartPos, HeaderBits, BunchBits, BunchInfo);
@@ -354,7 +358,7 @@ void FNetTrace::TraceBunch(FNetTraceCollector& DstCollector, const FNetTraceBunc
 	if (&DstCollector != BunchCollector)
 	{
 		FNetTrace::BeginBunch(DstCollector);
-		FNetTrace::FoldTraceCollector(&DstCollector, BunchCollector, true);
+		FNetTrace::FoldTraceCollector(&DstCollector, BunchCollector, 0U);
 	}
 		
 	FNetTrace::EndBunch(DstCollector, TraceName(BunchName), StartPos, HeaderBits, BunchBits, BunchInfo);
@@ -426,7 +430,7 @@ void FNetTrace::TracePacket(uint32 GameInstanceId, uint32 ConnectionId, uint32 P
 	}
 }
 
-void FNetTrace::TraceObjectCreated(uint32 GameInstanceId, uint32 NetObjectId, const FNetDebugName* DebugName, uint64 TypeIdentifier, uint32 OwnerId)
+void FNetTrace::TraceObjectCreated(uint32 GameInstanceId, uint32 NetObjectId, const UE::Net::FNetDebugName* DebugName, uint64 TypeIdentifier, uint32 OwnerId)
 {
 	if (!GNetTraceRuntimeVerbosity)
 	{
@@ -490,7 +494,7 @@ void FNetTrace::TraceConnectionClosed(uint32 GameInstanceId, uint32 ConnectionId
 	}
 }
 
-void FNetTrace::TracePacketStatsCounter(uint32 GameInstanceId, uint32 ConnectionId, FNetDebugNameId CounterNameId, uint32 StatValue)
+void FNetTrace::TracePacketStatsCounter(uint32 GameInstanceId, uint32 ConnectionId, UE::Net::FNetDebugNameId CounterNameId, uint32 StatValue)
 {
 	if (GNetTraceRuntimeVerbosity)
 	{
@@ -498,7 +502,7 @@ void FNetTrace::TracePacketStatsCounter(uint32 GameInstanceId, uint32 Connection
 	}
 }
 
-void FNetTrace::TraceFrameStatsCounter(uint32 GameInstanceId, FNetDebugNameId CounterNameId, uint32 StatValue)
+void FNetTrace::TraceFrameStatsCounter(uint32 GameInstanceId, UE::Net::FNetDebugNameId CounterNameId, uint32 StatValue)
 {
 	if (GNetTraceRuntimeVerbosity)
 	{
@@ -506,7 +510,7 @@ void FNetTrace::TraceFrameStatsCounter(uint32 GameInstanceId, FNetDebugNameId Co
 	}
 }
 
-FNetDebugNameId FNetTrace::TraceName(const TCHAR* Name)
+UE::Net::FNetDebugNameId FNetTrace::TraceName(const TCHAR* Name)
 {
 	if ((GNetTraceRuntimeVerbosity == 0U) | (Name == nullptr))
 	{
@@ -522,13 +526,13 @@ FNetDebugNameId FNetTrace::TraceName(const TCHAR* Name)
 
 	// Hash the name using CityHash64
 	const uint64 HashedName = CityHash64((const char*)Name, FCString::Strlen(Name) * sizeof(TCHAR));
-	if (const FNetDebugNameId* FoundNameId = ThreadBuffer->DynamicNameHashToNameIdMap.Find(HashedName))
+	if (const UE::Net::FNetDebugNameId* FoundNameId = ThreadBuffer->DynamicNameHashToNameIdMap.Find(HashedName))
 	{
 		return *FoundNameId;
 	}
 	else
 	{
-		const FNetDebugNameId NameId = FNetTraceInternal::GetNextNameId();
+		const UE::Net::FNetDebugNameId NameId = FNetTraceInternal::GetNextNameId();
 		ThreadBuffer->DynamicNameHashToNameIdMap.Add(HashedName, NameId);
 
 		FTCHARToUTF8 Converter(Name);
@@ -538,8 +542,10 @@ FNetDebugNameId FNetTrace::TraceName(const TCHAR* Name)
 	}
 }
 
-FNetDebugNameId FNetTrace::TraceName(FName Name)
+UE::Net::FNetDebugNameId FNetTrace::TraceName(FName Name)
 {
+	using namespace UE::Net;
+
 	if ((GNetTraceRuntimeVerbosity == 0U) || Name.IsNone())
 	{
 		return 0U;
@@ -571,7 +577,7 @@ FNetDebugNameId FNetTrace::TraceName(FName Name)
 	}
 }
 
-FNetDebugNameId FNetTrace::TraceName(const FNetDebugName* DebugName)
+UE::Net::FNetDebugNameId FNetTrace::TraceName(const UE::Net::FNetDebugName* DebugName)
 {
 	if ((GNetTraceRuntimeVerbosity == 0U) | (DebugName == nullptr))
 	{
@@ -580,7 +586,7 @@ FNetDebugNameId FNetTrace::TraceName(const FNetDebugName* DebugName)
 
 	if (DebugName->DebugNameId == 0U)
 	{
-		const FNetDebugNameId NameId = TraceName(DebugName->Name);
+		const UE::Net::FNetDebugNameId NameId = TraceName(DebugName->Name);
 		DebugName->DebugNameId = NameId;
 
 		return NameId;
diff --git a/Engine/Source/Runtime/Net/Core/Private/Net/Core/Trace/Reporters/NetTraceReporter.cpp b/Engine/Source/Runtime/Net/Core/Private/Net/Core/Trace/Reporters/NetTraceReporter.cpp
--- a/Engine/Source/Runtime/Net/Core/Private/Net/Core/Trace/Reporters/NetTraceReporter.cpp
+++ b/Engine/Source/Runtime/Net/Core/Private/Net/Core/Trace/Reporters/NetTraceReporter.cpp
@@ -22,7 +22,7 @@ UE_TRACE_EVENT_BEGIN(NetTrace, InitEvent)
 	UE_TRACE_EVENT_FIELD(uint32, NetTraceReporterVersion)
 UE_TRACE_EVENT_END()
 
-// Trace a name, the utf encoded name is attached as a attachment
+// Trace a name
 UE_TRACE_EVENT_BEGIN(NetTrace, NameEvent)
 	UE_TRACE_EVENT_FIELD(uint16, NameId)
 	UE_TRACE_EVENT_FIELD(UE::Trace::AnsiString, Name)
@@ -142,7 +142,7 @@ void FNetTraceReporter::ReportInstanceDestroyed(uint32 GameInstanceId)
 		<< InstanceDestroyedEvent.GameInstanceId(GameInstanceId);
 }
 
-void FNetTraceReporter::ReportAnsiName(FNetDebugNameId NameId, uint32 NameSize, const char* Name)
+void FNetTraceReporter::ReportAnsiName(UE::Net::FNetDebugNameId NameId, uint32 NameSize, const char* Name)
 {
 	UE_TRACE_LOG(NetTrace, NameEvent, NetChannel)
 		<< NameEvent.NameId(NameId)
@@ -302,7 +302,7 @@ void FNetTraceReporter::ReportConnectionClosed(uint32 GameInstanceId, uint32 Con
 		<< ConnectionClosedEvent.GameInstanceId(GameInstanceId);
 }
 
-void FNetTraceReporter::ReportPacketStatsCounter(uint32 GameInstanceId, uint32 ConnectionId, FNetDebugNameId CounterNameId, uint32 StatValue)
+void FNetTraceReporter::ReportPacketStatsCounter(uint32 GameInstanceId, uint32 ConnectionId, UE::Net::FNetDebugNameId CounterNameId, uint32 StatValue)
 {
 	UE_TRACE_LOG(NetTrace, PacketStatsCounterEvent, NetChannel)
 		<< PacketStatsCounterEvent.StatsValue(StatValue)
@@ -311,7 +311,7 @@ void FNetTraceReporter::ReportPacketStatsCounter(uint32 GameInstanceId, uint32 C
 		<< PacketStatsCounterEvent.GameInstanceId(GameInstanceId);
 }
 
-void FNetTraceReporter::ReportFrameStatsCounter(uint32 GameInstanceId, FNetDebugNameId CounterNameId, uint32 StatValue)
+void FNetTraceReporter::ReportFrameStatsCounter(uint32 GameInstanceId, UE::Net::FNetDebugNameId CounterNameId, uint32 StatValue)
 {
 	UE_TRACE_LOG(NetTrace, FrameStatsCounterEvent, NetChannel)
 		<< FrameStatsCounterEvent.Timestamp(FPlatformTime::Cycles64())
@@ -320,7 +320,7 @@ void FNetTraceReporter::ReportFrameStatsCounter(uint32 GameInstanceId, FNetDebug
 		<< FrameStatsCounterEvent.GameInstanceId(GameInstanceId);
 }
 
-void FNetTraceReporter::ReportObjectCreated(uint32 GameInstanceId, uint32 NetObjectId, FNetDebugNameId NameId, uint64 TypeIdentifier, uint32 OwnerId)
+void FNetTraceReporter::ReportObjectCreated(uint32 GameInstanceId, uint32 NetObjectId, UE::Net::FNetDebugNameId NameId, uint64 TypeIdentifier, uint32 OwnerId)
 {
 	UE_TRACE_LOG(NetTrace, ObjectCreatedEvent, NetChannel)
 		<< ObjectCreatedEvent.TypeId(TypeIdentifier)
