Iris experimental replication system.
Initial submit, compiled out by default.

- Fixed wording in comment

diff --git a/Engine/Source/Runtime/Experimental/Iris/Core/Public/Iris/ReplicationState/ReplicationStateUtil.h b/Engine/Source/Runtime/Experimental/Iris/Core/Public/Iris/ReplicationState/ReplicationStateUtil.h
--- a/Engine/Source/Runtime/Experimental/Iris/Core/Public/Iris/ReplicationState/ReplicationStateUtil.h
+++ b/Engine/Source/Runtime/Experimental/Iris/Core/Public/Iris/ReplicationState/ReplicationStateUtil.h
@@ -69,7 +69,7 @@ inline void MarkDirty(const UE::Net::FReplicationStateHeader& InternalState, FNe
 {
 	// If this state is bound to a replicated object, notify the replication system that we have data to copy
 	// Note that we only check the first bit of the changemask, as this is used to indicate whether the property is dirty or not
-	// some properties might use additional bits but will only be treated as dirty if the master bit is set.
+	// some properties might use additional bits but will only be treated as dirty if the parent bit is set.
 	if (InternalState.IsBound() && !MemberChangeMask.GetBit(ChangeMaskInfo.BitOffset))
 	{
 		MarkNetObjectStateDirty(InternalState);
