Fix GJK for Large Scaled Box

diff --git a/Engine/Source/Programs/HeadlessChaos/Private/HeadlessChaosTestGJK.cpp b/Engine/Source/Programs/HeadlessChaos/Private/HeadlessChaosTestGJK.cpp
--- a/Engine/Source/Programs/HeadlessChaos/Private/HeadlessChaosTestGJK.cpp
+++ b/Engine/Source/Programs/HeadlessChaos/Private/HeadlessChaosTestGJK.cpp
@@ -2021,8 +2021,7 @@ namespace ChaosTest
 		EXPECT_NEAR(OutPosition.Z, 14.63, 1e-1);
 	}
 
-	// Regression test: Disabled until fixed
-	GTEST_TEST(GJKTests, DISABLED_GJK_LargeScaledBoxBoxTest)
+	GTEST_TEST(GJKTests, GJK_LargeScaledBoxBoxTest)
 	{
 
 		TArray<FConvex::FVec3Type> ConvexParticles;
diff --git a/Engine/Source/Runtime/Experimental/Chaos/Public/Chaos/GJK.h b/Engine/Source/Runtime/Experimental/Chaos/Public/Chaos/GJK.h
--- a/Engine/Source/Runtime/Experimental/Chaos/Public/Chaos/GJK.h
+++ b/Engine/Source/Runtime/Experimental/Chaos/Public/Chaos/GJK.h
@@ -125,14 +125,17 @@ namespace Chaos
 		CalculateQueryMargins(A, B, ThicknessA, ThicknessB);
 		ThicknessA += InThicknessA;
 
-		const FReal InflationReal = ThicknessA + ThicknessB + static_cast<FReal>(1e-3);
+		FReal EpsilonScale = FMath::Max<FReal>(A.TGeometryA::BoundingBox().Extents().Max(), B.TGeometryB::BoundingBox().Extents().Max());
+		EpsilonScale = FMath::Min<FReal>(EpsilonScale, 1e5);
+		
+		const FReal InflationReal = ThicknessA + ThicknessB + 1e-6 * EpsilonScale;
 		const VectorRegister4Float Inflation = VectorSet1(static_cast<FRealSingle>(InflationReal));
 		const VectorRegister4Float Inflation2 = VectorMultiply(Inflation, Inflation);
 
 		int32 VertexIndexA = INDEX_NONE, VertexIndexB = INDEX_NONE;
 		do
 		{
-			if (!ensure(NumIterations++ < 32))	//todo: take this out
+ 			if (!ensure(NumIterations++ < 32))	//todo: take this out
 			{
 				break;	//if taking too long just stop. This should never happen
 			}
diff --git a/Engine/Source/Runtime/Experimental/Chaos/Public/Chaos/TriangleRegister.h b/Engine/Source/Runtime/Experimental/Chaos/Public/Chaos/TriangleRegister.h
--- a/Engine/Source/Runtime/Experimental/Chaos/Public/Chaos/TriangleRegister.h
+++ b/Engine/Source/Runtime/Experimental/Chaos/Public/Chaos/TriangleRegister.h
@@ -43,6 +43,19 @@ namespace Chaos
 		FORCEINLINE FReal GetMargin() const { return 0; }
 		FORCEINLINE FReal GetRadius() const { return 0; }
 
+		FORCEINLINE const FAABB3 BoundingBox() const
+		{
+			VectorRegister4Float MinSimd = VectorMin(VectorMin(A, B), C);
+			VectorRegister4Float MaxSimd = VectorMax(VectorMax(A, B), C);
+			
+			alignas(32) FReal AlignedArray[4];
+			VectorStoreAligned(MinSimd, AlignedArray);
+			const FVec3  Min(static_cast<FReal>(AlignedArray[0]), static_cast<FReal>(AlignedArray[1]), static_cast<FReal>(AlignedArray[2]));
+			VectorStoreAligned(MaxSimd, AlignedArray);
+			const FVec3  Max(static_cast<FReal>(AlignedArray[0]), static_cast<FReal>(AlignedArray[1]), static_cast<FReal>(AlignedArray[2]));
+			return FAABB3(Min, Max);
+		}
+
 
 	private:
 
