Fix geometry collection performance issue. Create one TLAS per geometry collection instead of one per transform.

diff --git a/Engine/Source/Runtime/Experimental/GeometryCollectionEngine/Private/GeometryCollection/GeometryCollectionSceneProxy.cpp b/Engine/Source/Runtime/Experimental/GeometryCollectionEngine/Private/GeometryCollection/GeometryCollectionSceneProxy.cpp
--- a/Engine/Source/Runtime/Experimental/GeometryCollectionEngine/Private/GeometryCollection/GeometryCollectionSceneProxy.cpp
+++ b/Engine/Source/Runtime/Experimental/GeometryCollectionEngine/Private/GeometryCollection/GeometryCollectionSceneProxy.cpp
@@ -1103,11 +1103,7 @@ void FGeometryCollectionSceneProxy::GetDynamicRayTracingInstances(FRayTracingMat
 
 			FRayTracingInstance RayTracingInstance;
 			RayTracingInstance.Geometry = &RayTracingGeometry;
-			RayTracingInstance.InstanceTransforms.Reserve(DynamicData->Transforms.Num());
-			for (int32 TransformIndex = 0; TransformIndex < DynamicData->Transforms.Num(); ++TransformIndex)
-			{
-				RayTracingInstance.InstanceTransforms.Emplace(GetLocalToWorld());
-			}
+			RayTracingInstance.InstanceTransforms.Emplace(GetLocalToWorld());
 
 			// Grab the material proxies we'll be using for each section
 			TArray<FMaterialRenderProxy*, TInlineAllocator<32>> MaterialProxies;
