Chaos - disable some excessive new stats

diff --git a/Engine/Source/Runtime/Experimental/Chaos/Private/Chaos/Collision/ParticlePairMidPhase.cpp b/Engine/Source/Runtime/Experimental/Chaos/Private/Chaos/Collision/ParticlePairMidPhase.cpp
--- a/Engine/Source/Runtime/Experimental/Chaos/Private/Chaos/Collision/ParticlePairMidPhase.cpp
+++ b/Engine/Source/Runtime/Experimental/Chaos/Private/Chaos/Collision/ParticlePairMidPhase.cpp
@@ -907,7 +907,7 @@ namespace Chaos
 		const FReal Dt,
 		FCollisionContext& Context)
 	{
-		SCOPE_CYCLE_COUNTER(STAT_Collisions_GenerateCollisions);
+		//SCOPE_CYCLE_COUNTER(STAT_Collisions_GenerateCollisions);
 		PHYSICS_CSV_SCOPED_EXPENSIVE(PhysicsVerbose, DetectCollisions_NarrowPhase);
 
 		if (!IsValid())
diff --git a/Engine/Source/Runtime/Experimental/Chaos/Private/Chaos/PBDJointConstraints.cpp b/Engine/Source/Runtime/Experimental/Chaos/Private/Chaos/PBDJointConstraints.cpp
--- a/Engine/Source/Runtime/Experimental/Chaos/Private/Chaos/PBDJointConstraints.cpp
+++ b/Engine/Source/Runtime/Experimental/Chaos/Private/Chaos/PBDJointConstraints.cpp
@@ -890,7 +890,7 @@ namespace Chaos
 
 	void FPBDJointConstraints::GatherInput(const FReal Dt, const int32 ConstraintIndex, const int32 Particle0Level, const int32 Particle1Level, FPBDIslandSolverData& SolverData)
 	{
-		SCOPE_CYCLE_COUNTER(STAT_Joints_Gather);
+		//SCOPE_CYCLE_COUNTER(STAT_Joints_Gather);
 		check(!ConstraintStates[ConstraintIndex].bDisabled);
 
 		const FPBDJointSettings& JointSettings = ConstraintSettings[ConstraintIndex];
