Don't occupy LLM enum values for Iris memory tracking.

diff --git a/Engine/Source/Runtime/Core/Private/HAL/LowLevelMemTracker.cpp b/Engine/Source/Runtime/Core/Private/HAL/LowLevelMemTracker.cpp
--- a/Engine/Source/Runtime/Core/Private/HAL/LowLevelMemTracker.cpp
+++ b/Engine/Source/Runtime/Core/Private/HAL/LowLevelMemTracker.cpp
@@ -210,9 +210,6 @@ DECLARE_LLM_MEMORY_STAT(TEXT("AudioSynthesis"), STAT_AudioSynthesisLLM, STATGROU
 DECLARE_LLM_MEMORY_STAT(TEXT("RealTimeCommunications"), STAT_RealTimeCommunicationsLLM, STATGROUP_LLMFULL);
 DECLARE_LLM_MEMORY_STAT(TEXT("FName"), STAT_FNameLLM, STATGROUP_LLMFULL);
 DECLARE_LLM_MEMORY_STAT(TEXT("Networking"), STAT_NetworkingLLM, STATGROUP_LLMFULL);
-DECLARE_LLM_MEMORY_STAT(TEXT("Iris"), STAT_IrisLLM, STATGROUP_LLMFULL);
-DECLARE_LLM_MEMORY_STAT(TEXT("IrisState"), STAT_IrisStateLLM, STATGROUP_LLMFULL);
-DECLARE_LLM_MEMORY_STAT(TEXT("IrisConnection"), STAT_IrisConnectionLLM, STATGROUP_LLMFULL);
 DECLARE_LLM_MEMORY_STAT(TEXT("Meshes"), STAT_MeshesLLM, STATGROUP_LLMFULL);
 DECLARE_LLM_MEMORY_STAT(TEXT("Stats"), STAT_StatsLLM, STATGROUP_LLMFULL);
 DECLARE_LLM_MEMORY_STAT(TEXT("Shaders"), STAT_ShadersLLM, STATGROUP_LLMFULL);
diff --git a/Engine/Source/Runtime/Core/Public/HAL/LowLevelMemTracker.h b/Engine/Source/Runtime/Core/Public/HAL/LowLevelMemTracker.h
--- a/Engine/Source/Runtime/Core/Public/HAL/LowLevelMemTracker.h
+++ b/Engine/Source/Runtime/Core/Public/HAL/LowLevelMemTracker.h
@@ -166,9 +166,6 @@ enum class ELLMTagSet : uint8
 	macro(RealTimeCommunications,				"RealTimeCommunications",		GET_STATFNAME(STAT_RealTimeCommunicationsLLM),				NAME_None,										-1)\
 	macro(FName,								"FName",						GET_STATFNAME(STAT_FNameLLM),								GET_STATFNAME(STAT_EngineSummaryLLM),			-1)\
 	macro(Networking,							"Networking",					GET_STATFNAME(STAT_NetworkingLLM),							GET_STATFNAME(STAT_NetworkingSummaryLLM),		-1)\
-	macro(Iris,									"Iris",							GET_STATFNAME(STAT_IrisLLM),								GET_STATFNAME(STAT_NetworkingSummaryLLM),			ELLMTag::Networking)\
-	macro(IrisState,							"IrisState",					GET_STATFNAME(STAT_IrisStateLLM),							GET_STATFNAME(STAT_NetworkingSummaryLLM),			ELLMTag::Networking)\
-	macro(IrisConnection,						"IrisConnection",				GET_STATFNAME(STAT_IrisStateLLM),							GET_STATFNAME(STAT_NetworkingSummaryLLM),			ELLMTag::Networking)\
 	macro(Meshes,								"Meshes",						GET_STATFNAME(STAT_MeshesLLM),								GET_STATFNAME(STAT_MeshesSummaryLLM),			-1)\
 	macro(Stats,								"Stats",						GET_STATFNAME(STAT_StatsLLM),								GET_STATFNAME(STAT_EngineSummaryLLM),			-1)\
 	macro(Shaders,								"Shaders",						GET_STATFNAME(STAT_ShadersLLM),								GET_STATFNAME(STAT_EngineSummaryLLM),			-1)\
diff --git a/Engine/Source/Runtime/Engine/Private/Net/Experimental/Iris/DataStreamChannel.cpp b/Engine/Source/Runtime/Engine/Private/Net/Experimental/Iris/DataStreamChannel.cpp
--- a/Engine/Source/Runtime/Engine/Private/Net/Experimental/Iris/DataStreamChannel.cpp
+++ b/Engine/Source/Runtime/Engine/Private/Net/Experimental/Iris/DataStreamChannel.cpp
@@ -15,6 +15,7 @@
 #include "Iris/Serialization/NetBitStreamWriter.h"
 #include "Iris/Serialization/NetSerializationContext.h"
 #include "Iris/Core/IrisProfiler.h"
+#include "Iris/Core/IrisMemoryTracker.h"
 #endif // UE_WITH_IRIS
 
 namespace UE::Net::Private
@@ -162,7 +163,7 @@ void UDataStreamChannel::Tick()
 	}
 
 	IRIS_PROFILER_SCOPE(UDataStreamChannel_Tick);
-	LLM_SCOPE(ELLMTag::Iris);
+	LLM_SCOPE_BYTAG(Iris);
 
 	// Currently we want to use a full bunch so we flush if we have to
 	bool bNeedsPreSendFlush = Connection->SendBuffer.GetNumBits() > MAX_PACKET_HEADER_BITS;
diff --git a/Engine/Source/Runtime/Engine/Private/NetDriver.cpp b/Engine/Source/Runtime/Engine/Private/NetDriver.cpp
--- a/Engine/Source/Runtime/Engine/Private/NetDriver.cpp
+++ b/Engine/Source/Runtime/Engine/Private/NetDriver.cpp
@@ -76,6 +76,7 @@
 #if UE_WITH_IRIS
 #include "Iris/IrisConfig.h"
 #include "Iris/Core/IrisProfiler.h"
+#include "Iris/Core/IrisMemoryTracker.h"
 #include "Iris/ReplicationSystem/ReplicationSystem.h"
 #include "Iris/ReplicationSystem/ReplicationView.h"
 #include "Iris/ReplicationSystem/Filtering/NetObjectFilter.h"
@@ -1741,7 +1742,7 @@ bool UNetDriver::InitBase(bool bInitAsClient, FNetworkNotify* InNotify, const FU
 #if UE_WITH_IRIS
 	if (IsUsingIrisReplication() && !ReplicationSystem)
 	{
-		LLM_SCOPE(ELLMTag::Iris);
+		LLM_SCOPE_BYTAG(Iris);
 
 		if (ensureAlways(InitReplicationBridgeClass()))
 		{
diff --git a/Engine/Source/Runtime/Experimental/Iris/Core/Private/Iris/Core/IrisMemoryTracker.cpp b/Engine/Source/Runtime/Experimental/Iris/Core/Private/Iris/Core/IrisMemoryTracker.cpp
new file mode 100644
--- /dev/null
+++ b/Engine/Source/Runtime/Experimental/Iris/Core/Private/Iris/Core/IrisMemoryTracker.cpp
@@ -0,0 +1,13 @@
+// Copyright Epic Games, Inc. All Rights Reserved.
+
+#include "Iris/Core/IrisMemoryTracker.h"
+#include "HAL/LowLevelMemStats.h"
+
+DECLARE_LLM_MEMORY_STAT(TEXT("Iris"), STAT_IrisLLM, STATGROUP_LLMFULL);
+DECLARE_LLM_MEMORY_STAT(TEXT("IrisState"), STAT_IrisStateLLM, STATGROUP_LLMFULL);
+DECLARE_LLM_MEMORY_STAT(TEXT("IrisConnection"), STAT_IrisConnectionLLM, STATGROUP_LLMFULL);
+
+// Tag name, display name, parent tag name, stat name, summary stat name
+LLM_DEFINE_TAG(Iris, NAME_None, NAME_None, GET_STATFNAME(STAT_IrisLLM), GET_STATFNAME(STAT_NetworkingSummaryLLM));
+LLM_DEFINE_TAG(IrisState, "State", "Iris", GET_STATFNAME(STAT_IrisStateLLM), GET_STATFNAME(STAT_NetworkingSummaryLLM));
+LLM_DEFINE_TAG(IrisConnection, "Connection", "Iris", GET_STATFNAME(STAT_IrisConnectionLLM), GET_STATFNAME(STAT_NetworkingSummaryLLM));
diff --git a/Engine/Source/Runtime/Experimental/Iris/Core/Private/Iris/ReplicationState/ReplicationStateStorage.cpp b/Engine/Source/Runtime/Experimental/Iris/Core/Private/Iris/ReplicationState/ReplicationStateStorage.cpp
--- a/Engine/Source/Runtime/Experimental/Iris/Core/Private/Iris/ReplicationState/ReplicationStateStorage.cpp
+++ b/Engine/Source/Runtime/Experimental/Iris/Core/Private/Iris/ReplicationState/ReplicationStateStorage.cpp
@@ -1,11 +1,11 @@
 // Copyright Epic Games, Inc. All Rights Reserved.
 
 #include "Iris/ReplicationState/ReplicationStateStorage.h"
+#include "Iris/Core/IrisMemoryTracker.h"
 #include "Iris/ReplicationSystem/NetHandleManager.h"
 #include "Iris/ReplicationSystem/ReplicationOperations.h"
 #include "Iris/ReplicationSystem/ReplicationOperationsInternal.h"
 #include "Iris/ReplicationSystem/ReplicationProtocol.h"
-#include "HAL/LowLevelMemTracker.h"
 #include "HAL/PlatformMath.h"
 #include "HAL/UnrealMemory.h"
 #include <limits>
@@ -85,7 +85,7 @@ const uint8* FReplicationStateStorage::GetState(uint32 ObjectIndex, EReplication
 
 uint8* FReplicationStateStorage::AllocBaseline(uint32 ObjectIndex, EReplicationStateType Base)
 {
-	LLM_SCOPE(ELLMTag::IrisState);
+	LLM_SCOPE_BYTAG(IrisState);
 
 	FPerObjectInfo* ObjectInfo = GetOrCreatePerObjectInfoForObject(ObjectIndex);
 	if (!ensureMsgf(ObjectInfo != nullptr, TEXT("Out of object infos. Max number of object infos is %u."), UsedPerObjectInfos.GetNumBits()))
@@ -176,7 +176,7 @@ void FReplicationStateStorage::FreeBaseline(uint32 ObjectIndex, uint8* Storage)
 
 FReplicationStateStorage::FBaselineReservation FReplicationStateStorage::ReserveBaseline(uint32 ObjectIndex, EReplicationStateType Base)
 {
-	LLM_SCOPE(ELLMTag::IrisState);
+	LLM_SCOPE_BYTAG(IrisState);
 
 	FBaselineReservation Reservation;
 
diff --git a/Engine/Source/Runtime/Experimental/Iris/Core/Private/Iris/ReplicationSystem/DeltaCompression/DeltaCompressionBaselineManager.cpp b/Engine/Source/Runtime/Experimental/Iris/Core/Private/Iris/ReplicationSystem/DeltaCompression/DeltaCompressionBaselineManager.cpp
--- a/Engine/Source/Runtime/Experimental/Iris/Core/Private/Iris/ReplicationSystem/DeltaCompression/DeltaCompressionBaselineManager.cpp
+++ b/Engine/Source/Runtime/Experimental/Iris/Core/Private/Iris/ReplicationSystem/DeltaCompression/DeltaCompressionBaselineManager.cpp
@@ -4,6 +4,7 @@
 #include "Iris/IrisConfigInternal.h"
 #include "Iris/Core/BitTwiddling.h"
 #include "Iris/Core/IrisProfiler.h"
+#include "Iris/Core/IrisMemoryTracker.h"
 #include "Iris/ReplicationSystem/ChangeMaskCache.h"
 #include "Iris/ReplicationSystem/NetHandleManager.h"
 #include "Iris/ReplicationSystem/ReplicationConnections.h"
@@ -12,9 +13,8 @@
 #include "Iris/ReplicationSystem/DeltaCompression/DeltaCompressionBaselineInvalidationTracker.h"
 #include "Iris/ReplicationSystem/ReplicationSystem.h"
 #include "HAL/IConsoleManager.h"
-#include "HAL/LowLevelMemTracker.h"
-#include <limits>
 #include "Net/Core/Trace/NetTrace.h"
+#include <limits>
 
 #if UE_NET_VALIDATE_DC_BASELINES
 #define UE_NET_DC_BASELINE_CHECK(...) check(__VA_ARGS__)
@@ -353,7 +353,7 @@ FDeltaCompressionBaselineManager::ObjectInfoIndexType FDeltaCompressionBaselineM
 		const SIZE_T ObjectInfoStorageIndex = ObjectInfoIndex*BytesPerObjectInfo;
 		if (ObjectInfoStorageIndex >= ObjectInfoStorage.Num())
 		{
-			LLM_SCOPE(ELLMTag::Iris);
+			LLM_SCOPE_BYTAG(Iris);
 			ObjectInfoStorage.AddUninitialized(ObjectInfoGrowCount*BytesPerObjectInfo);
 		}
 
diff --git a/Engine/Source/Runtime/Experimental/Iris/Core/Private/Iris/ReplicationSystem/ObjectReferenceCache.cpp b/Engine/Source/Runtime/Experimental/Iris/Core/Private/Iris/ReplicationSystem/ObjectReferenceCache.cpp
--- a/Engine/Source/Runtime/Experimental/Iris/Core/Private/Iris/ReplicationSystem/ObjectReferenceCache.cpp
+++ b/Engine/Source/Runtime/Experimental/Iris/Core/Private/Iris/ReplicationSystem/ObjectReferenceCache.cpp
@@ -2,6 +2,7 @@
 
 #include "ObjectReferenceCache.h"
 #include "Iris/Core/IrisLog.h"
+#include "Iris/Core/IrisMemoryTracker.h"
 #include "Iris/ReplicationSystem/ObjectReplicationBridge.h"
 #include "Iris/ReplicationSystem/ReplicationSystem.h"
 #include "Iris/ReplicationSystem/ReplicationSystemInternal.h"
@@ -289,7 +290,8 @@ bool FObjectReferenceCache::CreateObjectReferenceInternal(const UObject* Object,
 	CachedObject.bIsPending = false;
 	CachedObject.bIsBroken = false;
 
-	check( CachedObject.bIsPackage == (Cast< UPackage >( Object ) != NULL));		// Make sure it really is a package
+	// Make sure it really is a package
+	check(CachedObject.bIsPackage == (Cast<UPackage>(Object) != nullptr));
 
 	ReferenceHandleToCachedReference.Add(NetHandle, CachedObject);
 	ObjectToNetReferenceHandle.Add(Object, NetHandle);
@@ -361,7 +363,7 @@ void FObjectReferenceCache::AddRemoteReference(FNetHandle RefHandle, const UObje
 
 	if (ensure(RefHandle.IsDynamic()))
 	{
-		UE_LOG_REFERENCECACHE(Verbose, TEXT("AddRemoteReference: %s, Object: %s Pointer: %p" ), *RefHandle.ToString(), Object ? *Object->GetName() : TEXT( "NULL" ), Object);
+		UE_LOG_REFERENCECACHE(Verbose, TEXT("AddRemoteReference: %s, Object: %s Pointer: %p" ), *RefHandle.ToString(), Object ? *Object->GetName() : TEXT("NULL"), Object);
 
 	#if UE_NET_VALIDATE_REFERENCECACHE
 		// Dynamic object might have been referenced before it gets replicated, then it will then already be registered (with a path relative to its outer)
@@ -384,7 +386,7 @@ void FObjectReferenceCache::AddRemoteReference(FNetHandle RefHandle, const UObje
 		}
 		else
 		{
-			LLM_SCOPE(ELLMTag::Iris);
+			LLM_SCOPE_BYTAG(Iris);
 
 			FCachedNetObjectReference CachedObject;
 	
@@ -411,7 +413,7 @@ void FObjectReferenceCache::RemoveReference(FNetHandle RefHandle, const UObject*
 	// Cleanup cached data for dynamic instance
 	if (ensure(RefHandle.IsDynamic()))
 	{
-		UE_LOG_REFERENCECACHE(Verbose, TEXT("ObjectReferenceCache::RemoveReference %s, Object: %s Pointer: %p"), *RefHandle.ToString(), Object ? *Object->GetName() : TEXT( "NULL" ), Object);
+		UE_LOG_REFERENCECACHE(Verbose, TEXT("ObjectReferenceCache::RemoveReference %s, Object: %s Pointer: %p"), *RefHandle.ToString(), Object ? *Object->GetName() : TEXT("NULL"), Object);
 
 		// We create an iterator to avoid having to find the key twice
 		TMap<FNetHandle, FCachedNetObjectReference>::TKeyIterator It = ReferenceHandleToCachedReference.CreateKeyIterator(RefHandle);
@@ -440,7 +442,7 @@ void FObjectReferenceCache::RemoveReference(FNetHandle RefHandle, const UObject*
 					if (ReferenceHandleToCachedReference.RemoveAndCopyValue(SubRefHandle, RemovedSubObjectRef))
 					{
 						const UObject* SubRefObject = RemovedSubObjectRef.Object.Get();
-						UE_LOG_REFERENCECACHE(Verbose, TEXT("ObjectReferenceCache::RemoveReference %s, Referenced dynamic %s, Object: %s Pointer: %p"), *SubRefHandle.ToString(), *RefHandle.ToString(), SubRefObject ? *SubRefObject->GetName() : TEXT( "NULL" ), RemovedSubObjectRef.ObjectKey);
+						UE_LOG_REFERENCECACHE(Verbose, TEXT("ObjectReferenceCache::RemoveReference %s, Referenced dynamic %s, Object: %s Pointer: %p"), *SubRefHandle.ToString(), *RefHandle.ToString(), SubRefObject ? *SubRefObject->GetName() : TEXT("NULL"), RemovedSubObjectRef.ObjectKey);
 
 						ObjectToNetReferenceHandle.Remove(RemovedSubObjectRef.ObjectKey);
 					}
@@ -818,7 +820,7 @@ UObject* FObjectReferenceCache::ResolveObjectReferenceHandleInternal(FNetHandle
 		}
 	}
 
-	LLM_SCOPE(ELLMTag::Iris);
+	LLM_SCOPE_BYTAG(Iris);
 
 	if (bShouldAssign)
 	{
@@ -943,7 +945,7 @@ FNetObjectReference FObjectReferenceCache::GetOrCreateObjectReference(const FStr
 
 	if (Handle.IsValid())
 	{
-		LLM_SCOPE(ELLMTag::Iris);
+		LLM_SCOPE_BYTAG(Iris);
 
 		// Register cached ref
 		FCachedNetObjectReference CachedObject;
@@ -1193,7 +1195,7 @@ void FObjectReferenceCache::ReadFullReference(FNetSerializationContext& Context,
 
 	UE_NET_TRACE_SCOPE(FullNetObjectReference, *Reader, Context.GetTraceCollector(), ENetTraceVerbosity::Trace);
 	UE_NET_TRACE_NAMED_OBJECT_SCOPE(ReferenceScope, FNetHandle(), *Reader, Context.GetTraceCollector(), ENetTraceVerbosity::Verbose);
-	LLM_SCOPE(ELLMTag::Iris);
+	LLM_SCOPE_BYTAG(Iris);
 
 	// Handle client assigned reference
 	if (const bool bIsClientAssignedReference = Reader->ReadBool())
@@ -1278,7 +1280,7 @@ void FObjectReferenceCache::ReadReference(FNetSerializationContext& Context, FNe
 {
 	UE_NET_TRACE_SCOPE(NetObjectReference, *Context.GetBitStreamReader(), Context.GetTraceCollector(), ENetTraceVerbosity::Trace);
 	UE_NET_TRACE_NAMED_OBJECT_SCOPE(ReferenceScope, FNetHandle(), *Context.GetBitStreamReader(), Context.GetTraceCollector(), ENetTraceVerbosity::Verbose);
-	LLM_SCOPE(ELLMTag::Iris);
+	LLM_SCOPE_BYTAG(Iris);
 
 	FNetBitStreamReader* Reader = Context.GetBitStreamReader();
 
diff --git a/Engine/Source/Runtime/Experimental/Iris/Core/Private/Iris/ReplicationSystem/ObjectReplicationBridge.cpp b/Engine/Source/Runtime/Experimental/Iris/Core/Private/Iris/ReplicationSystem/ObjectReplicationBridge.cpp
--- a/Engine/Source/Runtime/Experimental/Iris/Core/Private/Iris/ReplicationSystem/ObjectReplicationBridge.cpp
+++ b/Engine/Source/Runtime/Experimental/Iris/Core/Private/Iris/ReplicationSystem/ObjectReplicationBridge.cpp
@@ -5,6 +5,7 @@
 #include "Iris/IrisConfigInternal.h"
 
 #include "Iris/Core/IrisLog.h"
+#include "Iris/Core/IrisMemoryTracker.h"
 #include "Iris/Core/IrisProfiler.h"
 
 #include "Net/Core/Trace/NetTrace.h"
@@ -110,7 +111,7 @@ UObjectReplicationBridge::~UObjectReplicationBridge()
 
 void UObjectReplicationBridge::Initialize(UReplicationSystem* InReplicationSystem)
 {
-	LLM_SCOPE(ELLMTag::Iris);
+	LLM_SCOPE_BYTAG(Iris);
 
 	Super::Initialize(InReplicationSystem);
 
@@ -171,7 +172,7 @@ UObject* UObjectReplicationBridge::GetReplicatedObject(FNetHandle Handle) const
 
 UE::Net::FNetHandle UObjectReplicationBridge::BeginReplication(UObject* Instance, const FCreateNetHandleParams& Params)
 {
-	LLM_SCOPE(ELLMTag::IrisState);
+	LLM_SCOPE_BYTAG(IrisState);
 
 	using namespace UE::Net;
 	using namespace UE::Net::Private;
@@ -301,7 +302,7 @@ UE::Net::FNetHandle UObjectReplicationBridge::BeginReplication(UObject* Instance
 
 UE::Net::FNetHandle UObjectReplicationBridge::BeginReplication(FNetHandle OwnerHandle, UObject* Instance, FNetHandle InsertRelativeToSubObjectHandle, ESubObjectInsertionOrder InsertionOrder, const FCreateNetHandleParams& Params)
 {
-	LLM_SCOPE(ELLMTag::IrisState);
+	LLM_SCOPE_BYTAG(IrisState);
 
 	using namespace UE::Net;
 	using namespace UE::Net::Private;
@@ -476,7 +477,7 @@ void UObjectReplicationBridge::UnregisterRemoteInstance(FNetHandle Handle, bool
 	
 UE::Net::FNetHandle UObjectReplicationBridge::CreateNetHandleFromRemote(FNetHandle SubObjectOwnerNetHandle, FNetHandle WantedNetHandle, FReplicationBridgeSerializationContext& Context)
 {
-	LLM_SCOPE(ELLMTag::IrisState);
+	LLM_SCOPE_BYTAG(IrisState);
 
 	using namespace UE::Net;
 	using namespace UE::Net::Private;
diff --git a/Engine/Source/Runtime/Experimental/Iris/Core/Private/Iris/ReplicationSystem/ReplicationSystem.cpp b/Engine/Source/Runtime/Experimental/Iris/Core/Private/Iris/ReplicationSystem/ReplicationSystem.cpp
--- a/Engine/Source/Runtime/Experimental/Iris/Core/Private/Iris/ReplicationSystem/ReplicationSystem.cpp
+++ b/Engine/Source/Runtime/Experimental/Iris/Core/Private/Iris/ReplicationSystem/ReplicationSystem.cpp
@@ -3,6 +3,7 @@
 #include "Iris/ReplicationSystem/ReplicationSystem.h"
 #include "Iris/Core/IrisDebugging.h"
 #include "Iris/Core/IrisLog.h"
+#include "Iris/Core/IrisMemoryTracker.h"
 #include "Iris/Core/IrisProfiler.h"
 #include "Iris/Core/NetObjectReference.h"
 #include "Net/Core/Trace/NetTrace.h"
@@ -67,7 +68,7 @@ public:
 
 	void Init(const UReplicationSystem::FReplicationSystemParams& Params)
 	{
-		LLM_SCOPE(ELLMTag::Iris);
+		LLM_SCOPE_BYTAG(Iris);
 
 #if !UE_BUILD_SHIPPING
 		IrisDebugHelperDummy = UE::Net::IrisDebugHelper::Init();
@@ -222,14 +223,14 @@ public:
 	void UpdateFiltering(const FNetBitArrayView& DirtyObjects)
 	{
 		IRIS_PROFILER_SCOPE(FReplicationSystem_UpdateFiltering);
-		LLM_SCOPE(ELLMTag::Iris);
+		LLM_SCOPE_BYTAG(Iris);
 
 		FReplicationFiltering& Filtering = ReplicationSystemInternal.GetFiltering();
 		Filtering.Filter(DirtyObjects);
 
 		// Iterate over all valid connections and propagate updated scopes
 		FReplicationConnections& Connections = ReplicationSystemInternal.GetConnections();
-		auto&& UpdateConnectionScope = [&Filtering, &Connections](uint32 ConnectionId)
+		auto UpdateConnectionScope = [&Filtering, &Connections](uint32 ConnectionId)
 		{
 			FReplicationConnection* Conn = Connections.GetConnection(ConnectionId);
 			const FNetBitArrayView& ObjectsInScope = Filtering.GetObjectsInScope(ConnectionId);
@@ -252,7 +253,7 @@ public:
 	void UpdatePrioritization(const FNetBitArrayView& ReplicatingConnections, const FNetBitArrayView& DirtyObjects)
 	{
 		IRIS_PROFILER_SCOPE(FReplicationSystem::FImpl::UpdatePrioritization);
-		LLM_SCOPE(ELLMTag::Iris);
+		LLM_SCOPE_BYTAG(Iris);
 
 		FReplicationPrioritization& Prioritization = ReplicationSystemInternal.GetPrioritization();
 		Prioritization.Prioritize(ReplicatingConnections, DirtyObjects);
@@ -266,7 +267,7 @@ public:
 		const FChangeMaskCache& UpdatedChangeMasks = ReplicationSystemInternal.GetChangeMaskCache();
 
 		// Iterate over connections and propagate dirty changemasks
-		auto&& UpdateDirtyChangeMasks = [&Connections, &UpdatedChangeMasks](uint32 ConnectionId)
+		auto UpdateDirtyChangeMasks = [&Connections, &UpdatedChangeMasks](uint32 ConnectionId)
 		{
 			FReplicationConnection* Conn = Connections.GetConnection(ConnectionId);
 			Conn->ReplicationWriter->UpdateDirtyChangeMasks(UpdatedChangeMasks);
@@ -278,7 +279,7 @@ public:
 	void CopyDirtyStateData()
 	{
 		IRIS_PROFILER_SCOPE(FReplicationSystem_CopyDirtyStateData);
-		LLM_SCOPE(ELLMTag::IrisState);
+		LLM_SCOPE_BYTAG(IrisState);
 
 		FNetHandleManager& NetHandleManager = ReplicationSystemInternal.GetNetHandleManager();
 		FChangeMaskCache& Cache = ReplicationSystemInternal.GetChangeMaskCache();
@@ -331,7 +332,7 @@ public:
 
 	void AddConnection(uint32 ConnectionId)
 	{
-		LLM_SCOPE(ELLMTag::IrisConnection);
+		LLM_SCOPE_BYTAG(IrisConnection);
 
 		FReplicationConnections& Connections = ReplicationSystemInternal.GetConnections();
 		{
@@ -413,7 +414,7 @@ public:
 	void UpdateUnresolvableReferenceTracking()
 	{
 		FReplicationConnections& Connections = ReplicationSystemInternal.GetConnections();
-		auto&& UpdateUnresolvableReferenceTracking = [&Connections](uint32 ConnectionId)
+		auto UpdateUnresolvableReferenceTracking = [&Connections](uint32 ConnectionId)
 		{
 			FReplicationConnection* Conn = Connections.GetConnection(ConnectionId);
 			Conn->ReplicationReader->UpdateUnresolvableReferenceTracking();
@@ -437,7 +438,7 @@ UReplicationSystem::UReplicationSystem()
 
 void UReplicationSystem::Init(uint32 InId, const FReplicationSystemParams& Params)
 {
-	LLM_SCOPE(ELLMTag::Iris);
+	LLM_SCOPE_BYTAG(Iris);
 
 	Id = InId;
 	bIsServer = Params.bIsServer;
@@ -911,14 +912,14 @@ void UReplicationSystem::RemoveSubObjectFilter(FName GroupName)
 
 UE::Net::FNetObjectGroupHandle UReplicationSystem::CreateGroup()
 {
-	LLM_SCOPE(ELLMTag::Iris);
+	LLM_SCOPE_BYTAG(Iris);
 
 	return Impl->ReplicationSystemInternal.GetGroups().CreateGroup();
 }
 
 void UReplicationSystem::AddToGroup(FNetObjectGroupHandle GroupHandle, FNetHandle Handle)
 {
-	LLM_SCOPE(ELLMTag::Iris);
+	LLM_SCOPE_BYTAG(Iris);
 
 	using namespace UE::Net::Private;
 
diff --git a/Engine/Source/Runtime/Experimental/Iris/Core/Private/Iris/ReplicationSystem/WorldLocations.cpp b/Engine/Source/Runtime/Experimental/Iris/Core/Private/Iris/ReplicationSystem/WorldLocations.cpp
--- a/Engine/Source/Runtime/Experimental/Iris/Core/Private/Iris/ReplicationSystem/WorldLocations.cpp
+++ b/Engine/Source/Runtime/Experimental/Iris/Core/Private/Iris/ReplicationSystem/WorldLocations.cpp
@@ -1,7 +1,7 @@
 // Copyright Epic Games, Inc. All Rights Reserved.
 
 #include "Iris/ReplicationSystem/WorldLocations.h"
-#include "HAL/LowLevelMemTracker.h"
+#include "Iris/Core/IrisMemoryTracker.h"
 
 namespace UE::Net
 {
@@ -18,7 +18,7 @@ void FWorldLocations::SetHasWorldLocation(uint32 ObjectIndex, bool bHasWorldLoca
 	{
 		if (ObjectIndex >= uint32(WorldLocations.Num()))
 		{
-			LLM_SCOPE(ELLMTag::Iris);
+			LLM_SCOPE_BYTAG(Iris);
 			WorldLocations.Add(ObjectIndex + 1U - WorldLocations.Num());
 		}
 		WorldLocations[ObjectIndex] = FVector::Zero();
diff --git a/Engine/Source/Runtime/Experimental/Iris/Core/Public/Iris/Core/IrisMemoryTracker.h b/Engine/Source/Runtime/Experimental/Iris/Core/Public/Iris/Core/IrisMemoryTracker.h
new file mode 100644
--- /dev/null
+++ b/Engine/Source/Runtime/Experimental/Iris/Core/Public/Iris/Core/IrisMemoryTracker.h
@@ -0,0 +1,9 @@
+// Copyright Epic Games, Inc. All Rights Reserved.
+
+#pragma once
+
+#include "HAL/LowLevelMemTracker.h"
+
+LLM_DECLARE_TAG_API(Iris, IRISCORE_API);
+LLM_DECLARE_TAG_API(IrisState, IRISCORE_API);
+LLM_DECLARE_TAG_API(IrisConnection, IRISCORE_API);
