Chaos : Flesh
Add template to corotated constraint. Add finite difference test to corotated constraint.

diff --git a/Engine/Source/Runtime/Experimental/Chaos/Public/Chaos/XPBDCorotatedConstraints.h b/Engine/Source/Runtime/Experimental/Chaos/Public/Chaos/XPBDCorotatedConstraints.h
--- a/Engine/Source/Runtime/Experimental/Chaos/Public/Chaos/XPBDCorotatedConstraints.h
+++ b/Engine/Source/Runtime/Experimental/Chaos/Public/Chaos/XPBDCorotatedConstraints.h
@@ -13,45 +13,45 @@ DECLARE_CYCLE_STAT(TEXT("Chaos XPBD Corotated Constraint Det Compute"), STAT_Cha
 
 namespace Chaos::Softs
 {
-
+	template <typename T, typename ParticleType>
 	class FXPBDCorotatedConstraints 
 	{
 
-		//TODO(Yizhou Chen): COuld be optimized . The SVD is using double so float -> double -> float all the time
+		//TODO(Yizhou Chen): COuld be optimized . The SVD is using Chaos::Softs::Tdouble so float -> double -> float all the time
 		//should change data type in accordance 
 
 	public:
 		//this one only accepts tetmesh input and mesh
 		FXPBDCorotatedConstraints(
-			const Chaos::Softs::FSolverParticles& InParticles,
+			const ParticleType& InParticles,
 			const TArray<TVector<int32, 4>>& InMesh,
 			const bool bRecordMetricIn = true,
-			const Chaos::Softs::FSolverReal& EMesh = (FSolverReal)10.0,
-			const Chaos::Softs::FSolverReal& NuMesh = (FSolverReal).3
+			const T& EMesh = (T)10.0,
+			const T& NuMesh = (T).3
 			)
 			: bRecordMetric(bRecordMetricIn), MeshConstraints(InMesh)
 		{	
 
-			LambdaArray.Init((FSolverReal)0., 2 * MeshConstraints.Num());
-			DmInverse.Init((FSolverReal)0., 9 * MeshConstraints.Num());
-			Measure.Init((FSolverReal)0.,  MeshConstraints.Num());
-			Lambda = EMesh * NuMesh / (((FSolverReal)1. + NuMesh) * ((FSolverReal)1. - (FSolverReal)2. * NuMesh));
-			Mu = EMesh / ((FSolverReal)2. * ((FSolverReal)1. + NuMesh));
+			LambdaArray.Init((T)0., 2 * MeshConstraints.Num());
+			DmInverse.Init((T)0., 9 * MeshConstraints.Num());
+			Measure.Init((T)0.,  MeshConstraints.Num());
+			Lambda = EMesh * NuMesh / (((T)1. + NuMesh) * ((T)1. - (T)2. * NuMesh));
+			Mu = EMesh / ((T)2. * ((T)1. + NuMesh));
 			for (int e = 0; e < InMesh.Num(); e++)
 			{
-				PMatrix<FSolverReal, 3, 3> Dm = DsInit(e, InParticles);
-				PMatrix<FSolverReal, 3, 3> DmInv = Dm.Inverse();
+				PMatrix<T, 3, 3> Dm = DsInit(e, InParticles);
+				PMatrix<T, 3, 3> DmInv = Dm.Inverse();
 				for (int r = 0; r < 3; r++) {
 					for (int c = 0; c < 3; c++) {
 						DmInverse[(3 * 3) * e + 3 * r + c] = DmInv.GetAt(r, c);
 					}
 				}
 
-				Measure[e] = Dm.Determinant() / (FSolverReal)6.;
+				Measure[e] = Dm.Determinant() / (T)6.;
 
 				//part of preprocessing: if inverted element is found, 
 				//invert it so that the measure is positive
-				if (Measure[e] < (FSolverReal)0.)
+				if (Measure[e] < (T)0.)
 				{
 
 					Measure[e] = -Measure[e];
@@ -64,8 +64,8 @@ namespace Chaos::Softs
 
 		virtual ~FXPBDCorotatedConstraints() {}
 
-		PMatrix<FSolverReal, 3, 3> DsInit(const int e, const FSolverParticles& InParticles) const {
-			PMatrix<FSolverReal, 3, 3> Result((FSolverReal)0.);
+		PMatrix<T, 3, 3> DsInit(const int e, const ParticleType& InParticles) const {
+			PMatrix<T, 3, 3> Result((T)0.);
 			for (int i = 0; i < 3; i++) {
 				for (int c = 0; c < 3; c++) {
 					Result.SetAt(c, i, InParticles.X(MeshConstraints[e][i + 1])[c] - InParticles.X(MeshConstraints[e][0])[c]);
@@ -75,8 +75,8 @@ namespace Chaos::Softs
 		}
 
 
-		PMatrix<FSolverReal, 3, 3> Ds(const int e, const FSolverParticles& InParticles) const {
-			PMatrix<FSolverReal, 3, 3> Result((FSolverReal)0.);
+		PMatrix<T, 3, 3> Ds(const int e, const ParticleType& InParticles) const {
+			PMatrix<T, 3, 3> Result((T)0.);
 			for (int i = 0; i < 3; i++) {
 				for (int c = 0; c < 3; c++) {
 					Result.SetAt(c, i, InParticles.P(MeshConstraints[e][i+1])[c] - InParticles.P(MeshConstraints[e][0])[c]);
@@ -86,12 +86,12 @@ namespace Chaos::Softs
 		}
 
 
-		PMatrix<FSolverReal, 3, 3> F(const int e, const FSolverParticles& InParticles) const {
+		PMatrix<T, 3, 3> F(const int e, const ParticleType& InParticles) const {
 			return ElementDmInv(e) * Ds(e, InParticles);
 		}
 
-		PMatrix<FSolverReal, 3, 3> ElementDmInv(const int e) const {
-			PMatrix<FSolverReal, 3, 3> DmInv((FSolverReal)0.);
+		PMatrix<T, 3, 3> ElementDmInv(const int e) const {
+			PMatrix<T, 3, 3> DmInv((T)0.);
 			for (int r = 0; r < 3; r++) {
 				for (int c = 0; c < 3; c++) {
 					DmInv.SetAt(r, c, DmInverse[(3 * 3) * e + 3 * r + c]);
@@ -102,20 +102,20 @@ namespace Chaos::Softs
 		
 		void Init() const 
 		{
-			for (FSolverReal& Lambdas : LambdaArray) { Lambdas = (FSolverReal)0.; }
+			for (T& Lambdas : LambdaArray) { Lambdas = (T)0.; }
 		}
 
-		virtual void ApplyInSerial(FSolverParticles& Particles, const FSolverReal Dt, const int32 ElementIndex) const
+		virtual void ApplyInSerial(ParticleType& Particles, const T Dt, const int32 ElementIndex) const
 		{
 			TRACE_CPUPROFILER_EVENT_SCOPE(TEXT("STAT_ChaosXPBDCorotatedApplySingle"));
-			TVec4<FSolverVec3> PolarDelta = GetPolarDelta(Particles, Dt, ElementIndex);
+			TVec4<TVector<T, 3>> PolarDelta = GetPolarDelta(Particles, Dt, ElementIndex);
 
 			for (int i = 0; i < 4; i++) 
 			{
 				Particles.P(MeshConstraints[ElementIndex][i]) += PolarDelta[i];
 			}
 			
-			TVec4<FSolverVec3> DetDelta = GetDeterminantDelta(Particles, Dt, ElementIndex);
+			TVec4<TVector<T, 3>> DetDelta = GetDeterminantDelta(Particles, Dt, ElementIndex);
 
 			for (int i = 0; i < 4; i++)
 			{
@@ -125,7 +125,7 @@ namespace Chaos::Softs
 
 		}
 
-		void ApplyInSerial(FSolverParticles& Particles, const FSolverReal Dt) const
+		void ApplyInSerial(ParticleType& Particles, const T Dt) const
 		{
 			SCOPE_CYCLE_COUNTER(STAT_ChaosXPBDCorotated);
 			TRACE_CPUPROFILER_EVENT_SCOPE(TEXT("STAT_ChaosXPBDCorotatedApplySerial"));
@@ -137,13 +137,13 @@ namespace Chaos::Softs
 
 		}
 
-		void ApplyInParallel(FSolverParticles& Particles, const FSolverReal Dt) const
+		void ApplyInParallel(ParticleType& Particles, const T Dt) const
 		{	
 			//code for error metric:
 			if (bRecordMetric)
 			{
-				GError.Init((FSolverReal)0., 3 * Particles.Size());
-				HErrorArray.Init((FSolverReal)0., 2 * MeshConstraints.Num());
+				GError.Init((T)0., 3 * Particles.Size());
+				HErrorArray.Init((T)0., 2 * MeshConstraints.Num());
 			}
 			
 			{
@@ -169,9 +169,101 @@ namespace Chaos::Softs
 
 		}
 
+		TVec4<TVector<T, 3>> GetPolarGradient(const PMatrix<T, 3, 3>& Fe, const PMatrix<T, 3, 3>& Re, const PMatrix<T, 3, 3>& DmInvT, const T C1) const
+		{
+			TVector<T, 81> dRdF((T)0.);
+			Chaos::dRdFCorotated(Fe, dRdF);
+			TVec4<TVector<T, 3>> dC1(TVector<T, 3>((T)0.));
+			//dC1 = dC1dF * dFdX
+			for (int alpha = 0; alpha < 3; alpha++) {
+				for (int l = 0; l < 3; l++) {
+
+					dC1[0][alpha] += (DmInvT * Re).GetAt(alpha, l) - (DmInvT * Fe).GetAt(alpha, l);
+
+				}
+			}
+			for (int ie = 0; ie < 3; ie++) {
+				for (int alpha = 0; alpha < 3; alpha++) {
+					dC1[ie + 1][alpha] = (DmInvT * Fe).GetAt(alpha, ie) - (DmInvT * Re).GetAt(alpha, ie);
+				}
+			}
+			//it's really ie-1 here
+			for (int ie = 0; ie < 3; ie++) {
+				for (int alpha = 0; alpha < 3; alpha++) {
+					for (int m = 0; m < 3; m++) {
+						for (int n = 0; n < 3; n++) {
+							for (int j = 0; j < 3; j++) {
+								dC1[ie + 1][alpha] -= (Fe.GetAt(m, n) - Re.GetAt(m, n)) * dRdF[9 * (alpha * 3 + j) + m * 3 + n] * DmInvT.GetAt(j, ie);
+							}
+						}
+					}
+				}
+			}
+			for (int alpha = 0; alpha < 3; alpha++) {
+				for (int m = 0; m < 3; m++) {
+					for (int n = 0; n < 3; n++) {
+						for (int l = 0; l < 3; l++) {
+							for (int j = 0; j < 3; j++) {
+								dC1[0][alpha] += (Fe.GetAt(m, n) - Re.GetAt(m, n)) * dRdF[9 * (alpha * 3 + j) + m * 3 + n] * DmInvT.GetAt(j, l);
+							}
+						}
+					}
+				}
+			}
+
+			if (C1 != 0)
+			{
+				for (int i = 0; i < 4; i++)
+				{
+					for (int j = 0; j < 3; j++)
+					{
+						dC1[i][j] /= C1;
+					}
+				}
+			}
+			return dC1;
+
+		}
+
+		TVec4<TVector<T, 3>> GetDeterminantGradient(const PMatrix<T, 3, 3>& Fe, const PMatrix<T, 3, 3>& DmInvT) const
+		{
+			TRACE_CPUPROFILER_EVENT_SCOPE(TEXT("STAT_ChaosXPBDCorotatedGetDetGradient"));
+			//SCOPE_CYCLE_COUNTER(STAT_ChaosXPBDCorotatedDet);
+			//const TVec4<int32>& Constraint = MeshConstraints[ElementIndex];
+
+			//const PMatrix<T, 3, 3> Fe = F(ElementIndex, Particles);
+			//PMatrix<T, 3, 3> DmInvT = ElementDmInv(ElementIndex).GetTransposed();
+			TVec4<TVector<T, 3>> dC2(TVector<T, 3>((T)0.));
+
+			PMatrix<T, 3, 3> JFinvT;
+			JFinvT.SetAt(0, 0, Fe.GetAt(1, 1) * Fe.GetAt(2, 2) - Fe.GetAt(2, 1) * Fe.GetAt(1, 2));
+			JFinvT.SetAt(0, 1, Fe.GetAt(2, 0) * Fe.GetAt(1, 2) - Fe.GetAt(1, 0) * Fe.GetAt(2, 2));
+			JFinvT.SetAt(0, 2, Fe.GetAt(1, 0) * Fe.GetAt(2, 1) - Fe.GetAt(2, 0) * Fe.GetAt(1, 1));
+			JFinvT.SetAt(1, 0, Fe.GetAt(2, 1) * Fe.GetAt(0, 2) - Fe.GetAt(0, 1) * Fe.GetAt(2, 2));
+			JFinvT.SetAt(1, 1, Fe.GetAt(0, 0) * Fe.GetAt(2, 2) - Fe.GetAt(2, 0) * Fe.GetAt(0, 2));
+			JFinvT.SetAt(1, 2, Fe.GetAt(2, 0) * Fe.GetAt(0, 1) - Fe.GetAt(0, 0) * Fe.GetAt(2, 1));
+			JFinvT.SetAt(2, 0, Fe.GetAt(0, 1) * Fe.GetAt(1, 2) - Fe.GetAt(1, 1) * Fe.GetAt(0, 2));
+			JFinvT.SetAt(2, 1, Fe.GetAt(1, 0) * Fe.GetAt(0, 2) - Fe.GetAt(0, 0) * Fe.GetAt(1, 2));
+			JFinvT.SetAt(2, 2, Fe.GetAt(0, 0) * Fe.GetAt(1, 1) - Fe.GetAt(1, 0) * Fe.GetAt(0, 1));
+
+			PMatrix<T, 3, 3> JinvTDmInvT = DmInvT * JFinvT;
+
+			for (int ie = 0; ie < 3; ie++) {
+				for (int alpha = 0; alpha < 3; alpha++) {
+					dC2[ie + 1][alpha] = JinvTDmInvT.GetAt(alpha, ie);
+				}
+			}
+			for (int alpha = 0; alpha < 3; alpha++) {
+				for (int l = 0; l < 3; l++) {
+					dC2[0][alpha] -= JinvTDmInvT.GetAt(alpha, l);
+				}
+			}
+			return dC2;
+		}
+
 	private:
 
-		void InitColor(const FSolverParticles& Particles)
+		void InitColor(const ParticleType& Particles)
 		{
 
 			{
@@ -179,8 +271,8 @@ namespace Chaos::Softs
 
 				// Reorder constraints based on color so each array in ConstraintsPerColor contains contiguous elements.
 				TArray<TVec4<int32>> ReorderedConstraints;
-				TArray<FSolverReal> ReorderedMeasure;
-				TArray<FSolverReal> ReorderedDmInverse;
+				TArray<T> ReorderedMeasure;
+				TArray<T> ReorderedDmInverse;
 				TArray<int32> OrigToReorderedIndices; // used to reorder stiffness indices
 				ReorderedConstraints.SetNumUninitialized(MeshConstraints.Num());
 				ReorderedMeasure.SetNumUninitialized(Measure.Num());
@@ -217,48 +309,27 @@ namespace Chaos::Softs
 
 	protected:
 
+		
+
 
-		TVec4<FSolverVec3> GetDeterminantDelta(const FSolverParticles& Particles, const FSolverReal Dt, const int32 ElementIndex, const FSolverReal Tol = 1e-3) const
+		TVec4<TVector<T, 3>> GetDeterminantDelta(const ParticleType& Particles, const T Dt, const int32 ElementIndex, const T Tol = 1e-3) const
 		{
 			TRACE_CPUPROFILER_EVENT_SCOPE(TEXT("STAT_ChaosXPBDCorotatedApplyDet"));
-			SCOPE_CYCLE_COUNTER(STAT_ChaosXPBDCorotatedDet);
-			const TVec4<int32>& Constraint = MeshConstraints[ElementIndex];
+			//SCOPE_CYCLE_COUNTER(STAT_ChaosXPBDCorotatedDet);
+			////const TVec4<int32>& Constraint = MeshConstraints[ElementIndex];
 
-			const PMatrix<FSolverReal, 3, 3> Fe = F(ElementIndex, Particles);
-			PMatrix<FSolverReal, 3, 3> DmInvT = ElementDmInv(ElementIndex).GetTransposed();
-			TVec4<FSolverVec3> dC2(FSolverVec3((FSolverReal)0.));
-
-			FSolverReal J = Fe.Determinant();
+			const PMatrix<T, 3, 3> Fe = F(ElementIndex, Particles);
+			PMatrix<T, 3, 3> DmInvT = ElementDmInv(ElementIndex).GetTransposed();
+			
+			T J = Fe.Determinant();
 			if (J - 1 < Tol)
 			{
-				return TVec4<FSolverVec3>(FSolverVec3((FSolverReal)0.));
+				return TVec4<TVector<T, 3>>(TVector<T, 3>((T)0.));
 			}
 
-			PMatrix<FSolverReal, 3, 3> JFinvT;
-			JFinvT.SetAt(0, 0, Fe.GetAt(1, 1) * Fe.GetAt(2, 2) - Fe.GetAt(2, 1) * Fe.GetAt(1, 2));
-			JFinvT.SetAt(0, 1, Fe.GetAt(2, 0) * Fe.GetAt(1, 2) - Fe.GetAt(1, 0) * Fe.GetAt(2, 2));
-			JFinvT.SetAt(0, 2, Fe.GetAt(1, 0) * Fe.GetAt(2, 1) - Fe.GetAt(2, 0) * Fe.GetAt(1, 1));
-			JFinvT.SetAt(1, 0, Fe.GetAt(2, 1) * Fe.GetAt(0, 2) - Fe.GetAt(0, 1) * Fe.GetAt(2, 2));
-			JFinvT.SetAt(1, 1, Fe.GetAt(0, 0) * Fe.GetAt(2, 2) - Fe.GetAt(2, 0) * Fe.GetAt(0, 2));
-			JFinvT.SetAt(1, 2, Fe.GetAt(2, 0) * Fe.GetAt(0, 1) - Fe.GetAt(0, 0) * Fe.GetAt(2, 1));
-			JFinvT.SetAt(2, 0, Fe.GetAt(0, 1) * Fe.GetAt(1, 2) - Fe.GetAt(1, 1) * Fe.GetAt(0, 2));
-			JFinvT.SetAt(2, 1, Fe.GetAt(1, 0) * Fe.GetAt(0, 2) - Fe.GetAt(0, 0) * Fe.GetAt(1, 2));
-			JFinvT.SetAt(2, 2, Fe.GetAt(0, 0) * Fe.GetAt(1, 1) - Fe.GetAt(1, 0) * Fe.GetAt(0, 1));
-
-			PMatrix<FSolverReal, 3, 3> JinvTDmInvT = DmInvT * JFinvT;
-
-			for (int ie = 0; ie < 3; ie++) {
-				for (int alpha = 0; alpha < 3; alpha++) {
-					dC2[ie + 1][alpha] = JinvTDmInvT.GetAt(alpha, ie);
-				}
-			}
-			for (int alpha = 0; alpha < 3; alpha++) {
-				for (int l = 0; l < 3; l++) {
-					dC2[0][alpha] -= JinvTDmInvT.GetAt(alpha, l);
-				}
-			}
+			TVec4<TVector<T, 3>> dC2 = GetDeterminantGradient(Fe, DmInvT);
 
-			FSolverReal AlphaTilde = (FSolverReal)2. / (Dt * Dt * Lambda * Measure[ElementIndex]);
+			T AlphaTilde = (T)2. / (Dt * Dt * Lambda * Measure[ElementIndex]);
 
 			if (bRecordMetric)
 			{
@@ -274,9 +345,9 @@ namespace Chaos::Softs
 			}
 			
 
-			FSolverReal DLambda = (1 - J) - AlphaTilde * LambdaArray[2 * ElementIndex + 1];
+			T DLambda = (1 - J) - AlphaTilde * LambdaArray[2 * ElementIndex + 1];
 
-			FSolverReal Denom = AlphaTilde;
+			T Denom = AlphaTilde;
 			for (int i = 0; i < 4; i++)
 			{
 				for (int j = 0; j < 3; j++)
@@ -286,7 +357,7 @@ namespace Chaos::Softs
 			}
 			DLambda /= Denom;
 			LambdaArray[2 * ElementIndex + 1] += DLambda;
-			TVec4<FSolverVec3> Delta(FSolverVec3((FSolverReal)0.));
+			TVec4<TVector<T, 3>> Delta(TVector<T, 3>((T)0.));
 			for (int i = 0; i < 4; i++)
 			{
 				for (int j = 0; j < 3; j++)
@@ -297,17 +368,19 @@ namespace Chaos::Softs
 			return Delta;
 		}
 
-		TVec4<FSolverVec3> GetPolarDelta(const FSolverParticles& Particles, const FSolverReal Dt, const int32 ElementIndex, const FSolverReal Tol = 1e-3) const
+
+
+		TVec4<TVector<T, 3>> GetPolarDelta(const ParticleType& Particles, const T Dt, const int32 ElementIndex, const T Tol = 1e-3) const
 		{	
 			TRACE_CPUPROFILER_EVENT_SCOPE(TEXT("STAT_ChaosXPBDCorotatedApplyPolar"));
 			SCOPE_CYCLE_COUNTER(STAT_ChaosXPBDCorotatedPolar);
-			const PMatrix<FSolverReal, 3, 3> Fe = F(ElementIndex, Particles);
+			const PMatrix<T, 3, 3> Fe = F(ElementIndex, Particles);
 
-			PMatrix<FSolverReal, 3, 3> Re((FSolverReal)0.), Se((FSolverReal)0.);
+			PMatrix<T, 3, 3> Re((T)0.), Se((T)0.);
 
 			Chaos::PolarDecomposition(Fe, Re, Se);
 
-			FSolverReal C1 = (FSolverReal)0.;
+			T C1 = (T)0.;
 			for (int i = 0; i < 3; i++)
 			{
 				for (int j = 0; j < 3; j++)
@@ -319,66 +392,68 @@ namespace Chaos::Softs
 
 			if (C1 < Tol)
 			{
-				return TVec4<FSolverVec3>(FSolverVec3((FSolverReal)0.));
+				return TVec4<TVector<T, 3>>(TVector<T, 3>((T)0.));
 			}
 
-			TVector<FSolverReal, 81> dRdF((FSolverReal)0.);
+			TVector<T, 81> dRdF((T)0.);
 			Chaos::dRdFCorotated(Fe, dRdF);
 
-			PMatrix<FSolverReal, 3, 3> DmInvT = ElementDmInv(ElementIndex).GetTransposed();
-
-			//TODO: deifnitely test the initialization
-			TVec4<FSolverVec3> dC1(FSolverVec3((FSolverReal)0.));
-			//dC1 = dC1dF * dFdX
-			for (int alpha = 0; alpha < 3; alpha++) {
-				for (int l = 0; l < 3; l++) {
-
-					dC1[0][alpha] += (DmInvT * Re).GetAt(alpha, l) - (DmInvT * Fe).GetAt(alpha, l);
-
-				}
-			}
-			for (int ie = 0; ie < 3; ie++) {
-				for (int alpha = 0; alpha < 3; alpha++) {
-					dC1[ie+1][alpha] = (DmInvT * Fe).GetAt(alpha, ie) - (DmInvT * Re).GetAt(alpha, ie);
-				}
-			}
-			//it's really ie-1 here
-			for (int ie = 0; ie < 3; ie++) {
-				for (int alpha = 0; alpha < 3; alpha++) {
-					for (int m = 0; m < 3; m++) {
-						for (int n = 0; n < 3; n++) {
-							for (int j = 0; j < 3; j++) {
-								dC1[ie + 1][alpha] -= (Fe.GetAt(m, n) - Re.GetAt(m, n)) * dRdF[9*(alpha*3+j)+ m * 3 + n] * DmInvT.GetAt(j, ie);
-							}
-						}
-					}
-				}
-			}
-			for (int alpha = 0; alpha < 3; alpha++) {
-				for (int m = 0; m < 3; m++) {
-					for (int n = 0; n < 3; n++) {
-						for (int l = 0; l < 3; l++) {
-							for (int j = 0; j < 3; j++) {
-								dC1[0][alpha] += (Fe.GetAt(m, n) - Re.GetAt(m, n)) * dRdF[9*(alpha*3+j) + m * 3 + n] * DmInvT.GetAt(j, l);
-							}
-						}
-					}
-				}
-			}
-
-			
-			if (C1 != 0)
-			{
-				for (int i = 0; i < 4; i++)
-				{
-					for (int j = 0; j < 3; j++)
-					{
-						dC1[i][j] /= C1;
-					}
-				}
-			}
-
-			FSolverReal AlphaTilde = (FSolverReal)1. / (Dt * Dt * Mu * Measure[ElementIndex]);
+			PMatrix<T, 3, 3> DmInvT = ElementDmInv(ElementIndex).GetTransposed();
+
+			////TODO: deifnitely test the initialization
+			//TVec4<TVector<T, 3>> dC1(TVector<T, 3>((T)0.));
+			////dC1 = dC1dF * dFdX
+			//for (int alpha = 0; alpha < 3; alpha++) {
+			//	for (int l = 0; l < 3; l++) {
+
+			//		dC1[0][alpha] += (DmInvT * Re).GetAt(alpha, l) - (DmInvT * Fe).GetAt(alpha, l);
+
+			//	}
+			//}
+			//for (int ie = 0; ie < 3; ie++) {
+			//	for (int alpha = 0; alpha < 3; alpha++) {
+			//		dC1[ie+1][alpha] = (DmInvT * Fe).GetAt(alpha, ie) - (DmInvT * Re).GetAt(alpha, ie);
+			//	}
+			//}
+			////it's really ie-1 here
+			//for (int ie = 0; ie < 3; ie++) {
+			//	for (int alpha = 0; alpha < 3; alpha++) {
+			//		for (int m = 0; m < 3; m++) {
+			//			for (int n = 0; n < 3; n++) {
+			//				for (int j = 0; j < 3; j++) {
+			//					dC1[ie + 1][alpha] -= (Fe.GetAt(m, n) - Re.GetAt(m, n)) * dRdF[9*(alpha*3+j)+ m * 3 + n] * DmInvT.GetAt(j, ie);
+			//				}
+			//			}
+			//		}
+			//	}
+			//}
+			//for (int alpha = 0; alpha < 3; alpha++) {
+			//	for (int m = 0; m < 3; m++) {
+			//		for (int n = 0; n < 3; n++) {
+			//			for (int l = 0; l < 3; l++) {
+			//				for (int j = 0; j < 3; j++) {
+			//					dC1[0][alpha] += (Fe.GetAt(m, n) - Re.GetAt(m, n)) * dRdF[9*(alpha*3+j) + m * 3 + n] * DmInvT.GetAt(j, l);
+			//				}
+			//			}
+			//		}
+			//	}
+			//}
+
+			//
+			//if (C1 != 0)
+			//{
+			//	for (int i = 0; i < 4; i++)
+			//	{
+			//		for (int j = 0; j < 3; j++)
+			//		{
+			//			dC1[i][j] /= C1;
+			//		}
+			//	}
+			//}
+
+			TVec4<TVector<T, 3>> dC1 = GetPolarGradient(Fe, Re, DmInvT, C1);
+
+			T AlphaTilde = (T)1. / (Dt * Dt * Mu * Measure[ElementIndex]);
 
 			if (bRecordMetric)
 			{
@@ -393,9 +468,9 @@ namespace Chaos::Softs
 				HErrorArray[2 * ElementIndex + 0] = C1 + AlphaTilde * LambdaArray[2 * ElementIndex + 0];
 			}
 
-			FSolverReal DLambda = -C1 - AlphaTilde* LambdaArray[2 * ElementIndex + 0];
+			T DLambda = -C1 - AlphaTilde* LambdaArray[2 * ElementIndex + 0];
 
-			FSolverReal Denom = AlphaTilde;
+			T Denom = AlphaTilde;
 			for (int i = 0; i < 4; i++)
 			{
 				for (int j = 0; j < 3; j++)
@@ -405,7 +480,7 @@ namespace Chaos::Softs
 			}
 			DLambda /= Denom;
 			LambdaArray[2 * ElementIndex + 0] += DLambda;
-			TVec4<FSolverVec3> Delta(FSolverVec3((FSolverReal)0.));
+			TVec4<TVector<T, 3>> Delta(TVector<T, 3>((T)0.));
 			for (int i = 0; i < 4; i++) 
 			{
 				for (int j = 0; j < 3; j++)
@@ -419,21 +494,21 @@ namespace Chaos::Softs
 
 
 	protected:
-		mutable TArray<FSolverReal> LambdaArray;
-		mutable TArray<FSolverReal> DmInverse;
+		mutable TArray<T> LambdaArray;
+		mutable TArray<T> DmInverse;
 
 		//material constants calculated from E:
-		FSolverReal Mu;
-		FSolverReal Lambda;
-		mutable FSolverReal HError;
-		mutable TArray<FSolverReal> HErrorArray;
+		T Mu;
+		T Lambda;
+		mutable T HError;
+		mutable TArray<T> HErrorArray;
 		bool bRecordMetric;
 
 		TArray<TVector<int32, 4>> MeshConstraints;
-		mutable TArray<FSolverReal> Measure;
-		FSolverParticles RestParticles;
+		mutable TArray<T> Measure;
+		ParticleType RestParticles;
 		TArray<int32> ConstraintsPerColorStartIndex; // Constraints are ordered so each batch is contiguous. This is ColorNum + 1 length so it can be used as start and end.
-		mutable TArray<FSolverReal> GError;
+		mutable TArray<T> GError;
 };
 
 }  // End namespace Chaos::Softs
diff --git a/Engine/Source/Runtime/Experimental/Chaos/Public/Chaos/XPBDCorotatedFiberConstraints.h b/Engine/Source/Runtime/Experimental/Chaos/Public/Chaos/XPBDCorotatedFiberConstraints.h
--- a/Engine/Source/Runtime/Experimental/Chaos/Public/Chaos/XPBDCorotatedFiberConstraints.h
+++ b/Engine/Source/Runtime/Experimental/Chaos/Public/Chaos/XPBDCorotatedFiberConstraints.h
@@ -11,11 +11,11 @@
 
 namespace Chaos::Softs
 {
-
-	class FXPBDCorotatedFiberConstraints : public FXPBDCorotatedConstraints
+	template <typename T, typename ParticleType>
+	class FXPBDCorotatedFiberConstraints : public FXPBDCorotatedConstraints <T, ParticleType>
 	{
 
-		typedef FXPBDCorotatedConstraints Base;
+		typedef FXPBDCorotatedConstraints<T, ParticleType> Base;
 		using Base::MeshConstraints;
 		using Base::LambdaArray;
 		using Base::Measure;
@@ -24,20 +24,24 @@ namespace Chaos::Softs
 	public:
 		//this one only accepts tetmesh input and mesh
 		FXPBDCorotatedFiberConstraints(
-			const FSolverParticles& InParticles,
+			const ParticleType& InParticles,
 			const TArray<TVector<int32, 4>>& InMesh,
 			const bool bRecordMetricIn = true,
-			const FSolverReal& EMesh = (FSolverReal)10.0,
-			const FSolverReal& NuMesh = (FSolverReal).3,
-			const FSolverVec3 InFiberDir = FSolverVec3((FSolverReal)1., (FSolverReal)0., (FSolverReal)0.),
-			const FSolverReal InSigmaMax = (FSolverReal)3e5
+			const T& EMesh = (T)10.0,
+			const T& NuMesh = (T).3,
+			const TVector<T, 3> InFiberDir = TVector<T, 3>((T)1., (T)0., (T)0.),
+			const T InSigmaMax = (T)3e5
 		)
 			: Base(InParticles, InMesh, bRecordMetricIn, EMesh, NuMesh), SigmaMax(InSigmaMax), FiberDir(InFiberDir)
 		{
-			LambdaArray.Init((FSolverReal)0., 3 * MeshConstraints.Num());
+			LambdaArray.Init((T)0., 3 * MeshConstraints.Num());
 		}
 
 		virtual ~FXPBDCorotatedFiberConstraints() {}
+		
+		TVector<T, 3> GetFiberDir() { return FiberDir; }
+
+		void SetActivation(const T AlphaIn) { AlphaActivation = AlphaIn; }
 
 		void SetTime(const float Time) const {
 			float FinalTime = 4.0f;
@@ -46,26 +50,26 @@ namespace Chaos::Softs
 			{
 				CurrentTime -= FinalTime;
 			}
-			AlphaActivation = (FSolverReal)1. - (FSolverReal)4. / FinalTime * FMath::Abs(CurrentTime - FinalTime / (FSolverReal)2.);
+			AlphaActivation = (T)1. - (T)4. / FinalTime * FMath::Abs(CurrentTime - FinalTime / (T)2.);
 		}
 
-		virtual void ApplyInSerial(FSolverParticles& Particles, const FSolverReal Dt, const int32 ElementIndex) const override
+		virtual void ApplyInSerial(ParticleType& Particles, const T Dt, const int32 ElementIndex) const override
 		{
-			TVec4<FSolverVec3> PolarDelta = Base::GetPolarDelta(Particles, Dt, ElementIndex);
+			TVec4<TVector<T, 3>> PolarDelta = Base::GetPolarDelta(Particles, Dt, ElementIndex);
 
 			for (int i = 0; i < 4; i++)
 			{
 				Particles.P(MeshConstraints[ElementIndex][i]) += PolarDelta[i];
 			}
 
-			TVec4<FSolverVec3> DetDelta = Base::GetDeterminantDelta(Particles, Dt, ElementIndex);
+			TVec4<TVector<T, 3>> DetDelta = Base::GetDeterminantDelta(Particles, Dt, ElementIndex);
 
 			for (int i = 0; i < 4; i++)
 			{
 				Particles.P(MeshConstraints[ElementIndex][i]) += DetDelta[i];
 			}
 
-			TVec4<FSolverVec3> FiberDelta = GetFiberDelta(Particles, Dt, ElementIndex);
+			TVec4<TVector<T, 3>> FiberDelta = GetFiberDelta(Particles, Dt, ElementIndex);
 
 			for (int i = 0; i < 4; i++)
 			{
@@ -75,23 +79,47 @@ namespace Chaos::Softs
 
 		}
 
+		TVec4<TVector<T, 3>> GetFiberGradient(const T dFpdL, const T dFadL,const T C3, const TVec4<TVector<T, 3>>& dLdX) const 
+		{
+			T dC3dL = (T)0.5 * (dFpdL + AlphaActivation * dFadL);
+			if (C3 == 0)
+			{
+				return TVec4<TVector<T, 3>>(TVector<T, 3>((T)0.));
+			}
+			else
+			{
+				dC3dL /= C3;
+			}
+
+			TVec4<TVector<T, 3>> dC3(TVector<T, 3>((T)0.));
+			for (int32 i = 0; i < 4; i++)
+			{
+				for (int32 j = 0; j < 3; j++)
+				{
+					dC3[i][j] = dC3dL * dLdX[i][j];
+				}
+			}
+
+			return dC3;
+		
+		}
 
 
 	private:
 
-		TVec4<FSolverVec3> GetFiberDelta(const FSolverParticles& Particles, const FSolverReal Dt, const int32 ElementIndex, const FSolverReal Tol = 1e-3) const
+		TVec4<TVector<T, 3>> GetFiberDelta(const ParticleType& Particles, const T Dt, const int32 ElementIndex, const T Tol = 1e-3) const
 		{
-			const PMatrix<FSolverReal, 3, 3> Fe = F(ElementIndex, Particles);
+			const PMatrix<T, 3, 3> Fe = Base::F(ElementIndex, Particles);
 
-			PMatrix<FSolverReal, 3, 3> Re((FSolverReal)0.), Se((FSolverReal)0.);
+			PMatrix<T, 3, 3> Re((T)0.), Se((T)0.);
 
 			Chaos::PolarDecomposition(Fe, Re, Se);
 
 			// l: fiber stretch, f = (vTCv)^(1/2)
-			FSolverVec3 FeV = Fe.GetTransposed() * FiberDir;
-			FSolverVec3 DmInverseV = ElementDmInv(ElementIndex).GetTransposed() * FiberDir;
-			FSolverReal L = FeV.Size();
-			TVec4<FSolverVec3> dLdX(FSolverVec3((FSolverReal)0.));
+			TVector<T, 3> FeV = Fe.GetTransposed() * FiberDir;
+			TVector<T, 3> DmInverseV = Base::ElementDmInv(ElementIndex).GetTransposed() * FiberDir;
+			T L = FeV.Size();
+			TVec4<TVector<T, 3>> dLdX(TVector<T, 3>((T)0.));
 			for (int32 alpha = 0; alpha < 3; alpha++)
 			{
 				for (int32 s = 0; s < 3; s++)
@@ -108,70 +136,72 @@ namespace Chaos::Softs
 				}
 			}
 
-			FSolverReal LambdaOFL = (FSolverReal)1.4;
-			FSolverReal P1 = (FSolverReal)0.05;
-			FSolverReal P2 = (FSolverReal)6.6;
-			FSolverReal FpIntegral = (FSolverReal)0.;
-			FSolverReal FaIntegral = (FSolverReal)0.;
+			T LambdaOFL = (T)1.4;
+			T P1 = (T)0.05;
+			T P2 = (T)6.6;
+			T FpIntegral = (T)0.;
+			T FaIntegral = (T)0.;
 
-			FSolverReal C3 = (FSolverReal)0.;
-			FSolverReal AlphaTilde = LambdaOFL / (SigmaMax * Dt * Dt * Measure[ElementIndex]);
-			FSolverReal dFpdL = (FSolverReal)0.;
-			FSolverReal dFadL = (FSolverReal)0.;
+			T C3 = (T)0.;
+			T AlphaTilde = LambdaOFL / (SigmaMax * Dt * Dt * Measure[ElementIndex]);
+			T dFpdL = (T)0.;
+			T dFadL = (T)0.;
 
 			if (L > LambdaOFL)
 			{
-				FpIntegral = P1 * LambdaOFL / P2 * FMath::Exp(P2 * (L / LambdaOFL - (FSolverReal)1.)) - P1 * (L - LambdaOFL);
-				dFpdL = P1 * FMath::Exp(P2 * (L / LambdaOFL - (FSolverReal)1.)) - P1;
+				FpIntegral = P1 * LambdaOFL / P2 * FMath::Exp(P2 * (L / LambdaOFL - (T)1.)) - P1 * (L - LambdaOFL);
+				dFpdL = P1 * FMath::Exp(P2 * (L / LambdaOFL - (T)1.)) - P1;
 			}
-			if (L > (FSolverReal)0.4 * LambdaOFL && L < (FSolverReal)0.6 * LambdaOFL)
+			if (L > (T)0.4 * LambdaOFL && L < (T)0.6 * LambdaOFL)
 			{
-				FaIntegral = (FSolverReal)3. * LambdaOFL * FMath::Pow((L / LambdaOFL - (FSolverReal)0.4), 3);
-				dFadL = (FSolverReal)9. * FMath::Pow((L / LambdaOFL - (FSolverReal)0.4), 2);
+				FaIntegral = (T)3. * LambdaOFL * FMath::Pow((L / LambdaOFL - (T)0.4), 3);
+				dFadL = (T)9. * FMath::Pow((L / LambdaOFL - (T)0.4), 2);
 			}
-			else if (L >= (FSolverReal)0.6 * LambdaOFL && L <= (FSolverReal)1.4 * LambdaOFL)
+			else if (L >= (T)0.6 * LambdaOFL && L <= (T)1.4 * LambdaOFL)
 			{
-				FaIntegral = (FSolverReal)3. * LambdaOFL * (FSolverReal)0.008 + L - (FSolverReal)4. / (FSolverReal)3. * LambdaOFL * FMath::Pow(L / LambdaOFL - (FSolverReal)1., 3) - (FSolverReal)0.6 * LambdaOFL - (FSolverReal)4. / (FSolverReal)3. * LambdaOFL * FMath::Pow((FSolverReal)0.4, 3);
-				dFadL = (FSolverReal)1. - (FSolverReal)4. * FMath::Pow(L / LambdaOFL - (FSolverReal)1., 2);
+				FaIntegral = (T)3. * LambdaOFL * (T)0.008 + L - (T)4. / (T)3. * LambdaOFL * FMath::Pow(L / LambdaOFL - (T)1., 3) - (T)0.6 * LambdaOFL - (T)4. / (T)3. * LambdaOFL * FMath::Pow((T)0.4, 3);
+				dFadL = (T)1. - (T)4. * FMath::Pow(L / LambdaOFL - (T)1., 2);
 			}
-			else if (L > (FSolverReal)1.4 * LambdaOFL && L <= (FSolverReal)1.6 * LambdaOFL)
+			else if (L > (T)1.4 * LambdaOFL && L <= (T)1.6 * LambdaOFL)
 			{
-				FaIntegral = (FSolverReal)3. * LambdaOFL * (FSolverReal)0.008 + (FSolverReal)0.8 * LambdaOFL - (FSolverReal)8. / (FSolverReal)3. * LambdaOFL * FMath::Pow((FSolverReal)0.4, 3) + (FSolverReal)3. * LambdaOFL * FMath::Pow(L / LambdaOFL - (FSolverReal)1.6, 3) + (FSolverReal)3. * LambdaOFL * (FSolverReal)0.008;
-				dFadL = (FSolverReal)9. * FMath::Pow((L / LambdaOFL - (FSolverReal)1.6), 2);
+				FaIntegral = (T)3. * LambdaOFL * (T)0.008 + (T)0.8 * LambdaOFL - (T)8. / (T)3. * LambdaOFL * FMath::Pow((T)0.4, 3) + (T)3. * LambdaOFL * FMath::Pow(L / LambdaOFL - (T)1.6, 3) + (T)3. * LambdaOFL * (T)0.008;
+				dFadL = (T)9. * FMath::Pow((L / LambdaOFL - (T)1.6), 2);
 
 			}
-			else if (L > (FSolverReal)1.6 * LambdaOFL)
+			else if (L > (T)1.6 * LambdaOFL)
 			{
-				FaIntegral = (FSolverReal)3. * LambdaOFL * (FSolverReal)0.008 + (FSolverReal)0.8 * LambdaOFL - (FSolverReal)8. / (FSolverReal)3. * LambdaOFL * FMath::Pow((FSolverReal)0.4, 3) + (FSolverReal)3. * LambdaOFL * (FSolverReal)0.008;
-				dFadL = (FSolverReal)0.;
+				FaIntegral = (T)3. * LambdaOFL * (T)0.008 + (T)0.8 * LambdaOFL - (T)8. / (T)3. * LambdaOFL * FMath::Pow((T)0.4, 3) + (T)3. * LambdaOFL * (T)0.008;
+				dFadL = (T)0.;
 			}
 
 
 			C3 = FMath::Sqrt(FpIntegral + AlphaActivation * FaIntegral);
 
 
-			FSolverReal dC3dL = (FSolverReal)0.5 * (dFpdL + AlphaActivation * dFadL);
-			if (C3 == 0)
-			{
-				return TVec4<FSolverVec3>(FSolverVec3((FSolverReal)0.));
-			}
-			else
-			{
-				dC3dL /= C3;
-			}
+			//T dC3dL = (T)0.5 * (dFpdL + AlphaActivation * dFadL);
+			//if (C3 == 0)
+			//{
+			//	return TVec4<TVector<T, 3>>(TVector<T, 3>((T)0.));
+			//}
+			//else
+			//{
+			//	dC3dL /= C3;
+			//}
 
-			TVec4<FSolverVec3> dC3(FSolverVec3((FSolverReal)0.));
-			for (int32 i = 0; i < 4; i++)
-			{
-				for (int32 j = 0; j < 3; j++)
-				{
-					dC3[i][j] = dC3dL * dLdX[i][j];
-				}
-			}
+			TVec4<TVector<T, 3>> dC3 = GetFiberGradient(dFpdL,dFadL, C3, dLdX);
+
+			//TVec4<TVector<T, 3>> dC3(TVector<T, 3>((T)0.));
+			//for (int32 i = 0; i < 4; i++)
+			//{
+			//	for (int32 j = 0; j < 3; j++)
+			//	{
+			//		dC3[i][j] = dC3dL * dLdX[i][j];
+			//	}
+			//}
 
-			FSolverReal DLambda = -C3 - AlphaTilde * LambdaArray[2 * ElementIndex + 2];
+			T DLambda = -C3 - AlphaTilde * LambdaArray[2 * ElementIndex + 2];
 
-			FSolverReal Denom = AlphaTilde;
+			T Denom = AlphaTilde;
 			for (int i = 0; i < 4; i++)
 			{
 				for (int j = 0; j < 3; j++)
@@ -181,7 +211,7 @@ namespace Chaos::Softs
 			}
 			DLambda /= Denom;
 			LambdaArray[2 * ElementIndex + 2] += DLambda;
-			TVec4<FSolverVec3> Delta(FSolverVec3((FSolverReal)0.));
+			TVec4<TVector<T, 3>> Delta(TVector<T, 3>((T)0.));
 			for (int i = 0; i < 4; i++)
 			{
 				for (int j = 0; j < 3; j++)
@@ -198,9 +228,9 @@ namespace Chaos::Softs
 	private:
 
 		//material constants calculated from E:
-		FSolverReal SigmaMax;
-		mutable FSolverReal AlphaActivation;
-		FSolverVec3 FiberDir;
+		T SigmaMax;
+		mutable T AlphaActivation;
+		TVector<T, 3> FiberDir;
 
 	};
 
