Fix CIS linux uninitialize property
#rnx

diff --git a/Engine/Source/Runtime/Experimental/Interchange/Engine/Public/InterchangeProjectSettings.h b/Engine/Source/Runtime/Experimental/Interchange/Engine/Public/InterchangeProjectSettings.h
--- a/Engine/Source/Runtime/Experimental/Interchange/Engine/Public/InterchangeProjectSettings.h
+++ b/Engine/Source/Runtime/Experimental/Interchange/Engine/Public/InterchangeProjectSettings.h
@@ -32,7 +32,7 @@ struct FInterchangeImportSettings
 
 	/** This tell interchange which pipeline stack to select when importing.*/
 	UPROPERTY(EditAnywhere, Category = "Pipeline")
-	FName DefaultPipelineStack;
+	FName DefaultPipelineStack = NAME_None;
 
 	/** This tell interchange which pipeline configuration dialog to popup when we need to configure the pipelines.*/
 	UPROPERTY(EditAnywhere, Category = "Pipeline")
@@ -40,7 +40,7 @@ struct FInterchangeImportSettings
 
 	/** If enabled, the pipeline stacks configuration dialog will show when interchange must choose a pipeline to import or re-import. If disabled interchange will use the DefaultPipelineStack.*/
 	UPROPERTY(EditAnywhere, Category = "Pipeline")
-	bool bShowPipelineStacksConfigurationDialog;
+	bool bShowPipelineStacksConfigurationDialog = true;
 };
 
 UCLASS(config=Engine, meta=(DisplayName=Interchange), MinimalAPI)
@@ -57,7 +57,7 @@ public:
 
 	/** If checked, will use Interchange when importing into level.*/
 	UPROPERTY(EditAnywhere, config, Category = "ImportIntoLevel|Experimental")
-	bool bUseInterchangeWhenImportingIntoLevel;
+	bool bUseInterchangeWhenImportingIntoLevel = false;
 
 	/**
 	 * Settings used when importing into a level.
