Chaos: managed array

added TArray<FVector3f>

diff --git a/Engine/Source/Runtime/Experimental/Chaos/Public/GeometryCollection/ManagedArrayTypeValues.inl b/Engine/Source/Runtime/Experimental/Chaos/Public/GeometryCollection/ManagedArrayTypeValues.inl
--- a/Engine/Source/Runtime/Experimental/Chaos/Public/GeometryCollection/ManagedArrayTypeValues.inl
+++ b/Engine/Source/Runtime/Experimental/Chaos/Public/GeometryCollection/ManagedArrayTypeValues.inl
@@ -49,6 +49,7 @@ MANAGED_ARRAY_TYPE(TArray<FIntVector2>, IntVector2Array)
 MANAGED_ARRAY_TYPE(TArray<int32>, Int32Array)
 MANAGED_ARRAY_TYPE(TArray<float>, FloatArray)
 MANAGED_ARRAY_TYPE(FVector4f, Vector4f)
+MANAGED_ARRAY_TYPE(TArray<FVector3f>, FVectorArray)
 
 // NOTE: new types must be added at the bottom to keep serialization from breaking
 
