Remove universal references where unneeded.

diff --git a/Engine/Source/Runtime/Engine/Private/Net/Iris/ReplicationSystem/ActorReplicationBridge.cpp b/Engine/Source/Runtime/Engine/Private/Net/Iris/ReplicationSystem/ActorReplicationBridge.cpp
--- a/Engine/Source/Runtime/Engine/Private/Net/Iris/ReplicationSystem/ActorReplicationBridge.cpp
+++ b/Engine/Source/Runtime/Engine/Private/Net/Iris/ReplicationSystem/ActorReplicationBridge.cpp
@@ -108,7 +108,7 @@ void UActorReplicationBridge::Initialize(UReplicationSystem* InReplicationSystem
 	Super::Initialize(InReplicationSystem);
 
 	{
-		auto&& ShouldSpatialize = [](const UClass* Class)
+		auto ShouldSpatialize = [](const UClass* Class)
 		{
 			if (AActor* CDO = Cast<AActor>(Class->GetDefaultObject()))
 			{
@@ -122,7 +122,7 @@ void UActorReplicationBridge::Initialize(UReplicationSystem* InReplicationSystem
 	}
 
 	{
-		auto&& ClassesAreRelevantEqual = [](const UClass* Class, const UClass* Subclass)
+		auto ClassesAreRelevantEqual = [](const UClass* Class, const UClass* Subclass)
 		{
 			const AActor* CDO = Cast<AActor>(Class->GetDefaultObject());
 			const AActor* SubCDO = Cast<AActor>(Subclass->GetDefaultObject());
diff --git a/Engine/Source/Runtime/Experimental/Iris/Core/Private/Iris/Core/IrisDebugging.cpp b/Engine/Source/Runtime/Experimental/Iris/Core/Private/Iris/Core/IrisDebugging.cpp
--- a/Engine/Source/Runtime/Experimental/Iris/Core/Private/Iris/Core/IrisDebugging.cpp
+++ b/Engine/Source/Runtime/Experimental/Iris/Core/Private/Iris/Core/IrisDebugging.cpp
@@ -345,7 +345,7 @@ void NetObjectProtocolReferencesToString(FStringBuilderBase& StringBuilder, uint
 			StringBuilder << TEXT("Protocol: ") << ToCStr(Protocol->DebugName);
 			StringBuilder.Appendf(TEXT("Id: 0x%" UINT64_x_FMT " Created From : 0x%p"), Protocol->ProtocolIdentifier, ArchetypeOrCDOUsedAsKey) << TEXT(" Used by : \n");
 
-			auto&& FindMatchingProtocol = [&StringBuilder, &Protocol, &NetHandleManager](uint32 InternalObjectIndex)
+			auto FindMatchingProtocol = [&StringBuilder, &Protocol, &NetHandleManager](uint32 InternalObjectIndex)
 			{
 				const FNetHandleManager::FReplicatedObjectData& ObjectData = NetHandleManager.GetReplicatedObjectDataNoCheck(InternalObjectIndex);
 				if (ObjectData.Protocol == Protocol)
diff --git a/Engine/Source/Runtime/Experimental/Iris/Core/Private/Iris/ReplicationSystem/DeltaCompression/DeltaCompressionBaselineManager.cpp b/Engine/Source/Runtime/Experimental/Iris/Core/Private/Iris/ReplicationSystem/DeltaCompression/DeltaCompressionBaselineManager.cpp
--- a/Engine/Source/Runtime/Experimental/Iris/Core/Private/Iris/ReplicationSystem/DeltaCompression/DeltaCompressionBaselineManager.cpp
+++ b/Engine/Source/Runtime/Experimental/Iris/Core/Private/Iris/ReplicationSystem/DeltaCompression/DeltaCompressionBaselineManager.cpp
@@ -376,7 +376,7 @@ void FDeltaCompressionBaselineManager::FreePerObjectInfo(ObjectInfoIndexType Ind
 
 void FDeltaCompressionBaselineManager::FreeAllPerObjectInfos()
 {
-	auto&& DestructObjectInfo = [this](uint32 ObjectInfoIndex)
+	auto DestructObjectInfo = [this](uint32 ObjectInfoIndex)
 	{
 		FPerObjectInfo* ObjectInfo = GetPerObjectInfo(ObjectInfoIndex);
 		DestructPerObjectInfo(ObjectInfo);
@@ -535,7 +535,7 @@ void FDeltaCompressionBaselineManager::AdjustBaselineCount(const FPerObjectInfo*
 
 void FDeltaCompressionBaselineManager::ReleaseBaselinesForConnection(uint32 ConnId)
 {
-	auto&& ReleaseObjectBaselines = [this, ConnId](uint32 ObjectInfoIndex)
+	auto ReleaseObjectBaselines = [this, ConnId](uint32 ObjectInfoIndex)
 	{
 		FPerObjectInfo* ObjectInfo = GetPerObjectInfo(ObjectInfoIndex);
 		FObjectBaselineInfo& BaselineInfo = ObjectInfo->BaselinesForConnections[ConnId];
@@ -570,7 +570,7 @@ void FDeltaCompressionBaselineManager::ConstructBaselineSharingContext(FBaseline
 
 void FDeltaCompressionBaselineManager::DestructBaselineSharingContext(FBaselineSharingContext& SharingContext)
 {
-	auto&& ReleaseBaseline = [this, SharingContext](uint32 ObjectInfoIndex)
+	auto ReleaseBaseline = [this, SharingContext](uint32 ObjectInfoIndex)
 	{
 		const DeltaCompressionBaselineStateInfoIndexType BaselineStateInfoIndex = SharingContext.ObjectInfoIndexToBaselineInfoIndex[ObjectInfoIndex];
 		BaselineStorage.OptionallyCommitAndDoReleaseBaseline(BaselineStateInfoIndex);
diff --git a/Engine/Source/Runtime/Experimental/Iris/Core/Private/Iris/ReplicationSystem/DeltaCompression/DeltaCompressionBaselineStorage.cpp b/Engine/Source/Runtime/Experimental/Iris/Core/Private/Iris/ReplicationSystem/DeltaCompression/DeltaCompressionBaselineStorage.cpp
--- a/Engine/Source/Runtime/Experimental/Iris/Core/Private/Iris/ReplicationSystem/DeltaCompression/DeltaCompressionBaselineStorage.cpp
+++ b/Engine/Source/Runtime/Experimental/Iris/Core/Private/Iris/ReplicationSystem/DeltaCompression/DeltaCompressionBaselineStorage.cpp
@@ -212,7 +212,7 @@ void FDeltaCompressionBaselineStorage::FreeBaselineStateInfo(DeltaCompressionBas
 
 void FDeltaCompressionBaselineStorage::FreeAllBaselineStateInfos()
 {
-	auto&& DestructStateInfo = [this](uint32 StateInfoIndex)
+	auto DestructStateInfo = [this](uint32 StateInfoIndex)
 	{
 		FInternalBaselineStateInfo* ObjectInfo = GetBaselineStateInfo(StateInfoIndex);
 		DestructBaselineStateInfo(ObjectInfo);
diff --git a/Engine/Source/Runtime/Experimental/Iris/Core/Private/Iris/ReplicationSystem/Filtering/ReplicationFiltering.cpp b/Engine/Source/Runtime/Experimental/Iris/Core/Private/Iris/ReplicationSystem/Filtering/ReplicationFiltering.cpp
--- a/Engine/Source/Runtime/Experimental/Iris/Core/Private/Iris/ReplicationSystem/Filtering/ReplicationFiltering.cpp
+++ b/Engine/Source/Runtime/Experimental/Iris/Core/Private/Iris/ReplicationSystem/Filtering/ReplicationFiltering.cpp
@@ -299,7 +299,7 @@ bool FReplicationFiltering::SetFilter(uint32 ObjectIndex, FNetObjectFilterHandle
 	}
 
 	const FNetBitArray& ObjectsInScope = NetHandleManager->GetScopableInternalIndices();
-	auto&& ClearAndDirtyStaticFilter = [this](uint32 ObjIndex)
+	auto ClearAndDirtyStaticFilter = [this](uint32 ObjIndex)
 	{
 		this->bHasDirtyConnectionFilter = 1;
 
@@ -308,7 +308,7 @@ bool FReplicationFiltering::SetFilter(uint32 ObjectIndex, FNetObjectFilterHandle
 		this->FreePerObjectInfoForObject(ObjIndex);
 	};
 
-	auto&& TrySetDynamicFilter = [this, &ObjectData, &ReplicatedObjectsStateBuffers](uint32 ObjIndex, uint32 FilterIndex)
+	auto TrySetDynamicFilter = [this, &ObjectData, &ReplicatedObjectsStateBuffers](uint32 ObjIndex, uint32 FilterIndex)
 	{
 		FNetObjectFilteringInfo& NetObjectFilteringInfo = this->NetObjectFilteringInfos[ObjIndex];
 		NetObjectFilteringInfo = {};
@@ -508,7 +508,7 @@ void FReplicationFiltering::InitNewConnections()
 
 	bHasNewConnection = 0;
 	
-	auto&& InitNewConnection = [this](uint32 ConnectionId)
+	auto InitNewConnection = [this](uint32 ConnectionId)
 	{
 		// Copy default scope
 		const FNetBitArray& ScopableInternalIndices = NetHandleManager->GetScopableInternalIndices();
@@ -534,7 +534,7 @@ void FReplicationFiltering::InitNewConnections()
 
 		// Update group filter by iterating over all registered groups for the connection
 		{
-			auto&& InitGroupFilterForConnection = [this, &ConnectionInfo, ConnectionId](uint32 GroupIndex)
+			auto InitGroupFilterForConnection = [this, &ConnectionInfo, ConnectionId](uint32 GroupIndex)
 			{
 				const FNetObjectGroupHandle GroupHandle = GroupIndex;
 				const FPerObjectInfo* ConnectionState = GetPerObjectInfo(this->GroupInfos[GroupIndex].ConnectionStateIndex);
@@ -562,7 +562,7 @@ void FReplicationFiltering::InitNewConnections()
 
 		// Update connection scope with owner filtering
 		{
-			auto&& MaskObjectToOwner = [this, ConnectionId, &ConnectionInfo](uint32 ObjectIndex)
+			auto MaskObjectToOwner = [this, ConnectionId, &ConnectionInfo](uint32 ObjectIndex)
 			{
 				const bool bIsOwner = (ConnectionId == ObjectIndexToOwningConnection.GetData()[ObjectIndex]);
 				ConnectionInfo.ConnectionFilteredObjects.SetBitValue(ObjectIndex, bIsOwner);
@@ -573,7 +573,7 @@ void FReplicationFiltering::InitNewConnections()
 
 		// Update connection scope with connection filtering
 		{
-			auto&& MaskObjectToConnection = [this, ConnectionId, &ConnectionInfo](uint32 ObjectIndex)
+			auto MaskObjectToConnection = [this, ConnectionId, &ConnectionInfo](uint32 ObjectIndex)
 			{
 				const PerObjectInfoIndexType ObjectInfoIndex = ObjectIndexToPerObjectInfoIndex.GetData()[ObjectIndex];
 				const FPerObjectInfo* ObjectInfo = GetPerObjectInfo(ObjectInfoIndex);
@@ -605,7 +605,7 @@ void FReplicationFiltering::ResetRemovedConnections()
 	// Reset group filter status
 	// We might want to introduce a way to specify default state for group filters
 	// Currently we just clear the filter
-	auto&& ResetGroupFilterStatus = [this](uint32 GroupIndex)
+	auto ResetGroupFilterStatus = [this](uint32 GroupIndex)
 	{
 		FPerObjectInfo* ConnectionStateInfo = GetPerObjectInfo(this->GroupInfos[GroupIndex].ConnectionStateIndex);
 
@@ -830,7 +830,7 @@ void FReplicationFiltering::UpdateOwnerAndConnectionFiltering()
 	if (bHasDirtyOwner)
 	{
 		uint16* const ObjectIndexToOwningConnectionStorage = ObjectIndexToOwningConnection.GetData();
-		auto&& UpdateOwners = [this, ObjectIndexToOwningConnectionStorage](uint32 ObjectIndex)
+		auto UpdateOwners = [this, ObjectIndexToOwningConnectionStorage](uint32 ObjectIndex)
 		{
 			const uint32 OwningConnectionId = ObjectIndexToOwningConnectionStorage[ObjectIndex];
 
@@ -852,7 +852,7 @@ void FReplicationFiltering::UpdateOwnerAndConnectionFiltering()
 	// Update filtering
 	if (bHasDirtyConnectionFilter)
 	{
-		auto&& UpdateConnectionScope = [this, &GlobalObjectsInScope](uint32 ConnectionId)
+		auto UpdateConnectionScope = [this, &GlobalObjectsInScope](uint32 ConnectionId)
 		{
 			FPerConnectionInfo& ConnectionInfo = this->ConnectionInfos[ConnectionId];
 			FNetBitArrayView ConnectionScope = MakeNetBitArrayView(ConnectionInfo.ConnectionFilteredObjects);
@@ -861,7 +861,7 @@ void FReplicationFiltering::UpdateOwnerAndConnectionFiltering()
 
 			// Update filter info
 			{
-				auto&& MaskObject = [this, &ConnectionScope, &GroupFilteredOutObjects, &ObjectsInScopeBeforeDynamicFiltering, ConnectionId, &GlobalObjectsInScope](uint32 ObjectIndex)
+				auto MaskObject = [this, &ConnectionScope, &GroupFilteredOutObjects, &ObjectsInScopeBeforeDynamicFiltering, ConnectionId, &GlobalObjectsInScope](uint32 ObjectIndex)
 				{
 					bool bObjectIsInScope = true;
 					if (HasOwnerFilter(ObjectIndex))
@@ -914,7 +914,7 @@ void FReplicationFiltering::UpdateOwnerAndConnectionFiltering()
 void FReplicationFiltering::UpdateGroupFiltering()
 {
 	// Adding objects to an active group filter is deferred in order to avoid triggering constant filter updates.
-	auto&& UpdateGroupFilter = [this](uint32 GroupIndex)
+	auto UpdateGroupFilter = [this](uint32 GroupIndex)
 	{
 		if (this->FilterGroups.GetBit(GroupIndex))
 		{
@@ -922,7 +922,7 @@ void FReplicationFiltering::UpdateGroupFiltering()
 			const FNetObjectGroup* Group = this->Groups->GetGroup(GroupIndex);
 			FPerConnectionInfo* LocalConnectionInfos = this->ConnectionInfos.GetData();
 
-			auto&& UpdateGroupFilterForConnection = [this, ConnectionStateInfo, Group, LocalConnectionInfos](uint32 ConnectionId)
+			auto UpdateGroupFilterForConnection = [this, ConnectionStateInfo, Group, LocalConnectionInfos](uint32 ConnectionId)
 			{
 				if (this->GetConnectionFilterStatus(*ConnectionStateInfo, ConnectionId) == ENetFilterStatus::Disallow)
 				{
@@ -1346,7 +1346,7 @@ void FReplicationFiltering::RemoveSubObjectFilter(FNetObjectGroupHandle GroupHan
 void FReplicationFiltering::UpdateSubObjectFilters()
 {
 	// We want to remove all groups that have no members and no enabled connections
-	auto&& UpdateSubObjectFilterGroup = [this](uint32 GroupIndex)
+	auto UpdateSubObjectFilterGroup = [this](uint32 GroupIndex)
 	{
 		FNetObjectGroupHandle GroupHandle(GroupIndex);
 		if (const FNetObjectGroup* Group = Groups->GetGroup(FNetObjectGroupHandle(GroupIndex)))
@@ -1603,7 +1603,7 @@ bool FReplicationFiltering::UpdateGroupFilterEffectsForObject(uint32 ObjectIndex
 	FNetBitArray& ConnectionFilteredObjects = ConnectionInfo.ConnectionFilteredObjects;
 	FNetBitArray& ObjectsInScopeBeforeDynamicFiltering = ConnectionInfo.ObjectsInScopeBeforeDynamicFiltering;
 
-	auto&& ClearGroupFilterForObject = [&GroupFilteredOutObjects, &ConnectionFilteredObjects, &ObjectsInScopeBeforeDynamicFiltering](uint32 InternalObjectIndex)
+	auto ClearGroupFilterForObject = [&GroupFilteredOutObjects, &ConnectionFilteredObjects, &ObjectsInScopeBeforeDynamicFiltering](uint32 InternalObjectIndex)
 	{
 		GroupFilteredOutObjects.ClearBit(InternalObjectIndex);
 		ObjectsInScopeBeforeDynamicFiltering.SetBitValue(InternalObjectIndex, ConnectionFilteredObjects.GetBit(InternalObjectIndex));
@@ -1748,7 +1748,7 @@ void FReplicationFiltering::NotifyObjectRemovedFromGroup(FNetObjectGroupHandle G
 		FPerObjectInfo* ConnectionState = GetPerObjectInfo(GroupInfos[GroupHandle].ConnectionStateIndex);
 		const FNetBitArray& ScopedObjects = NetHandleManager->GetScopableInternalIndices();
 
-		auto&& MarkGroupsDirty = [this, ObjectIndex, GroupHandle, ConnectionState, &ScopedObjects](uint32 ConnectionId)
+		auto MarkGroupsDirty = [this, ObjectIndex, GroupHandle, ConnectionState, &ScopedObjects](uint32 ConnectionId)
 		{
 			if (GetConnectionFilterStatus(*ConnectionState, ConnectionId) == ENetFilterStatus::Disallow)
 			{
diff --git a/Engine/Source/Runtime/Experimental/Iris/Core/Private/Iris/ReplicationSystem/ReplicationReader.cpp b/Engine/Source/Runtime/Experimental/Iris/Core/Private/Iris/ReplicationSystem/ReplicationReader.cpp
--- a/Engine/Source/Runtime/Experimental/Iris/Core/Private/Iris/ReplicationSystem/ReplicationReader.cpp
+++ b/Engine/Source/Runtime/Experimental/Iris/Core/Private/Iris/ReplicationSystem/ReplicationReader.cpp
@@ -644,7 +644,7 @@ void FReplicationReader::UpdateObjectReferenceTracking(FReplicatedObjectInfo* Re
 		}
 
 		// Replace each entry in UnresolvedObjectReferences for the given changemask
-		auto&& UpdateUnresolvedReferencesForChange = [ReplicationInfo, &NewUnresolvedReferences](uint32 ChangeBit)
+		auto UpdateUnresolvedReferencesForChange = [ReplicationInfo, &NewUnresolvedReferences](uint32 ChangeBit)
 		{
 			FObjectReferenceTracker& UnresolvedObjectReferences = ReplicationInfo->UnresolvedObjectReferences;
 
@@ -715,7 +715,7 @@ void FReplicationReader::UpdateObjectReferenceTracking(FReplicatedObjectInfo* Re
 		}
 
 		// Replace each entry in ResolvedDynamicObjectReferences for the given changemask
-		auto&& UpdateResolvedReferencesForChange = [ReplicationInfo, &NewMappedDynamicReferences](uint32 ChangeBit)
+		auto UpdateResolvedReferencesForChange = [ReplicationInfo, &NewMappedDynamicReferences](uint32 ChangeBit)
 		{
 			FObjectReferenceTracker& ResolvedDynamicObjectReferences = ReplicationInfo->ResolvedDynamicObjectReferences;
 
@@ -1013,7 +1013,7 @@ void FReplicationReader::DispatchStateData(FNetSerializationContext& Context)
 	uint32 NumObjectsPendingPostDistpatch = 0U;
 
 	// Function to flush all objects pending post dispatch
-	auto&& FlushPostDispatchForBatch = [PostDispatchObjectInfos, &NumObjectsPendingPostDistpatch, this, &Context]()
+	auto FlushPostDispatchForBatch = [PostDispatchObjectInfos, &NumObjectsPendingPostDistpatch, this, &Context]()
 	{
 		// When all received states have been applied we invoke PostReplicate and RepNotifies
 		for (FPostDispatchObjectInfo& PostDispatchObjectInfo : MakeArrayView(PostDispatchObjectInfos, NumObjectsPendingPostDistpatch))
diff --git a/Engine/Source/Runtime/Experimental/Iris/Core/Private/Iris/ReplicationSystem/ReplicationWriter.cpp b/Engine/Source/Runtime/Experimental/Iris/Core/Private/Iris/ReplicationSystem/ReplicationWriter.cpp
--- a/Engine/Source/Runtime/Experimental/Iris/Core/Private/Iris/ReplicationSystem/ReplicationWriter.cpp
+++ b/Engine/Source/Runtime/Experimental/Iris/Core/Private/Iris/ReplicationSystem/ReplicationWriter.cpp
@@ -494,7 +494,7 @@ void FReplicationWriter::UpdateScope(const FNetBitArrayView& UpdatedScope)
 {
 	IRIS_PROFILER_SCOPE(FReplicationWriter_ScopeUpdate);
 
-	auto&& NewObjectFunctor = [this](uint32 Index)
+	auto NewObjectFunctor = [this](uint32 Index)
 	{
 		// We can only start replicating an object that is not currently replicated
 		// Later on we will support "cancelling" destroy objects that are flushing state data
@@ -546,7 +546,7 @@ void FReplicationWriter::UpdateScope(const FNetBitArrayView& UpdatedScope)
 		}
 	};
 
-	auto&& DestroyedObjectFunctor = [this](uint32 Index) 
+	auto DestroyedObjectFunctor = [this](uint32 Index) 
 	{
 		// Request object to be destroyed
 		FReplicationInfo& Info = GetReplicationInfo(Index);
@@ -697,7 +697,7 @@ const FNetBitArray& FReplicationWriter::GetObjectsRequiringPriorityUpdate() cons
 void FReplicationWriter::UpdatePriorities(const float* UpdatedPriorities)
 {
 	IRIS_PROFILER_SCOPE(FReplicationWriter_UpdatePriorities);
-	auto&& UpdatePriority = [LocalPriorities = SchedulingPriorities.GetData(), UpdatedPriorities](uint32 Index)
+	auto UpdatePriority = [LocalPriorities = SchedulingPriorities.GetData(), UpdatedPriorities](uint32 Index)
 	{
 		LocalPriorities[Index] += UpdatedPriorities[Index];
 	};
@@ -721,7 +721,7 @@ uint32 FReplicationWriter::ScheduleObjects(FScheduleObjectInfo* OutScheduledObje
 	const FNetBitArray& SubObjects = NetHandleManager->GetSubObjectInternalIndices();
 
 	// Fill IndexList with all objects with dirty state that have positive priority excluding subobjects
-	auto&& FillIndexListFunc = [&LocalPriorities, &ScheduledObjectIndices, &ScheduledObjectCount](uint32 Index)
+	auto FillIndexListFunc = [&LocalPriorities, &ScheduledObjectIndices, &ScheduledObjectCount](uint32 Index)
 	{
 		const float UpdatedPriority = LocalPriorities[Index];
 
@@ -2184,7 +2184,7 @@ uint32 FReplicationWriter::WriteObjects(FNetSerializationContext& Context)
 
 	bool bContinue = WriteContext.bHasUpdatedObjectsToSend;
 
-	auto&& SendObjectFunction = [this, &Context, &WrittenObjectCount](FInternalNetHandle InternalIndex)
+	auto SendObjectFunction = [this, &Context, &WrittenObjectCount](FInternalNetHandle InternalIndex)
 	{
 		if (!this->WriteContext.ObjectsWrittenThisTick.GetBit(InternalIndex) && this->CanSendObject(InternalIndex))
 		{
diff --git a/Engine/Source/Runtime/Experimental/Iris/Core/Private/Iris/Serialization/NetBitStreamUtil.cpp b/Engine/Source/Runtime/Experimental/Iris/Core/Private/Iris/Serialization/NetBitStreamUtil.cpp
--- a/Engine/Source/Runtime/Experimental/Iris/Core/Private/Iris/Serialization/NetBitStreamUtil.cpp
+++ b/Engine/Source/Runtime/Experimental/Iris/Core/Private/Iris/Serialization/NetBitStreamUtil.cpp
@@ -596,7 +596,7 @@ void WriteSparseBitArray(FNetBitStreamWriter* Writer, const uint32* Data, uint32
 {
 	if (Hint == ESparseBitArraySerializationHint::None)
 	{
-		auto&& GetDataFunc = [](const uint32 Value) { return Value; };
+		auto GetDataFunc = [](const uint32 Value) { return Value; };
 		// Support large bit arrays
 		uint32 RemainingBits = BitCount;
 		while (RemainingBits >= Private::SerializeSparseArrayMaxBitCount)
@@ -612,7 +612,7 @@ void WriteSparseBitArray(FNetBitStreamWriter* Writer, const uint32* Data, uint32
 	}
 	else
 	{
-		auto&& GetDataFunc = [](const uint32 Value) { return ~Value; };
+		auto GetDataFunc = [](const uint32 Value) { return ~Value; };
 		// Support large bit arrays
 		uint32 RemainingBits = BitCount;
 		while (RemainingBits >= Private::SerializeSparseArrayMaxBitCount)
@@ -632,7 +632,7 @@ void ReadSparseBitArray(FNetBitStreamReader* Reader, uint32* OutData, uint32 Bit
 {
 	if (Hint == ESparseBitArraySerializationHint::None)
 	{
-		auto&& SetDataFunc = [](const uint32 Value) { return Value; };
+		auto SetDataFunc = [](const uint32 Value) { return Value; };
 		// Support large bit arrays
 		uint32 RemainingBits = BitCount;
 		while (RemainingBits >= Private::SerializeSparseArrayMaxBitCount)
@@ -648,7 +648,7 @@ void ReadSparseBitArray(FNetBitStreamReader* Reader, uint32* OutData, uint32 Bit
 	}
 	else
 	{
-		auto&& SetDataFunc = [](const uint32 Value) { return ~Value; };
+		auto SetDataFunc = [](const uint32 Value) { return ~Value; };
 		// Support large bit arrays
 		uint32 RemainingBits = BitCount;
 		while (RemainingBits >= Private::SerializeSparseArrayMaxBitCount)
diff --git a/Engine/Source/Runtime/Net/Core/Classes/Net/Serialization/FastArraySerializerImplementation.h b/Engine/Source/Runtime/Net/Core/Classes/Net/Serialization/FastArraySerializerImplementation.h
--- a/Engine/Source/Runtime/Net/Core/Classes/Net/Serialization/FastArraySerializerImplementation.h
+++ b/Engine/Source/Runtime/Net/Core/Classes/Net/Serialization/FastArraySerializerImplementation.h
@@ -8,7 +8,7 @@
 #define UE_NET_IMPLEMENT_FASTARRAY(FastArrayType) \
 UE::Net::CreateAndRegisterReplicationFragmentFunc FastArrayType::GetFastArrayCreateReplicationFragmentFunction() \
 { \
-	auto&& CreateFastArrayReplicationFragmentFunction = [](UObject* InLocalOwner, const UE::Net::FReplicationStateDescriptor* InLocalDescriptor, UE::Net::FFragmentRegistrationContext& InLocalContext) \
+	auto CreateFastArrayReplicationFragmentFunction = [](UObject* InLocalOwner, const UE::Net::FReplicationStateDescriptor* InLocalDescriptor, UE::Net::FFragmentRegistrationContext& InLocalContext) \
 	{ \
 		return UE::Net::Private::CreateAndRegisterFragment<FastArrayType>(InLocalOwner, InLocalDescriptor, InLocalContext); \
 	}; \
