removed useless dynamic offsets and added editable conditions for offset properties

diff --git a/Engine/Source/Runtime/Experimental/Animation/Constraints/Private/TransformConstraint.cpp b/Engine/Source/Runtime/Experimental/Animation/Constraints/Private/TransformConstraint.cpp
--- a/Engine/Source/Runtime/Experimental/Animation/Constraints/Private/TransformConstraint.cpp
+++ b/Engine/Source/Runtime/Experimental/Animation/Constraints/Private/TransformConstraint.cpp
@@ -242,7 +242,7 @@ void UTickableTranslationConstraint::PostEditChangeProperty(FPropertyChangedEven
 			
 			const FTransform ChildGlobalTransform = GetChildGlobalTransform();
 			const FTransform ParentWorldTransform = GetParentGlobalTransform();
-			DynamicOffsetTranslation = ChildGlobalTransform.GetLocation() - ParentWorldTransform.GetLocation();
+			OffsetTranslation = ChildGlobalTransform.GetLocation() - ParentWorldTransform.GetLocation();
 			
 			Evaluate();
 		}
@@ -258,11 +258,9 @@ void UTickableTranslationConstraint::ComputeOffset()
 	const FTransform InitChildTransform = GetChildGlobalTransform();
 	
 	OffsetTranslation = FVector::ZeroVector;
-	DynamicOffsetTranslation = FVector::ZeroVector;
-	if (bMaintainOffset)
+	if (bMaintainOffset || bDynamicOffset)
 	{
 		OffsetTranslation = InitChildTransform.GetLocation() - InitParentTransform.GetLocation();
-		DynamicOffsetTranslation = OffsetTranslation;
 	}	
 }
 
@@ -283,9 +281,7 @@ FConstraintTickFunction::ConstraintFunction UTickableTranslationConstraint::GetF
 
 		const FVector ParentTranslation = GetParentGlobalTransform().GetLocation();
 		FTransform Transform = GetChildGlobalTransform();
-		FVector NewTranslation = bDynamicOffset ? ParentTranslation + DynamicOffsetTranslation :
-								bMaintainOffset ? ParentTranslation + OffsetTranslation :
-								ParentTranslation;
+		FVector NewTranslation = (!bMaintainOffset) ? ParentTranslation : ParentTranslation + OffsetTranslation;
 		if (ClampedWeight < 1.0f - KINDA_SMALL_NUMBER)
 		{
 			NewTranslation = FMath::Lerp<FVector>(Transform.GetLocation(), NewTranslation, ClampedWeight);
@@ -298,7 +294,7 @@ FConstraintTickFunction::ConstraintFunction UTickableTranslationConstraint::GetF
 
 void UTickableTranslationConstraint::OnHandleModified(UTransformableHandle* InHandle, bool bUpdate)
 {
-	if (!bDynamicOffset)
+	if (!Active || !bDynamicOffset)
 	{
 		return;
 	}
@@ -319,12 +315,12 @@ void UTickableTranslationConstraint::OnHandleModified(UTransformableHandle* InHa
 		{
 			const FTransform ParentWorldTransform = GetParentGlobalTransform();
 			const FTransform ChildGlobalTransform = GetChildGlobalTransform();
-			DynamicOffsetTranslation = ChildGlobalTransform.GetLocation() - ParentWorldTransform.GetLocation();
+			OffsetTranslation = ChildGlobalTransform.GetLocation() - ParentWorldTransform.GetLocation();
 		}
 		else
 		{
 			const FTransform ChildLocalTransform = GetChildLocalTransform();
-			DynamicOffsetTranslation = ChildLocalTransform.GetTranslation();
+			OffsetTranslation = ChildLocalTransform.GetTranslation();
 		}
 	}
 }
@@ -368,7 +364,7 @@ void UTickableRotationConstraint::PostEditChangeProperty(FPropertyChangedEvent&
 			
 			const FTransform ChildGlobalTransform = GetChildGlobalTransform();
 			const FTransform ParentWorldTransform = GetParentGlobalTransform();
-			DynamicOffsetRotation = ParentWorldTransform.GetRotation().Inverse() * ChildGlobalTransform.GetRotation();
+			OffsetRotation = ParentWorldTransform.GetRotation().Inverse() * ChildGlobalTransform.GetRotation();
 			
 			Evaluate();
 		}
@@ -383,12 +379,11 @@ void UTickableRotationConstraint::ComputeOffset()
 	const FTransform InitParentTransform = GetParentGlobalTransform();
 	const FTransform InitChildTransform = GetChildGlobalTransform();
 	
-	OffsetRotation = DynamicOffsetRotation = FQuat::Identity;
-	if (bMaintainOffset)
+	OffsetRotation = FQuat::Identity;
+	if (bMaintainOffset || bDynamicOffset)
 	{
 		OffsetRotation = InitParentTransform.GetRotation().Inverse() * InitChildTransform.GetRotation();
-		OffsetRotation.Normalize();
-		DynamicOffsetRotation = OffsetRotation; 
+		OffsetRotation.Normalize(); 
 	}
 }
 
@@ -410,9 +405,7 @@ FConstraintTickFunction::ConstraintFunction UTickableRotationConstraint::GetFunc
 		const FQuat ParentRotation = GetParentGlobalTransform().GetRotation();
 		FTransform Transform = GetChildGlobalTransform();
 
-		FQuat NewRotation = bDynamicOffset ? ParentRotation * DynamicOffsetRotation :
-							bMaintainOffset ? ParentRotation * OffsetRotation :
-							ParentRotation;
+		FQuat NewRotation = (!bMaintainOffset) ? ParentRotation : ParentRotation * OffsetRotation;
 		if (ClampedWeight < 1.0f - KINDA_SMALL_NUMBER)
 		{
 			NewRotation = FQuat::Slerp(Transform.GetRotation(), NewRotation, ClampedWeight);
@@ -425,7 +418,7 @@ FConstraintTickFunction::ConstraintFunction UTickableRotationConstraint::GetFunc
 
 void UTickableRotationConstraint::OnHandleModified(UTransformableHandle* InHandle, bool bUpdate)
 {
-	if (!bDynamicOffset)
+	if (!Active || !bDynamicOffset)
 	{
 		return;
 	}
@@ -446,11 +439,11 @@ void UTickableRotationConstraint::OnHandleModified(UTransformableHandle* InHandl
 		{
 			const FTransform ParentWorldTransform = GetParentGlobalTransform();
 			const FTransform ChildGlobalTransform = GetChildGlobalTransform();
-			DynamicOffsetRotation = ParentWorldTransform.GetRotation().Inverse() * ChildGlobalTransform.GetRotation();
+			OffsetRotation = ParentWorldTransform.GetRotation().Inverse() * ChildGlobalTransform.GetRotation();
 		}
 		else
 		{
-			DynamicOffsetRotation = GetChildLocalTransform().GetRotation();
+			OffsetRotation = GetChildLocalTransform().GetRotation();
 		}
 	}
 }
@@ -512,7 +505,7 @@ FConstraintTickFunction::ConstraintFunction UTickableScaleConstraint::GetFunctio
 		
 		const FVector ParentScale = GetParentGlobalTransform().GetScale3D();
 		FTransform Transform = GetChildGlobalTransform();
-		FVector NewScale = bMaintainOffset ? ParentScale * OffsetScale : ParentScale;
+		FVector NewScale = (!bMaintainOffset) ? ParentScale : ParentScale * OffsetScale;
 		if (ClampedWeight < 1.0f - KINDA_SMALL_NUMBER)
 		{
 			NewScale = FMath::Lerp<FVector>(Transform.GetScale3D(), NewScale, ClampedWeight);
@@ -538,11 +531,9 @@ void UTickableParentConstraint::ComputeOffset()
 	const FTransform InitChildTransform = GetChildGlobalTransform();
 	
 	OffsetTransform = FTransform::Identity;
-	DynamicOffsetTransform = FTransform::Identity;
-	if (bMaintainOffset)
+	if (bMaintainOffset || bDynamicOffset)
 	{
-		OffsetTransform = InitChildTransform.GetRelativeTransform(InitParentTransform);
-		DynamicOffsetTransform = OffsetTransform; 
+		OffsetTransform = InitChildTransform.GetRelativeTransform(InitParentTransform); 
 	}
 }
 
@@ -593,9 +584,7 @@ FConstraintTickFunction::ConstraintFunction UTickableParentConstraint::GetFuncti
 
 		const FTransform ParentTransform = GetParentGlobalTransform();
 		
-		FTransform TargetTransform = bDynamicOffset ? DynamicOffsetTransform * ParentTransform :
-									bMaintainOffset ? OffsetTransform * ParentTransform :
-									ParentTransform;
+		FTransform TargetTransform = (!bMaintainOffset) ? ParentTransform : OffsetTransform * ParentTransform;
 		//apply weight if needed
 		LerpTransform(GetChildGlobalTransform(), TargetTransform);
 		SetChildGlobalTransform(TargetTransform);
@@ -604,7 +593,7 @@ FConstraintTickFunction::ConstraintFunction UTickableParentConstraint::GetFuncti
 
 void UTickableParentConstraint::OnHandleModified(UTransformableHandle* InHandle, bool bUpdate)
 {
-	if (!bDynamicOffset)
+	if (!Active || !bDynamicOffset)
 	{
 		return;
 	}
@@ -625,11 +614,11 @@ void UTickableParentConstraint::OnHandleModified(UTransformableHandle* InHandle,
 		{
 			const FTransform ParentWorldTransform = GetParentGlobalTransform();
 			const FTransform ChildGlobalTransform = GetChildGlobalTransform();
-			DynamicOffsetTransform = ChildGlobalTransform.GetRelativeTransform(ParentWorldTransform);
+			OffsetTransform = ChildGlobalTransform.GetRelativeTransform(ParentWorldTransform);
 		}
 		else
 		{
-			DynamicOffsetTransform = GetChildLocalTransform();
+			OffsetTransform = GetChildLocalTransform();
 		}
 	}
 }
@@ -649,7 +638,7 @@ void UTickableParentConstraint::PostEditChangeProperty(FPropertyChangedEvent& Pr
 			
 			const FTransform ChildGlobalTransform = GetChildGlobalTransform();
 			const FTransform ParentWorldTransform = GetParentGlobalTransform();
-			DynamicOffsetTransform = ChildGlobalTransform.GetRelativeTransform(ParentWorldTransform);
+			OffsetTransform = ChildGlobalTransform.GetRelativeTransform(ParentWorldTransform);
 			
 			Evaluate();
 		}
@@ -669,9 +658,7 @@ UTickableLookAtConstraint::UTickableLookAtConstraint()
 }
 
 void UTickableLookAtConstraint::ComputeOffset()
-{
-	// TODO compute offset
-}
+{}
 
 FConstraintTickFunction::ConstraintFunction UTickableLookAtConstraint::GetFunction() const
 {
diff --git a/Engine/Source/Runtime/Experimental/Animation/Constraints/Public/TransformConstraint.h b/Engine/Source/Runtime/Experimental/Animation/Constraints/Public/TransformConstraint.h
--- a/Engine/Source/Runtime/Experimental/Animation/Constraints/Public/TransformConstraint.h
+++ b/Engine/Source/Runtime/Experimental/Animation/Constraints/Public/TransformConstraint.h
@@ -58,7 +58,7 @@ public:
 	float Weight = 1.f;
 
 	/** Should the child be able to change it's offset dynamically. */
-	UPROPERTY(EditAnywhere, BlueprintReadWrite, Category="Offset")
+	UPROPERTY(EditAnywhere, BlueprintReadWrite, Category="Offset", meta=(EditCondition="bMaintainOffset"))
 	bool bDynamicOffset = false;
 
 	/** Returns the constraint type (Position, Parent, Aim...). */
@@ -157,13 +157,9 @@ protected:
 	virtual void ComputeOffset() override;
 
 	/** Defines the local child's translation offset in the parent space. */
-	UPROPERTY(EditAnywhere, BlueprintReadWrite, Category="Offset")
+	UPROPERTY(EditAnywhere, BlueprintReadWrite, Category="Offset", meta=(EditCondition="bMaintainOffset"))
 	FVector OffsetTranslation = FVector::ZeroVector;
 
-	/** Defines the local child's translation dynamic offset in the parent space. */
-	UPROPERTY(EditAnywhere, BlueprintReadWrite, Category="Offset")
-	FVector DynamicOffsetTranslation = FVector::ZeroVector;
-
 #if WITH_EDITOR
 public:
 	// UObject interface
@@ -204,13 +200,9 @@ protected:
 	virtual void OnHandleModified(UTransformableHandle* InHandle, bool bUpdate) override;
 	
 	/** Defines the local child's rotation offset in the parent space. */
-	UPROPERTY(EditAnywhere, BlueprintReadWrite, Category="Offset")
+	UPROPERTY(EditAnywhere, BlueprintReadWrite, Category="Offset", meta=(EditCondition="bMaintainOffset"))
 	FQuat OffsetRotation = FQuat::Identity;
 
-	/** Defines the local child's rotation dynamic offset in the parent space. */
-	UPROPERTY(EditAnywhere, BlueprintReadWrite, Category="Offset")
-	FQuat DynamicOffsetRotation = FQuat::Identity;
-
 #if WITH_EDITOR
 public:
 	// UObject interface
@@ -239,7 +231,7 @@ protected:
 	virtual void ComputeOffset() override;
 
 	/** Defines the local child's scale offset in the parent space. */
-	UPROPERTY(EditAnywhere, BlueprintReadWrite, Category="Offset")
+	UPROPERTY(EditAnywhere, BlueprintReadWrite, Category="Offset", meta=(EditCondition="bMaintainOffset"))
 	FVector OffsetScale = FVector::OneVector;
 };
 
@@ -274,13 +266,9 @@ protected:
 	virtual void ComputeOffset() override;
 
 	/** Defines the local child's transform offset in the parent space. */
-	UPROPERTY(EditAnywhere, BlueprintReadWrite, Category="Offset")
+	UPROPERTY(EditAnywhere, BlueprintReadWrite, Category="Offset", meta=(EditCondition="bMaintainOffset"))
 	FTransform OffsetTransform = FTransform::Identity;
 
-	/** Defines the local child's dynamic transform offset in the parent space. */
-	UPROPERTY(EditAnywhere, BlueprintReadWrite, Category="Offset")
-	FTransform DynamicOffsetTransform = FTransform::Identity;
-
 #if WITH_EDITOR
 public:
 	// UObject interface
