Chaos - island management
- SolverIsland constraint and particle lists are now transient and not maintained between frames (cleared at the end of frame)
- Constraint and Particle cached GraphIndex is always set when in the graph and cleared when not in the graph
- Islands now wake automatically - no explicit WakeIsland calls from evolution
- Lots of debug tracking (on a define, enabled in unit tests)
- Various unit tests

diff --git a/Engine/Source/Programs/HeadlessChaos/HeadlessChaos.Target.cs b/Engine/Source/Programs/HeadlessChaos/HeadlessChaos.Target.cs
--- a/Engine/Source/Programs/HeadlessChaos/HeadlessChaos.Target.cs
+++ b/Engine/Source/Programs/HeadlessChaos/HeadlessChaos.Target.cs
@@ -24,5 +24,9 @@ public class HeadlessChaosTarget : TestTargetRules
 
 		// our gtest does not have debug crt libs
 		bDebugBuildsActuallyUseDebugCRT = false;
+
+		// Enable additional debug functionality in the constraint management system to allow better unit testing
+		GlobalDefinitions.Add("CHAOS_CONSTRAINTHANDLE_DEBUG_ENABLED=1");
+		BuildEnvironment = TargetBuildEnvironment.Unique;
 	}
 }
diff --git a/Engine/Source/Programs/HeadlessChaos/Private/HeadlessChaosTestClustering.cpp b/Engine/Source/Programs/HeadlessChaos/Private/HeadlessChaosTestClustering.cpp
--- a/Engine/Source/Programs/HeadlessChaos/Private/HeadlessChaosTestClustering.cpp
+++ b/Engine/Source/Programs/HeadlessChaos/Private/HeadlessChaosTestClustering.cpp
@@ -141,11 +141,8 @@ namespace ChaosTest {
 			{
 				EXPECT_NE(Particle.Handle(), ClusteredParticles.Handle(BoxID));	//make sure boxes are not in non disabled array
 			}
-			
-			for (int32 Island = 0; Island < Evolution.NumIslands(); ++Island)
-			{
-				EXPECT_TRUE(Evolution.GetIslandParticles(Island).Contains(ClusteredParticles.Handle(BoxID)) == false);
-			}
+
+			EXPECT_TRUE(Evolution.GetConstraintGraph().DebugCheckParticleNotInGraph(ClusteredParticles.Handle(BoxID)));
 		}
 
 		for (Chaos::FPBDRigidParticleHandle* ClusterHandle : ClusterHandles)
@@ -171,10 +168,7 @@ namespace ChaosTest {
 				EXPECT_NE(Particle.Handle(), ClusterHandle);	//make sure boxes are not in non disabled array
 			}
 
-			for (int32 Island = 0; Island < Evolution.NumIslands(); ++Island)
-			{
-				EXPECT_TRUE(Evolution.GetIslandParticles(Island).Contains(ClusterHandle) == false);
-			}
+			EXPECT_TRUE(Evolution.GetConstraintGraph().DebugCheckParticleNotInGraph(ClusterHandle));
 		}
 
 		EXPECT_EQ(Particles.GetNonDisabledView().Num(), NumBoxes);
@@ -270,10 +264,7 @@ namespace ChaosTest {
 		{
 			EXPECT_TRUE(ClusteredParticles.Disabled(BoxID));	//no boxes should be active yet
 			EXPECT_TRUE(Evolution.GetRigidClustering().GetTopLevelClusterParents().Contains(ClusteredParticles.Handle(BoxID)) == false);
-			for (int32 Island = 0; Island < Evolution.NumIslands(); ++Island)
-			{
-				EXPECT_TRUE(Evolution.GetIslandParticles(Island).Contains(ClusteredParticles.Handle(BoxID)) == false);
-			}
+			EXPECT_TRUE(Evolution.GetConstraintGraph().DebugCheckParticleNotInGraph(ClusteredParticles.Handle(BoxID)));
 		}
 
 		SolverStrainArray[NumBoxes + NumBoxes / 4 + 1] = (FReal)1;
@@ -290,10 +281,7 @@ namespace ChaosTest {
 		{
 			EXPECT_TRUE(ClusteredParticles.Disabled(BoxID));	//no boxes should be active yet
 			EXPECT_TRUE(Evolution.GetRigidClustering().GetTopLevelClusterParents().Contains(ClusteredParticles.Handle(BoxID)) == false);
-			for (int32 Island = 0; Island < Evolution.NumIslands(); ++Island)
-			{
-				EXPECT_TRUE(Evolution.GetIslandParticles(Island).Contains(ClusteredParticles.Handle(BoxID)) == false);
-			}
+			EXPECT_TRUE(Evolution.GetConstraintGraph().DebugCheckParticleNotInGraph(ClusteredParticles.Handle(BoxID)));
 		}
 		
 	}
diff --git a/Engine/Source/Programs/HeadlessChaos/Private/HeadlessChaosTestGraph.cpp b/Engine/Source/Programs/HeadlessChaos/Private/HeadlessChaosTestGraph.cpp
--- a/Engine/Source/Programs/HeadlessChaos/Private/HeadlessChaosTestGraph.cpp
+++ b/Engine/Source/Programs/HeadlessChaos/Private/HeadlessChaosTestGraph.cpp
@@ -22,25 +22,8 @@ namespace ChaosTest {
 	template<int32 T_TYPEID>
 	class TMockGraphConstraints;
 
-
 	template<int32 T_TYPEID>
-	class TMockGraphConstraintHandle : public TIndexedContainerConstraintHandle<TMockGraphConstraints<T_TYPEID>>
-	{
-	public:
-		TMockGraphConstraintHandle(TMockGraphConstraints<T_TYPEID>* InConstraintContainer, int32 ConstraintIndex)
-			: TIndexedContainerConstraintHandle<TMockGraphConstraints<T_TYPEID>>(InConstraintContainer, ConstraintIndex)
-		{
-		}
-
-		virtual void SetEnabled(bool InEnabled) {};
-		virtual bool IsEnabled() const { return true; };
-
-		static const FConstraintHandleTypeID& StaticType()
-		{
-			static FConstraintHandleTypeID STypeID(TEXT("FMockConstraintHandle"), &FIndexedConstraintHandle::StaticType());
-			return STypeID;
-		}
-	};
+	class TMockGraphConstraintHandle;
 
 	/**
 	 * Constraint Container with minimal API required to test the Graph.
@@ -55,6 +38,8 @@ namespace ChaosTest {
 		struct FMockConstraint
 		{
 			TVec2<int32> ConstrainedParticles;
+			bool bEnabled;
+			bool bSleeping;
 		};
 
 		TMockGraphConstraints()
@@ -72,7 +57,7 @@ namespace ChaosTest {
 
 		void AddConstraint(const TVec2<int32>& InConstraintedParticles)
 		{
-			Constraints.Emplace(FMockConstraint({ InConstraintedParticles }));
+			Constraints.Emplace(FMockConstraint({ InConstraintedParticles, true, false }));
 			Handles.Emplace(HandleAllocator.AllocHandle(this, Handles.Num()));
 		}
 
@@ -127,8 +112,54 @@ namespace ChaosTest {
 		TArray<FMockConstraint> Constraints;
 		TArray<FConstraintContainerHandle*> Handles;
 		TConstraintHandleAllocator<TMockGraphConstraints<T_TYPEID>> HandleAllocator;
+	};
+
+	template<int32 T_TYPEID>
+	class TMockGraphConstraintHandle : public TIndexedContainerConstraintHandle<TMockGraphConstraints<T_TYPEID>>
+	{
+	protected:
+		using Base = TIndexedContainerConstraintHandle<TMockGraphConstraints<T_TYPEID>>;
+		using Base::ConcreteContainer;
+	public:
+		using Base::GetConstraintIndex;
+
+	public:
+
+		TMockGraphConstraintHandle(TMockGraphConstraints<T_TYPEID>* InConstraintContainer, int32 ConstraintIndex)
+			: TIndexedContainerConstraintHandle<TMockGraphConstraints<T_TYPEID>>(InConstraintContainer, ConstraintIndex)
+		{
+		}
+
+		virtual void SetEnabled(bool InEnabled) override
+		{
+			ConcreteContainer()->Constraints[GetConstraintIndex()].bEnabled = InEnabled;
+		};
+
+		virtual bool IsEnabled() const override
+		{
+			return ConcreteContainer()->Constraints[GetConstraintIndex()].bEnabled;
+		};
+
+		virtual void SetIsSleeping(bool bInSleeping) override
+		{
+			ConcreteContainer()->Constraints[GetConstraintIndex()].bSleeping = bInSleeping;
+		}
+
+		virtual bool IsSleeping() const override
+		{
+			return ConcreteContainer()->Constraints[GetConstraintIndex()].bSleeping;
+		}
+
+		virtual FParticlePair GetConstrainedParticles() const override { return FParticlePair(nullptr, nullptr); }
 
+		static const FConstraintHandleTypeID& StaticType()
+		{
+			static FConstraintHandleTypeID STypeID(TEXT("FMockConstraintHandle"), &FIndexedConstraintHandle::StaticType());
+			return STypeID;
+		}
 	};
+
+
 	template class TMockGraphConstraints<0>;
 	template class TMockGraphConstraints<1>;
 
@@ -588,20 +619,17 @@ namespace ChaosTest {
 				
 				++Index;
 			}
-			Graph.ResetIndices();
+
+			Graph.EndTick();
+			Graph.Reset();
 		}
 	}
 
 	void HelpTickConstraints(FPBDRigidsSOAs& SOAs, const TArray<TPBDRigidParticleHandle<FReal, 3>*>& Particles,
-		FPBDConstraintGraph& Graph, const TArray<TVec2<int32>>& ConstrainedParticles,
+		FPBDConstraintGraph& Graph, TMockGraphConstraints<0> Constraints,
 		const TArrayCollectionArray<TSerializablePtr<FChaosPhysicsMaterial>>& PhysicsMaterials,
 		const THandleArray<FChaosPhysicsMaterial>& PhysicalMaterials)
 	{
-		TMockGraphConstraints<0> Constraints;
-		for(const auto& ConstrainedParticleIndices : ConstrainedParticles)
-		{
-			Constraints.AddConstraint(ConstrainedParticleIndices);
-		}
 
 		SOAs.ClearTransientDirty();
 		Graph.InitializeGraph(SOAs.GetNonDisabledDynamicView());
@@ -624,9 +652,12 @@ namespace ChaosTest {
 				{
 					SOAs.DeactivateParticle(Particle);
 				}
+
+				Graph.SleepIsland(SOAs, IslandIndex);
 			}
 		}
 
+		Graph.EndTick();
 	}
 
 	bool ContainsHelper(const TParticleView<TPBDRigidParticles<FReal,3>>& View,const TGeometryParticleHandle<FReal,3>* InParticle)
@@ -690,6 +721,12 @@ namespace ChaosTest {
 				{3, 4},
 			};
 
+			TMockGraphConstraints<0> Constraints;
+			for (const auto& ConstrainedParticleIndices : ConstrainedParticles)
+			{
+				Constraints.AddConstraint(ConstrainedParticleIndices);
+			}
+
 			FPBDConstraintGraph Graph;
 			for (int32 LoopIndex = 0; LoopIndex < 5 + PhysicalMaterial->SleepCounterThreshold; ++LoopIndex)
 			{
@@ -700,7 +737,7 @@ namespace ChaosTest {
 					Particles[ParticleIndex]->ResetSmoothedVelocities();
 				}
 
-				HelpTickConstraints(SOAs,Particles,Graph,ConstrainedParticles,PhysicsMaterials,PhysicalMaterials);
+				HelpTickConstraints(SOAs,Particles,Graph,Constraints,PhysicsMaterials,PhysicalMaterials);
 			
 				// Particles 0-2 are always awake
 				EXPECT_FALSE(Particles[0]->Sleeping());
@@ -766,13 +803,17 @@ namespace ChaosTest {
 				{3, 4},
 			};
 
-			TArray<TVec2<int32>> ConstrainedParticlesAfterSleep =
+			TArray<TVec2<int32>> ConstrainedParticlesOnWake =
 			{
-				{0,1},
 				{1,3},	//will merge islands and wake up 3,4
-				{3,4}
 			};
 
+			TMockGraphConstraints<0> Constraints;
+			for (const auto& ConstrainedParticleIndices : ConstrainedParticles)
+			{
+				Constraints.AddConstraint(ConstrainedParticleIndices);
+			}
+
 			FPBDConstraintGraph Graph;
 			const int32 WakeUpFrame = 5 + PhysicalMaterial->SleepCounterThreshold;
 			for (int32 LoopIndex = 0; LoopIndex < WakeUpFrame + 5; ++LoopIndex)
@@ -784,14 +825,15 @@ namespace ChaosTest {
 					Particles[ParticleIndex]->ResetSmoothedVelocities();
 				}
 
-				if(LoopIndex < WakeUpFrame)
-				{
-					HelpTickConstraints(SOAs,Particles,Graph,ConstrainedParticles,PhysicsMaterials,PhysicalMaterials);
-				}
-				else
+				if (LoopIndex == WakeUpFrame)
 				{
-					HelpTickConstraints(SOAs,Particles,Graph,ConstrainedParticlesAfterSleep,PhysicsMaterials,PhysicalMaterials);
+					for (const auto& ConstrainedParticleIndices : ConstrainedParticlesOnWake)
+					{
+						Constraints.AddConstraint(ConstrainedParticleIndices);
+					}
 				}
+
+				HelpTickConstraints(SOAs,Particles,Graph,Constraints,PhysicsMaterials,PhysicalMaterials);
 			
 				// Particles 0-2 are always awake
 				EXPECT_FALSE(Particles[0]->Sleeping());
@@ -857,16 +899,20 @@ namespace ChaosTest {
 			TArray<TVec2<int32>> ConstrainedParticles =
 			{
 				{0,1},
-			{3,4},
+				{3,4},
 			};
 
-			TArray<TVec2<int32>> ConstrainedParticlesAfterSleep =
+			TArray<TVec2<int32>> ConstrainedParticlesOnWake =
 			{
-				{0,1},
-			{1,3},	//will merge islands and wake up 3,4
-			{3,4}
+				{1,3},	//will merge islands and wake up 3,4
 			};
 
+			TMockGraphConstraints<0> Constraints;
+			for (const auto& ConstrainedParticleIndices : ConstrainedParticles)
+			{
+				Constraints.AddConstraint(ConstrainedParticleIndices);
+			}
+
 			FPBDConstraintGraph Graph;
 			const int32 MergeFrame = 5 + PhysicalMaterial->SleepCounterThreshold;
 			const int32 SleepAfterMergeFrame = MergeFrame + PhysicalMaterial->SleepCounterThreshold + 1;	//first frame after merge must wake up
@@ -886,15 +932,16 @@ namespace ChaosTest {
 					Particles[ParticleIndex]->ResetSmoothedVelocities();
 				}
 
-				if(LoopIndex < MergeFrame)
-				{
-					HelpTickConstraints(SOAs,Particles,Graph,ConstrainedParticles,PhysicsMaterials,PhysicalMaterials);
-				}
-				else
+				if (LoopIndex == MergeFrame)
 				{
-					HelpTickConstraints(SOAs,Particles,Graph,ConstrainedParticlesAfterSleep,PhysicsMaterials,PhysicalMaterials);
+					for (const auto& ConstrainedParticleIndices : ConstrainedParticlesOnWake)
+					{
+						Constraints.AddConstraint(ConstrainedParticleIndices);
+					}
 				}
 
+				HelpTickConstraints(SOAs,Particles,Graph,Constraints,PhysicsMaterials,PhysicalMaterials);
+
 				// Particle 2 is always awake
 				EXPECT_FALSE(Particles[2]->Sleeping());
 				EXPECT_TRUE(ContainsHelper(SOAs.GetActiveParticlesView(),Particles[2]));
@@ -1263,7 +1310,7 @@ namespace ChaosTest {
 		{
 			FPBDRigidParticleHandle* Particle0 = AllParticles[Indices[0]]->CastToRigidParticle();
 			FPBDRigidParticleHandle* Particle1 = AllParticles[Indices[1]]->CastToRigidParticle();
-			auto* NewConstraint = JointConstraints.AddConstraint(FPBDJointConstraints::FParticlePair(Particle0,Particle1), FRigidTransform3());
+			auto* NewConstraint = JointConstraints.AddConstraint(FParticlePair(Particle0,Particle1), FRigidTransform3());
 			Graph.AddConstraint(3, NewConstraint, TVec2<FGeometryParticleHandle*>(Particle0, Particle1));
 		}
 		ConstraintGraphValidation::ValidateConstraintGraphParticleHandles(Graph);
diff --git a/Engine/Source/Programs/HeadlessChaos/Private/HeadlessChaosTestGraphEvolution.cpp b/Engine/Source/Programs/HeadlessChaos/Private/HeadlessChaosTestGraphEvolution.cpp
new file mode 100644
--- /dev/null
+++ b/Engine/Source/Programs/HeadlessChaos/Private/HeadlessChaosTestGraphEvolution.cpp
@@ -0,0 +1,592 @@
+// Copyright Epic Games, Inc. All Rights Reserved.
+#include "CoreMinimal.h"
+#include "HeadlessChaos.h"
+#include "HeadlessChaosTestUtility.h"
+
+#include "Chaos/Island/IslandManager.h"
+#include "Chaos/ParticleHandle.h"
+#include "Chaos/PBDConstraintRule.h"
+#include "Chaos/PBDNullConstraints.h"
+#include "Chaos/PBDRigidsEvolutionGBF.h"
+#include "Chaos/Utilities.h"
+#include "Modules/ModuleManager.h"
+
+namespace ChaosTest
+{
+	using namespace Chaos;
+
+	class FGraphEvolutionTest
+	{
+	public:
+		FGraphEvolutionTest(const int32 NumParticles)
+			: UniqueIndices()
+			, Particles(UniqueIndices)
+			, PhysicalMaterials()
+			, Evolution(Particles, PhysicalMaterials)
+			, Constraints()
+			, ConstraintsRule(Constraints)
+			, TickCount(0)
+		{
+			IslandManager = &Evolution.GetConstraintGraph();
+
+			// Bind the constraints to the evolution
+			Evolution.AddConstraintRule(&ConstraintsRule);
+
+			Evolution.GetGravityForces().SetAcceleration(FVec3(0));
+
+			// Create some particles and constraints connecting them in a chain: 0-1-2-3-...-N
+			ParticleHandles = Evolution.CreateDynamicParticles(NumParticles);
+		}
+
+		// Connect all the particles in a chain
+		void MakeChain()
+		{
+			for (int32 ParticleIndex = 0; ParticleIndex < ParticleHandles.Num() - 1; ++ParticleIndex)
+			{
+				ConstraintHandles.Add(Constraints.AddConstraint({ ParticleHandles[ParticleIndex], ParticleHandles[ParticleIndex + 1] }));
+			}
+		}
+
+		// Treat particle0 like a kinematic floor with all the other particles sat on it
+		void MakeFloor()
+		{
+			Evolution.SetParticleObjectState(ParticleHandles[0], EObjectStateType::Kinematic);
+
+			for (int32 ParticleIndex = 0; ParticleIndex < ParticleHandles.Num() - 1; ++ParticleIndex)
+			{
+				ConstraintHandles.Add(Constraints.AddConstraint({ ParticleHandles[0], ParticleHandles[ParticleIndex + 1] }));
+			}
+		}
+
+		void Advance()
+		{
+			Evolution.AdvanceOneTimeStep(FReal(1.0 / 60.0));
+			++TickCount;
+		}
+
+		void AdvanceUntilSleeping(int32 MaxIterations)
+		{
+			const int32 MaxTickCount = TickCount + MaxIterations;
+			bool bIsSleeping = false;
+			while (!bIsSleeping && (TickCount < MaxTickCount))
+			{
+				Advance();
+
+				bIsSleeping = true;
+				for (int32 IslandIndex = 0; IslandIndex < IslandManager->NumIslands(); ++IslandIndex)
+				{
+					bIsSleeping = bIsSleeping && IslandManager->GetSolverIsland(IslandIndex)->IsSleeping();
+				}
+			}
+
+			EXPECT_TRUE(bIsSleeping);
+			EXPECT_LT(TickCount, MaxTickCount);
+		}
+
+		FParticleUniqueIndicesMultithreaded UniqueIndices;
+		FPBDRigidsSOAs Particles;
+		THandleArray<FChaosPhysicsMaterial> PhysicalMaterials;
+		FPBDRigidsEvolutionGBF Evolution;
+
+		// Create a constraint container and bind it to the evolution
+		FPBDNullConstraints Constraints;
+		TPBDConstraintIslandRule<FPBDNullConstraints> ConstraintsRule;
+
+		TArray<FPBDRigidParticleHandle*> ParticleHandles;
+		TArray<FPBDNullConstraintHandle*> ConstraintHandles;
+		FPBDIslandManager* IslandManager;
+
+		int32 TickCount;
+	};
+
+	// Veryify that the Null Constraint mockup is working as intended. We can create the container and constraints, and they are correctly bound to the evolution
+	GTEST_TEST(GraphEvolutionTests, TestNullConstraint)
+	{
+		FGraphEvolutionTest Test(4);
+		Test.MakeChain();
+
+		Test.Advance();
+
+		EXPECT_EQ(Test.IslandManager->NumIslands(), 1);
+	}
+
+	// Start with an island containing 4 particle connected in a chain, then make the second one kinematic.
+	// Check that the island splits.
+	//		A-B-C-D 
+	// =>	A-B, B-C-D
+	//
+	GTEST_TEST(GraphEvolutionTests, TestConstraintGraph_ToKinematic)
+	{
+		FGraphEvolutionTest Test(4);
+		Test.MakeChain();
+
+		Test.Advance();
+
+		EXPECT_EQ(Test.IslandManager->NumIslands(), 1);
+
+		// Convert particle B to kinematic to split the islands: A-B, B-C-D
+		Test.Evolution.SetParticleObjectState(Test.ParticleHandles[1], EObjectStateType::Kinematic);
+
+		Test.Advance();
+
+		// Should now have 2 islands
+		EXPECT_EQ(Test.IslandManager->NumIslands(), 2);
+	}
+
+	// Wait for all particles to go to sleep naturally (i.e., as part of the tick and not by explicitly setting the state)
+	// then check that the islands are preserved.
+	GTEST_TEST(GraphEvolutionTests, TestConstraintGraph_ParticleSleep_Natural)
+	{
+		FGraphEvolutionTest Test(4);
+		Test.MakeChain();
+
+		Test.Advance();
+
+		// All constraints in graph in 1 island that is awake
+		EXPECT_EQ(Test.IslandManager->NumIslands(), 1);
+		EXPECT_NE(Test.ConstraintHandles[0]->ConstraintGraphIndex(), INDEX_NONE);
+		EXPECT_NE(Test.ConstraintHandles[1]->ConstraintGraphIndex(), INDEX_NONE);
+		EXPECT_NE(Test.ConstraintHandles[2]->ConstraintGraphIndex(), INDEX_NONE);
+		EXPECT_FALSE(Test.IslandManager->GetSolverIsland(0)->IsSleeping());
+
+		// Make all the particles sleep
+		Test.AdvanceUntilSleeping(50);
+
+		// Island should be asleep
+		EXPECT_EQ(Test.IslandManager->NumIslands(), 1);
+		EXPECT_TRUE(Test.IslandManager->GetSolverIsland(0)->IsSleeping());
+	}
+
+	// Sleep all particles in the scene and ensure that the island manager puts the island to sleep
+	// but retains all the constraints and particles in the island.
+	// 
+	// @todo(chaos): this test fails. Manually setting all particles in an island to the sleep state
+	// does put the island to sleep, but only after all the constraints have been removed because
+	// InitializeGraph is called before the sleep state is updated.
+	//
+	GTEST_TEST(GraphEvolutionTests, DISABLED_TestConstraintGraph_ParticleSleep_Manual)
+	{
+		FGraphEvolutionTest Test(4);
+		Test.MakeChain();
+
+		Test.Advance();
+
+		// All constraints in graph in 1 island that is awake
+		EXPECT_EQ(Test.IslandManager->NumIslands(), 1);
+		EXPECT_NE(Test.ConstraintHandles[0]->ConstraintGraphIndex(), INDEX_NONE);
+		EXPECT_NE(Test.ConstraintHandles[1]->ConstraintGraphIndex(), INDEX_NONE);
+		EXPECT_NE(Test.ConstraintHandles[2]->ConstraintGraphIndex(), INDEX_NONE);
+		EXPECT_FALSE(Test.IslandManager->GetSolverIsland(0)->IsSleeping());
+
+		// Make all the particles sleep
+		for (FPBDRigidParticleHandle* ParticleHandle : Test.ParticleHandles)
+		{
+			Test.Evolution.SetParticleObjectState(ParticleHandle, EObjectStateType::Sleeping);
+		}
+
+		Test.Advance();
+
+		// Island should be asleep
+		EXPECT_EQ(Test.IslandManager->NumIslands(), 1);
+		EXPECT_TRUE(Test.IslandManager->GetSolverIsland(0)->IsSleeping());
+	}
+
+	// Start with an island containing 4 particle connected in a chain, then make the middle two kinematic.
+	// This makes the B-C constraint kinematic ("invalid" in the Constraint Graph terminolgy) which means
+	// it does not belong in any island.
+	// Check that the island manager handles kinematic-kinematic constraints
+	//		A-B-C-D 
+	// =>	A-B  C-D
+	//
+	GTEST_TEST(GraphEvolutionTests, TestConstraintGraph_KinematicKinematic)
+	{
+		FGraphEvolutionTest Test(4);
+		Test.MakeChain();
+
+		Test.Advance();
+
+		// All constraints in graph in 1 island
+		EXPECT_EQ(Test.IslandManager->NumIslands(), 1);
+		EXPECT_NE(Test.ConstraintHandles[0]->ConstraintGraphIndex(), INDEX_NONE);
+		EXPECT_NE(Test.ConstraintHandles[1]->ConstraintGraphIndex(), INDEX_NONE);
+		EXPECT_NE(Test.ConstraintHandles[2]->ConstraintGraphIndex(), INDEX_NONE);
+
+		// Convert a particle B and C to kinematic to split the islands: A-B, C-D
+		Test.Evolution.SetParticleObjectState(Test.ParticleHandles[1], EObjectStateType::Kinematic);
+		Test.Evolution.SetParticleObjectState(Test.ParticleHandles[2], EObjectStateType::Kinematic);
+
+		Test.Advance();
+
+		// Constraint[1] is not in the graph or an island
+		EXPECT_NE(Test.ConstraintHandles[0]->ConstraintGraphIndex(), INDEX_NONE);
+		EXPECT_EQ(Test.ConstraintHandles[1]->ConstraintGraphIndex(), INDEX_NONE);	// Not in graph
+		EXPECT_NE(Test.ConstraintHandles[2]->ConstraintGraphIndex(), INDEX_NONE);
+		EXPECT_TRUE(Test.IslandManager->DebugCheckConstraintNotInGraph(Test.ConstraintHandles[1]));
+
+		// Should now have 2 islands
+		EXPECT_EQ(Test.IslandManager->NumIslands(), 2);
+	}
+
+
+	// Same as TestConstraintGraph_KinematicKinematic but islands are sleeping when the change is made.
+	// Sleeping islands are not cleared of constraints, but they should be awoken by the dynamic->kinematic conversion
+	GTEST_TEST(GraphEvolutionTests, TestConstraintGraph_KinematicKinematic_Sleeping)
+	{
+		FGraphEvolutionTest Test(4);
+		Test.MakeChain();
+
+		Test.Advance();
+
+		// All constraints in graph in 1 island that is awake
+		EXPECT_EQ(Test.IslandManager->NumIslands(), 1);
+		EXPECT_NE(Test.ConstraintHandles[0]->ConstraintGraphIndex(), INDEX_NONE);
+		EXPECT_NE(Test.ConstraintHandles[1]->ConstraintGraphIndex(), INDEX_NONE);
+		EXPECT_NE(Test.ConstraintHandles[2]->ConstraintGraphIndex(), INDEX_NONE);
+		EXPECT_FALSE(Test.IslandManager->GetSolverIsland(0)->IsSleeping());
+
+		// Wait for the sleep state
+		Test.AdvanceUntilSleeping(50);
+
+		// Island should be asleep but still contain all the particles and constraints
+		EXPECT_EQ(Test.IslandManager->NumIslands(), 1);
+		EXPECT_TRUE(Test.IslandManager->GetSolverIsland(0)->IsSleeping());
+		EXPECT_NE(Test.ConstraintHandles[0]->ConstraintGraphIndex(), INDEX_NONE);
+		EXPECT_NE(Test.ConstraintHandles[1]->ConstraintGraphIndex(), INDEX_NONE);
+		EXPECT_NE(Test.ConstraintHandles[2]->ConstraintGraphIndex(), INDEX_NONE);
+
+		// Convert a particle B and C to kinematic to split the islands: A-B, C-D
+		// NOTE: the island won't split until the island is woken
+		Test.Evolution.SetParticleObjectState(Test.ParticleHandles[1], EObjectStateType::Kinematic);
+		Test.Evolution.SetParticleObjectState(Test.ParticleHandles[2], EObjectStateType::Kinematic);
+
+		Test.Advance();
+
+		// Island is asleep and won't have split yet
+		EXPECT_EQ(Test.IslandManager->NumIslands(), 1);
+
+		// The kinematic-kinematic constraint will still be in the graph, and still in its island (it will not be removed until the island wakes)
+		EXPECT_NE(Test.ConstraintHandles[1]->ConstraintGraphIndex(), INDEX_NONE);
+		EXPECT_NE(Test.IslandManager->GetConstraintIsland(Test.ConstraintHandles[1]), INDEX_NONE);
+
+		// Wake a dynamic particle the island and we should get the split at the next update
+		Test.Evolution.SetParticleObjectState(Test.ParticleHandles[0], EObjectStateType::Dynamic);
+
+		Test.Advance();
+
+		// Should now have 2 islands (with 2 particles and 1 constraint each though we can't check that atm)
+		EXPECT_EQ(Test.IslandManager->NumIslands(), 2);
+
+		// The kinematic-kinematic constraint will still be in the graph, but not in any island
+		EXPECT_NE(Test.ConstraintHandles[1]->ConstraintGraphIndex(), INDEX_NONE);
+		EXPECT_EQ(Test.IslandManager->GetConstraintIsland(Test.ConstraintHandles[1]), INDEX_NONE);
+	}
+
+	// 3 objects sat on the floor awake. Make the floor dynamic.
+	// This tests what happens when a kinematic in multiple islands gets converted to a dynamic.
+	// 
+	// (d=dynamic, s=sleeping, k=kinematic)
+	// Bd   Cd   Dd          Bd   Cd   Dd
+	//  \   |   /	   =>     \   |   /
+	//      Ak		              Ad
+	//
+	GTEST_TEST(GraphEvolutionTests, TestConstraintGraph_KinematicToDynamic)
+	{
+		FGraphEvolutionTest Test(4);
+		Test.MakeFloor();
+
+		Test.Advance();
+
+		// Each particle in its own island (kinematic will be in all 3)
+		EXPECT_EQ(Test.IslandManager->NumIslands(), 3);
+		EXPECT_EQ(Test.IslandManager->FindParticleIslands(Test.ParticleHandles[0]).Num(), 3);
+		EXPECT_EQ(Test.IslandManager->FindParticleIslands(Test.ParticleHandles[1]).Num(), 1);
+		EXPECT_EQ(Test.IslandManager->FindParticleIslands(Test.ParticleHandles[2]).Num(), 1);
+		EXPECT_EQ(Test.IslandManager->FindParticleIslands(Test.ParticleHandles[3]).Num(), 1);
+
+		// Convert A to dynamic which should merge all the islands into 1
+		Test.Evolution.SetParticleObjectState(Test.ParticleHandles[0], EObjectStateType::Dynamic);
+
+		Test.Advance();
+
+		// All particles in same island
+		EXPECT_EQ(Test.IslandManager->NumIslands(), 1);
+		EXPECT_EQ(Test.IslandManager->FindParticleIslands(Test.ParticleHandles[0]).Num(), 1);
+		EXPECT_EQ(Test.IslandManager->FindParticleIslands(Test.ParticleHandles[1]).Num(), 1);
+		EXPECT_EQ(Test.IslandManager->FindParticleIslands(Test.ParticleHandles[2]).Num(), 1);
+		EXPECT_EQ(Test.IslandManager->FindParticleIslands(Test.ParticleHandles[3]).Num(), 1);
+	}
+
+	// 3 objects sat on the floor asleep. Make the floor dynamic.
+	// This tests what happens when a kinematic in multiple sleeping islands gets converted to a dynamic.
+	// 
+	// (d=dynamic, s=sleeping, k=kinematic)
+	// Bs   Cs   Ds          Bd   Cd   Dd
+	//  \   |   /	   =>     \   |   /
+	//      Ak		              Ad
+	//
+	GTEST_TEST(GraphEvolutionTests, TestConstraintGraph_KinematicToDynamic_WithSleep)
+	{
+		FGraphEvolutionTest Test(4);
+		Test.MakeFloor();
+
+		Test.AdvanceUntilSleeping(50);
+
+		// Each particle in its own island (kinematic will be in all 3)
+		EXPECT_EQ(Test.IslandManager->NumIslands(), 3);
+		EXPECT_EQ(Test.IslandManager->FindParticleIslands(Test.ParticleHandles[0]).Num(), 3);
+		EXPECT_EQ(Test.IslandManager->FindParticleIslands(Test.ParticleHandles[1]).Num(), 1);
+		EXPECT_EQ(Test.IslandManager->FindParticleIslands(Test.ParticleHandles[2]).Num(), 1);
+		EXPECT_EQ(Test.IslandManager->FindParticleIslands(Test.ParticleHandles[3]).Num(), 1);
+
+		// Convert A to dynamic which should merge all the islands into 1
+		Test.Evolution.SetParticleObjectState(Test.ParticleHandles[0], EObjectStateType::Dynamic);
+
+		Test.Advance();
+
+		// All particles in one awake island
+		EXPECT_EQ(Test.IslandManager->NumIslands(), 1);
+		EXPECT_FALSE(Test.IslandManager->GetSolverIsland(0)->IsSleeping());
+		EXPECT_EQ(Test.IslandManager->FindParticleIslands(Test.ParticleHandles[0]).Num(), 1);
+		EXPECT_EQ(Test.IslandManager->FindParticleIslands(Test.ParticleHandles[1]).Num(), 1);
+		EXPECT_EQ(Test.IslandManager->FindParticleIslands(Test.ParticleHandles[2]).Num(), 1);
+		EXPECT_EQ(Test.IslandManager->FindParticleIslands(Test.ParticleHandles[3]).Num(), 1);
+
+		// All particles awake
+		EXPECT_FALSE(FConstGenericParticleHandle(Test.ParticleHandles[0])->Sleeping());
+		EXPECT_FALSE(FConstGenericParticleHandle(Test.ParticleHandles[1])->Sleeping());
+		EXPECT_FALSE(FConstGenericParticleHandle(Test.ParticleHandles[2])->Sleeping());
+		EXPECT_FALSE(FConstGenericParticleHandle(Test.ParticleHandles[3])->Sleeping());
+	}
+
+	// 3 Objects sat on the floor asleep. Make the floor dynamic and asleep.
+	// This tests that adding a sleeping body to an island does not wake it.
+	// This is required for streaming to work which adds bodies over multiple frames.
+	// 
+	// (d=dynamic, s=sleeping, k=kinematic)
+	// Bs   Cs   Ds          Bs   Cs   Ds
+	//  \   |   /      =>     \   |   /
+	//      Ak                    As
+	//
+	GTEST_TEST(GraphEvolutionTests, TestConstraintGraph_KinematicToDynamic_WithSleep2)
+	{
+		FGraphEvolutionTest Test(4);
+		Test.MakeFloor();
+
+		Test.AdvanceUntilSleeping(50);
+
+		// Each particle in its own island (kinematic will be in all 3)
+		EXPECT_EQ(Test.IslandManager->NumIslands(), 3);
+		EXPECT_EQ(Test.IslandManager->FindParticleIslands(Test.ParticleHandles[0]).Num(), 3);
+		EXPECT_EQ(Test.IslandManager->FindParticleIslands(Test.ParticleHandles[1]).Num(), 1);
+		EXPECT_EQ(Test.IslandManager->FindParticleIslands(Test.ParticleHandles[2]).Num(), 1);
+		EXPECT_EQ(Test.IslandManager->FindParticleIslands(Test.ParticleHandles[3]).Num(), 1);
+
+		// Convert A to dynamic sleeping which should merge all the islands into 1 but leave it asleep
+		Test.Evolution.SetParticleObjectState(Test.ParticleHandles[0], EObjectStateType::Sleeping);
+
+		Test.Advance();
+
+		// All particles in one asleep island
+		EXPECT_EQ(Test.IslandManager->NumIslands(), 1);
+		EXPECT_TRUE(Test.IslandManager->GetSolverIsland(0)->IsSleeping());
+		EXPECT_EQ(Test.IslandManager->FindParticleIslands(Test.ParticleHandles[0]).Num(), 1);
+		EXPECT_EQ(Test.IslandManager->FindParticleIslands(Test.ParticleHandles[1]).Num(), 1);
+		EXPECT_EQ(Test.IslandManager->FindParticleIslands(Test.ParticleHandles[2]).Num(), 1);
+		EXPECT_EQ(Test.IslandManager->FindParticleIslands(Test.ParticleHandles[3]).Num(), 1);
+
+		// All particles asleep
+		EXPECT_TRUE(FConstGenericParticleHandle(Test.ParticleHandles[0])->Sleeping());
+		EXPECT_TRUE(FConstGenericParticleHandle(Test.ParticleHandles[1])->Sleeping());
+		EXPECT_TRUE(FConstGenericParticleHandle(Test.ParticleHandles[2])->Sleeping());
+		EXPECT_TRUE(FConstGenericParticleHandle(Test.ParticleHandles[3])->Sleeping());
+
+	}
+
+	// 3 Objects sat on the floor, 2 asleep and 1 awake. Make the floor dynamic and asleep.
+	// In this case we should get 1 awake island and all particles should wake.
+	// 
+	// (d=dynamic, s=sleeping, k=kinematic)
+	// Bs   Cs   Dd          Bd   Cd   Dd
+	//  \   |   /      =>     \   |   /
+	//      Ak	                  Ad
+	//
+	GTEST_TEST(GraphEvolutionTests, TestConstraintGraph_KinematicToDynamic_WithSleep3)
+	{
+		FGraphEvolutionTest Test(4);
+		Test.MakeFloor();
+
+		Test.AdvanceUntilSleeping(50);
+
+		// Wake D
+		Test.Evolution.SetParticleObjectState(Test.ParticleHandles[3], EObjectStateType::Dynamic);
+
+		Test.Advance();
+
+		EXPECT_TRUE(FConstGenericParticleHandle(Test.ParticleHandles[1])->Sleeping());
+		EXPECT_TRUE(FConstGenericParticleHandle(Test.ParticleHandles[2])->Sleeping());
+		EXPECT_FALSE(FConstGenericParticleHandle(Test.ParticleHandles[3])->Sleeping());
+
+		// Each particle in its own island (kinematic will be in all 3)
+		EXPECT_EQ(Test.IslandManager->NumIslands(), 3);
+		EXPECT_EQ(Test.IslandManager->FindParticleIslands(Test.ParticleHandles[0]).Num(), 3);
+		EXPECT_EQ(Test.IslandManager->FindParticleIslands(Test.ParticleHandles[1]).Num(), 1);
+		EXPECT_EQ(Test.IslandManager->FindParticleIslands(Test.ParticleHandles[2]).Num(), 1);
+		EXPECT_EQ(Test.IslandManager->FindParticleIslands(Test.ParticleHandles[3]).Num(), 1);
+
+		// Convert A to dynamic sleeping
+		Test.Evolution.SetParticleObjectState(Test.ParticleHandles[0], EObjectStateType::Sleeping);
+
+		Test.Advance();
+
+		// All particles in one awake island (D was awake so it would wake the island)
+		EXPECT_EQ(Test.IslandManager->NumIslands(), 1);
+		EXPECT_FALSE(Test.IslandManager->GetSolverIsland(0)->IsSleeping());
+		EXPECT_EQ(Test.IslandManager->FindParticleIslands(Test.ParticleHandles[0]).Num(), 1);
+		EXPECT_EQ(Test.IslandManager->FindParticleIslands(Test.ParticleHandles[1]).Num(), 1);
+		EXPECT_EQ(Test.IslandManager->FindParticleIslands(Test.ParticleHandles[2]).Num(), 1);
+		EXPECT_EQ(Test.IslandManager->FindParticleIslands(Test.ParticleHandles[3]).Num(), 1);
+
+		// All particles awake
+		EXPECT_FALSE(FConstGenericParticleHandle(Test.ParticleHandles[0])->Sleeping());
+		EXPECT_FALSE(FConstGenericParticleHandle(Test.ParticleHandles[1])->Sleeping());
+		EXPECT_FALSE(FConstGenericParticleHandle(Test.ParticleHandles[2])->Sleeping());
+		EXPECT_FALSE(FConstGenericParticleHandle(Test.ParticleHandles[3])->Sleeping());
+	}
+
+	// 3 Objects sat on the floor, 2 asleep and 1 awake. Make the floor dynamic and asleep.
+	// Same as TestConstraintGraph_KinematicToDynamic_WithSleep3 except we wake a different particle to 
+	// be sure we weren't just lucky above (when we make a kinematic into a dynamic we add it to one
+	// of the islands it is in. This is testing that this is ok).
+	// 
+	// (d=dynamic, s=sleeping, k=kinematic)
+	// Bd   Cs   Ds          Bd   Cd   Dd
+	//  \   |   /      =>     \   |   /
+	//      Ak	                  Ad
+	//
+	GTEST_TEST(GraphEvolutionTests, TestConstraintGraph_KinematicToDynamic_WithSleep4)
+	{
+		FGraphEvolutionTest Test(4);
+		Test.MakeFloor();
+
+		Test.AdvanceUntilSleeping(50);
+
+		// Wake B
+		Test.Evolution.SetParticleObjectState(Test.ParticleHandles[1], EObjectStateType::Dynamic);
+
+		Test.Advance();
+
+		EXPECT_FALSE(FConstGenericParticleHandle(Test.ParticleHandles[1])->Sleeping());
+		EXPECT_TRUE(FConstGenericParticleHandle(Test.ParticleHandles[2])->Sleeping());
+		EXPECT_TRUE(FConstGenericParticleHandle(Test.ParticleHandles[3])->Sleeping());
+
+		// Each particle in its own island (kinematic will be in all 3)
+		EXPECT_EQ(Test.IslandManager->NumIslands(), 3);
+		EXPECT_EQ(Test.IslandManager->FindParticleIslands(Test.ParticleHandles[0]).Num(), 3);
+		EXPECT_EQ(Test.IslandManager->FindParticleIslands(Test.ParticleHandles[1]).Num(), 1);
+		EXPECT_EQ(Test.IslandManager->FindParticleIslands(Test.ParticleHandles[2]).Num(), 1);
+		EXPECT_EQ(Test.IslandManager->FindParticleIslands(Test.ParticleHandles[3]).Num(), 1);
+
+		// Convert A to dynamic sleeping
+		Test.Evolution.SetParticleObjectState(Test.ParticleHandles[0], EObjectStateType::Sleeping);
+
+		Test.Advance();
+
+		// All particles in one awake island (B was awake so it would wake the island, even though A was set to sleeping)
+		EXPECT_EQ(Test.IslandManager->NumIslands(), 1);
+		EXPECT_FALSE(Test.IslandManager->GetSolverIsland(0)->IsSleeping());
+		EXPECT_EQ(Test.IslandManager->FindParticleIslands(Test.ParticleHandles[0]).Num(), 1);
+		EXPECT_EQ(Test.IslandManager->FindParticleIslands(Test.ParticleHandles[1]).Num(), 1);
+		EXPECT_EQ(Test.IslandManager->FindParticleIslands(Test.ParticleHandles[2]).Num(), 1);
+		EXPECT_EQ(Test.IslandManager->FindParticleIslands(Test.ParticleHandles[3]).Num(), 1);
+
+		// All particles awake
+		EXPECT_FALSE(FConstGenericParticleHandle(Test.ParticleHandles[0])->Sleeping());
+		EXPECT_FALSE(FConstGenericParticleHandle(Test.ParticleHandles[1])->Sleeping());
+		EXPECT_FALSE(FConstGenericParticleHandle(Test.ParticleHandles[2])->Sleeping());
+		EXPECT_FALSE(FConstGenericParticleHandle(Test.ParticleHandles[3])->Sleeping());
+	}
+
+	// Start with an island containing 4 sleeping particles connected in a chain. Then wake the island
+	// using the WakeIsland method.
+	//
+	//		(d=dynamic, s=sleeping, k=kinematic)
+	//		As - Bs - Cs - Ds  =>	Ad - Bd - Cd - Dd
+	//
+	GTEST_TEST(GraphEvolutionTests, TestConstraintGraph_WakeIsland)
+	{
+		FGraphEvolutionTest Test(4);
+		Test.MakeChain();
+
+		Test.AdvanceUntilSleeping(50);
+
+		// All particles asleep
+		EXPECT_TRUE(FConstGenericParticleHandle(Test.ParticleHandles[0])->Sleeping());
+		EXPECT_TRUE(FConstGenericParticleHandle(Test.ParticleHandles[1])->Sleeping());
+		EXPECT_TRUE(FConstGenericParticleHandle(Test.ParticleHandles[2])->Sleeping());
+		EXPECT_TRUE(FConstGenericParticleHandle(Test.ParticleHandles[3])->Sleeping());
+		EXPECT_EQ(Test.IslandManager->NumIslands(), 1);
+
+		// Flag the island for wake up
+		if (Test.IslandManager->NumIslands() == 1)
+		{
+			EXPECT_TRUE(Test.IslandManager->GetSolverIsland(0)->IsSleeping());
+			Test.IslandManager->WakeIsland(Test.Particles, 0);
+		}
+
+		Test.Advance();
+
+		// Island and all particles should now be awake
+		EXPECT_EQ(Test.IslandManager->NumIslands(), 1);
+		if (Test.IslandManager->NumIslands() == 1)
+		{
+			EXPECT_FALSE(Test.IslandManager->GetSolverIsland(0)->IsSleeping());
+		}
+		EXPECT_FALSE(FConstGenericParticleHandle(Test.ParticleHandles[0])->Sleeping());
+		EXPECT_FALSE(FConstGenericParticleHandle(Test.ParticleHandles[1])->Sleeping());
+		EXPECT_FALSE(FConstGenericParticleHandle(Test.ParticleHandles[2])->Sleeping());
+		EXPECT_FALSE(FConstGenericParticleHandle(Test.ParticleHandles[3])->Sleeping());
+	}
+
+	// Start with an island containing 4 awake particles connected in a chain. Then sleep the island
+	// using the SleepIsland method.
+	//
+	//		(d=dynamic, s=sleeping, k=kinematic)
+	//		Ad - Bd - Cd - Dd  =>	As - Bs - Cs - Ds
+	//
+	// @todo(chaos): this test fails. Explicit call to SleepIsland do not work atm
+	GTEST_TEST(GraphEvolutionTests, DISABLED_TestConstraintGraph_SleepIsland)
+	{
+		FGraphEvolutionTest Test(4);
+		Test.MakeChain();
+
+		Test.Advance();
+
+		// All particles awake
+		EXPECT_FALSE(FConstGenericParticleHandle(Test.ParticleHandles[0])->Sleeping());
+		EXPECT_FALSE(FConstGenericParticleHandle(Test.ParticleHandles[1])->Sleeping());
+		EXPECT_FALSE(FConstGenericParticleHandle(Test.ParticleHandles[2])->Sleeping());
+		EXPECT_FALSE(FConstGenericParticleHandle(Test.ParticleHandles[3])->Sleeping());
+		EXPECT_EQ(Test.IslandManager->NumIslands(), 1);
+
+		// Flag the island for sleep
+		if (Test.IslandManager->NumIslands() == 1)
+		{
+			EXPECT_FALSE(Test.IslandManager->GetSolverIsland(0)->IsSleeping());
+			Test.IslandManager->SleepIsland(Test.Particles, 0);
+		}
+
+		Test.Advance();
+
+		// Island and all particles should now be asleep
+		EXPECT_EQ(Test.IslandManager->NumIslands(), 1);
+		if (Test.IslandManager->NumIslands() == 1)
+		{
+			EXPECT_TRUE(Test.IslandManager->GetSolverIsland(0)->IsSleeping());
+		}
+		EXPECT_TRUE(FConstGenericParticleHandle(Test.ParticleHandles[0])->Sleeping());
+		EXPECT_TRUE(FConstGenericParticleHandle(Test.ParticleHandles[1])->Sleeping());
+		EXPECT_TRUE(FConstGenericParticleHandle(Test.ParticleHandles[2])->Sleeping());
+		EXPECT_TRUE(FConstGenericParticleHandle(Test.ParticleHandles[3])->Sleeping());
+	}
+
+}
diff --git a/Engine/Source/Programs/HeadlessChaos/Private/HeadlessChaosTestIsland.cpp b/Engine/Source/Programs/HeadlessChaos/Private/HeadlessChaosTestIsland.cpp
--- a/Engine/Source/Programs/HeadlessChaos/Private/HeadlessChaosTestIsland.cpp
+++ b/Engine/Source/Programs/HeadlessChaos/Private/HeadlessChaosTestIsland.cpp
@@ -15,7 +15,7 @@ namespace ChaosTest
 	TEST(IslandTests,TestMergeIslands)
 	{
 		static constexpr int32 NumParticles = 21;
-		using IslandGraphType = Chaos::FIslandGraph<int32, int32, int32>;
+		using IslandGraphType = Chaos::FIslandGraph<int32, int32, int32, TNullIslandGraphOwner<int32, int32>>;
 
 		IslandGraphType IslandGraph;
 
diff --git a/Engine/Source/Runtime/Engine/Private/PhysicsEngine/ImmediatePhysics/ImmediatePhysicsChaos/ImmediatePhysicsSimulation_Chaos.cpp b/Engine/Source/Runtime/Engine/Private/PhysicsEngine/ImmediatePhysics/ImmediatePhysicsChaos/ImmediatePhysicsSimulation_Chaos.cpp
--- a/Engine/Source/Runtime/Engine/Private/PhysicsEngine/ImmediatePhysics/ImmediatePhysicsChaos/ImmediatePhysicsSimulation_Chaos.cpp
+++ b/Engine/Source/Runtime/Engine/Private/PhysicsEngine/ImmediatePhysics/ImmediatePhysicsChaos/ImmediatePhysicsSimulation_Chaos.cpp
@@ -301,8 +301,6 @@ namespace ImmediatePhysics_Chaos
 	struct FSimulation::FImplementation
 	{
 	public:
-		using FParticlePair = Chaos::TVec2<Chaos::FGeometryParticleHandle*>;
-
 		FImplementation()
 			: Particles(UniqueIndices)
 			, Joints()
@@ -335,7 +333,7 @@ namespace ImmediatePhysics_Chaos
 		}
 
 		// @todo(ccaulfield): Look into these...
-		TArray<FParticlePair> ActivePotentiallyCollidingPairs;
+		TArray<Chaos::FParticlePair> ActivePotentiallyCollidingPairs;
 		TArray<FActorHandle*> StaticParticles;
 		Chaos::TArrayCollectionArray<bool> CollidedParticles;
 		Chaos::TArrayCollectionArray<Chaos::TSerializablePtr<Chaos::FChaosPhysicsMaterial>> ParticleMaterials;
@@ -366,7 +364,7 @@ namespace ImmediatePhysics_Chaos
 		// @todo(ccaulfield): Optimize
 		TMap<const Chaos::FGeometryParticleHandle*, TSet<const Chaos::FGeometryParticleHandle*>> IgnoreCollisionParticlePairTable;
 
-		TArray<FParticlePair> PotentiallyCollidingPairs;
+		TArray<Chaos::FParticlePair> PotentiallyCollidingPairs;
 
 		Chaos::FSimulationSpace SimulationSpace;
 
@@ -596,7 +594,7 @@ namespace ImmediatePhysics_Chaos
 			FGeometryParticleHandle* Particle1 = OtherActorHandle->GetParticle();
 			if ((OtherActorHandle != ActorHandle) && OtherActorHandle->IsSimulated())
 			{
-				Implementation->PotentiallyCollidingPairs.Emplace(typename FImplementation::FParticlePair(Particle0, Particle1));
+				Implementation->PotentiallyCollidingPairs.Emplace(FParticlePair(Particle0, Particle1));
 			}
 		}
 		Implementation->bActorsDirty = true;
@@ -604,7 +602,8 @@ namespace ImmediatePhysics_Chaos
 
 	void FSimulation::RemoveFromCollidingPairs(FActorHandle* ActorHandle)
 	{
-		for (typename FImplementation::FParticlePair& ParticlePair : Implementation->PotentiallyCollidingPairs)
+		using namespace Chaos;
+		for (FParticlePair& ParticlePair : Implementation->PotentiallyCollidingPairs)
 		{
 			if ((ParticlePair[0] == ActorHandle->GetParticle()) || (ParticlePair[1] == ActorHandle->GetParticle()))
 			{
@@ -677,7 +676,7 @@ namespace ImmediatePhysics_Chaos
 					bool bIgnoreActorHandlePair = (Particle0IgnoreSet != nullptr) && Particle0IgnoreSet->Contains(Particle1);
 					if (!bIgnoreActorHandlePair)
 					{
-						Implementation->PotentiallyCollidingPairs.Emplace(typename FImplementation::FParticlePair(Particle0, Particle1));
+						Implementation->PotentiallyCollidingPairs.Emplace(FParticlePair(Particle0, Particle1));
 					}
 				}
 			}
@@ -707,7 +706,7 @@ namespace ImmediatePhysics_Chaos
 		using namespace Chaos;
 
 		Implementation->ActivePotentiallyCollidingPairs.Reset();
-		for (const typename FImplementation::FParticlePair& ParticlePair : Implementation->PotentiallyCollidingPairs)
+		for (const FParticlePair& ParticlePair : Implementation->PotentiallyCollidingPairs)
 		{
 			if ((ParticlePair[0] != nullptr) && (ParticlePair[1] != nullptr))
 			{
diff --git a/Engine/Source/Runtime/Experimental/Chaos/Private/Chaos/ChaosDebugDraw.cpp b/Engine/Source/Runtime/Experimental/Chaos/Private/Chaos/ChaosDebugDraw.cpp
--- a/Engine/Source/Runtime/Experimental/Chaos/Private/Chaos/ChaosDebugDraw.cpp
+++ b/Engine/Source/Runtime/Experimental/Chaos/Private/Chaos/ChaosDebugDraw.cpp
@@ -992,6 +992,16 @@ namespace Chaos
 							FDebugDrawQueue::GetInstance().DrawDebugBox(WorldPlaneLocation, FVec3(BoxScale, BoxScale, FReal(0.01)), FRotation3(FRotationMatrix::MakeFromZ(WorldPlaneNormal)), FColor::Orange, false, Duration, Settings.DrawPriority, 0.5f * Settings.LineThickness);
 						}
 					}
+
+					// Sleeping
+					if (Settings.ContactInfoWidth > 0)
+					{
+						if (Contact.IsSleeping())
+						{
+							const FReal BoxScale = Settings.DrawScale * Settings.ContactInfoWidth * 1.1f;
+							FDebugDrawQueue::GetInstance().DrawDebugBox(WorldPlaneLocation, FVec3(BoxScale, BoxScale, FReal(0.01)), FRotation3(FRotationMatrix::MakeFromZ(WorldPlaneNormal)), FColor::Black, false, Duration, Settings.DrawPriority, 0.5f * Settings.LineThickness);
+						}
+					}
 				}
 
 				// AccumulatedImpulse
@@ -1324,43 +1334,49 @@ namespace Chaos
 				}
 			};
 
-			for (int32 IslandIndex = 0; IslandIndex < Graph.NumIslands(); ++IslandIndex)
+			if (bChaosDebugDebugDrawIslands)
 			{
-				const FPBDIslandSolver* Island = Graph.GetSolverIsland(IslandIndex);
+				TArray<FAABB3> IslandBounds;
+				IslandBounds.SetNum(Graph.NumIslands());
 
-				FAABB3 IslandAABB = FAABB3::EmptyAABB();
-				const TArray<FGeometryParticleHandle*> Particles = Island->GetParticles();
-				for (const FGeometryParticleHandle* GeoParticle : Particles)
+				const typename FPBDIslandManager::GraphType* IslandGraph = Graph.GetIslandGraph();
+				for (const auto& GraphNode : IslandGraph->GraphNodes)
 				{
-					FConstGenericParticleHandle Particle = GeoParticle;
+					FConstGenericParticleHandle Particle = GraphNode.NodeItem;
 					if (Particle->IsDynamic() && Particle->HasBounds())
 					{
-						IslandAABB.GrowToInclude(Particle->BoundingBox());
+						for (int32 ParticleIsland : Graph.FindParticleIslands(Particle->Handle()))
+						{
+							IslandBounds[ParticleIsland].GrowToInclude(Particle->BoundingBox());
+						}
 					}
 				}
-
-				if (bChaosDebugDebugDrawIslands)
+			
+				for (int32 IslandIndex = 0; IslandIndex < IslandBounds.Num(); ++IslandIndex)
 				{
-					const FColor IslandColor = GetIslandColor(IslandIndex, !Island->IsSleeping());
+					FAABB3 IslandAABB = IslandBounds[IslandIndex];
+
+					const FColor IslandColor = GetIslandColor(IslandIndex, !Graph.GetSolverIsland(IslandIndex)->IsSleeping());
 					const FAABB3 Bounds = IslandAABB.TransformedAABB(SpaceTransform);
 					FDebugDrawQueue::GetInstance().DrawDebugBox(Bounds.Center(), 0.5f * Bounds.Extents(), SpaceTransform.GetRotation(), IslandColor, false, UE_KINDA_SMALL_NUMBER, Settings.DrawPriority, 3.0f * Settings.LineThickness);
 				}
+			}
 
-				if (bChaosDebugDebugDrawContactGraph || bChaosDebugDebugDrawContactGraphUnused || bChaosDebugDebugDrawContactGraphUsed)
+			if (bChaosDebugDebugDrawContactGraph || bChaosDebugDebugDrawContactGraphUnused || bChaosDebugDebugDrawContactGraphUsed)
+			{
+				const typename FPBDIslandManager::GraphType* IslandGraph = Graph.GetIslandGraph();
+				for (const auto& GraphEdge : IslandGraph->GraphEdges)
 				{
-					const auto& Constraints = Island->GetConstraints();
-					for (int32 ConstraintIndex = 0; ConstraintIndex < Constraints.Num(); ++ConstraintIndex)
+					const FConstraintHandle* Constraint = GraphEdge.EdgeItem;
+					if (const FPBDCollisionConstraintHandle* Collision = Constraint->As<FPBDCollisionConstraintHandle>())
 					{
-						const FConstraintHandle* Constraint = Constraints[ConstraintIndex];
-
-						if (const FPBDCollisionConstraintHandle* Collision = Constraint->As<FPBDCollisionConstraintHandle>())
-						{
-							// @chaos(todo): store level and color in the constraint? Only for debug draw...
-							const int32 LevelIndex = INDEX_NONE;
-							const int32 ColorIndex = INDEX_NONE;
-							const bool bIsUsed = !Collision->GetConstraint()->AccumulatedImpulse.IsNearlyZero();
-							DrawGraphCollision(SpaceTransform, Collision->GetConstraint(), IslandIndex, LevelIndex, ColorIndex, ConstraintIndex, bIsUsed, Settings);
-						}
+						const int32 IslandIndex = GraphEdge.IslandIndex;
+						// @chaos(todo): would be nice to have this data retained for debug draw...
+						const int32 LevelIndex = INDEX_NONE;
+						const int32 ColorIndex = INDEX_NONE;
+						const int32 OrderIndex = INDEX_NONE;
+						const bool bIsUsed = !Collision->GetConstraint()->AccumulatedImpulse.IsNearlyZero();
+						DrawGraphCollision(SpaceTransform, Collision->GetConstraint(), IslandIndex, LevelIndex, ColorIndex, OrderIndex, bIsUsed, Settings);
 					}
 				}
 			}
diff --git a/Engine/Source/Runtime/Experimental/Chaos/Private/Chaos/Collision/PBDCollisionConstraint.cpp b/Engine/Source/Runtime/Experimental/Chaos/Private/Chaos/Collision/PBDCollisionConstraint.cpp
--- a/Engine/Source/Runtime/Experimental/Chaos/Private/Chaos/Collision/PBDCollisionConstraint.cpp
+++ b/Engine/Source/Runtime/Experimental/Chaos/Private/Chaos/Collision/PBDCollisionConstraint.cpp
@@ -3,9 +3,9 @@
 #include "Chaos/Collision/CollisionConstraintAllocator.h"
 #include "Chaos/Collision/PBDCollisionConstraintHandle.h"
 #include "Chaos/Collision/SolverCollisionContainer.h"
-#include "Chaos/PBDCollisionConstraints.h"
 #include "Chaos/Evolution/SolverBody.h"
 #include "Chaos/Evolution/SolverBodyContainer.h"
+#include "Chaos/Island/IslandManager.h"
 #include "Chaos/ParticleHandle.h"
 #include "Chaos/Particle/ParticleUtilities.h"
 #include "Chaos/PBDCollisionConstraints.h"
@@ -176,6 +176,9 @@ namespace Chaos
 		// @todo(chaos): this should probably be handled by the copy constructor
 		Constraint.GetContainerCookie().ClearContainerData();
 
+		// We are not in the constraint graph, even if Source was
+		Constraint.SetConstraintGraphIndex(INDEX_NONE);
+
 		return Constraint;
 	}
 
@@ -253,6 +256,14 @@ namespace Chaos
 	{
 	}
 
+	FPBDCollisionConstraint::~FPBDCollisionConstraint()
+	{
+#if !UE_BUILD_TEST && !UE_BUILD_SHIPPING
+		// Make sure we have been dropped by the graph
+		ensure(ConstraintGraphIndex() == INDEX_NONE);
+#endif
+	}
+
 	void FPBDCollisionConstraint::Setup(
 		const ECollisionCCDType InCCDType,
 		const EContactShapesType InShapesType,
@@ -380,6 +391,15 @@ namespace Chaos
 		CCDTargetPenetration = FMath::Min(MinBounds * CVars::CCDAllowedDepthBoundsScale, CCDEnablePenetration);
 	}
 
+	bool FPBDCollisionConstraint::IsSleeping() const
+	{
+		if (ContainerCookie.MidPhase != nullptr)
+		{
+			return ContainerCookie.MidPhase->IsSleeping();
+		}
+		return false;
+	}
+
 	void FPBDCollisionConstraint::SetIsSleeping(const bool bInIsSleeping)
 	{
 		// This actually sets the sleeping state on all constraints between the same particle pair so calling this with multiple
diff --git a/Engine/Source/Runtime/Experimental/Chaos/Private/Chaos/Collision/PBDCollisionConstraintHandle.cpp b/Engine/Source/Runtime/Experimental/Chaos/Private/Chaos/Collision/PBDCollisionConstraintHandle.cpp
--- a/Engine/Source/Runtime/Experimental/Chaos/Private/Chaos/Collision/PBDCollisionConstraintHandle.cpp
+++ b/Engine/Source/Runtime/Experimental/Chaos/Private/Chaos/Collision/PBDCollisionConstraintHandle.cpp
@@ -7,78 +7,4 @@
 
 namespace Chaos
 {
-	const FPBDCollisionConstraints* FPBDCollisionConstraintHandle::ConcreteContainer() const
-	{
-		return static_cast<FPBDCollisionConstraints*>(ConstraintContainer);
-	}
-
-	FPBDCollisionConstraints* FPBDCollisionConstraintHandle::ConcreteContainer()
-	{
-		return static_cast<FPBDCollisionConstraints*>(ConstraintContainer);
-	}
-
-	const FPBDCollisionConstraint& FPBDCollisionConstraintHandle::GetContact() const
-	{
-		return *GetConstraint();
-	}
-
-	FPBDCollisionConstraint& FPBDCollisionConstraintHandle::GetContact()
-	{
-		return *GetConstraint();
-	}
-
-	ECollisionCCDType FPBDCollisionConstraintHandle::GetCCDType() const
-	{
-		return GetContact().GetCCDType();
-	}
-
-	void FPBDCollisionConstraintHandle::SetEnabled(bool InEnabled)
-	{
-		GetContact().SetDisabled(!InEnabled);
-	}
-
-	bool FPBDCollisionConstraintHandle::IsEnabled() const
-	{
-		return !GetContact().GetDisabled();
-	}
-
-	bool FPBDCollisionConstraintHandle::IsProbe() const
-	{
-		return GetContact().GetIsProbe();
-	}
-
-	FVec3 FPBDCollisionConstraintHandle::GetAccumulatedImpulse() const
-	{
-		return GetContact().AccumulatedImpulse;
-	}
-
-	TVector<const TGeometryParticleHandle<FReal, 3>*, 2> FPBDCollisionConstraintHandle::GetConstrainedParticles() const
-	{
-		return { GetContact().GetParticle0(), GetContact().GetParticle1() };
-	}
-
-	TVector<TGeometryParticleHandle<FReal, 3>*, 2> FPBDCollisionConstraintHandle::GetConstrainedParticles()
-	{
-		return { GetContact().GetParticle0(), GetContact().GetParticle1() };
-	}
-
-	void FPBDCollisionConstraintHandle::PreGatherInput(const FReal Dt, FPBDIslandSolverData& SolverData)
-	{
-		ConcreteContainer()->PreGatherInput(Dt, GetContact(), SolverData);
-	}
-
-	void FPBDCollisionConstraintHandle::GatherInput(FReal Dt, const int32 Particle0Level, const int32 Particle1Level, FPBDIslandSolverData& SolverData)
-	{
-		ConcreteContainer()->GatherInput(Dt, GetContact(), Particle0Level, Particle1Level, SolverData);
-	}
-
-	FSolverBody* FPBDCollisionConstraintHandle::GetSolverBody0()
-	{
-		return GetContact().GetSolverBody0();
-	}
-
-	FSolverBody* FPBDCollisionConstraintHandle::GetSolverBody1()
-	{
-		return GetContact().GetSolverBody1();
-	}
 }
\ No newline at end of file
diff --git a/Engine/Source/Runtime/Experimental/Chaos/Private/Chaos/Collision/ParticlePairMidPhase.cpp b/Engine/Source/Runtime/Experimental/Chaos/Private/Chaos/Collision/ParticlePairMidPhase.cpp
--- a/Engine/Source/Runtime/Experimental/Chaos/Private/Chaos/Collision/ParticlePairMidPhase.cpp
+++ b/Engine/Source/Runtime/Experimental/Chaos/Private/Chaos/Collision/ParticlePairMidPhase.cpp
@@ -29,6 +29,21 @@ namespace Chaos
 	}
 	using namespace CVars;
 
+#if CHAOS_CONSTRAINTHANDLE_DEBUG_ENABLED
+	void CheckConstraintContainsDynamic(const FPBDCollisionConstraint* Constraint)
+	{
+		if (Constraint != nullptr)
+		{
+			const bool bIsDynamic0 = FConstGenericParticleHandle(Constraint->GetParticle0())->IsDynamic();
+			const bool bIsDynamic1 = FConstGenericParticleHandle(Constraint->GetParticle1())->IsDynamic();
+			if (!bIsDynamic0 && !bIsDynamic1)
+			{
+				ensure(false);
+			}
+		}
+	}
+#endif
+
 	inline bool ImplicitOverlapOBBToAABB(
 		const FImplicitObject* Implicit0,
 		const FImplicitObject* Implicit1,
@@ -436,8 +451,13 @@ namespace Chaos
 	{
 		if ((Constraint != nullptr) && IsUsedSince(SleepEpoch))
 		{
+#if CHAOS_CONSTRAINTHANDLE_DEBUG_ENABLED
+			// If a particle changed to kinematic, all constraints without a dynamic should have been removed
+			CheckConstraintContainsDynamic(Constraint.Get());
+#endif
+
 			// We need to refresh the epoch so that the constraint state will be used as the previous
-			// state iof the pair is still colliding in the next tick
+			// state if the pair is still colliding in the next tick
 			const int32 CurrentEpoch = MidPhase.GetCollisionAllocator().GetCurrentEpoch();
 			Constraint->GetContainerCookie().LastUsedEpoch = MidPhase.GetCollisionAllocator().GetCurrentEpoch();
 			LastUsedEpoch = CurrentEpoch;
@@ -964,6 +984,7 @@ namespace Chaos
 				if (((Shape0 == ShapePair.GetShape0()) && (Shape1 == ShapePair.GetShape1())) || ((Shape0 == ShapePair.GetShape1()) && (Shape1 == ShapePair.GetShape0())))
 				{
 					ShapePair.SetCollision(Constraint);
+					break;
 				}
 			}
 		}
@@ -972,6 +993,8 @@ namespace Chaos
 			// @todo(chaos): implement cluster Resim restore
 			ensure(false);
 		}
+
+		LastUsedEpoch = CollisionAllocator->GetCurrentEpoch();
 	}
 
 	void FParticlePairMidPhase::SetIsSleeping(const bool bInIsSleeping)
@@ -1015,6 +1038,22 @@ namespace Chaos
 		}
 	}
 
+	bool FParticlePairMidPhase::IsInConstraintGraph() const
+	{
+		// @todo(chaos): optimize
+		bool bInGraph = false;
+		VisitConstCollisions([&bInGraph](const FPBDCollisionConstraint& Constraint)
+		{
+			if (Constraint.ConstraintGraphIndex() != INDEX_NONE)
+			{
+				bInGraph = true;
+				return ECollisionVisitorResult::Stop;
+			}
+			return ECollisionVisitorResult::Continue;
+		}, false);
+		return bInGraph;
+	}
+
 }
 
 
diff --git a/Engine/Source/Runtime/Experimental/Chaos/Private/Chaos/ConstraintHandle.cpp b/Engine/Source/Runtime/Experimental/Chaos/Private/Chaos/ConstraintHandle.cpp
new file mode 100644
--- /dev/null
+++ b/Engine/Source/Runtime/Experimental/Chaos/Private/Chaos/ConstraintHandle.cpp
@@ -0,0 +1,26 @@
+// Copyright Epic Games, Inc. All Rights Reserved.
+#include "Chaos/ConstraintHandle.h"
+
+#if CHAOS_CONSTRAINTHANDLE_DEBUG_ENABLED
+// For debugging only
+#include "Chaos/Collision/PBDCollisionConstraintHandle.h"
+#endif
+
+namespace Chaos
+{
+#if CHAOS_CONSTRAINTHANDLE_DEBUG_ENABLED
+	void FConstraintHandleHolder::InitDebugData()
+	{
+		ConstraintType = nullptr;
+		Particles[0] = nullptr;
+		Particles[1] = nullptr;
+		if (Handle != nullptr)
+		{
+			ConstraintType = &Handle->GetType();
+
+			Particles[0] = Handle->GetConstrainedParticles()[0];
+			Particles[1] = Handle->GetConstrainedParticles()[1];
+		}
+	}
+#endif
+}
\ No newline at end of file
diff --git a/Engine/Source/Runtime/Experimental/Chaos/Private/Chaos/Evolution/PBDMinEvolution.cpp b/Engine/Source/Runtime/Experimental/Chaos/Private/Chaos/Evolution/PBDMinEvolution.cpp
--- a/Engine/Source/Runtime/Experimental/Chaos/Private/Chaos/Evolution/PBDMinEvolution.cpp
+++ b/Engine/Source/Runtime/Experimental/Chaos/Private/Chaos/Evolution/PBDMinEvolution.cpp
@@ -404,6 +404,7 @@ namespace Chaos
 		}
 
 		CollisionDetector.DetectCollisions(Dt, nullptr);
+		CollisionDetector.GetCollisionContainer().GetConstraintAllocator().PruneExpiredItems();
 		CollisionDetector.GetCollisionContainer().GetConstraintAllocator().SortConstraintsHandles();
 	}
 
diff --git a/Engine/Source/Runtime/Experimental/Chaos/Private/Chaos/Island/IslandGraph.cpp b/Engine/Source/Runtime/Experimental/Chaos/Private/Chaos/Island/IslandGraph.cpp
--- a/Engine/Source/Runtime/Experimental/Chaos/Private/Chaos/Island/IslandGraph.cpp
+++ b/Engine/Source/Runtime/Experimental/Chaos/Private/Chaos/Island/IslandGraph.cpp
@@ -2,36 +2,51 @@
 
 #include "Chaos/Island/IslandGraph.h"
 #include "Chaos/Island/SolverIsland.h"
+#include "Chaos/Island/IslandManager.h"
 #include "Chaos/ParticleHandle.h"
 #include "Chaos/ConstraintHandle.h"
 
 #include "ChaosStats.h"
 #include "ChaosLog.h"
 
+//PRAGMA_DISABLE_OPTIMIZATION
+
 namespace Chaos
 {
 
 DECLARE_CYCLE_STAT(TEXT("MergeIslandGraph"), STAT_MergeIslandGraph, STATGROUP_Chaos);
 DECLARE_CYCLE_STAT(TEXT("SplitIslandGraph"), STAT_SplitIslandGraph, STATGROUP_Chaos);
 
-template<typename NodeType, typename EdgeType, typename IslandType>
-void FIslandGraph<NodeType, EdgeType, IslandType>::ResetNodes()
+template<typename NodeType, typename EdgeType, typename IslandType, typename OwnerType>
+void FIslandGraph<NodeType, EdgeType, IslandType, OwnerType>::ResetNodes()
 {
+	// If we remove all nodes, we must also remove edges
+	ResetEdges();
+
+	// Notify all nodes they were removed
+	if (Owner != nullptr)
+	{
+		for (auto& ItemNode : GraphNodes)
+		{
+			Owner->GraphNodeRemoved(ItemNode.NodeItem);
+		}
+	}
+
 	GraphNodes.Reset();
 	ItemNodes.Reset();
 	GraphIslands.Reset();
 }
 
-template<typename NodeType, typename EdgeType, typename IslandType>
-void FIslandGraph<NodeType, EdgeType, IslandType>::ReserveNodes(const int32 NumNodes)
+template<typename NodeType, typename EdgeType, typename IslandType, typename OwnerType>
+void FIslandGraph<NodeType, EdgeType, IslandType, OwnerType>::ReserveNodes(const int32 NumNodes)
 {
 	GraphNodes.Reserve(NumNodes);
 	ItemNodes.Reserve(NumNodes);
 	GraphIslands.Reserve(NumNodes);
 }
 
-template<typename NodeType, typename EdgeType, typename IslandType>
-void FIslandGraph<NodeType, EdgeType, IslandType>::ParentIslands(const int32 FirstIsland, const int32 SecondIsland, const bool bIsEdgeMoving)
+template<typename NodeType, typename EdgeType, typename IslandType, typename OwnerType>
+void FIslandGraph<NodeType, EdgeType, IslandType, OwnerType>::ParentIslands(const int32 FirstIsland, const int32 SecondIsland, const bool bIsEdgeMoving)
 {
 	if (GraphIslands.IsValidIndex(FirstIsland) && GraphIslands.IsValidIndex(SecondIsland) && (FirstIsland != SecondIsland))
 	{
@@ -47,64 +62,70 @@ void FIslandGraph<NodeType, EdgeType, IslandType>::ParentIslands(const int32 Fir
 	}
 }
 	
-template<typename NodeType, typename EdgeType, typename IslandType>
-void FIslandGraph<NodeType, EdgeType, IslandType>::UpdateNode(const NodeType& NodeItem, const bool bValidNode, const int32 IslandIndex, const bool bStationaryNode, const int32 NodeIndex)
+template<typename NodeType, typename EdgeType, typename IslandType, typename OwnerType>
+void FIslandGraph<NodeType, EdgeType, IslandType, OwnerType>::UpdateNode(const NodeType& NodeItem, const bool bValidNode, const int32 IslandIndex, const bool bStationaryNode, const int32 NodeIndex)
 {
 	if (GraphNodes.IsValidIndex(NodeIndex))
 	{
 		FGraphNode& GraphNode = GraphNodes[NodeIndex];
-		
-		// In case the item is changing its state to be valid
-		// we merge all the connected islands 
-		if (bValidNode && !GraphNode.bValidNode)
+#if CHAOS_CONSTRAINTHANDLE_DEBUG_ENABLED
+		ensure(GraphNode.NodeItem == NodeItem);
+#endif
+
+		// Update node state (must come first because IsEdgeMoving relies in state on Stationary being up to date)
+		const bool bWasValidNode = GraphNode.bValidNode;
+		GraphNode.bValidNode = bValidNode;
+		GraphNode.bStationaryNode = bStationaryNode;
+
+		// In case the item is changing its state to be valid (Kinematic->Dynamic/Sleeping)
+		// we merge all the connected islands (this will wake the islands)
+		if (bValidNode && !bWasValidNode)
 		{
+			// @todo(chaos): we could just use the GraphNode.NodeIslands here if it were maintained internally
 			int32 MasterIsland = INDEX_NONE;
 			for (int32& EdgeIndex : GraphNode.NodeEdges)
 			{
 				FGraphEdge& GraphEdge = GraphEdges[EdgeIndex];
+
 				const bool bIsEdgeMoving = IsEdgeMoving(EdgeIndex);
-				
 				ParentIslands(MasterIsland, GraphEdge.IslandIndex, bIsEdgeMoving);
+			
 				MasterIsland = GraphEdge.IslandIndex;
-				if(bIsEdgeMoving)
-				{
-					GraphIslands[GraphEdge.IslandIndex].bIsPersistent = false;
-				}
 			}
+
+			// Put the valid node into one of the islands - they will be merged anyway so doesn't matter which
+			// If we did not have an island (no edges), one will be assigned later
+			GraphNode.IslandIndex = MasterIsland;
+			GraphNode.NodeIslands.Reset();
 		}
-		// If node is not valid and was valid before we check that all the node edges
-		// have at least one valid node. otherwise we remove the edge. 
-		if (!bValidNode && GraphNode.bValidNode)
+
+		// If we are changing to invalid (Dynamic/Sleeping->Kinematic) wake the island
+		if (!bValidNode && bWasValidNode)
 		{
-			for (int32 NodeEdge = GraphNode.NodeEdges.GetMaxIndex() - 1; NodeEdge >= 0; --NodeEdge)
+			// Wake the node's island if the kinematic is moving
+			if (GraphIslands.IsValidIndex(GraphNode.IslandIndex) && !bStationaryNode)
 			{
-				if (GraphNode.NodeEdges.IsValidIndex(NodeEdge))
-				{
-					const int32 EdgeIndex = GraphNode.NodeEdges[NodeEdge];
-
-					if (GraphEdges.IsValidIndex(EdgeIndex))
-					{
-						FGraphEdge& GraphEdge = GraphEdges[EdgeIndex];
+				GraphIslands[GraphNode.IslandIndex].bIsPersistent = false;
+			}
 
-						// Get the opposite node index for the given edge
-						const int32 OtherIndex = (NodeIndex == GraphEdge.FirstNode) ?
-							GraphEdge.SecondNode : GraphEdge.FirstNode;
+			// Invalid node island lists are built later (see PopulateIslands)
+			GraphNode.IslandIndex = INDEX_NONE;
+			GraphNode.NodeIslands.Reset();
+		}
 
-						if (!GraphNodes.IsValidIndex(OtherIndex) || (GraphNodes.IsValidIndex(OtherIndex) && !GraphNodes[OtherIndex].bValidNode))
-						{
-							RemoveEdge(EdgeIndex);
-						}
-					}
-				}
+		// If we change validity, we may have to change the validity of some of the edges
+		if (bValidNode != bWasValidNode)
+		{
+			for (int32 EdgeIndex : GraphNode.NodeEdges)
+			{
+				UpdateEdge(EdgeIndex);
 			}
 		}
-		GraphNode.bValidNode = bValidNode;
-		GraphNode.bStationaryNode = bStationaryNode;
 	}
 }
 
-template<typename NodeType, typename EdgeType, typename IslandType>
-int32 FIslandGraph<NodeType, EdgeType, IslandType>::AddNode(const NodeType& NodeItem, const bool bValidNode, const int32 IslandIndex, const bool bStationaryNode)
+template<typename NodeType, typename EdgeType, typename IslandType, typename OwnerType>
+int32 FIslandGraph<NodeType, EdgeType, IslandType, OwnerType>::AddNode(const NodeType& NodeItem, const bool bValidNode, const int32 IslandIndex, const bool bStationaryNode)
 {
 	int32 NodeIndex = INDEX_NONE;
 	int32* ItemIndex = ItemNodes.Find(NodeItem);
@@ -117,18 +138,27 @@ int32 FIslandGraph<NodeType, EdgeType, IslandType>::AddNode(const NodeType& Node
 	{
 		FGraphNode GraphNode;
 		GraphNode.NodeItem = NodeItem;
-		GraphNode.IslandIndex = IslandIndex;
+
+		if (bValidNode)
+		{
+			GraphNode.IslandIndex = IslandIndex;
+		}
 		
 		GraphNode.bValidNode = bValidNode;
 		GraphNode.bStationaryNode = bStationaryNode;
 
 		NodeIndex =  ItemNodes.Add(NodeItem, GraphNodes.Emplace(GraphNode));
+
+		if (Owner != nullptr)
+		{
+			Owner->GraphNodeAdded(NodeItem, NodeIndex);
+		}
 	}
 	return NodeIndex;
 }
 
-template<typename NodeType, typename EdgeType, typename IslandType>
-void FIslandGraph<NodeType, EdgeType, IslandType>::RemoveNode(const NodeType& NodeItem)
+template<typename NodeType, typename EdgeType, typename IslandType, typename OwnerType>
+void FIslandGraph<NodeType, EdgeType, IslandType, OwnerType>::RemoveNode(const NodeType& NodeItem)
 {
 	if (const int32* IndexPtr = ItemNodes.Find(NodeItem))
 	{
@@ -157,6 +187,11 @@ void FIslandGraph<NodeType, EdgeType, IslandType>::RemoveNode(const NodeType& No
 			}
 			ItemNodes.Remove(NodeItem);
 			GraphNodes.RemoveAt(NodeIndex);
+
+			if (Owner != nullptr)
+			{
+				Owner->GraphNodeRemoved(NodeItem);
+			}
 		}
 		else
 		{
@@ -165,22 +200,40 @@ void FIslandGraph<NodeType, EdgeType, IslandType>::RemoveNode(const NodeType& No
 	}
 }
 
-template<typename NodeType, typename EdgeType, typename IslandType>
-void FIslandGraph<NodeType, EdgeType, IslandType>::ResetEdges()
+template<typename NodeType, typename EdgeType, typename IslandType, typename OwnerType>
+void FIslandGraph<NodeType, EdgeType, IslandType, OwnerType>::ResetEdges()
 {
+	if (Owner != nullptr)
+	{
+		for (auto& GraphEdge : GraphEdges)
+		{
+			Owner->GraphEdgeRemoved(GraphEdge.EdgeItem);
+		}
+	}
+
+	for (auto& GraphIsland : GraphIslands)
+	{
+		GraphIsland.NumEdges = 0;
+	}
+
+	for (auto& GraphNode : GraphNodes)
+	{
+		GraphNode.NodeEdges.Reset();
+	}
+
 	GraphEdges.Reset();
 	ItemEdges.Reset();
 }
 
-template<typename NodeType, typename EdgeType, typename IslandType>
-void FIslandGraph<NodeType, EdgeType, IslandType>::ReserveEdges(const int32 NumEdges)
+template<typename NodeType, typename EdgeType, typename IslandType, typename OwnerType>
+void FIslandGraph<NodeType, EdgeType, IslandType, OwnerType>::ReserveEdges(const int32 NumEdges)
 {
 	GraphEdges.Reserve(NumEdges);
 	ItemEdges.Reserve(NumEdges);
 }
 
-template<typename NodeType, typename EdgeType, typename IslandType>
-void FIslandGraph<NodeType, EdgeType, IslandType>::AttachIslands(const int32 EdgeIndex)
+template<typename NodeType, typename EdgeType, typename IslandType, typename OwnerType>
+void FIslandGraph<NodeType, EdgeType, IslandType, OwnerType>::AttachIslands(const int32 EdgeIndex)
 {
 	if (GraphEdges.IsValidIndex(EdgeIndex))
 	{
@@ -200,11 +253,19 @@ void FIslandGraph<NodeType, EdgeType, IslandType>::AttachIslands(const int32 Edg
 			// if yes we set the invalid node island index 
 			// and the edge one to be the valid one 
 			// If none are valid we create a new island 
-			if (bFirstValidIsland && !bSecondValidIsland)
+			if (!FirstNode.bValidNode && !SecondNode.bValidNode)
+			{
+				// If we have two invalid nodes, we just remove the edge from its island
+				GraphEdge.IslandIndex = INDEX_NONE;
+			}
+			else if (bFirstValidIsland && !bSecondValidIsland)
 			{
 				GraphEdge.IslandIndex = FirstNode.IslandIndex;
-				SecondNode.IslandIndex = GraphEdge.IslandIndex;
-				if(SecondNode.bValidNode && bIsEdgeMoving)
+				if (SecondNode.bValidNode)
+				{
+					SecondNode.IslandIndex = GraphEdge.IslandIndex;
+				}
+				if (SecondNode.bValidNode && bIsEdgeMoving)
 				{
 					GraphIslands[GraphEdge.IslandIndex].bIsPersistent = false;
 				}
@@ -212,8 +273,11 @@ void FIslandGraph<NodeType, EdgeType, IslandType>::AttachIslands(const int32 Edg
 			else if (!bFirstValidIsland && bSecondValidIsland)
 			{
 				GraphEdge.IslandIndex = SecondNode.IslandIndex;
-				FirstNode.IslandIndex = GraphEdge.IslandIndex;
-				if(FirstNode.bValidNode && bIsEdgeMoving)
+				if (FirstNode.bValidNode)
+				{
+					FirstNode.IslandIndex = GraphEdge.IslandIndex;
+				}
+				if (FirstNode.bValidNode && bIsEdgeMoving)
 				{
 					GraphIslands[GraphEdge.IslandIndex].bIsPersistent = false;
 				}
@@ -224,8 +288,14 @@ void FIslandGraph<NodeType, EdgeType, IslandType>::AttachIslands(const int32 Edg
 				GraphEdge.IslandIndex = GraphIslands.Emplace(GraphIsland);
 
 				// We set both island indices to be the equal to the edge one
-				FirstNode.IslandIndex = GraphEdge.IslandIndex;
-				SecondNode.IslandIndex = GraphEdge.IslandIndex;
+				if (FirstNode.bValidNode)
+				{
+					FirstNode.IslandIndex = GraphEdge.IslandIndex;
+				}
+				if (SecondNode.bValidNode)
+				{
+					SecondNode.IslandIndex = GraphEdge.IslandIndex;
+				}
 			}
 			else
 			{
@@ -238,7 +308,7 @@ void FIslandGraph<NodeType, EdgeType, IslandType>::AttachIslands(const int32 Edg
 		}
 		else if (GraphNodes.IsValidIndex(GraphEdge.FirstNode) && !GraphNodes.IsValidIndex(GraphEdge.SecondNode) && GraphNodes[GraphEdge.FirstNode].bValidNode)
 		{
-			// If only the first node is valid, check if its island index is valid and if yes use it for the edge.
+			// If only the first node exists and is valid, check if its island index is valid and if yes use it for the edge.
 			// otherwise we create a new island
 			if(!GraphIslands.IsValidIndex(GraphNodes[GraphEdge.FirstNode].IslandIndex))
 			{
@@ -249,7 +319,7 @@ void FIslandGraph<NodeType, EdgeType, IslandType>::AttachIslands(const int32 Edg
 		}
 		else if (!GraphNodes.IsValidIndex(GraphEdge.FirstNode) && GraphNodes.IsValidIndex(GraphEdge.SecondNode) && GraphNodes[GraphEdge.SecondNode].bValidNode)
 		{
-			// If only the second node is valid, check if its island index is valid and if yes use it for the edge.
+			// If only the second node exists and is valid, check if its island index is valid and if yes use it for the edge.
 			// otherwise we create a new island
 			if(!GraphIslands.IsValidIndex(GraphNodes[GraphEdge.SecondNode].IslandIndex))
 			{
@@ -261,14 +331,14 @@ void FIslandGraph<NodeType, EdgeType, IslandType>::AttachIslands(const int32 Edg
 	}
 }
 
-template<typename NodeType, typename EdgeType, typename IslandType>
-int32 FIslandGraph<NodeType, EdgeType, IslandType>::AddEdge(const EdgeType& EdgeItem, const int32 ItermContainer, const int32 FirstNode, const int32 SecondNode)
+template<typename NodeType, typename EdgeType, typename IslandType, typename OwnerType>
+int32 FIslandGraph<NodeType, EdgeType, IslandType, OwnerType>::AddEdge(const EdgeType& EdgeItem, const int32 ItermContainer, const int32 FirstNode, const int32 SecondNode)
 {
 	int32 EdgeIndex = INDEX_NONE;
 	// We only add an edge if one of the 2 nodes are valid
 	if (GraphNodes.IsValidIndex(FirstNode) || GraphNodes.IsValidIndex(SecondNode))
 	{
-		// This test could be quite slow. will probably be better to use the graph index that is stored on the constraint handle
+		// @todo(chaos): This test could be quite slow. will probably be better to use the graph index that is stored on the constraint handle
 		int32* ItemIndex = ItemEdges.Find(EdgeItem);
 		if (ItemIndex)
 		{
@@ -284,14 +354,20 @@ int32 FIslandGraph<NodeType, EdgeType, IslandType>::AddEdge(const EdgeType& Edge
 			GraphEdge.SecondNode = SecondNode;
 			GraphEdge.IslandIndex = INDEX_NONE;
 			GraphEdge.ItemContainer = ItermContainer;
+			GraphEdge.bValidEdge = true;
 
 			EdgeIndex = GraphEdges.Emplace(GraphEdge);
 			ItemEdges.Add(EdgeItem, EdgeIndex);
 
-			GraphEdges[EdgeIndex].FirstEdge = GraphNodes.IsValidIndex(FirstNode) ? GraphNodes[FirstNode].NodeEdges.Add(EdgeIndex) : -1;
-			GraphEdges[EdgeIndex].SecondEdge = GraphNodes.IsValidIndex(SecondNode) ? GraphNodes[SecondNode].NodeEdges.Add(EdgeIndex) : -1;
+			GraphEdges[EdgeIndex].FirstEdge = GraphNodes.IsValidIndex(FirstNode) ? GraphNodes[FirstNode].NodeEdges.Add(EdgeIndex) : INDEX_NONE;
+			GraphEdges[EdgeIndex].SecondEdge = GraphNodes.IsValidIndex(SecondNode) ? GraphNodes[SecondNode].NodeEdges.Add(EdgeIndex) : INDEX_NONE;
 
 			AttachIslands(EdgeIndex);
+
+			if (Owner != nullptr)
+			{
+				Owner->GraphEdgeAdded(EdgeItem, EdgeIndex);
+			}
 		}
 	}	
 	else
@@ -301,13 +377,12 @@ int32 FIslandGraph<NodeType, EdgeType, IslandType>::AddEdge(const EdgeType& Edge
 	return EdgeIndex;
 }
 
-template<typename NodeType, typename EdgeType, typename IslandType>
-void FIslandGraph<NodeType, EdgeType, IslandType>::RemoveEdge(const int32 EdgeIndex)
+template<typename NodeType, typename EdgeType, typename IslandType, typename OwnerType>
+void FIslandGraph<NodeType, EdgeType, IslandType, OwnerType>::RemoveEdge(const int32 EdgeIndex)
 {
 	if (GraphEdges.IsValidIndex(EdgeIndex))
 	{
 		FGraphEdge& GraphEdge = GraphEdges[EdgeIndex];
-		// We first remove the edge from the nodes edges list
 		if(GraphNodes.IsValidIndex(GraphEdge.FirstNode) && GraphNodes[GraphEdge.FirstNode].NodeEdges.IsValidIndex(GraphEdge.FirstEdge))
 		{
 			GraphNodes[GraphEdge.FirstNode].NodeEdges.RemoveAt(GraphEdge.FirstEdge);
@@ -319,6 +394,11 @@ void FIslandGraph<NodeType, EdgeType, IslandType>::RemoveEdge(const int32 EdgeIn
 		// We then remove the edge from the item and graph edges
 		ItemEdges.Remove(GraphEdge.EdgeItem);
 		GraphEdges.RemoveAt(EdgeIndex);
+
+		if (Owner != nullptr)
+		{
+			Owner->GraphEdgeRemoved(GraphEdge.EdgeItem);
+		}
 	}
 	else
 	{
@@ -326,8 +406,20 @@ void FIslandGraph<NodeType, EdgeType, IslandType>::RemoveEdge(const int32 EdgeIn
 	}
 }
 
-template<typename NodeType, typename EdgeType, typename IslandType>
-void FIslandGraph<NodeType, EdgeType, IslandType>::MergeIslands(const int32 ParentIndex, const int32 ChildIndex)
+template<typename NodeType, typename EdgeType, typename IslandType, typename OwnerType>
+void FIslandGraph<NodeType, EdgeType, IslandType, OwnerType>::UpdateEdge(const int32 EdgeIndex)
+{
+	if (GraphEdges.IsValidIndex(EdgeIndex))
+	{
+		FGraphEdge& GraphEdge = GraphEdges[EdgeIndex];
+		const bool bFirstNodeValid = GraphNodes.IsValidIndex(GraphEdge.FirstNode) && GraphNodes[GraphEdge.FirstNode].bValidNode;
+		const bool bSecondNodeValid = GraphNodes.IsValidIndex(GraphEdge.SecondNode) && GraphNodes[GraphEdge.SecondNode].bValidNode;
+		GraphEdge.bValidEdge = bFirstNodeValid || bSecondNodeValid;
+	}
+}
+
+template<typename NodeType, typename EdgeType, typename IslandType, typename OwnerType>
+void FIslandGraph<NodeType, EdgeType, IslandType, OwnerType>::MergeIslands(const int32 ParentIndex, const int32 ChildIndex)
 {
 	int32 CurrentIndex;
 	TQueue<int32> ChildQueue;
@@ -352,8 +444,8 @@ void FIslandGraph<NodeType, EdgeType, IslandType>::MergeIslands(const int32 Pare
 	}
 }
 	
-template<typename NodeType, typename EdgeType, typename IslandType>
-void FIslandGraph<NodeType, EdgeType, IslandType>::MergeIslands()
+template<typename NodeType, typename EdgeType, typename IslandType, typename OwnerType>
+void FIslandGraph<NodeType, EdgeType, IslandType, OwnerType>::MergeIslands()
 {
 	GraphCounter = (GraphCounter + 1) % MaxCount;
 
@@ -396,8 +488,8 @@ void FIslandGraph<NodeType, EdgeType, IslandType>::MergeIslands()
 	}
 }
 
-template<typename NodeType, typename EdgeType, typename IslandType>
-void FIslandGraph<NodeType, EdgeType, IslandType>::InitSorting()
+template<typename NodeType, typename EdgeType, typename IslandType, typename OwnerType>
+void FIslandGraph<NodeType, EdgeType, IslandType, OwnerType>::InitSorting()
 {
 	// Reset nodes levels and colors
 	for(auto& GraphNode : GraphNodes)
@@ -419,8 +511,8 @@ void FIslandGraph<NodeType, EdgeType, IslandType>::InitSorting()
 	}
 }
 	
-template<typename NodeType, typename EdgeType, typename IslandType>
-void FIslandGraph<NodeType, EdgeType, IslandType>::UpdateLevels(const int32 NodeIndex, const int32 ContainerId)
+template<typename NodeType, typename EdgeType, typename IslandType, typename OwnerType>
+void FIslandGraph<NodeType, EdgeType, IslandType, OwnerType>::UpdateLevels(const int32 NodeIndex, const int32 ContainerId)
 {
 	if(GraphNodes.IsValidIndex(NodeIndex))
 	{
@@ -451,8 +543,8 @@ void FIslandGraph<NodeType, EdgeType, IslandType>::UpdateLevels(const int32 Node
 	}
 }
 
-template<typename NodeType, typename EdgeType, typename IslandType>
-void FIslandGraph<NodeType, EdgeType, IslandType>::ComputeLevels(const int32 ContainerId)
+template<typename NodeType, typename EdgeType, typename IslandType, typename OwnerType>
+void FIslandGraph<NodeType, EdgeType, IslandType, OwnerType>::ComputeLevels(const int32 ContainerId)
 {
 	// First enqueue all the static/kinematic nodes for levels 0
 	NodeQueue.Reset();
@@ -489,8 +581,8 @@ void FIslandGraph<NodeType, EdgeType, IslandType>::ComputeLevels(const int32 Con
 	}
 }
 	
-template<typename NodeType, typename EdgeType, typename IslandType>
-int32 FIslandGraph<NodeType, EdgeType, IslandType>::PickColor(const FGraphNode& GraphNode, const int32 OtherIndex)
+template<typename NodeType, typename EdgeType, typename IslandType, typename OwnerType>
+int32 FIslandGraph<NodeType, EdgeType, IslandType, OwnerType>::PickColor(const FGraphNode& GraphNode, const int32 OtherIndex)
 {
 	int32 ColorToUse = 0;
 	if (GraphNodes.IsValidIndex(OtherIndex) && GraphNodes[OtherIndex].bValidNode)
@@ -519,8 +611,8 @@ int32 FIslandGraph<NodeType, EdgeType, IslandType>::PickColor(const FGraphNode&
 	return ColorToUse;
 }
 
-template<typename NodeType, typename EdgeType, typename IslandType>
-void FIslandGraph<NodeType, EdgeType, IslandType>::UpdateColors(const int32 NodeIndex, const int32 ContainerId, const int32 MinEdges)
+template<typename NodeType, typename EdgeType, typename IslandType, typename OwnerType>
+void FIslandGraph<NodeType, EdgeType, IslandType, OwnerType>::UpdateColors(const int32 NodeIndex, const int32 ContainerId, const int32 MinEdges)
 {
 	if (GraphNodes.IsValidIndex(NodeIndex))
 	{
@@ -552,8 +644,8 @@ void FIslandGraph<NodeType, EdgeType, IslandType>::UpdateColors(const int32 Node
 	}
 }
 
-template<typename NodeType, typename EdgeType, typename IslandType>
-void FIslandGraph<NodeType, EdgeType, IslandType>::ComputeColors(const int32 ContainerId, const int32 MinEdges)
+template<typename NodeType, typename EdgeType, typename IslandType, typename OwnerType>
+void FIslandGraph<NodeType, EdgeType, IslandType, OwnerType>::ComputeColors(const int32 ContainerId, const int32 MinEdges)
 {
 	GraphCounter = (GraphCounter + 1) % MaxCount;
 
@@ -581,8 +673,8 @@ void FIslandGraph<NodeType, EdgeType, IslandType>::ComputeColors(const int32 Con
 	}
 }
 
-template<typename NodeType, typename EdgeType, typename IslandType>
-void FIslandGraph<NodeType, EdgeType, IslandType>::SplitIsland(const int32 RootIndex, const int32 IslandIndex)
+template<typename NodeType, typename EdgeType, typename IslandType, typename OwnerType>
+void FIslandGraph<NodeType, EdgeType, IslandType, OwnerType>::SplitIsland(const int32 RootIndex, const int32 IslandIndex)
 {
 	NodeQueue.PushLast(RootIndex);
 	int32 NodeIndex = RootIndex;
@@ -598,7 +690,17 @@ void FIslandGraph<NodeType, EdgeType, IslandType>::SplitIsland(const int32 RootI
 		if (GraphNode.NodeCounter != GraphCounter)
 		{
 			GraphNode.NodeCounter = GraphCounter;
-			GraphNode.IslandIndex = IslandIndex;
+
+			// We are always awake when SplitIslands is called so NodeIslands will be rebuilt
+			if (GraphNode.bValidNode)
+			{
+				GraphNode.IslandIndex = IslandIndex;
+			}
+			else
+			{
+				GraphNode.IslandIndex = INDEX_NONE;
+			}
+			GraphNode.NodeIslands.Reset();
 
 			// Loop over the node edges to continue the island discovery
 			for (int32& EdgeIndex : GraphNode.NodeEdges)
@@ -612,7 +714,7 @@ void FIslandGraph<NodeType, EdgeType, IslandType>::SplitIsland(const int32 RootI
 				const int32 OtherIndex = (NodeIndex == GraphEdge.FirstNode) ?
 					GraphEdge.SecondNode : GraphEdge.FirstNode;
 
-				// Only the valid nodes are allowed to continue the graph traversal
+				// Only the valid nodes (Sleeping/Dynamic Particles) are allowed to continue the graph traversal (connect islands)
 				if (GraphNodes.IsValidIndex(OtherIndex) && GraphNodes[OtherIndex].NodeCounter != GraphCounter && GraphNodes[OtherIndex].bValidNode)
 				{
 					NodeQueue.PushLast(OtherIndex);
@@ -622,8 +724,8 @@ void FIslandGraph<NodeType, EdgeType, IslandType>::SplitIsland(const int32 RootI
 	}
 }
 
-template<typename NodeType, typename EdgeType, typename IslandType>
-void FIslandGraph<NodeType, EdgeType, IslandType>::SplitIslands()
+template<typename NodeType, typename EdgeType, typename IslandType, typename OwnerType>
+void FIslandGraph<NodeType, EdgeType, IslandType, OwnerType>::SplitIslands()
 {
 	SCOPE_CYCLE_COUNTER(STAT_SplitIslandGraph);
 
@@ -659,14 +761,15 @@ void FIslandGraph<NodeType, EdgeType, IslandType>::SplitIslands()
 				else
 				{
 					RootNode.IslandIndex = INDEX_NONE;
+					RootNode.NodeIslands.Reset();
 				}
 			}
 		}
 	}
 }
 	
-template<typename NodeType, typename EdgeType, typename IslandType>
-void FIslandGraph<NodeType, EdgeType, IslandType>::ReassignIslands()
+template<typename NodeType, typename EdgeType, typename IslandType, typename OwnerType>
+void FIslandGraph<NodeType, EdgeType, IslandType, OwnerType>::ReassignIslands()
 {
 	// Update all the nodes indices
 	for (FGraphIsland& GraphIsland : GraphIslands)
@@ -675,18 +778,20 @@ void FIslandGraph<NodeType, EdgeType, IslandType>::ReassignIslands()
 		GraphIsland.NumEdges = 0;
 	}
 	
-	// Update all the edges indices
+	// Update all the edges island indices
 	for (FGraphEdge& GraphEdge : GraphEdges)
 	{
-		if(GraphIslands.IsValidIndex(GraphEdge.IslandIndex))
+		int32 EdgeIslandIndex = INDEX_NONE;
+		if (GraphIslands.IsValidIndex(GraphEdge.IslandIndex))
 		{
 			const int32 ParentIndex = GraphIslands[GraphEdge.IslandIndex].ParentIsland;
 			if(GraphIslands.IsValidIndex(ParentIndex))
 			{
-				GraphEdge.IslandIndex = ParentIndex;
+				EdgeIslandIndex = ParentIndex;
 				GraphIslands[ParentIndex].NumEdges++;
 			}
 		}
+		GraphEdge.IslandIndex = EdgeIslandIndex;
 	}
 
 	// Update all the nodes indices
@@ -697,7 +802,17 @@ void FIslandGraph<NodeType, EdgeType, IslandType>::ReassignIslands()
 			const int32 ParentIndex = GraphIslands[GraphNode.IslandIndex].ParentIsland;
 			if(GraphIslands.IsValidIndex(ParentIndex))
 			{
-				GraphNode.IslandIndex = ParentIndex;
+				// This is only called from MergeIslands so we know we are awake. NodeIslands will be rebuilt (See PopulateIslands)
+				if (GraphNode.bValidNode)
+				{
+					GraphNode.IslandIndex = ParentIndex;
+				}
+				else
+				{
+					GraphNode.IslandIndex = INDEX_NONE;
+				}
+				GraphNode.NodeIslands.Reset();
+
 				if(GraphNode.bValidNode)
 				{
 					GraphIslands[ParentIndex].NumNodes++;
@@ -707,8 +822,8 @@ void FIslandGraph<NodeType, EdgeType, IslandType>::ReassignIslands()
 	}
 }
 
-template<typename NodeType, typename EdgeType, typename IslandType>
-void FIslandGraph<NodeType, EdgeType, IslandType>::UpdateGraph()
+template<typename NodeType, typename EdgeType, typename IslandType, typename OwnerType>
+void FIslandGraph<NodeType, EdgeType, IslandType, OwnerType>::UpdateGraph()
 {
 	SCOPE_CYCLE_COUNTER(STAT_MergeIslandGraph);
 
@@ -721,23 +836,36 @@ void FIslandGraph<NodeType, EdgeType, IslandType>::UpdateGraph()
 		if (GraphNodes.IsValidIndex(NodeIndex))
 		{
 			FGraphNode& GraphNode = GraphNodes[NodeIndex];
+
 			// Add new islands for the all the particles that are not connected into the graph 
-			if(!GraphIslands.IsValidIndex(GraphNode.IslandIndex) && GraphNode.bValidNode)
+			if (GraphNode.bValidNode && !GraphIslands.IsValidIndex(GraphNode.IslandIndex))
 			{
-				FGraphIsland GraphIsland = { 0, 1 };
+				FGraphIsland GraphIsland = { 0, 1 }; // {NumEdges, NumNodes}
 				GraphNode.IslandIndex = GraphIslands.Emplace(GraphIsland);
+
+#if CHAOS_CONSTRAINTHANDLE_DEBUG_ENABLED
+				// Should never have any NodeIslands if we get here
+				ensure(GraphNode.NodeIslands.Num() == 0);
+#endif
 			}
-			// Update of the sleeping flag on the island
-			if(!GraphNode.bStationaryNode)
+
+			// Clear the sleeping flag on the island for moving nodes
+			if (!GraphNode.bStationaryNode)
 			{
-				if(GraphNode.NodeEdges.Num() == 0 && GraphIslands.IsValidIndex(GraphNode.IslandIndex))
+				if (GraphNode.bValidNode)
 				{
+					// NOTE: All valid nodes should have an island index here (see above)
+					check(GraphIslands.IsValidIndex(GraphNode.IslandIndex));
+
+					// Valid node (Sleeping/Dynamic Particle) that is moving - island must be awake
 					GraphIslands[GraphNode.IslandIndex].bIsSleeping = false;
 				}
-				else
+				else if (!GraphNode.bValidNode)
 				{
-					// if a node is a moving kinematic and is switching its stationary state
-					// we have to iterate through all the edges since the kinematic node could belong to several islands
+					// Invalid node (Kinematic) that is moving - islands must be awake.
+					// We have to iterate through all the edges since the node could belong to several islands
+					// and the NodeIslands array is managed outside this class and may not be up to date.
+					// @todo(chaos): move NodeIslands array management into FIslandGraph if possible
 					for (auto& EdgeIndex : GraphNode.NodeEdges)
 					{
 						GraphIslands[GraphEdges[EdgeIndex].IslandIndex].bIsSleeping = false;
@@ -756,33 +884,49 @@ void FIslandGraph<NodeType, EdgeType, IslandType>::UpdateGraph()
 		}
 	}
 	
-	// Split the islands if this one is persistent and not sleeping
+	// Split the islands that are persistent and not sleeping if possible
 	SplitIslands();
+
+	// Remove edges from their island if:
+	// - the island is awake and the edge is invalid (no valid nodes)
+	// - the island was destroyed/merged but the edge wasn't moved to the new island (because it is invalid)
+	for (FGraphEdge& GraphEdge : GraphEdges)
+	{
+		const bool bValidIsland = GraphIslands.IsValidIndex(GraphEdge.IslandIndex);
+		const bool bAwakeIsland = bValidIsland && !GraphIslands[GraphEdge.IslandIndex].bIsSleeping;
+		const bool bValidEdge = GraphEdge.bValidEdge;
+		if (!bValidIsland || (bAwakeIsland && !bValidEdge))
+		{
+			GraphEdge.IslandIndex = INDEX_NONE;
+		}
+	}
 }
 
-template<typename NodeType, typename EdgeType, typename IslandType>
-void FIslandGraph<NodeType, EdgeType, IslandType>::InitIslands()
+template<typename NodeType, typename EdgeType, typename IslandType, typename OwnerType>
+void FIslandGraph<NodeType, EdgeType, IslandType, OwnerType>::InitIslands()
 {
 	// Remove of all the non sleeping edges
 	for(int32 EdgeIndex = GraphEdges.GetMaxIndex(); EdgeIndex >= 0; --EdgeIndex)
 	{
-		if(GraphEdges.IsValidIndex(EdgeIndex) && GraphIslands.IsValidIndex(GraphEdges[EdgeIndex].IslandIndex) && !GraphIslands[GraphEdges[EdgeIndex].IslandIndex].bIsSleeping)
+		if(GraphEdges.IsValidIndex(EdgeIndex) && GraphIslands.IsValidIndex(GraphEdges[EdgeIndex].IslandIndex))
 		{
-			RemoveEdge(EdgeIndex);
+			if (!GraphIslands[GraphEdges[EdgeIndex].IslandIndex].bIsSleeping)
+			{
+				RemoveEdge(EdgeIndex);
+			}
 		}
 	} 
 
 	// Reset of the sleeping flag for the graph islands
+	// See UpdateGraph() which sets it to false again if there are any awake nodes
 	for (auto& GraphIsland : GraphIslands)
 	{
 		GraphIsland.bIsSleeping = true;
-		GraphIsland.ChildrenIslands.Reset();
-		GraphIsland.ParentIsland = INDEX_NONE;
 	}
 }
 
-template<typename NodeType, typename EdgeType, typename IslandType>
-void FIslandGraph<NodeType, EdgeType, IslandType>::ResetIslands()
+template<typename NodeType, typename EdgeType, typename IslandType, typename OwnerType>
+void FIslandGraph<NodeType, EdgeType, IslandType, OwnerType>::ResetIslands()
 {
 	GraphEdges.Reset();
 	ItemEdges.Reset();
@@ -798,6 +942,6 @@ void FIslandGraph<NodeType, EdgeType, IslandType>::ResetIslands()
 	}
 }
 
-template class FIslandGraph<FGeometryParticleHandle*, FConstraintHandleHolder, FPBDIslandSolver*>;
-template class FIslandGraph<int32, int32, int32>;
+template class FIslandGraph<FGeometryParticleHandle*, FConstraintHandleHolder, FPBDIslandSolver*, FPBDIslandManager>;
+template class FIslandGraph<int32, int32, int32, TNullIslandGraphOwner<int32, int32>>;
 }
diff --git a/Engine/Source/Runtime/Experimental/Chaos/Private/Chaos/Island/IslandManager.cpp b/Engine/Source/Runtime/Experimental/Chaos/Private/Chaos/Island/IslandManager.cpp
--- a/Engine/Source/Runtime/Experimental/Chaos/Private/Chaos/Island/IslandManager.cpp
+++ b/Engine/Source/Runtime/Experimental/Chaos/Private/Chaos/Island/IslandManager.cpp
@@ -12,6 +12,8 @@
 #include "Framework/Threading.h"
 #include "Misc/App.h"
 
+//PRAGMA_DISABLE_OPTIMIZATION
+
 /** Cvar to enable/disable the island sleeping */
 bool bChaosSolverSleepEnabled = true;
 FAutoConsoleVariableRef CVarChaosSolverSleepEnabled(TEXT("p.Chaos.Solver.SleepEnabled"), bChaosSolverSleepEnabled, TEXT(""));
@@ -32,7 +34,7 @@ FAutoConsoleVariableRef CVarChaosSolverCollisionDefaultLinearSleepThreshold(TEXT
 Chaos::FRealSingle ChaosSolverCollisionDefaultAngularSleepThresholdCVar = 0.0087f;  //~1/2 unit mass degree
 FAutoConsoleVariableRef CVarChaosSolverCollisionDefaultAngularSleepThreshold(TEXT("p.ChaosSolverCollisionDefaultAngularSleepThreshold"), ChaosSolverCollisionDefaultAngularSleepThresholdCVar, TEXT("Default angular threshold for sleeping.[def:0.0087]"));
 
-bool bChaosSolverValidateGraph = false;
+bool bChaosSolverValidateGraph = (CHAOS_CONSTRAINTHANDLE_DEBUG_ENABLED != 0);
 FAutoConsoleVariableRef CVarChaosSolverValidateGraph(TEXT("p.Chaos.Solver.ValidateGraph"), bChaosSolverValidateGraph, TEXT(""));
 
 namespace Chaos
@@ -52,6 +54,7 @@ FORCEINLINE bool IsStationaryParticle(const FGeometryParticleHandle* ParticleHan
 	if (ParticleHandle->ObjectState() == EObjectStateType::Kinematic)
 	{
 		const FKinematicGeometryParticleHandle* KinematicParticle = ParticleHandle->CastToKinematicParticle();
+		// @todo(chaos): should this also be checking Kinematic W()?
 		return KinematicParticle->V().IsZero() && KinematicParticle->KinematicTarget().GetMode() == EKinematicTargetMode::None;
 	}
 	else
@@ -179,39 +182,27 @@ inline void UpdateSyncState(FPBDIslandSolver* IslandSolver, FPBDRigidsSOAs& Part
 		IslandSolver->SetNeedsResim(bNeedsResim);
 	}
 }
-	/** Add all the graph particle and constraints to the solver islands*/
+
+/** Add all the graph particle and constraints to the solver islands*/
 inline void PopulateIslands(FPBDIslandManager::GraphType* IslandGraph)
 {
 	auto AddNodeToIsland = [IslandGraph](const int32 IslandIndex, FPBDIslandManager::GraphType::FGraphNode& GraphNode)
 	{
-		if(IslandGraph->GraphIslands.IsValidIndex(IslandIndex))
+		if (IslandGraph->GraphIslands.IsValidIndex(IslandIndex))
 		{
 			FPBDIslandSolver* IslandSolver = IslandGraph->GraphIslands[IslandIndex].IslandItem;
-			if(!IslandSolver->IsSleeping())
+			if (!IslandSolver->IsSleeping())
 			{
 				IslandSolver->AddParticle(GraphNode.NodeItem);
 			}
 		}
 	};
-	TSet<int32> NodeIslands;
-	for(auto& GraphNode : IslandGraph->GraphNodes)
+	for (auto& GraphNode : IslandGraph->GraphNodes)
 	{
-		NodeIslands = GraphNode.NodeIslands;
 		GraphNode.NodeIslands.Reset();
 
-		// First we re-add the island indices coming from a sleeping islands since we are not adding them here.
-		for (auto& IslandIndex : NodeIslands)
-		{
-			if (IslandGraph->GraphIslands.IsValidIndex(IslandIndex))
-			{
-				if (IslandGraph->GraphIslands[IslandIndex].IslandItem->IsSleeping())
-				{
-					GraphNode.NodeIslands.Add(IslandIndex);
-				}
-			}
-		}
 		// If the node is valid : only one island
-		if( GraphNode.bValidNode)
+		if (GraphNode.bValidNode)
 		{
 			AddNodeToIsland(GraphNode.IslandIndex, GraphNode);
 			GraphNode.NodeIslands.Add(GraphNode.IslandIndex);
@@ -220,26 +211,30 @@ inline void PopulateIslands(FPBDIslandManager::GraphType* IslandGraph)
 		{
 			// A particle could belong to several islands when static/kinematic (not valid)
 			// First we compute the unique particle set of islands
-			for( auto& NodeEdge : GraphNode.NodeEdges)
+			for (auto& NodeEdge : GraphNode.NodeEdges)
 			{
 				GraphNode.NodeIslands.Add(IslandGraph->GraphEdges[NodeEdge].IslandIndex);
 			}
 			// Loop over the set of islands and add the particle to the solver island
-			for(auto& NodeIsland : GraphNode.NodeIslands)
+			for (auto& NodeIsland : GraphNode.NodeIslands)
 			{
 				AddNodeToIsland(NodeIsland, GraphNode);
 			}
 		}
 	}
 	
-	// Loop over the graph edges to transfer them into the solver island
-	for(auto& GraphEdge : IslandGraph->GraphEdges)
+	// Loop over the graph edges to transfer them into the solver island (if awake and involves a dynamic particle)
+	// At this point, no awake island should contain any invalid nodes or edges
+	for (auto& GraphEdge : IslandGraph->GraphEdges)
 	{
-		if(IslandGraph->GraphIslands.IsValidIndex(GraphEdge.IslandIndex))
+		if (IslandGraph->GraphIslands.IsValidIndex(GraphEdge.IslandIndex))
 		{
 			FPBDIslandSolver* IslandSolver = IslandGraph->GraphIslands[GraphEdge.IslandIndex].IslandItem;
-			if(!IslandSolver->IsSleeping())
+			if (!IslandSolver->IsSleeping())
 			{
+#if CHAOS_CONSTRAINTHANDLE_DEBUG_ENABLED
+				ensure(GraphEdge.bValidEdge);
+#endif
 				IslandSolver->AddConstraint(GraphEdge.EdgeItem);
 			}
 		}
@@ -301,33 +296,43 @@ inline bool ComputeSleepingThresholds(FPBDIslandSolver* IslandSolver,
 	return bHaveSleepThreshold;
 }
 
-FPBDIslandManager::FPBDIslandManager() : IslandSolvers(), IslandGraph(MakeUnique<GraphType>()), MaxParticleIndex(INDEX_NONE)
+FPBDIslandManager::FPBDIslandManager() 
+	: IslandSolvers()
+	, IslandGraph(MakeUnique<GraphType>())
+	, MaxParticleIndex(INDEX_NONE)
+	, bIslandsPopulated(false)
+	, bEndTickCalled(true)
 {
+	IslandGraph->SetOwner(this);
+
 	InitializeGroups();
 }
 
-FPBDIslandManager::FPBDIslandManager(const TParticleView<FPBDRigidParticles>& PBDRigids) : IslandSolvers(), IslandGraph(MakeUnique<GraphType>())
+FPBDIslandManager::FPBDIslandManager(const TParticleView<FPBDRigidParticles>& PBDRigids) 
+	: IslandSolvers()
+	, IslandGraph(MakeUnique<GraphType>())
+	, MaxParticleIndex(INDEX_NONE)
+	, bIslandsPopulated(false)
+	, bEndTickCalled(true)
 {
 	InitializeGraph(PBDRigids);
 	InitializeGroups();
 }
 
 FPBDIslandManager::~FPBDIslandManager()
-{}
+{
+}
 
-void FPBDIslandManager::ResetIndices()
+void FPBDIslandManager::Reset()
 {
-	// Reset of all the particles / constraint graph index to INDEX_NONE
-	for(auto& ItemNode : IslandGraph->ItemNodes)
-	{
-		if(FPBDRigidParticleHandle* PBDRigid = ItemNode.Key->CastToRigidParticle())
-		{
-			PBDRigid->SetConstraintGraphIndex(INDEX_NONE);
-		}
-	}
-	for(auto& ItemEdge : IslandGraph->ItemEdges)
+	RemoveConstraints();
+	RemoveParticles();
+
+	IslandSolvers.Reset();
+
+	for (auto& IslandGroup : IslandGroups)
 	{
-		ItemEdge.Key->SetConstraintGraphIndex(INDEX_NONE);
+		IslandGroup->NumConstraints() = 0;
 	}
 }
 
@@ -348,14 +353,33 @@ void FPBDIslandManager::InitializeGroups()
 
 void FPBDIslandManager::InitializeGraph(const TParticleView<FPBDRigidParticles>& PBDRigids)
 {
+	// If we hit this, there's a missing call to EndTick since the previous tick
+	ensure(bEndTickCalled);
+	bEndTickCalled = false;
+
+	// @todo(chaos): make this handle when the user sleeps all particles in an island. Currently this
+	// will put the island to sleep as expected, but only after removing all the constraints from the island.
+	// The result is that the particles in the island will have no collisions on the first frame after waking.
+	// See unit test: DISABLED_TestConstraintGraph_ParticleSleep_Manual
+
 	MaxParticleIndex = 0;
 	ReserveParticles(PBDRigids.Num());
+
+	// Reset the island merge tracking data - adding particles may cause islands to merge (e.g., if a particle has switched from kinematic to dynamic
+	// @todo(chaos): put this into a method on the graph. Or maybe this isn't necessary since Merge/Split should always leave a valid empty state...
+	for (auto& GraphIsland : IslandGraph->GraphIslands)
+	{
+		GraphIsland.ChildrenIslands.Reset();
+		GraphIsland.ParentIsland = INDEX_NONE;
+	}
+
 	// We are adding all the particles from the solver in
 	// case some were just created/activated
 	for (auto& RigidParticle : PBDRigids)
 	{
 		AddParticle(RigidParticle.Handle());
 	}
+
 	// we update the valid/steady state of the nodes in case any state changed
 	for(int32 NodeIndex = 0, NumNodes = IslandGraph->GraphNodes.GetMaxIndex(); NodeIndex < NumNodes; ++NodeIndex)
 	{
@@ -367,41 +391,75 @@ void FPBDIslandManager::InitializeGraph(const TParticleView<FPBDRigidParticles>&
 			IslandGraph->UpdateNode(ParticleHandle, IsDynamicParticle(ParticleHandle), GraphNode.IslandIndex,
 													IsStationaryParticle(ParticleHandle), NodeIndex);
 
+			// @todo(chaos): we should not be doing this here - we should remove particles when they are disabled...is this a leftover temp fix?
 			if (ParticleHandle->CastToRigidParticle() && ParticleHandle->CastToRigidParticle()->Disabled())
 			{
 				RemoveParticle(ParticleHandle);
 			}
 		}
 	}
-	// For now we are resetting all the constraints but we should keep the
-	// persistent collisions, joints...over time
+
+	// For now we are resetting all the constraints but we should keep the persistent constraints.
+	// NOTE: This only removes awake constraints from the graph and it leaves the sleeping ones. This is important for 
+	// persistent collisions because we do not run collision detection on sleeping particles, so we need the graph
+	// to keep it's sleeping edges so we know what collisions to wake when an island is awakened. The constraint
+	// management system also used the sleeping state as a lock to prevent constraint destruction.
 	IslandGraph->InitIslands();
-	for(auto& IslandSolver : IslandSolvers)
+
+	// If the user explicitly woke any islands, transfer that change to the graph
+	// @todo(chaos): this is not great. If WakeIsland is called before InitializeGraph we set the wake state 
+	// on the graph but it gets reset in InitIslands. Maybe we should reset the state at the end of the tick and
+	// then any user changes could be retained. For now we check to see if the solver sleep state was changed
+	// and if it was, use that to initialize the graph island.
+	for (int32 IslandIndex = 0; IslandIndex < IslandGraph->GraphIslands.GetMaxIndex(); ++IslandIndex)
 	{
-		if(!IslandSolver->IsSleeping())
+		if (IslandGraph->GraphIslands.IsValidIndex(IslandIndex))
 		{
-			IslandSolver->ClearConstraints();
+			if (IslandGraph->GraphIslands[IslandIndex].IslandItem->SleepingChanged())
+			{
+				IslandGraph->GraphIslands[IslandIndex].bIsSleeping = IslandGraph->GraphIslands[IslandIndex].IslandItem->IsSleeping();
+				IslandGraph->GraphIslands[IslandIndex].IslandItem->ResetSleepingChanged();
+			}
 		}
 	}
 }
 
-void FPBDIslandManager::RemoveConstraints()
+void FPBDIslandManager::EndTick()
 {
-	for(auto& ItemEdge : IslandGraph->ItemEdges)
+	// IslandSolvers are transient for use during the solver phase - clear them at the end of the frame
+	for (auto& IslandSolver : IslandSolvers)
 	{
-		ItemEdge.Key->SetConstraintGraphIndex(INDEX_NONE);
+		IslandSolver->ClearParticles();
+		IslandSolver->ClearConstraints();
+		IslandSolver->ResetSleepingChanged();
 	}
-	IslandGraph->ItemEdges.Reset();
-	IslandGraph->GraphEdges.Reset();
-	
-	for(auto& GraphIsland : IslandGraph->GraphIslands)
+
+	// We can no longer use the Island's particle and constraint lists
+	bIslandsPopulated = false;
+
+	bEndTickCalled = true;
+}
+
+void FPBDIslandManager::RemoveParticles()
+{
+	RemoveConstraints();
+
+	IslandGraph->ResetNodes();
+
+	for (auto& IslandSolver : IslandSolvers)
 	{
-		GraphIsland.NumEdges = 0;
+		IslandSolver->ClearParticles();
 	}
-	for(auto& GraphNode: IslandGraph->GraphNodes)
+	for (auto& IslandGroup : IslandGroups)
 	{
-		GraphNode.NodeEdges.Reset();
+		IslandGroup->NumParticles() = 0;
 	}
+}
+
+void FPBDIslandManager::RemoveConstraints()
+{
+	IslandGraph->ResetEdges();
+
 	for(auto& IslandSolver : IslandSolvers)
 	{
 		IslandSolver->ClearConstraints();
@@ -414,7 +472,7 @@ void FPBDIslandManager::RemoveConstraints()
 
 int32 FPBDIslandManager::ReserveParticles(const int32 NumParticles)
 {
-	const int32 MaxIndex = IslandGraph->NumNodes();
+	const int32 MaxIndex = IslandGraph->MaxNumNodes();
 	IslandGraph->ReserveNodes(NumParticles);
 	
 	IslandSolvers.Reserve(NumParticles);
@@ -428,41 +486,68 @@ void FPBDIslandManager::ReserveConstraints(const int32 NumConstraints)
 	IslandGraph->ReserveEdges(NumConstraints);
 }
 
+// Callback from the IslandGraph to allow us to store the node index
+void FPBDIslandManager::GraphNodeAdded(FGeometryParticleHandle* ParticleHandle, const int32 NodeIndex)
+{
+	// @todo(chaos): we should store the node index on kinematics and statics too
+	if (FPBDRigidParticleHandle* PBDRigid = ParticleHandle->CastToRigidParticle())
+	{
+		PBDRigid->SetConstraintGraphIndex(NodeIndex);
+	}
+}
+
+void FPBDIslandManager::GraphNodeRemoved(FGeometryParticleHandle* ParticleHandle)
+{
+	if (FPBDRigidParticleHandle* PBDRigid = ParticleHandle->CastToRigidParticle())
+	{
+		PBDRigid->SetConstraintGraphIndex(INDEX_NONE);
+	}
+}
+
+// Callback from the IslandGraph to allow us to store the edge index
+void FPBDIslandManager::GraphEdgeAdded(const FConstraintHandleHolder& ConstraintHandle, const int32 EdgeIndex)
+{
+	ConstraintHandle->SetConstraintGraphIndex(EdgeIndex);
+}
+
+void FPBDIslandManager::GraphEdgeRemoved(const FConstraintHandleHolder& ConstraintHandle)
+{
+	ConstraintHandle->SetConstraintGraphIndex(INDEX_NONE);
+}
+
 int32 FPBDIslandManager::AddParticle(FGeometryParticleHandle* ParticleHandle, const int32 IslandIndex, const bool bOnlyDynamic)
 {
 	if (ParticleHandle)
 	{
 		const bool bIsDynamic = IsDynamicParticle(ParticleHandle);
-		if( !bOnlyDynamic || (bOnlyDynamic && bIsDynamic) )
+		if (!bOnlyDynamic || (bOnlyDynamic && bIsDynamic))
 		{
 			MaxParticleIndex = FMath::Max(MaxParticleIndex, ParticleHandle->UniqueIdx().Idx);
-			if(FPBDRigidParticleHandle* PBDRigid = ParticleHandle->CastToRigidParticle())
+
+			FPBDRigidParticleHandle* PBDRigid = ParticleHandle->CastToRigidParticle();
+			const bool bIsStationary = IsStationaryParticle(ParticleHandle);
+
+			// Assign or update the particle's graph node
+			int32 NodeIndex = INDEX_NONE;
+			if ((PBDRigid != nullptr) && (PBDRigid->ConstraintGraphIndex() != INDEX_NONE))
 			{
-				// If the rigid already have a graph index we just update the node information based on the new particles state...
-				if(PBDRigid->ConstraintGraphIndex() != INDEX_NONE)
-				{
-					IslandGraph->UpdateNode(ParticleHandle, bIsDynamic, IslandIndex,
-						IsStationaryParticle(ParticleHandle), PBDRigid->ConstraintGraphIndex());
-					return PBDRigid->ConstraintGraphIndex();
-				}
+				// If the rigid already has a graph index we just update the node information based on the new particles state...
+				NodeIndex = PBDRigid->ConstraintGraphIndex();
+				IslandGraph->UpdateNode(ParticleHandle, bIsDynamic, IslandIndex, bIsStationary, NodeIndex);
 			}
-			// It could be nice to have a graph index on the particle handle the same way
-			// we have one on the constraint handle. It will allow us to skip the query on the TSet to
-			// check if the particle is already there, which could be quite slow
-			const int32 NodeIndex = IslandGraph->AddNode(ParticleHandle, bIsDynamic, IslandIndex,
-				IsStationaryParticle(ParticleHandle));
-
-			if(FPBDRigidParticleHandle* PBDRigid = ParticleHandle->CastToRigidParticle())
+			else
 			{
-				PBDRigid->SetConstraintGraphIndex(NodeIndex);
+				// If we get here we have a new particle, or a particle that may be in multiple islands (Static/Kinematic)
+				NodeIndex = IslandGraph->AddNode(ParticleHandle, bIsDynamic, IslandIndex, bIsStationary);
 			}
+
 			return NodeIndex;
 		}
 	}
 	return INDEX_NONE;
 }
 	
-int32 FPBDIslandManager::AddConstraint(const uint32 ContainerId, FConstraintHandle* ConstraintHandle, const TVec2<FGeometryParticleHandle*>& ConstrainedParticles)
+void FPBDIslandManager::AddConstraint(const uint32 ContainerId, FConstraintHandle* ConstraintHandle, const TVec2<FGeometryParticleHandle*>& ConstrainedParticles)
 {
 	if (ConstraintHandle)
 	{
@@ -470,108 +555,95 @@ int32 FPBDIslandManager::AddConstraint(const uint32 ContainerId, FConstraintHand
 		const bool bValidParticle0 = ConstrainedParticles[0] && IsDynamicParticle(ConstrainedParticles[0]);
 		const bool bValidParticle1 = ConstrainedParticles[1] && IsDynamicParticle(ConstrainedParticles[1]);
 		
-		// We are checking if one of the 2 particle is dynamic to add the constraint to the graph
+		// We are checking if one of the 2 particle is dynamic to add the constraint to the graph.
+		// NOTE: This is also where kinematic particles get added to the graph since they were not passed to InitializeGraph.
 		if(bValidParticle0 || bValidParticle1)
 		{
 			const int32 NodeIndex0 = AddParticle(ConstrainedParticles[0], INDEX_NONE, false);
 			const int32 NodeIndex1 = AddParticle(ConstrainedParticles[1], INDEX_NONE, false);
 
-			const int32 EdgeIndex = IslandGraph->AddEdge(ConstraintHandle, ContainerId, NodeIndex0, NodeIndex1);
-			ConstraintHandle->SetConstraintGraphIndex(EdgeIndex);
-
-			if(IslandGraph->GraphEdges.IsValidIndex(EdgeIndex))
-			{
-				// Make sure to sync the state of the constraint with its owning island, otherwise the constraint may be flag as destroyable and leave a dangling pointer in the islands
-				const int32 IslandIndex = IslandGraph->GraphEdges[EdgeIndex].IslandIndex;
-				if (IslandGraph->GraphIslands.IsValidIndex(IslandIndex) && IslandGraph->GraphIslands[IslandIndex].bIsSleeping)
-				{
-					ConstraintHandle->SetIsSleeping(true);
-				}
-			}
-
-			return EdgeIndex;
-		}
-		else
-		{
-			ConstraintHandle->SetConstraintGraphIndex(INDEX_NONE);
+			IslandGraph->AddEdge(ConstraintHandle, ContainerId, NodeIndex0, NodeIndex1);
 		}
-		return INDEX_NONE;
 	}
-	return INDEX_NONE;
 }
 
 void FPBDIslandManager::RemoveParticle(FGeometryParticleHandle* ParticleHandle)
 {
 	if (ParticleHandle)
 	{
-		// @todo(chaos) : this is sub-optimal fix for a crash where trying to remove a particle from the data in the graph is not enough 
-		// edgelist is empty and particle is in multiple island solver 
-		// could possibly be because of the order we delete constraint before removing the particle
-		/*for (auto& IslandSolver : IslandSolvers)
+		IslandGraph->RemoveNode(ParticleHandle);
+
+#if CHAOS_CONSTRAINTHANDLE_DEBUG_DETAILED_ENABLED
+		DebugCheckParticleNotInGraph(ParticleHandle);
+#endif
+	}
+}
+
+void FPBDIslandManager::RemoveConstraint(const uint32 ContainerId, FConstraintHandle* ConstraintHandle)
+{
+	if (ConstraintHandle)
+	{
+		const int32 EdgeIndex = ConstraintHandle->ConstraintGraphIndex();
+		if (IslandGraph->GraphEdges.IsValidIndex(EdgeIndex))
 		{
-			IslandSolver->RemoveParticle(ParticleHandle);
-		}*/
+			IslandGraph->RemoveEdge(EdgeIndex);
+		}
 
-		// @todo(chaos) : test that new version to see if we still need the version above 
-		if (const int32* NodeIndex = IslandGraph->ItemNodes.Find(ParticleHandle))
+#if CHAOS_CONSTRAINTHANDLE_DEBUG_DETAILED_ENABLED
+		DebugCheckConstraintNotInGraph(ConstraintHandle);
+#endif
+	}
+}
+
+void FPBDIslandManager::RemoveParticleConstraints(FGeometryParticleHandle* ParticleHandle, const uint32 ContainerId)
+{
+	if (ParticleHandle)
+	{
+		// Find the graph node for this particle
+		if (const int32* PNodeIndex = IslandGraph->ItemNodes.Find(ParticleHandle))
 		{
-			if (IslandGraph->GraphNodes.IsValidIndex(*NodeIndex))
+			const int32 NodeIndex = *PNodeIndex;
+			if (IslandGraph->GraphNodes.IsValidIndex(NodeIndex))
 			{
-				const FGraphNode& GraphNode = IslandGraph->GraphNodes[*NodeIndex];
+				const FGraphNode& GraphNode = IslandGraph->GraphNodes[NodeIndex];
 
-				// We loop over all the connected edges to find all the islands in which the particle
-				// is (static/kinematic particles could belong to several islands)
-				// And remove the particle from the solver island. It will allow the solver
-				// islands to be updated directly and not at the next sync.
-				for (const int32& EdgeIndex : GraphNode.NodeEdges)
+				// Loop over the edges attached to this node and remove the constraints in the specified container
+				// NOTE: Can't use ranged-for iterator when removing items, but we can just loop and remove since it's a sparse array
+				for (int32 NodeEdgeIndex = 0, NumNodeEdges = GraphNode.NodeEdges.GetMaxIndex(); NodeEdgeIndex < NumNodeEdges; ++NodeEdgeIndex)
 				{
-					const int32 IslandIndex = IslandGraph->GraphEdges[EdgeIndex].IslandIndex;
-					if (IslandSolvers.IsValidIndex(IslandIndex))
+					if (GraphNode.NodeEdges.IsValidIndex(NodeEdgeIndex))
 					{
-						//We need to remove as well the constraint from the island solver since the edges are removed from the graph
-						IslandSolvers[IslandIndex]->RemoveConstraint(IslandGraph->GraphEdges[EdgeIndex].EdgeItem);
-						IslandGraph->GraphEdges[EdgeIndex].EdgeItem->SetConstraintGraphIndex(INDEX_NONE);
+						const int32 EdgeIndex = GraphNode.NodeEdges[NodeEdgeIndex];
+						if (IslandGraph->GraphEdges.IsValidIndex(EdgeIndex))
+						{
+							const FGraphEdge& GraphEdge = IslandGraph->GraphEdges[EdgeIndex];
+
+							// Remove the constraint if it is the right type (same ContainerId)
+							if (GraphEdge.ItemContainer == ContainerId)
+							{
+#if CHAOS_CONSTRAINTHANDLE_DEBUG_ENABLED
+								// All constraints should know their graph index
+								ensure(GraphEdge.EdgeItem->ConstraintGraphIndex() != INDEX_NONE);
+#endif
+
+								RemoveConstraint(ContainerId, GraphEdge.EdgeItem.Get());
+							}
+						}
 					}
 				}
 
-				// If the node is valid : we need to also remove it from its own island 
-				for (const int32& IslandIndex : GraphNode.NodeIslands)
-				{
-					if (IslandSolvers.IsValidIndex(IslandIndex))
-					{
-						IslandSolvers[IslandIndex]->RemoveParticle(ParticleHandle);
-					}
-				}
+				// Note: we may have left the particle in the island without any constraints on it. Is that ok?
+				// It should be fine as long as the node knows that it is still in it/them, which it does.
 			}
 		}
 
-		IslandGraph->RemoveNode(ParticleHandle);
-		if(FPBDRigidParticleHandle* PBDRigid = ParticleHandle->CastToRigidParticle())
-		{
-			PBDRigid->SetConstraintGraphIndex(INDEX_NONE);
-		}
+#if CHAOS_CONSTRAINTHANDLE_DEBUG_ENABLED
+		DebugParticleConstraintsNotInGraph(ParticleHandle, ContainerId);
+#endif
 	}
 }
 
-void FPBDIslandManager::RemoveConstraint(const uint32 ContainerId, FConstraintHandle* ConstraintHandle)
-{
-	if (ConstraintHandle)
-	{
-		const int32 EdgeIndex = ConstraintHandle->ConstraintGraphIndex();
-		if (IslandGraph->GraphEdges.IsValidIndex(EdgeIndex))
-		{
-			const int32 IslandIndex = IslandGraph->GraphEdges[EdgeIndex].IslandIndex;
-			if (IslandSolvers.IsValidIndex(IslandIndex))
-			{
-				IslandSolvers[IslandIndex]->RemoveConstraint(ConstraintHandle);
-			}
 
-			IslandGraph->RemoveEdge(EdgeIndex);
-			ConstraintHandle->SetConstraintGraphIndex(INDEX_NONE);
-		}
-	}
-}
-	
 void FPBDIslandManager::EnableParticle(FGeometryParticleHandle* ChildParticle, const FGeometryParticleHandle* ParentParticle)
 {
 	if (ParentParticle && ChildParticle)
@@ -622,13 +694,13 @@ inline bool SolverIslandSortPredicate(const TUniquePtr<FPBDIslandSolver>& Solver
 	
 void FPBDIslandManager::SyncIslands(FPBDRigidsSOAs& Particles, const int32 NumContainers)
 {
-	IslandSolvers.Reserve(IslandGraph->NumIslands());
-	IslandIndexing.SetNum(IslandGraph->NumIslands(),false);
-	SortedIslands.SetNum(IslandGraph->NumIslands(),false);
+	IslandSolvers.Reserve(IslandGraph->MaxNumIslands());
+	IslandIndexing.SetNum(IslandGraph->MaxNumIslands(),false);
+	SortedIslands.SetNum(IslandGraph->MaxNumIslands(),false);
 	int32 LocalIsland = 0;
 
 	// Sync of the solver islands first and reserve the required space
-	for (int32 IslandIndex = 0, NumIslands = IslandGraph->NumIslands(); IslandIndex < NumIslands; ++IslandIndex)
+	for (int32 IslandIndex = 0, NumIslands = IslandGraph->MaxNumIslands(); IslandIndex < NumIslands; ++IslandIndex)
 	{
 		if (IslandGraph->GraphIslands.IsValidIndex(IslandIndex))
 		{
@@ -652,18 +724,16 @@ void FPBDIslandManager::SyncIslands(FPBDRigidsSOAs& Particles, const int32 NumCo
 			SortedIslands[LocalIsland] = IslandIndex;
 			LocalIsland++;
 
-			// We finally update the solver islands based on the new particles and constraints if
-			// the island is not persistent or not sleeping. 
-			if(!IslandSolver->IsSleeping())
-			{
-				IslandSolver->ReserveParticles(IslandGraph->GraphIslands[IslandIndex].NumNodes);
-				IslandSolver->ReserveConstraints(IslandGraph->GraphIslands[IslandIndex].NumEdges);
-			}
+			// Clear the island solver ready for repopulation . This leaves the sleeping islands empty, but they 
+			// are not used until they wake at which point they will be filled again.
+			IslandSolver->ReserveParticles(IslandGraph->GraphIslands[IslandIndex].NumNodes);
+			IslandSolver->ReserveConstraints(IslandGraph->GraphIslands[IslandIndex].NumEdges);
+
 			// Reset of the sleep counter if the island is :
-			// - Non persistent since we are starting incrementing 
-			// the counter once the island is persistent and if values below the threshold
+			// - Non persistent since we are starting incrementing the counter once the island is persistent and if values below the threshold
 			// - Sleeping since as soon as it wakes up we could start incrementing the counter as well
-			if (!IslandSolver->IsPersistent() || IslandSolver->IsSleeping())
+			// - Just woken because we may not have been asleep for a whole frame to hit the
+			if (!IslandSolver->IsPersistent() || IslandSolver->IsSleeping() || IslandSolver->SleepingChanged())
 			{ 
 				IslandSolver->SetSleepCounter(0);
 			}
@@ -675,9 +745,12 @@ void FPBDIslandManager::SyncIslands(FPBDRigidsSOAs& Particles, const int32 NumCo
 			IslandSolvers.RemoveAt(IslandIndex);
 		}
 	}
+
+	IslandIndexing.SetNum(LocalIsland, false);
+	SortedIslands.SetNum(LocalIsland, false);
+
 	PopulateIslands(IslandGraph.Get());
 
-	
 	// Update of the sync and sleep state for each islands
 	for(auto& IslandSolver : IslandSolvers)
 	{
@@ -687,12 +760,17 @@ void FPBDIslandManager::SyncIslands(FPBDRigidsSOAs& Particles, const int32 NumCo
 			UpdateSleepState(IslandSolver.Get(), Particles);
 		}
 	}
-	
-	IslandIndexing.SetNum(LocalIsland,false);
-	SortedIslands.SetNum(LocalIsland,false);
+
+#if CHAOS_CONSTRAINTHANDLE_DEBUG_ENABLED
+	DebugCheckIslands();
+#endif
+
+	// We can now safely use the GetIslandParticles and GetIslandConstraints accessors
+	bIslandsPopulated = true;
 
 	// Build all the island groups
 	BuildGroups(NumContainers);
+
 }
 
 void FPBDIslandManager::BuildGroups(const int32 NumContainers)
@@ -810,41 +888,403 @@ bool FPBDIslandManager::SleepInactive(const int32 IslandIndex,
 
 void FPBDIslandManager::SleepIsland(FPBDRigidsSOAs& Particles, const int32 IslandIndex)
 {
-	const int32 GraphIndex = GetGraphIndex(IslandIndex);
-	if (IslandSolvers.IsValidIndex(GraphIndex) && !IslandSolvers[GraphIndex]->IsSleeping())
-	{
-		IslandSolvers[GraphIndex]->SetIsSleeping(true);
-		IslandGraph->GraphIslands[GraphIndex].bIsSleeping = true;
-		UpdateSleepState(IslandSolvers[GraphIndex].Get(),Particles);
-	}
+	// @todo(chaos): this function will not work when called manually at the start of the frame based on user request.
+	// Currently this is not required so not important. See disabled unit test: TestConstraintGraph_SleepIsland
+	ensure(bIslandsPopulated);
+
+	SetIslandSleeping(Particles, IslandIndex, true);
 }
 
 void FPBDIslandManager::WakeIsland(FPBDRigidsSOAs& Particles, const int32 IslandIndex)
 {
+	SetIslandSleeping(Particles, IslandIndex, false);
+}
+
+void FPBDIslandManager::SetIslandSleeping(FPBDRigidsSOAs& Particles, const int32 IslandIndex, const bool bIsSleeping)
+{
+	// NOTE: This Sleep function is called from the graph update to set the state of the 
+	// island and all particles in it.
 	const int32 GraphIndex = GetGraphIndex(IslandIndex);
 	if (IslandSolvers.IsValidIndex(GraphIndex))
 	{
-		IslandSolvers[GraphIndex]->SetIsSleeping(false);
-		IslandGraph->GraphIslands[GraphIndex].bIsSleeping = false;
-		UpdateSleepState(IslandSolvers[GraphIndex].Get(),Particles);
-		
-		IslandSolvers[GraphIndex]->SetSleepCounter(0);
+		const bool bWasSleeping = IslandSolvers[GraphIndex]->IsSleeping();
+		if (bWasSleeping != bIsSleeping)
+		{
+			IslandSolvers[GraphIndex]->SetIsSleeping(bIsSleeping);
+			IslandGraph->GraphIslands[GraphIndex].bIsSleeping = bIsSleeping;
+			UpdateSleepState(IslandSolvers[GraphIndex].Get(),Particles);
+		}
+
+		// Reset the sleep counter with every wake call, even if already awake
+		if (!bIsSleeping)
+		{
+			IslandSolvers[GraphIndex]->SetSleepCounter(0);
+		}
 	}
 }
 	
 const TArray<FGeometryParticleHandle*>& FPBDIslandManager::GetIslandParticles(const int32 IslandIndex) const
 {
+#if CHAOS_CONSTRAINTHANDLE_DEBUG_ENABLED
+	ensure(bIslandsPopulated);
+#endif
 	return IslandSolvers[GetGraphIndex(IslandIndex)]->GetParticles();
 }
 
 const TArray<FConstraintHandleHolder>& FPBDIslandManager::GetIslandConstraints(const int32 IslandIndex) const
 {
+#if CHAOS_CONSTRAINTHANDLE_DEBUG_ENABLED
+	ensure(bIslandsPopulated);
+#endif
 	return IslandSolvers[GetGraphIndex(IslandIndex)]->GetConstraints();
 }
 
+TArray<int32> FPBDIslandManager::FindParticleIslands(const FGeometryParticleHandle* ParticleHandle) const
+{
+	TArray<int32> IslandIndices;
+
+	// Initially build the list of sparse island indices based on the connected edges (for kinematics)
+	if (ParticleHandle)
+	{
+		// Find the graph node for this particle
+		if (const int32* PNodeIndex = IslandGraph->ItemNodes.Find(ParticleHandle))
+		{
+			const int32 NodeIndex = *PNodeIndex;
+			if (IslandGraph->GraphNodes.IsValidIndex(NodeIndex))
+			{
+				const FGraphNode& GraphNode = IslandGraph->GraphNodes[NodeIndex];
+
+				if (GraphNode.bValidNode)
+				{
+					// A dynamic particle is only in one island
+					if ((GraphNode.IslandIndex != INDEX_NONE) && IslandGraph->GraphIslands.IsValidIndex(GraphNode.IslandIndex))
+					{
+						IslandIndices.Add(GraphNode.IslandIndex);
+					}
+				}
+				else
+				{
+					// A non-dynamic particle can be in many islands, depending on the connections to dynamics
+					for (int32 EdgeIndex : GraphNode.NodeEdges)
+					{
+						if (IslandGraph->GraphEdges.IsValidIndex(EdgeIndex))
+						{
+							const FGraphEdge& GraphEdge = IslandGraph->GraphEdges[EdgeIndex];
+							if ((GraphEdge.IslandIndex != INDEX_NONE) && IslandGraph->GraphIslands.IsValidIndex(GraphEdge.IslandIndex))
+							{
+								IslandIndices.Add(GraphEdge.IslandIndex);
+							}
+						}
+					}
+				}
+			}
+		}
+
+		// Map the islands into non-sparse indices
+		for (int32 IslandIndicesIndex = 0; IslandIndicesIndex < IslandIndices.Num(); ++IslandIndicesIndex)
+		{
+			const int32 GraphIslandIndex = IslandIndices[IslandIndicesIndex];
+			const int32 SolverIslandIndex = IslandGraph->GraphIslands[GraphIslandIndex].IslandItem->GetIslandIndex();
+			IslandIndices[IslandIndicesIndex] = SolverIslandIndex;
+		}
+	}
+
+	return IslandIndices;
+}
+
+int32 FPBDIslandManager::GetConstraintIsland(const FConstraintHandle* ConstraintHandle) const
+{
+	const int32 EdgeIndex = ConstraintHandle->ConstraintGraphIndex();
+	if (IslandGraph->GraphEdges.IsValidIndex(EdgeIndex))
+	{
+		int32 GraphIslandIndex = IslandGraph->GraphEdges[EdgeIndex].IslandIndex;
+		if (GraphIslandIndex != INDEX_NONE)
+		{
+			const int32 SolverIslandIndex = IslandGraph->GraphIslands[GraphIslandIndex].IslandItem->GetIslandIndex();
+			return SolverIslandIndex;
+		}
+	}
+	return INDEX_NONE;
+}
+
 bool FPBDIslandManager::IslandNeedsResim(const int32 IslandIndex) const
 {
 	return IslandSolvers[GetGraphIndex(IslandIndex)]->NeedsResim();
 }
 
+#if CHAOS_CONSTRAINTHANDLE_DEBUG_ENABLED
+bool FPBDIslandManager::DebugCheckParticleNotInGraph(const FGeometryParticleHandle* ParticleHandle) const
+{
+	bool bSuccess = true;
+	for (const auto& Node : IslandGraph->GraphNodes)
+	{
+		if (Node.NodeItem == ParticleHandle)
+		{
+			bSuccess = false;
+			ensure(false);
+		}
+	}
+
+	for (const auto& ItemNode : IslandGraph->ItemNodes)
+	{
+		if (ItemNode.Key == ParticleHandle)
+		{
+			bSuccess = false;
+			ensure(false);
+		}
+	}
+
+	for (const auto& Edge : IslandGraph->GraphEdges)
+	{
+		if (Edge.EdgeItem.GetParticle0() == ParticleHandle)
+		{
+			bSuccess = false;
+			ensure(false);
+		}
+		if (Edge.EdgeItem.GetParticle1() == ParticleHandle)
+		{
+			bSuccess = false;
+			ensure(false);
+		}
+	}
+	return bSuccess;
+}
+
+bool FPBDIslandManager::DebugParticleConstraintsNotInGraph(const FGeometryParticleHandle* ParticleHandle, const int32 ContainerId) const
+{
+	bool bSuccess = true;
+	for (const auto& Edge : IslandGraph->GraphEdges)
+	{
+		if (Edge.ItemContainer == ContainerId)
+		{
+			if ((Edge.EdgeItem.GetParticle0() == ParticleHandle) || (Edge.EdgeItem.GetParticle1() == ParticleHandle))
+			{
+				bSuccess = false;
+				ensure(false);
+			}
+		}
+	}
+	return bSuccess;
+}
+
+bool FPBDIslandManager::DebugCheckConstraintNotInGraph(const FConstraintHandle* ConstraintHandle) const
+{
+	bool bSuccess = true;
+	for (const auto& Edge : IslandGraph->GraphEdges)
+	{
+		if (Edge.EdgeItem.Get() == ConstraintHandle)
+		{
+			bSuccess = false;
+			ensure(false);
+		}
+	}
+
+	for (const auto& ItemEdge : IslandGraph->ItemEdges)
+	{
+		if (ItemEdge.Key.Get() == ConstraintHandle)
+		{
+			bSuccess = false;
+			ensure(false);
+		}
+	}
+
+	for (int32 IslandIndex = 0; IslandIndex < NumIslands(); ++IslandIndex)
+	{
+		const FPBDIslandSolver* IslandSolver = GetSolverIsland(IslandIndex);
+		for (const FConstraintHandleHolder& HandleHolder : IslandSolver->GetConstraints())
+		{
+			if (HandleHolder.Get() == ConstraintHandle)
+			{
+				bSuccess = false;
+				ensure(false);
+			}
+		}
+	}
+	return bSuccess;
+}
+
+bool FPBDIslandManager::DebugCheckIslandConstraints() const
+{
+	bool bSuccess = true;
+
+	for (const auto& GraphEdge : IslandGraph->GraphEdges)
+	{
+		const FConstraintHandleHolder& ConstraintHolder = GraphEdge.EdgeItem;
+
+		// NOTE: The check for both null here is for unit tests that do not use particles...not ideal...
+		if ((ConstraintHolder.GetParticle0() != nullptr) || (ConstraintHolder.GetParticle1() != nullptr))
+		{
+			// Make sure that if we have an IslandIndex it is a valid one
+			const bool bIslandIndexOK = (GraphEdge.IslandIndex == INDEX_NONE) || IslandGraph->GraphIslands.IsValidIndex(GraphEdge.IslandIndex);
+			bSuccess = bSuccess && bIslandIndexOK;
+			ensure(bIslandIndexOK);
+
+			// At least one particle on each constraint must be dynamic if we are in an awake island.
+			// If we do not have a dynamic, we must have no island or it must be asleep
+			const bool bInAwakeIsland = (GraphEdge.IslandIndex != INDEX_NONE) && IslandGraph->GraphIslands.IsValidIndex(GraphEdge.IslandIndex) && !IslandGraph->GraphIslands[GraphEdge.IslandIndex].bIsSleeping;
+			const bool bIsDynamic0 = (ConstraintHolder.GetParticle0() != nullptr) ? FConstGenericParticleHandle(ConstraintHolder.GetParticle0())->IsDynamic() : false;
+			const bool bIsDynamic1 = (ConstraintHolder.GetParticle1() != nullptr) ? FConstGenericParticleHandle(ConstraintHolder.GetParticle1())->IsDynamic() : false;
+			if (bInAwakeIsland)
+			{
+				bSuccess = bSuccess && (bIsDynamic0 || bIsDynamic1);
+				ensure(bIsDynamic0 || bIsDynamic1);
+			}
+			if (!bIsDynamic0 && !bIsDynamic1)
+			{
+				bSuccess = bSuccess && !bInAwakeIsland;
+				ensure(!bInAwakeIsland);
+			}
+		}
+	}
+
+	return bSuccess;
+}
+
+void FPBDIslandManager::DebugCheckParticleIslands(const FGeometryParticleHandle* ParticleHandle) const
+{
+	// Check that every particle is included in the islands they think and no others
+	const int32* PNodeIndex = IslandGraph->ItemNodes.Find(ParticleHandle);
+	if (PNodeIndex != nullptr)
+	{
+		const int32 NodeIndex = *PNodeIndex;
+		const auto& Node = IslandGraph->GraphNodes[NodeIndex];
+		ensure(Node.NodeItem == ParticleHandle);
+
+		// Find the islands containing this node (either directly, or referenced by an constraint)
+		TArray<int32> NodeSolverIslands;
+		for (int32 IslandIndex = 0; IslandIndex < NumIslands(); ++IslandIndex)
+		{
+			const FPBDIslandSolver* IslandSolver = GetSolverIsland(IslandIndex);
+			for (const FGeometryParticleHandle* IslandParticleHandle : IslandSolver->GetParticles())
+			{
+				if (IslandParticleHandle == ParticleHandle)
+				{
+					NodeSolverIslands.AddUnique(IslandIndex);
+				}
+			}
+			for (const FConstraintHandleHolder& HandleHolder : IslandSolver->GetConstraints())
+			{
+				if ((HandleHolder.GetParticle0() == ParticleHandle) || (HandleHolder.GetParticle1() == ParticleHandle))
+				{
+					NodeSolverIslands.AddUnique(IslandIndex);
+				}
+			}
+		}
+
+		// NOTE: Particles are always assigneed to an island, but if we do not have any constraints, the node may not be in the IslandSolver
+		if (FConstGenericParticleHandle(ParticleHandle)->IsDynamic())
+		{
+			// Dynamic particles should be in 1 island
+			ensure(Node.IslandIndex != INDEX_NONE);
+			ensure(Node.NodeIslands.Num() == 1);
+			ensure(Node.NodeIslands.Contains(Node.IslandIndex));
+		}
+		else
+		{
+			// Static/Kinematic particles should be in at least 1 island if it has ever been in an awake island...
+			// We can't check that explicitly, but if we have no edges we may legitimately have no island
+			ensure(Node.IslandIndex == INDEX_NONE);
+			ensure((Node.NodeEdges.Num() == 0) || (Node.NodeIslands.Num() >= 1));
+		}
+
+		// If we found the node in some solver islands, make sure the nodes know about them
+		// NOTE: SolverIsland indexes are into the IslandIndexingArray, not the Islands array
+		for (int32 NodeSolverIsland : NodeSolverIslands)
+		{
+			const int32 NodeIsland = IslandIndexing[NodeSolverIsland];
+			ensure(Node.NodeIslands.Contains(NodeIsland));
+		}
+	}
+}
+
+void FPBDIslandManager::DebugCheckConstraintIslands(const FConstraintHandle* ConstraintHandle, const int32 EdgeIslandIndex) const
+{
+	// Find all the islands containing this edge
+	TArray<int32> EdgeSolverIslands;
+	for (int32 SolverIslandIndex = 0; SolverIslandIndex < NumIslands(); ++SolverIslandIndex)
+	{
+		const FPBDIslandSolver* IslandSolver = GetSolverIsland(SolverIslandIndex);
+		for (const FConstraintHandleHolder& HandleHolder : IslandSolver->GetConstraints())
+		{
+			if (HandleHolder.Get() == ConstraintHandle)
+			{
+				EdgeSolverIslands.AddUnique(SolverIslandIndex);
+			}
+		}
+	}
+
+	if (EdgeIslandIndex != INDEX_NONE)
+	{
+		if (IslandGraph->GraphIslands[EdgeIslandIndex].bIsSleeping)
+		{
+			// Make sure the constraint and its island agree on sleepiness
+			if (ConstraintHandle->SupportsSleeping())
+			{
+				ensure(ConstraintHandle->IsSleeping());
+			}
+
+			// If we are sleeping, we should not appear in any solver islands (we should only be in the graph)
+			ensure(EdgeSolverIslands.Num() == 0);
+		}
+		else
+		{
+			// Make sure the constraint and its island agree on sleepiness
+			if (ConstraintHandle->SupportsSleeping())
+			{
+				ensure(!ConstraintHandle->IsSleeping());
+			}
+
+			// We are awake, make sure our island search produced the correct island (and only that one)
+			ensure(EdgeSolverIslands.Num() == 1);
+			if (EdgeSolverIslands.Num() > 0)
+			{
+				const int32 EdgeIsland = IslandIndexing[EdgeSolverIslands[0]];
+				ensure(EdgeIslandIndex == EdgeIsland);
+			}
+		}
+	}
+	else
+	{
+		// We are not yet added to any islands
+		ensure(EdgeSolverIslands.Num() == 0);
+	}
+}
+
+void FPBDIslandManager::DebugCheckIslands() const
+{
+	// Check that the solver island state matches the graph island state
+	for (int32 IslandIndex = 0; IslandIndex < IslandGraph->GraphIslands.GetMaxIndex(); ++IslandIndex)
+	{
+		if (IslandGraph->GraphIslands.IsValidIndex(IslandIndex))
+		{
+			if (IslandGraph->GraphIslands[IslandIndex].bIsSleeping != IslandGraph->GraphIslands[IslandIndex].IslandItem->IsSleeping())
+			{
+				ensure(false);
+			}
+
+			const int32 SolverIslandIndex = IslandGraph->GraphIslands[IslandIndex].IslandItem->GetIslandIndex();
+			if (IslandIndexing[SolverIslandIndex] != IslandIndex)
+			{
+				ensure(false);
+			}
+		}
+	}
+
+	// Make sure islands only contain valid constraints
+	DebugCheckIslandConstraints();
+
+	// Check that every particle is included in the islands they think and no others
+	for (const auto& Node : IslandGraph->GraphNodes)
+	{
+		DebugCheckParticleIslands(Node.NodeItem);
+	}
+
+	// Check that every constraint is included in the island they think and no others
+	for (const auto& Edge : IslandGraph->GraphEdges)
+	{
+		DebugCheckConstraintIslands(Edge.EdgeItem, Edge.IslandIndex);
+	}
+}
+#endif
+
 }
\ No newline at end of file
diff --git a/Engine/Source/Runtime/Experimental/Chaos/Private/Chaos/Island/SolverIsland.cpp b/Engine/Source/Runtime/Experimental/Chaos/Private/Chaos/Island/SolverIsland.cpp
--- a/Engine/Source/Runtime/Experimental/Chaos/Private/Chaos/Island/SolverIsland.cpp
+++ b/Engine/Source/Runtime/Experimental/Chaos/Private/Chaos/Island/SolverIsland.cpp
@@ -60,11 +60,11 @@ void FPBDIslandSolver::AddConstraint(FConstraintHandle* ConstraintHandle)
 {
 	if (ConstraintHandle)
 	{
-		IslandConstraints.Add(ConstraintHandle);
-
 		const int32 ContainerId = ConstraintHandle->GetContainerId();
-		if(ConstraintCounts.IsValidIndex(ContainerId))
+		if (ConstraintCounts.IsValidIndex(ContainerId))
 		{
+			IslandConstraints.Add(ConstraintHandle);
+
 			ConstraintCounts[ContainerId]++;
 		}
 	}
@@ -85,14 +85,6 @@ void FPBDIslandSolver::ReserveConstraints(const int32 NumConstraints)
 	IslandConstraints.Reserve(NumConstraints);
 }
 
-void FPBDIslandSolver::ResetIndices()
-{
-	for(FConstraintHandleHolder& ConstraintHandle : IslandConstraints)
-	{
-		ConstraintHandle->SetConstraintGraphIndex(INDEX_NONE);
-	}
-}
-
 void FPBDIslandSolver::ClearConstraints()
 {
 	IslandConstraints.Reset();
diff --git a/Engine/Source/Runtime/Experimental/Chaos/Private/Chaos/PBDCollisionConstraints.cpp b/Engine/Source/Runtime/Experimental/Chaos/Private/Chaos/PBDCollisionConstraints.cpp
--- a/Engine/Source/Runtime/Experimental/Chaos/Private/Chaos/PBDCollisionConstraints.cpp
+++ b/Engine/Source/Runtime/Experimental/Chaos/Private/Chaos/PBDCollisionConstraints.cpp
@@ -301,6 +301,9 @@ namespace Chaos
 				Contact->SetContainer(this);
 				UpdateConstraintMaterialProperties(*Contact);
 			}
+
+			// Reset constraint modifications and accumulators
+			Contact->Activate();
 		}
 	}
 
diff --git a/Engine/Source/Runtime/Experimental/Chaos/Private/Chaos/PBDConstraintColor.cpp b/Engine/Source/Runtime/Experimental/Chaos/Private/Chaos/PBDConstraintColor.cpp
--- a/Engine/Source/Runtime/Experimental/Chaos/Private/Chaos/PBDConstraintColor.cpp
+++ b/Engine/Source/Runtime/Experimental/Chaos/Private/Chaos/PBDConstraintColor.cpp
@@ -759,8 +759,8 @@ void FPBDConstraintColor::InitializeColor(const FPBDConstraintGraph& ConstraintG
 	Edges.SetNum(ConstraintGraph.GetGraphEdges().GetMaxIndex());
 	IslandData.SetNum(ConstraintGraph.NumIslands());
 	ParticleToLevel.Init(0, ConstraintGraph.GetMaxParticleUniqueIdx() + 1);
-	DynamicGlobalToIslandLocalNodeIndicesArray.Init(-1, ConstraintGraph.GetIslandGraph()->NumNodes());
-	NodeToParticle.Init(nullptr, ConstraintGraph.GetIslandGraph()->NumNodes());
+	DynamicGlobalToIslandLocalNodeIndicesArray.Init(-1, ConstraintGraph.GetIslandGraph()->MaxNumNodes());
+	NodeToParticle.Init(nullptr, ConstraintGraph.GetIslandGraph()->MaxNumNodes());
 	IslandToGBFContactGraphData.SetNum(ConstraintGraph.NumIslands());
 }
 
diff --git a/Engine/Source/Runtime/Experimental/Chaos/Private/Chaos/PBDConstraintRule.cpp b/Engine/Source/Runtime/Experimental/Chaos/Private/Chaos/PBDConstraintRule.cpp
--- a/Engine/Source/Runtime/Experimental/Chaos/Private/Chaos/PBDConstraintRule.cpp
+++ b/Engine/Source/Runtime/Experimental/Chaos/Private/Chaos/PBDConstraintRule.cpp
@@ -5,11 +5,14 @@
 #include "Chaos/Island/SolverIsland.h"
 #include "Chaos/PBDCollisionConstraints.h"
 #include "Chaos/PBDJointConstraints.h"
+#include "Chaos/PBDNullConstraints.h"
 #include "Chaos/PBDPositionConstraints.h"
 #include "Chaos/PBDSuspensionConstraints.h"
 #include "Chaos/PBDRigidDynamicSpringConstraints.h"
 #include "Chaos/PBDRigidSpringConstraints.h"
 
+//PRAGMA_DISABLE_OPTIMIZATION
+
 namespace Chaos
 {
 	int32 ChaosShockPropagationVelocityPerLevelIterations = 1;
@@ -571,17 +574,20 @@ namespace Chaos
 	{
 		for(auto& GraphEdge : ConstraintGraph->GetIslandGraph()->GraphEdges)
 		{
-			auto& GraphIsland = ConstraintGraph->GetIslandGraph()->GraphIslands[GraphEdge.IslandIndex];
-			if(GraphEdge.ItemContainer == Constraints.GetContainerId() && !GraphIsland.IslandItem->IsSleeping())
+			if (GraphEdge.IslandIndex != INDEX_NONE)
 			{
-				const int32 EdgeColor = IsSortingUsingColors() ? FMath::Max(0,GraphEdge.ColorIndex) : 0;
-				const int32 EdgeLevel = IsSortingUsingLevels() ? FMath::Max(0,GraphEdge.LevelIndex) : 0;
-				const int32 MaxColors = IsSortingUsingColors() ? FMath::Max(1,GraphIsland.MaxColors+1) : 1;
+				auto& GraphIsland = ConstraintGraph->GetIslandGraph()->GraphIslands[GraphEdge.IslandIndex];
+				if(GraphEdge.ItemContainer == Constraints.GetContainerId() && !GraphIsland.IslandItem->IsSleeping())
+				{
+					const int32 EdgeColor = IsSortingUsingColors() ? FMath::Max(0,GraphEdge.ColorIndex) : 0;
+					const int32 EdgeLevel = IsSortingUsingLevels() ? FMath::Max(0,GraphEdge.LevelIndex) : 0;
+					const int32 MaxColors = IsSortingUsingColors() ? FMath::Max(1,GraphIsland.MaxColors+1) : 1;
 
-				const int32 IslandOffset = IslandOffsets[GraphIsland.IslandItem->GetIslandIndex()];
-				const int32 OffsetIndex = IslandOffset + EdgeLevel * MaxColors + EdgeColor;
+					const int32 IslandOffset = IslandOffsets[GraphIsland.IslandItem->GetIslandIndex()];
+					const int32 OffsetIndex = IslandOffset + EdgeLevel * MaxColors + EdgeColor;
 
-				InFunction(OffsetIndex,GraphEdge);
+					InFunction(OffsetIndex,GraphEdge);
+				}
 			}
 		}
 	}
@@ -609,11 +615,10 @@ namespace Chaos
 				}
 			}
 			// Initialization of the constraint offsets and the offsets counters
-			ConstraintOffsets.SetNum(IslandOffset, false);
-			OffsetCounters.SetNum(IslandOffset, false);
-			
-			FMemory::Memzero(ConstraintOffsets.GetData(), ConstraintOffsets.Num() * sizeof(int32));
-			FMemory::Memzero(OffsetCounters.GetData(), ConstraintOffsets.Num() * sizeof(int32));
+			ConstraintOffsets.Reset();
+			OffsetCounters.Reset();
+			ConstraintOffsets.SetNumZeroed(IslandOffset, false);
+			OffsetCounters.SetNumZeroed(IslandOffset, false);
 			
 			auto& LocalConstraintOffsets = ConstraintOffsets;
 			auto& LocalOffsetCounters = OffsetCounters;
@@ -651,10 +656,12 @@ namespace Chaos
 
 	template class TSimpleConstraintRule<FPBDCollisionConstraints>;
 	template class TSimpleConstraintRule<FPBDJointConstraints>;
+	template class TSimpleConstraintRule<FPBDNullConstraints>;
 	template class TSimpleConstraintRule<FPBDRigidSpringConstraints>;
 
 	template class TPBDConstraintGraphRuleImpl<FPBDCollisionConstraints>;
 	template class TPBDConstraintGraphRuleImpl<FPBDJointConstraints>;
+	template class TPBDConstraintGraphRuleImpl<FPBDNullConstraints>;
 	template class TPBDConstraintGraphRuleImpl<FPBDPositionConstraints>;
 	template class TPBDConstraintGraphRuleImpl<FPBDSuspensionConstraints>;
 	template class TPBDConstraintGraphRuleImpl<FPBDRigidDynamicSpringConstraints>;
@@ -662,6 +669,7 @@ namespace Chaos
 
 	template class TPBDConstraintColorRule<FPBDCollisionConstraints>;
 	template class TPBDConstraintIslandRule<FPBDJointConstraints>;
+	template class TPBDConstraintIslandRule<FPBDNullConstraints>;
 	template class TPBDConstraintIslandRule<FPBDPositionConstraints>;
 	template class TPBDConstraintIslandRule<FPBDSuspensionConstraints>;
 	template class TPBDConstraintIslandRule<FPBDRigidDynamicSpringConstraints>;
diff --git a/Engine/Source/Runtime/Experimental/Chaos/Private/Chaos/PBDJointConstraints.cpp b/Engine/Source/Runtime/Experimental/Chaos/Private/Chaos/PBDJointConstraints.cpp
--- a/Engine/Source/Runtime/Experimental/Chaos/Private/Chaos/PBDJointConstraints.cpp
+++ b/Engine/Source/Runtime/Experimental/Chaos/Private/Chaos/PBDJointConstraints.cpp
@@ -164,7 +164,7 @@ namespace Chaos
 		}
 	}
 
-	TVector<FGeometryParticleHandle*, 2> FPBDJointConstraintHandle::GetConstrainedParticles() const 
+	FParticlePair FPBDJointConstraintHandle::GetConstrainedParticles() const
 	{ 
 		return ConcreteContainer()->GetConstrainedParticles(ConstraintIndex);
 	}
@@ -676,7 +676,7 @@ namespace Chaos
 	}
 
 	
-	const typename FPBDJointConstraints::FParticlePair& FPBDJointConstraints::GetConstrainedParticles(int32 ConstraintIndex) const
+	const FParticlePair& FPBDJointConstraints::GetConstrainedParticles(int32 ConstraintIndex) const
 	{
 		return ConstraintParticles[ConstraintIndex];
 	}
diff --git a/Engine/Source/Runtime/Experimental/Chaos/Private/Chaos/PBDPositionConstraints.cpp b/Engine/Source/Runtime/Experimental/Chaos/Private/Chaos/PBDPositionConstraints.cpp
--- a/Engine/Source/Runtime/Experimental/Chaos/Private/Chaos/PBDPositionConstraints.cpp
+++ b/Engine/Source/Runtime/Experimental/Chaos/Private/Chaos/PBDPositionConstraints.cpp
@@ -5,7 +5,7 @@
 
 namespace Chaos
 {
-	TVector<FGeometryParticleHandle*, 2> FPBDPositionConstraintHandle::GetConstrainedParticles() const 
+	FParticlePair FPBDPositionConstraintHandle::GetConstrainedParticles() const
 	{ 
 		return ConcreteContainer()->GetConstrainedParticles(ConstraintIndex);
 	}
diff --git a/Engine/Source/Runtime/Experimental/Chaos/Private/Chaos/PBDRigidDynamicSpringConstraints.cpp b/Engine/Source/Runtime/Experimental/Chaos/Private/Chaos/PBDRigidDynamicSpringConstraints.cpp
--- a/Engine/Source/Runtime/Experimental/Chaos/Private/Chaos/PBDRigidDynamicSpringConstraints.cpp
+++ b/Engine/Source/Runtime/Experimental/Chaos/Private/Chaos/PBDRigidDynamicSpringConstraints.cpp
@@ -7,7 +7,7 @@
 
 using namespace Chaos;
 
-TVec2<FGeometryParticleHandle*> FPBDRigidDynamicSpringConstraintHandle::GetConstrainedParticles() const
+FParticlePair FPBDRigidDynamicSpringConstraintHandle::GetConstrainedParticles() const
 {
 	return ConcreteContainer()->GetConstrainedParticles(ConstraintIndex);
 }
diff --git a/Engine/Source/Runtime/Experimental/Chaos/Private/Chaos/PBDRigidSpringConstraints.cpp b/Engine/Source/Runtime/Experimental/Chaos/Private/Chaos/PBDRigidSpringConstraints.cpp
--- a/Engine/Source/Runtime/Experimental/Chaos/Private/Chaos/PBDRigidSpringConstraints.cpp
+++ b/Engine/Source/Runtime/Experimental/Chaos/Private/Chaos/PBDRigidSpringConstraints.cpp
@@ -21,9 +21,9 @@ namespace Chaos
 	{
 		ConcreteContainer()->SetConstraintPositions(ConstraintIndex, ConstraintPositions);
 	}
-	
-	TVector<typename FPBDRigidSpringConstraintHandle::FGeometryParticleHandle*, 2> FPBDRigidSpringConstraintHandle::GetConstrainedParticles() const
-	{ 
+
+	FParticlePair FPBDRigidSpringConstraintHandle::GetConstrainedParticles() const
+	{
 		return ConcreteContainer()->GetConstrainedParticles(ConstraintIndex);
 	}
 
diff --git a/Engine/Source/Runtime/Experimental/Chaos/Private/Chaos/PBDRigidsEvolution.cpp b/Engine/Source/Runtime/Experimental/Chaos/Private/Chaos/PBDRigidsEvolution.cpp
--- a/Engine/Source/Runtime/Experimental/Chaos/Private/Chaos/PBDRigidsEvolution.cpp
+++ b/Engine/Source/Runtime/Experimental/Chaos/Private/Chaos/PBDRigidsEvolution.cpp
@@ -7,6 +7,8 @@
 #include "Chaos/ParticleHandle.h"
 #include "Chaos/SpatialAccelerationCollection.h"
 
+//PRAGMA_DISABLE_OPTIMIZATION
+
 CSV_DECLARE_CATEGORY_EXTERN(ChaosPhysicsTimers);
 
 int32 ChaosRigidsEvolutionApplyAllowEarlyOutCVar = 1;
@@ -1022,7 +1024,12 @@ namespace Chaos
 
 	void FPBDRigidsEvolutionBase::SetParticleObjectState(FPBDRigidParticleHandle* Particle, EObjectStateType ObjectState)
 	{
-		EObjectStateType InitialState = Particle->ObjectState();
+		const EObjectStateType InitialState = Particle->ObjectState();
+
+		// Should never be called for statics - they cannot change state. Except...
+		// @todo(chaos): We should be able to assert this but FRigidClustering::CreateClusterParticle and ReleaseClusterParticle change the state after creation
+		//check(InitialState != EObjectStateType::Static)
+		//check(ObjectState != EObjectStateType::Static);
 
 		Particle->SetObjectStateLowLevel(ObjectState);
 
@@ -1036,20 +1043,12 @@ namespace Chaos
 			Particles.SetDynamicParticleSOA(Particle);
 		}
 
-		// @todo(chaos): this doesn't seem to handle Sleeping -> Kinematic, or Kinematic -> Sleeping
 		if (InitialState != ObjectState && !Particle->Disabled())
 		{
-			if (InitialState == EObjectStateType::Sleeping)
-			{
-				if (Particle->IslandIndex() != INDEX_NONE)
-				{
-					// GT has forced a wake so have to wake everything in the island
-					IslandsToWake.Enqueue(Particle->IslandIndex());
-				}
-			}
-			else if (ObjectState != EObjectStateType::Dynamic)
+			// If we were just put to sleep or made kinematic, we should still report info back to GT
+			// @todo(chaos): why? add reason to the comments
+			if (ObjectState != EObjectStateType::Dynamic)
 			{
-				// even though we went to sleep, we should still report info back to GT
 				Particles.MarkTransientDirtyParticle(Particle);
 			}
 		}
@@ -1075,15 +1074,7 @@ namespace Chaos
 		Particles.SetDynamicParticleSOA(Particle);
 		if (InitialState != ObjectState)
 		{
-			if (InitialState == EObjectStateType::Sleeping)
-			{
-				if (Particle->IslandIndex() != INDEX_NONE)
-				{
-					// GT has forced a wake so have to wake everything in the island
-					IslandsToWake.Enqueue(Particle->IslandIndex());
-				}
-			}
-			else if(ObjectState != EObjectStateType::Dynamic)
+			if(ObjectState != EObjectStateType::Dynamic)
 			{
 				// even though we went to sleep, we should still report info back to GT
 				Particles.MarkTransientDirtyParticle(Particle);
diff --git a/Engine/Source/Runtime/Experimental/Chaos/Private/Chaos/PBDRigidsEvolutionGBF.cpp b/Engine/Source/Runtime/Experimental/Chaos/Private/Chaos/PBDRigidsEvolutionGBF.cpp
--- a/Engine/Source/Runtime/Experimental/Chaos/Private/Chaos/PBDRigidsEvolutionGBF.cpp
+++ b/Engine/Source/Runtime/Experimental/Chaos/Private/Chaos/PBDRigidsEvolutionGBF.cpp
@@ -29,7 +29,7 @@
 namespace Chaos
 {
 #if !UE_BUILD_SHIPPING
-	CHAOS_API bool bPendingHierarchyDump = false;
+	CHAOS_API bool bPendingHierarchyDump = false; 
 #else
 	const bool bPendingHierarchyDump = false;
 #endif
@@ -126,6 +126,7 @@ namespace Chaos
 		DECLARE_CYCLE_STAT(TEXT("FPBDRigidsEvolutionGBF::ComputeIntermediateSpatialAcceleration"), STAT_Evolution_ComputeIntermediateSpatialAcceleration, STATGROUP_Chaos);
 		DECLARE_CYCLE_STAT(TEXT("FPBDRigidsEvolutionGBF::CreateConstraintGraph"), STAT_Evolution_CreateConstraintGraph, STATGROUP_Chaos);
 		DECLARE_CYCLE_STAT(TEXT("FPBDRigidsEvolutionGBF::CreateIslands"), STAT_Evolution_CreateIslands, STATGROUP_Chaos);
+		DECLARE_CYCLE_STAT(TEXT("FPBDRigidsEvolutionGBF::PruneCollisions"), STAT_Evolution_PruneCollisions, STATGROUP_Chaos);
 		DECLARE_CYCLE_STAT(TEXT("FPBDRigidsEvolutionGBF::UpdateAccelerationStructures"), STAT_Evolution_UpdateAccelerationStructures, STATGROUP_Chaos);
 		DECLARE_CYCLE_STAT(TEXT("FPBDRigidsEvolutionGBF::AddSleepingContacts"), STAT_Evolution_AddSleepingContacts, STATGROUP_Chaos);
 		DECLARE_CYCLE_STAT(TEXT("FPBDRigidsEvolutionGBF::PreApplyCallback"), STAT_Evolution_PreApplyCallback, STATGROUP_Chaos);
@@ -455,11 +456,19 @@ void FPBDRigidsEvolutionGBF::AdvanceOneTimeStepImpl(const FReal Dt, const FSubSt
 		CSV_SCOPED_TIMING_STAT(PhysicsVerbose, StepSolver_GraphColor);
 		SortConstraints();
 	}
+	{
+		// Once the graph is built (constraints are removed/added) we can destroy any collision that are not required
+		SCOPE_CYCLE_COUNTER(STAT_Evolution_PruneCollisions);
+		CSV_SCOPED_TIMING_STAT(PhysicsVerbose, StepSolver_PruneCollisions);
+		CollisionConstraints.GetConstraintAllocator().PruneExpiredItems();
+	}
+
 	if (PreApplyCallback != nullptr)
 	{
 		SCOPE_CYCLE_COUNTER(STAT_Evolution_PreApplyCallback);
 		PreApplyCallback();
 	}
+
 	CollisionConstraints.SetGravity(GetGravityForces().GetAcceleration());
 
 	TArray<bool> SleepedIslands;
@@ -623,6 +632,9 @@ void FPBDRigidsEvolutionGBF::AdvanceOneTimeStepImpl(const FReal Dt, const FSubSt
 
 	ParticleUpdatePosition(Particles.GetDirtyParticlesView(), Dt);
 
+	// Clean up the transient data from the constraint graph (e.g., IslandSolvers get cleared here)
+	GetConstraintGraph().EndTick();
+
 	if (bChaos_Solver_TestMode)
 	{
 		TestModeResetParticles();
@@ -901,8 +913,9 @@ FPBDRigidsEvolutionGBF::FPBDRigidsEvolutionGBF(FPBDRigidsSOAs& InParticles,THand
 
 FPBDRigidsEvolutionGBF::~FPBDRigidsEvolutionGBF()
 {
-	// Not sure we need to reset the indices since all the particles and constraints are going to be destroyed
-	//GetConstraintGraph().ResetIndices();
+	// This is really only needed to ensure proper cleanup (we verify that constraints have been removed from 
+	// the graph in the destructor). This can be optimized if it's a problem but it shouldn't be
+	GetConstraintGraph().Reset();
 }
 
 void FPBDRigidsEvolutionGBF::Serialize(FChaosArchive& Ar)
@@ -1101,6 +1114,18 @@ void FPBDRigidsEvolutionGBF::DestroyParticleCollisionsInAllocator(FGeometryParti
 	}
 }
 
+void FPBDRigidsEvolutionGBF::DestroyTransientConstraints(FGeometryParticleHandle* Particle)
+{
+	if (Particle != nullptr)
+	{
+		// Remove all the particle's collisions from the graph
+		GetConstraintGraph().RemoveParticleConstraints(Particle, CollisionConstraints.GetContainerId());
+
+		// Mark all collision constraints for destruction
+		DestroyParticleCollisionsInAllocator(Particle);
+	}
+}
+
 void FPBDRigidsEvolutionGBF::ParticleMaterialChanged(FGeometryParticleHandle* Particle)
 {
 	Particle->ParticleCollisions().VisitCollisions([this](FPBDCollisionConstraint& Collision)
diff --git a/Engine/Source/Runtime/Experimental/Chaos/Private/Chaos/PBDSuspensionConstraints.cpp b/Engine/Source/Runtime/Experimental/Chaos/Private/Chaos/PBDSuspensionConstraints.cpp
--- a/Engine/Source/Runtime/Experimental/Chaos/Private/Chaos/PBDSuspensionConstraints.cpp
+++ b/Engine/Source/Runtime/Experimental/Chaos/Private/Chaos/PBDSuspensionConstraints.cpp
@@ -49,7 +49,7 @@ namespace Chaos
 		ConcreteContainer()->SetSettings(ConstraintIndex, Settings);
 	}
 
-	TVec2<FGeometryParticleHandle*> FPBDSuspensionConstraintHandle::GetConstrainedParticles() const
+	FParticlePair FPBDSuspensionConstraintHandle::GetConstrainedParticles() const
 	{
 		return ConcreteContainer()->GetConstrainedParticles(ConstraintIndex);
 	}
diff --git a/Engine/Source/Runtime/Experimental/Chaos/Private/PBDRigidsSolver.cpp b/Engine/Source/Runtime/Experimental/Chaos/Private/PBDRigidsSolver.cpp
--- a/Engine/Source/Runtime/Experimental/Chaos/Private/PBDRigidsSolver.cpp
+++ b/Engine/Source/Runtime/Experimental/Chaos/Private/PBDRigidsSolver.cpp
@@ -775,7 +775,6 @@ namespace Chaos
 		QUICK_SCOPE_CYCLE_COUNTER(STAT_StartedSceneSimulation);
 
 		GetEvolution()->GetBroadPhase().GetIgnoreCollisionManager().PopStorageData_Internal(GetEvolution()->LatestExternalTimestampConsumed_Internal);
-		GetEvolution()->WakeIslands();
 	}
 
 	void FPBDRigidsSolver::DestroyPendingProxies_Internal()
diff --git a/Engine/Source/Runtime/Experimental/Chaos/Private/PhysicsProxy/SingleParticlePhysicsProxy.cpp b/Engine/Source/Runtime/Experimental/Chaos/Private/PhysicsProxy/SingleParticlePhysicsProxy.cpp
--- a/Engine/Source/Runtime/Experimental/Chaos/Private/PhysicsProxy/SingleParticlePhysicsProxy.cpp
+++ b/Engine/Source/Runtime/Experimental/Chaos/Private/PhysicsProxy/SingleParticlePhysicsProxy.cpp
@@ -66,12 +66,10 @@ void PushToPhysicsStateImp(const Chaos::FDirtyPropertiesManager& Manager, Chaos:
 
 		if(NewNonFrequentData)
 		{
-			Handle->SetNonFrequentData(*NewNonFrequentData);
-
-			// geometry may have changed, we need to invalidate the particle so it can be removed from caching structures in the evolution
-			// @todo(chaos): Remove collision constraints only. Invalidate particle may remove joint constraints in constraint graph. If joint constraint is persistent, this may cause issues.
+			// Geometry may have changed, we need to remove the particle and its collisions from the graph
 			Evolution.InvalidateParticle(Handle);
-			Evolution.DestroyParticleCollisionsInAllocator(Handle);
+
+			Handle->SetNonFrequentData(*NewNonFrequentData);
 		}
 
 		auto NewVelocities = bHasKinematicData ? ParticleData.FindVelocities(Manager, DataIdx) : nullptr;
diff --git a/Engine/Source/Runtime/Experimental/Chaos/Public/Chaos/Collision/CollisionConstraintAllocator.h b/Engine/Source/Runtime/Experimental/Chaos/Public/Chaos/Collision/CollisionConstraintAllocator.h
--- a/Engine/Source/Runtime/Experimental/Chaos/Public/Chaos/Collision/CollisionConstraintAllocator.h
+++ b/Engine/Source/Runtime/Experimental/Chaos/Public/Chaos/Collision/CollisionConstraintAllocator.h
@@ -152,8 +152,14 @@ namespace Chaos
 			bInCollisionDetectionPhase = false;
 
 			ProcessNewItems();
+		}
 
-			PruneExpiredItems();
+		/**
+		 * @brief Called each tick after the graph is updated to remove unused collisions
+		*/
+		void PruneExpiredItems()
+		{
+			PruneExpiredMidPhases();
 		}
 
 		/**
@@ -309,11 +315,6 @@ namespace Chaos
 			ProcessNewConstraints();
 		}
 
-		void PruneExpiredItems()
-		{
-			PruneExpiredMidPhases();
-		}
-
 		FParticlePairMidPhase* FindParticlePairMidPhaseImpl(FGeometryParticleHandle* Particle0, FGeometryParticleHandle* Particle1, FGeometryParticleHandle* SearchParticle)
 		{
 			// Find the existing midphase from one of the particle's lists of midphases
@@ -380,12 +381,19 @@ namespace Chaos
 			for (int32 Index = 0; Index < NumParticlePairMidPhases; ++Index)
 			{
 				TUniquePtr<FParticlePairMidPhase>& MidPhase = ParticlePairMidPhases[Index];
+
+				// We could also add !MidPhase->IsInConstraintGraph() here, but we know that we will not be in the graph if we were
+				// not active this tick and were not asleep. The constraint graph ejects all non-sleeping constraints each tick.
+				// (There is a check in the collision destructor that verified this).
 				if ((MidPhase != nullptr) && !MidPhase->IsUsedSince(CurrentEpoch) && !MidPhase->IsSleeping())
 				{
+					// Remove from the particles' lists of contacts
 					DetachParticlePairMidPhase(MidPhase.Get());
 
-					// This array can get very large so we allow it to shrink when possible
-					constexpr bool bAllowShrink = true;
+					// Destroy the midphase. ParticlePairMidPhases can get large so we allow it to shrink from time to time
+					const int32 MaxSlack = 1000;
+					const int32 Slack = ParticlePairMidPhases.Max() - ParticlePairMidPhases.Num();
+					const bool bAllowShrink = (Slack > MaxSlack);
 					ParticlePairMidPhases.RemoveAtSwap(Index, 1, bAllowShrink);
 
 					--Index;
@@ -421,9 +429,6 @@ namespace Chaos
 			}
 
 			Cookie.LastUsedEpoch = CurrentEpoch;
-
-			// An opportunity for the constraint to initialize data before CCD or solver
-			CollisionConstraint->Activate();
 		}
 
 		// All of the overlapping particle pairs in the scene
@@ -460,12 +465,14 @@ namespace Chaos
 	///////////////////////////////////////////////////////////////////////////////////////////////
 
 	template<typename TLambda>
-	inline ECollisionVisitorResult FMultiShapePairCollisionDetector::VisitCollisions(const int32 LastEpoch, const TLambda& Visitor)
+	inline ECollisionVisitorResult FMultiShapePairCollisionDetector::VisitCollisions(const int32 LastEpoch, const TLambda& Visitor, const bool bOnlyActive)
 	{
 		for (auto& KVP : Constraints)
 		{
 			const TUniquePtr<FPBDCollisionConstraint>& Constraint = KVP.Value;
-			if (Constraint->GetContainerCookie().LastUsedEpoch >= LastEpoch)
+
+			// If we only want active constraints, check the timestamp
+			if ((Constraint != nullptr) && (!bOnlyActive || (Constraint->GetContainerCookie().LastUsedEpoch >= LastEpoch)))
 			{
 				if (Visitor(*Constraint) == ECollisionVisitorResult::Stop)
 				{
@@ -477,14 +484,16 @@ namespace Chaos
 	}
 
 	template<typename TLambda>
-	inline ECollisionVisitorResult FMultiShapePairCollisionDetector::VisitConstCollisions(const int32 LastEpoch, const TLambda& Visitor) const
+	inline ECollisionVisitorResult FMultiShapePairCollisionDetector::VisitConstCollisions(const int32 LastEpoch, const TLambda& Visitor, const bool bOnlyActive) const
 	{
 		for (auto& KVP : Constraints)
 		{
 			const TUniquePtr<FPBDCollisionConstraint>& Constraint = KVP.Value;
-			if (Constraint->GetContainerCookie().LastUsedEpoch >= LastEpoch)
+
+			// If we only want active constraints, check the timestamp
+			if ((Constraint != nullptr) && (!bOnlyActive || (Constraint->GetContainerCookie().LastUsedEpoch >= LastEpoch)))
 			{
-				if (Visitor(*Constraint) == ECollisionVisitorResult::Stop)
+				if (!bOnlyActive || (Visitor(*Constraint) == ECollisionVisitorResult::Stop))
 				{
 					return ECollisionVisitorResult::Stop;
 				}
@@ -494,12 +503,13 @@ namespace Chaos
 	}
 
 	template<typename TLambda>
-	inline ECollisionVisitorResult FParticlePairMidPhase::VisitCollisions(const TLambda& Visitor)
+	inline ECollisionVisitorResult FParticlePairMidPhase::VisitCollisions(const TLambda& Visitor, const bool bOnlyActive)
 	{
 		const int32 LastEpoch = IsSleeping() ? LastUsedEpoch : GetCurrentEpoch();
 		for (FSingleShapePairCollisionDetector& ShapePair : ShapePairDetectors)
 		{
-			if (ShapePair.IsUsedSince(LastEpoch))
+			// If we only want active constraints, check the timestamp
+			if ((ShapePair.GetConstraint() != nullptr) && (!bOnlyActive || ShapePair.IsUsedSince(LastEpoch)))
 			{
 				if (Visitor(*ShapePair.GetConstraint()) == ECollisionVisitorResult::Stop)
 				{
@@ -510,7 +520,7 @@ namespace Chaos
 
 		for (FMultiShapePairCollisionDetector& MultiShapePair : MultiShapePairDetectors)
 		{
-			if (MultiShapePair.VisitCollisions(LastEpoch, Visitor) == ECollisionVisitorResult::Stop)
+			if (MultiShapePair.VisitCollisions(LastEpoch, Visitor, bOnlyActive) == ECollisionVisitorResult::Stop)
 			{
 				return ECollisionVisitorResult::Stop;
 			}
@@ -521,12 +531,13 @@ namespace Chaos
 
 
 	template<typename TLambda>
-	inline ECollisionVisitorResult FParticlePairMidPhase::VisitConstCollisions(const TLambda& Visitor) const
+	inline ECollisionVisitorResult FParticlePairMidPhase::VisitConstCollisions(const TLambda& Visitor, const bool bOnlyActive) const
 	{
 		const int32 LastEpoch = IsSleeping() ? LastUsedEpoch : GetCurrentEpoch();
 		for (const FSingleShapePairCollisionDetector& ShapePair : ShapePairDetectors)
 		{
-			if (ShapePair.IsUsedSince(LastEpoch))
+			// If we only want active constraints, check the timestamp
+			if ((ShapePair.GetConstraint() != nullptr) && (!bOnlyActive || ShapePair.IsUsedSince(LastEpoch)))
 			{
 				if (Visitor(*ShapePair.GetConstraint()) == ECollisionVisitorResult::Stop)
 				{
@@ -537,7 +548,7 @@ namespace Chaos
 
 		for (const FMultiShapePairCollisionDetector& MultiShapePair : MultiShapePairDetectors)
 		{
-			if (MultiShapePair.VisitConstCollisions(LastEpoch, Visitor) == ECollisionVisitorResult::Stop)
+			if (MultiShapePair.VisitConstCollisions(LastEpoch, Visitor, bOnlyActive) == ECollisionVisitorResult::Stop)
 			{
 				return ECollisionVisitorResult::Stop;
 			}
diff --git a/Engine/Source/Runtime/Experimental/Chaos/Public/Chaos/Collision/PBDCollisionConstraint.h b/Engine/Source/Runtime/Experimental/Chaos/Public/Chaos/Collision/PBDCollisionConstraint.h
--- a/Engine/Source/Runtime/Experimental/Chaos/Public/Chaos/Collision/PBDCollisionConstraint.h
+++ b/Engine/Source/Runtime/Experimental/Chaos/Public/Chaos/Collision/PBDCollisionConstraint.h
@@ -225,7 +225,7 @@ namespace Chaos
 	 * to reduce unnecessary indirection.
 	 * 
 	*/
-	class CHAOS_API FPBDCollisionConstraint : public FPBDCollisionConstraintHandle
+	class CHAOS_API FPBDCollisionConstraint final : public FPBDCollisionConstraintHandle
 	{
 		friend class FCollisionConstraintAllocator;
 		friend class FMultiShapePairCollisionDetector;
@@ -294,6 +294,7 @@ namespace Chaos
 		static FPBDCollisionConstraint MakeCopy(const FPBDCollisionConstraint& Source);
 
 		FPBDCollisionConstraint();
+		virtual ~FPBDCollisionConstraint();
 
 		/**
 		 * @brief The current CCD state of this constraint
@@ -335,12 +336,9 @@ namespace Chaos
 		// API
 		//
 
-		FGeometryParticleHandle* GetParticle0() { return Particle[0]; }
-		const FGeometryParticleHandle* GetParticle0() const { return Particle[0]; }
-		FGeometryParticleHandle* GetParticle1() { return Particle[1]; }
-		const FGeometryParticleHandle* GetParticle1() const { return Particle[1]; }
-		FGeometryParticleHandle* GetParticle(const int32 ParticleIndex) { check((ParticleIndex >= 0) && (ParticleIndex < 2)); return Particle[ParticleIndex]; }
-		const FGeometryParticleHandle* GetParticle(const int32 ParticleIndex) const { check((ParticleIndex >= 0) && (ParticleIndex < 2)); return Particle[ParticleIndex]; }
+		FGeometryParticleHandle* GetParticle0() const { return Particle[0]; }
+		FGeometryParticleHandle* GetParticle1() const { return Particle[1]; }
+		FGeometryParticleHandle* GetParticle(const int32 ParticleIndex) const { check((ParticleIndex >= 0) && (ParticleIndex < 2)); return Particle[ParticleIndex]; }
 
 		const FImplicitObject* GetImplicit0() const { return Implicit[0]; }
 		const FImplicitObject* GetImplicit1() const { return Implicit[1]; }
@@ -381,7 +379,9 @@ namespace Chaos
 		void SetIsProbe(bool bInProbe) { Flags.bIsProbe = bInProbe; }
 		bool GetIsProbe() const { return Flags.bIsProbe; }
 
-		virtual void SetIsSleeping(const bool bInIsSleeping) override;
+		virtual bool SupportsSleeping() const override final { return true; }
+		virtual bool IsSleeping() const override final;
+		virtual void SetIsSleeping(const bool bInIsSleeping) override final;
 
 		// Get the world-space normal of the closest manifold point
 		// @todo(chaos): remove (used by legacy RBAN collision solver)
diff --git a/Engine/Source/Runtime/Experimental/Chaos/Public/Chaos/Collision/PBDCollisionConstraintHandle.h b/Engine/Source/Runtime/Experimental/Chaos/Public/Chaos/Collision/PBDCollisionConstraintHandle.h
--- a/Engine/Source/Runtime/Experimental/Chaos/Public/Chaos/Collision/PBDCollisionConstraintHandle.h
+++ b/Engine/Source/Runtime/Experimental/Chaos/Public/Chaos/Collision/PBDCollisionConstraintHandle.h
@@ -44,6 +44,7 @@ namespace Chaos
 	/**
 	 * @brief A handle to a contact constraint.
 	 * @note This is an intrusive handle, so you can use a contact pointer as a handle.
+	 * @todo(chaos): remove this class - it can just be a "using FPBDCollisionConstraintHandle = FPBDCollisionConstraint"
 	*/
 	class CHAOS_API FPBDCollisionConstraintHandle : public TIntrusiveConstraintHandle<FPBDCollisionConstraint>
 	{
@@ -85,9 +86,8 @@ namespace Chaos
 
 		FVec3 GetAccumulatedImpulse() const;
 
-		TVector<const TGeometryParticleHandle<FReal, 3>*, 2> GetConstrainedParticles() const;
-
-		TVector<TGeometryParticleHandle<FReal, 3>*, 2> GetConstrainedParticles();
+		// Declared final so that TPBDConstraintGraphRuleImpl::AddToGraph() does not need to hit vtable
+		virtual FParticlePair GetConstrainedParticles() const override final;
 
 		void PreGatherInput(const FReal Dt, FPBDIslandSolverData& SolverData);
 		void GatherInput(FReal Dt, const int32 Particle0Level, const int32 Particle1Level, FPBDIslandSolverData& SolverData);
diff --git a/Engine/Source/Runtime/Experimental/Chaos/Public/Chaos/Collision/ParticlePairBroadPhase.h b/Engine/Source/Runtime/Experimental/Chaos/Public/Chaos/Collision/ParticlePairBroadPhase.h
--- a/Engine/Source/Runtime/Experimental/Chaos/Public/Chaos/Collision/ParticlePairBroadPhase.h
+++ b/Engine/Source/Runtime/Experimental/Chaos/Public/Chaos/Collision/ParticlePairBroadPhase.h
@@ -20,7 +20,6 @@ namespace Chaos
 	{
 	public:
 		using FParticleHandle = TGeometryParticleHandle<FReal, 3>;
-		using FParticlePair = TVector<FParticleHandle*, 2>;
 		using FAABB = FAABB3;
 
 		FBasicBroadPhase(const TArray<FParticlePair>* InParticlePairs, const TArray<FParticleHandle*>* InParticlesA, const TArray<FParticleHandle*>* InParticlesB)
diff --git a/Engine/Source/Runtime/Experimental/Chaos/Public/Chaos/Collision/ParticlePairMidPhase.h b/Engine/Source/Runtime/Experimental/Chaos/Public/Chaos/Collision/ParticlePairMidPhase.h
--- a/Engine/Source/Runtime/Experimental/Chaos/Public/Chaos/Collision/ParticlePairMidPhase.h
+++ b/Engine/Source/Runtime/Experimental/Chaos/Public/Chaos/Collision/ParticlePairMidPhase.h
@@ -220,13 +220,13 @@ namespace Chaos
 		 * @brief Visit all the collisions
 		*/
 		template<typename TLambda>
-		ECollisionVisitorResult VisitCollisions(const int32 LastEpoch, const TLambda& Visitor);
+		ECollisionVisitorResult VisitCollisions(const int32 LastEpoch, const TLambda& Visitor, const bool bOnlyActive = true);
 
 		/**
 		 * @brief Visit all the collisions
 		*/
 		template<typename TLambda>
-		ECollisionVisitorResult VisitConstCollisions(const int32 LastEpoch, const TLambda& Visitor) const;
+		ECollisionVisitorResult VisitConstCollisions(const int32 LastEpoch, const TLambda& Visitor, const bool bOnlyActive = true) const;
 
 	private:
 		FPBDCollisionConstraint* FindConstraint(const FCollisionParticlePairConstraintKey& Key);
@@ -335,6 +335,8 @@ namespace Chaos
 		*/
 		void SetIsSleeping(const bool bInIsSleeping);
 
+		bool IsInConstraintGraph() const;
+
 		/**
 		 * @brief Destroy all collisions and prevent this midphasae from being used any more. Called when one of its particles is destoyed.
 		 * It will be culled at the next Prune in the CollisionConstraintAllocator. We don't delete it immediately so that we don't
@@ -370,14 +372,14 @@ namespace Chaos
 		 * but are still kept around as an optimization.
 		*/
 		template<typename TLambda>
-		ECollisionVisitorResult VisitCollisions(const TLambda& Visitor);
+		ECollisionVisitorResult VisitCollisions(const TLambda& Visitor, const bool bOnlyActive = true);
 
 
 		/**
 		 * @brief Call a lambda on each active collision constraint
 		 */
 		template<typename TLambda>
-		ECollisionVisitorResult VisitConstCollisions(const TLambda& Visitor) const;
+		ECollisionVisitorResult VisitConstCollisions(const TLambda& Visitor, const bool bOnlyActive = true) const;
 
 		/**
 		 * @brief Cookie for use by FParticleCollisions
diff --git a/Engine/Source/Runtime/Experimental/Chaos/Public/Chaos/ConstraintHandle.h b/Engine/Source/Runtime/Experimental/Chaos/Public/Chaos/ConstraintHandle.h
--- a/Engine/Source/Runtime/Experimental/Chaos/Public/Chaos/ConstraintHandle.h
+++ b/Engine/Source/Runtime/Experimental/Chaos/Public/Chaos/ConstraintHandle.h
@@ -4,18 +4,33 @@
 #include "Chaos/Core.h"
 #include "Chaos/Declares.h"
 #include "Chaos/Vector.h"
+#include "Chaos/ParticleHandleFwd.h"
 
-#if UE_BUILD_DEBUG
-#define CHAOS_CONSTRAINTHANDLE_DEBUG_ENABLED  1
-#else
+// Set to 1 to help trap erros in the constraint management which typically manifests as a dangling constraint handle in the IslandManager
+// 
+// WARNING: Do not submit with either of these set to 1 !!
+//
+#ifndef CHAOS_CONSTRAINTHANDLE_DEBUG_ENABLED
 #define CHAOS_CONSTRAINTHANDLE_DEBUG_ENABLED  0
 #endif
+#ifndef CHAOS_CONSTRAINTHANDLE_DEBUG_DETAILED_ENABLED
+#define CHAOS_CONSTRAINTHANDLE_DEBUG_DETAILED_ENABLED (CHAOS_CONSTRAINTHANDLE_DEBUG_ENABLED && 0)
+#endif
+
+// This is a proxy for not allowing the above to be checked in - CIS should fail if it is left enabled
+#if UE_BUILD_SHIPPING || UE_BUILD_TEST
+static_assert(CHAOS_CONSTRAINTHANDLE_DEBUG_ENABLED == 0, "CHAOS_CONSTRAINTHANDLE_DEBUG_ENABLED should be 0");
+static_assert(CHAOS_CONSTRAINTHANDLE_DEBUG_DETAILED_ENABLED == 0, "CHAOS_CONSTRAINTHANDLE_DEBUG_DETAILED_ENABLED should be 0");
+#endif
 
 namespace Chaos
 {
 	class FPBDConstraintContainer;
 	class FPBDIndexedConstraintContainer;
 
+	using FParticlePair = TVec2<FGeometryParticleHandle*>;
+	using FConstParticlePair = TVec2<const FGeometryParticleHandle*>;
+
 	/**
 	 * @brief A type id for constraint handles to support safe up/down casting (including intermediate classes in the hierrachy)
 	 *
@@ -125,12 +140,17 @@ namespace Chaos
 			return (GraphIndex != INDEX_NONE);
 		}
 
+		virtual TVec2<FGeometryParticleHandle*> GetConstrainedParticles() const = 0;
+
 		virtual void SetEnabled(bool InEnabled) = 0;
 
 		virtual bool IsEnabled() const = 0;
 
 		virtual bool IsProbe() const { return false; }
 
+		// Does this constraint have the concept of sleep? (only really used for debug validation)
+		virtual bool SupportsSleeping() const { return false; }
+
 		virtual bool IsSleeping() const { return false; }
 		virtual void SetIsSleeping(const bool bInIsSleeping) {}
 		
@@ -358,30 +378,23 @@ namespace Chaos
 	 * @brief A debugging utility for tracking down dangling constraint issues
 	 * This acts as a FConstraintHandle*, but caches some extra debug data useful in tracking
 	 * down dangling pointer issues when they arise.
-	 * @todo(chaos): improve constraint lifetime management so that we don't get these problems!
 	*/
 	class CHAOS_API FConstraintHandleHolder
 	{
 	public:
 		FConstraintHandleHolder()
 			: Handle(nullptr)
+		{
 #if CHAOS_CONSTRAINTHANDLE_DEBUG_ENABLED
-			, ConstraintType(nullptr)
+			InitDebugData();
 #endif
-		{
 		}
 
 		FConstraintHandleHolder(FConstraintHandle* InHandle)
 			: Handle(InHandle)
-#if CHAOS_CONSTRAINTHANDLE_DEBUG_ENABLED
-			, ConstraintType(nullptr)
-#endif
 		{
 #if CHAOS_CONSTRAINTHANDLE_DEBUG_ENABLED
-			if (Handle != nullptr)
-			{
-				ConstraintType = &InHandle->GetType();
-			}
+			InitDebugData();
 #endif
 		}
 
@@ -396,8 +409,17 @@ namespace Chaos
 
 	private:
 		FConstraintHandle* Handle;
+
 #if CHAOS_CONSTRAINTHANDLE_DEBUG_ENABLED
+	public:
+		const FGeometryParticleHandle* GetParticle0() const { return Particles[0]; }
+		const FGeometryParticleHandle* GetParticle1() const { return Particles[1]; }
+
+	private:
+		void InitDebugData();
+
 		const FConstraintHandleTypeID* ConstraintType;
+		const FGeometryParticleHandle* Particles[2];
 #endif
 	};
 }
diff --git a/Engine/Source/Runtime/Experimental/Chaos/Public/Chaos/Island/IslandGraph.h b/Engine/Source/Runtime/Experimental/Chaos/Public/Chaos/Island/IslandGraph.h
--- a/Engine/Source/Runtime/Experimental/Chaos/Public/Chaos/Island/IslandGraph.h
+++ b/Engine/Source/Runtime/Experimental/Chaos/Public/Chaos/Island/IslandGraph.h
@@ -16,23 +16,34 @@ namespace Chaos
 		// @todo(chaos): should we really have a sparse array here? The implementation probably wastes a lot of memory
 		TSparseArray<int32> NodeEdges;
 
-		/** Node Island Index (for static/kinematic particles could belong to several islands) */
+		/**
+		 * Node Island Index. Only used by valid nodes (Dynamic/Sleeping Particles).
+		 * Static/kinematic particles could belong to several islands and this is always INDEX_NONE.
+		 */
 		int32 IslandIndex = INDEX_NONE;
 
-		/** List of islands in which the node is referenced */
-		// @todo(chaos): should we really have a set here? The implementation probably wastes a lot of memory
+		/**
+		 * List of islands in which the node is referenced. Should only be 1 item for valied nodes(
+		 * Dynamic/Sleeping Particles), and may be more if invalid (Static/Kinematic).
+		 * @note this array is populated externally (see see PopulateIslands in IslandManager.cpp)
+		 * @todo(chaos): should we really have a set here? The implementation probably wastes a lot of memory
+		 */
 		TSet<int32> NodeIslands;
 
-		/** Check if a node is valid (checked for graph partitioning) */
+		/** Check if a node is valid (checked for graph partitioning into islands).
+		 * In the context of particles (nodes) and constraints (edges), bValidNode is true when
+		 * a Particle is Dynamic or Sleeping, but not when Static or Kinematic. This is because we
+		 * do not need to merge islands that are only connected through a Static/Kinematic Particle node.
+		 */
 		bool bValidNode = true;
 
 		/** Check if a node is steady */
 		bool bStationaryNode = true;
 
-		/** Node counter to filter nodes already processed */
+		/** Node counter to filter nodes already processed (@see SplitIsland) */
 		int32 NodeCounter = 0;
 
-		/** Node item that is stored per node */
+		/** Node item that is stored per node. This is a ParticleHandle. */
 		NodeType NodeItem;
 
 		/** Node level index */
@@ -61,7 +72,10 @@ namespace Chaos
 		/** Unique edge island index */
 		int32 IslandIndex = INDEX_NONE;
 
-		/** Edge counter to filter edges already processed */
+		/** True if one or both nodes are valid (dynamic) */
+		bool bValidEdge = true;
+
+		/** Edge counter to filter edges already processed (@see SplitIsland) */
 		int32 EdgeCounter = 0;
 
 		/** Edge item that is stored per node */
@@ -90,7 +104,7 @@ namespace Chaos
 		/** Island counter to filter islands already processed */
 		int32 IslandCounter = 0;
 
-		/** Boolean to check if an island is persistent or not */
+		/** Boolean to check if an island is persistent or not. A persistent island is one where no particles or constraints were added or removed */
 		bool bIsPersistent = true;
 
 		/** Boolean to check if an island is sleeping or not */
@@ -119,10 +133,10 @@ namespace Chaos
  * This template implements a graph that will be stored on the island manager
  * The goal here is to minimize memory allocation while doing graph operations. 
  *
- * @param NodeType The node type of the item that will be stored in each nodes
- * @param EdgeType The edge type of the item that will be stored in each edges
+ * @param NodeType The node type of the item that will be stored in each node (e.g., ParticleHandle)
+ * @param EdgeType The edge type of the item that will be stored in each edge (e.g., ConstraintHandle)
  */
-template<typename NodeType, typename EdgeType, typename IslandType>
+template<typename NodeType, typename EdgeType, typename IslandType, typename OwnerType>
 class CHAOS_API FIslandGraph
 {
 public:
@@ -131,9 +145,18 @@ public:
 	using FGraphNode = TIslandGraphNode<NodeType>;
 	using FGraphEdge = TIslandGraphEdge<EdgeType>;
 	using FGraphIsland = TIslandGraphIsland<IslandType>;
+	using FGraphOwner = OwnerType;
 
 	static constexpr const int32 MaxCount = 100000;
 
+	/**
+	 * @brief Set the graph owner. This will allow the owner to cache node and edge indices with the node/edge items
+	*/
+	void SetOwner(FGraphOwner* InOwner)
+	{
+		Owner = InOwner;
+	}
+
 	/**
 	 * Update current graph to merge connected islands 
 	 */
@@ -146,14 +169,14 @@ public:
 	void ReserveNodes(const int32 NumNodes);
 
 	/**
-	 * Clear the graph nodes without deallocating memory
+	 * Remove all nodes (and attached edges - so all of them) from the graph
 	 */
 	void ResetNodes();
 
 	/**
 	 * Given a node item, add a node to the graph nodes
 	 * @param NodeItem Node Item to be stored in the graph node
-	 * @param bValidNode Check if a node should be considered for graph partitioning
+	 * @param bValidNode Check if a node should be considered for graph partitioning (into islands)
 	 * @param IslandIndex Potential island index we want the node to belong to
 	 * @param bStationaryNode Boolean to check if a node is steady or not
 	 * @return Node index that has just been added
@@ -163,7 +186,7 @@ public:
 	/**
 	* Given a node item, update the graph node information (valid,discard...)
 	* @param NodeItem Node Item to be stored in the graph node
-	* @param bValidNode Check if a node should be considered for graph partitioning
+	* @param bValidNode Check if a node should be considered for graph partitioning (into islands)
 	* @param IslandIndex Potential island index we want the node to belong to
 	* @param bStationaryNode Boolean to check if a node is steady or not
 	*  @param NodeIndex Index to consider to update the graph node information
@@ -204,6 +227,8 @@ public:
 	 */
 	void RemoveEdge(const int32 EdgeIndex);
 
+	void UpdateEdge(const int32 EdgeIndex);
+
 	/**
 	* Remove all the edges from the graph
 	*/
@@ -215,22 +240,19 @@ public:
 	void InitIslands();
 
 	/**
-	 * Get the number of edges that are inside a graph
-	 * @return Number of graph edges
+	 * Get the maximum number of edges that are inside a graph (SparseArray size)
 	 */
-	FORCEINLINE int32 NumEdges() const { return GraphEdges.GetMaxIndex(); }
+	FORCEINLINE int32 MaxNumEdges() const { return GraphEdges.GetMaxIndex(); }
 
 	/**
-	 * Get the number of nodes that are inside a graph
-	 * @return Number of graph nodes
+	 * Get the maximum number of nodes that are inside a graph (SparseArray size)
 	 */
-	FORCEINLINE int32 NumNodes() const { return GraphNodes.GetMaxIndex(); }
+	FORCEINLINE int32 MaxNumNodes() const { return GraphNodes.GetMaxIndex(); }
 
 	/**
-	 * Get the number of islands that are inside a graph
-	 * @return Number of graph islands
+	 * Get the maximum number of islands that are inside a graph (SparseArray size)
 	 */
-	FORCEINLINE int32 NumIslands() const { return GraphIslands.GetMaxIndex(); }
+	FORCEINLINE int32 MaxNumIslands() const { return GraphIslands.GetMaxIndex(); }
 
 	/**
 	 * Link two islands to each other for future island traversal
@@ -335,7 +357,9 @@ public:
 		}
 		return false;
 	}
-	
+
+	FGraphOwner* Owner = nullptr;
+
 	/** List of graph nodes */
 	TSparseArray<FGraphNode> GraphNodes;
 
@@ -358,4 +382,16 @@ public:
 	TDeque<int32> NodeQueue;
 };
 
+// For use with unit tests to pretend to be the owner of the graph
+template<typename NodeType, typename EdgeType>
+struct TNullIslandGraphOwner
+{
+	void GraphNodeAdded(const NodeType& NodeItem, const int32 NodeIndex) {}
+	void GraphNodeRemoved(const NodeType& NodeItem) {}
+
+	void GraphEdgeAdded(const EdgeType& EdgeItem, const int32 EdgeIndex) {}
+	void GraphEdgeRemoved(const EdgeType& EdgeItem) {}
+};
+
+
 }
\ No newline at end of file
diff --git a/Engine/Source/Runtime/Experimental/Chaos/Public/Chaos/Island/IslandManager.h b/Engine/Source/Runtime/Experimental/Chaos/Public/Chaos/Island/IslandManager.h
--- a/Engine/Source/Runtime/Experimental/Chaos/Public/Chaos/Island/IslandManager.h
+++ b/Engine/Source/Runtime/Experimental/Chaos/Public/Chaos/Island/IslandManager.h
@@ -32,15 +32,15 @@ class FChaosPhysicsMaterial;
 
 using FPBDRigidParticles = TPBDRigidParticles<FReal, 3>;
 
-
 /** Island manager responsible to create the list of solver islands that will be persistent over time */
 class CHAOS_API FPBDIslandManager
 {
 public:
-	using GraphType = FIslandGraph<FGeometryParticleHandle*, FConstraintHandleHolder, FPBDIslandSolver*>;
+	using GraphType = FIslandGraph<FGeometryParticleHandle*, FConstraintHandleHolder, FPBDIslandSolver*, FPBDIslandManager>;
 	using FGraphNode = GraphType::FGraphNode;
 	using FGraphEdge = GraphType::FGraphEdge;
-	
+	using FGraphIsland = GraphType::FGraphIsland;
+
 	/**
 	  * Default island manager constructor
 	  */
@@ -54,7 +54,17 @@ public:
 	/**
 	* Reset nodes and edges graph indices 
 	*/
-	void ResetIndices();
+	//void ResetIndices();
+
+	/**
+	 * Clear all nodes and edges from the graph
+	*/
+	void Reset();
+
+	/**
+	 * Remove all particles (and their constraints) from the graph
+	*/
+	void RemoveParticles();
 
 	/**
 	* Remove all the constraints from the graph
@@ -74,6 +84,11 @@ public:
 	*/
 	void InitializeGraph(const TParticleView<FPBDRigidParticles>& Particles);
 
+	/**
+	 * @brief Called at the end of the frame to perform cleanup
+	*/
+	void EndTick();
+
 	/**
 	  * Reserve space in the graph for NumConstraints additional constraints.
 	  * @param NumConstraints Number of constraints to be used to reserved memory space
@@ -87,7 +102,7 @@ public:
 	  * @param ConstrainedParticles List of 2 particles handles that are used within the constraint
 	  * @return Edge index within the sparse graph edges list
 	  */
-	int32 AddConstraint(const uint32 ContainerId, FConstraintHandle* ConstraintHandle, const TVector<FGeometryParticleHandle*, 2>& ConstrainedParticles);
+	void AddConstraint(const uint32 ContainerId, FConstraintHandle* ConstraintHandle, const TVector<FGeometryParticleHandle*, 2>& ConstrainedParticles);
 
 	/**
 	  * Remove a constraint from the graph
@@ -96,6 +111,13 @@ public:
 	  * @param ConstrainedParticles List of 2 particles handles that are used within the constraint
 	  */
 	void RemoveConstraint(const uint32 ContainerId, FConstraintHandle* ConstraintHandle);
+
+	/**
+	 * @brief Remove all the constraints of the specified type (ContainerId) belonging to the Particle
+	 * @param Particle the particle to remove the constraints from
+	 * @param ContainerId Container id the constraint is belonging to
+	*/
+	void RemoveParticleConstraints(FGeometryParticleHandle* Particle, const uint32 ContainerId);
 	
 	/**
 	  * Preallocate buffers for \p Num particles.
@@ -174,6 +196,7 @@ public:
 	  * Get the island particles stored in the solver island
 	  * @param IslandIndex Island index
 	  * @return List of particles handles in the island
+	  * @note This will be empty for sleeping islands, though the particles will still be in the graph
 	  */
 	const TArray<FGeometryParticleHandle*>& GetIslandParticles(const int32 IslandIndex) const;
 
@@ -181,9 +204,18 @@ public:
 	  * Get the island constraints stored in the solver island
 	  * @param IslandIndex Island index
 	  * @return List of constraints handles in the island
+	  * @note This will be empty for sleeping islands, though the constraints will still be in the graph
 	  */
 	const TArray<FConstraintHandleHolder>& GetIslandConstraints(const int32 IslandIndex) const;
 
+	/**
+	 * Get the set of islands that the particle belongs to. Will be 0-1 island for dynamics and could be any number for non-dynamics.
+	 * This works for sleeping and waking particles (but we must iterate over constraints for this to work).
+	 */
+	TArray<int32> FindParticleIslands(const FGeometryParticleHandle* ParticleHandle) const;
+
+	int32 GetConstraintIsland(const FConstraintHandle* ConstraintHandle) const;
+
 	/**
 	 * When resim is used, tells us whether we need to resolve island
 	 * @param IslandIndex Island index
@@ -259,18 +291,6 @@ public:
 		}
 	}
 
-	// /** Add a constraint container to all the solver island given a container id
-	// * @param ContainerId Constraints container id from which the solver constraint containers are being built
-	// */
-	// template<typename ConstraintType>
-	// void GatherSolverInput(const FReal Dt, int32 IslandIndex, const int32 ContainerId)
-	// {
-	// 	if(IslandIndexing.IsValidIndex(IslandIndex))
-	// 	{
-	// 		IslandSolvers[IslandIndexing[IslandIndex]]->GatherSolverInput<ConstraintType>(Dt, IslandIndex, ContainerId);
-	// 	}
-	// }
-
 	/** Accessors of the graph group */
 	const FPBDIslandGroup* GetIslandGroup(const int32 GroupIndex) const {return IslandGroups.IsValidIndex(GroupIndex) ? IslandGroups[GroupIndex].Get() : nullptr; }
 	FPBDIslandGroup* GetIslandGroup(const int32 GroupIndex) {return IslandGroups.IsValidIndex(GroupIndex) ? IslandGroups[GroupIndex].Get()  : nullptr; }
@@ -298,9 +318,28 @@ public:
 	
 	/** Get the constraints edges */
 	const TMap<FConstraintHandleHolder, int32>& GetConstraintEdges() const { return IslandGraph->ItemEdges; }
-	
+
+	// Graph Owner API Implementation
+	void GraphNodeAdded(FGeometryParticleHandle* NodeItem, const int32 NodeIndex);
+	void GraphNodeRemoved(FGeometryParticleHandle* NodeItem);
+	void GraphEdgeAdded(const FConstraintHandleHolder& EdgeItem, const int32 EdgeIndex);
+	void GraphEdgeRemoved(const FConstraintHandleHolder& EdgeItem);
+
+#if CHAOS_CONSTRAINTHANDLE_DEBUG_ENABLED
+	bool DebugCheckParticleNotInGraph(const FGeometryParticleHandle* ParticleHandle) const;
+	bool DebugParticleConstraintsNotInGraph(const FGeometryParticleHandle* Particle, const int32 ContainerId) const;
+	bool DebugCheckConstraintNotInGraph(const FConstraintHandle* ConstraintHandle) const;
+#endif
+
 protected:
 
+#if CHAOS_CONSTRAINTHANDLE_DEBUG_ENABLED
+	bool DebugCheckIslandConstraints() const;
+	void DebugCheckParticleIslands(const FGeometryParticleHandle* ParticleHandle) const;
+	void DebugCheckConstraintIslands(const FConstraintHandle* ConstraintHandle, const int32 IslandIndex) const;
+	void DebugCheckIslands() const;
+#endif
+
 	/**
 	* Initialize the groups according to the number of threads
 	*/
@@ -316,6 +355,12 @@ protected:
 	*/
 	void SyncIslands(FPBDRigidsSOAs& Particles, const int32 NumContainers);
 
+
+	/**
+	* Set the island and all particles' sleep state (called from the graph update).
+	*/
+	void SetIslandSleeping(FPBDRigidsSOAs& Particles, const int32 IslandIndex, const bool bIsSleeping);
+
 	void ValidateIslands() const;
 
 	/** List of solver islands in the manager */
@@ -331,9 +376,15 @@ protected:
 	int32									MaxParticleIndex = INDEX_NONE;
 	
 	/** List of island groups in the manager */
-	TArray<TUniquePtr<FPBDIslandGroup>> IslandGroups;
+	TArray<TUniquePtr<FPBDIslandGroup>>		IslandGroups;
 
 	/** Sorted list of islands */
 	TArray<int32>							SortedIslands;
+
+	/** Islands are only populated with constraints and particles during the solver phase. Trying to access them outside that is an error we can trap */
+	bool									bIslandsPopulated;
+
+	/** Make sure EndTick is called. Sanity check mostlyhelpful for unit tests */
+	bool									bEndTickCalled;
 };
 }
\ No newline at end of file
diff --git a/Engine/Source/Runtime/Experimental/Chaos/Public/Chaos/Island/SolverIsland.h b/Engine/Source/Runtime/Experimental/Chaos/Public/Chaos/Island/SolverIsland.h
--- a/Engine/Source/Runtime/Experimental/Chaos/Public/Chaos/Island/SolverIsland.h
+++ b/Engine/Source/Runtime/Experimental/Chaos/Public/Chaos/Island/SolverIsland.h
@@ -77,11 +77,6 @@ class CHAOS_API FPBDIslandSolver : public FPBDIslandSolverData
 	*/
 	void RemoveConstraint(FConstraintHandle* ConstraintHandle);
 
-	/**
-	* Reset all the island constraint graph index
-	*/
-	void ResetIndices();
-
 	/**
 	* Sort the islands constraints
 	*/
@@ -121,7 +116,8 @@ class CHAOS_API FPBDIslandSolver : public FPBDIslandSolverData
 	* Members accessors
 	*/
 	FORCEINLINE bool IsSleeping() const {return bIsSleeping;}
-	FORCEINLINE void SetIsSleeping(const bool bIsSleepingIn ) { bSleepingChanged = (bIsSleeping != bIsSleepingIn); bIsSleeping = bIsSleepingIn; }
+	FORCEINLINE void SetIsSleeping(const bool bIsSleepingIn ) { bSleepingChanged = bSleepingChanged || (bIsSleeping != bIsSleepingIn); bIsSleeping = bIsSleepingIn; }
+	FORCEINLINE void ResetSleepingChanged() { bSleepingChanged = false; }
 	FORCEINLINE bool IsPersistent() const { return bIsPersistent; }
 	FORCEINLINE void SetIsPersistent(const bool bIsPersistentIn) { bIsPersistent = bIsPersistentIn; }
 	FORCEINLINE bool NeedsResim() const { return bNeedsResim; }
diff --git a/Engine/Source/Runtime/Experimental/Chaos/Public/Chaos/PBDCollisionConstraints.h b/Engine/Source/Runtime/Experimental/Chaos/Public/Chaos/PBDCollisionConstraints.h
--- a/Engine/Source/Runtime/Experimental/Chaos/Public/Chaos/PBDCollisionConstraints.h
+++ b/Engine/Source/Runtime/Experimental/Chaos/Public/Chaos/PBDCollisionConstraints.h
@@ -5,6 +5,7 @@
 #include "Chaos/Collision/CollisionApplyType.h"
 #include "Chaos/Collision/CollisionConstraintAllocator.h"
 #include "Chaos/Collision/CollisionContext.h"
+#include "Chaos/Collision/PBDCollisionConstraint.h"
 #include "Chaos/Collision/PBDCollisionConstraintHandle.h"
 #include "Chaos/Collision/SolverCollisionContainer.h"
 #include "Chaos/PBDConstraintContainer.h"
@@ -377,4 +378,81 @@ private:
 
 	EConstraintSolverType SolverType;
 };
+
+//
+//
+// Inlined FPBDCollisionConstraintHandle functions. Here to avoid circular deps
+//
+//
+
+inline const FPBDCollisionConstraints* FPBDCollisionConstraintHandle::ConcreteContainer() const
+{
+	return static_cast<FPBDCollisionConstraints*>(ConstraintContainer);
+}
+
+inline FPBDCollisionConstraints* FPBDCollisionConstraintHandle::ConcreteContainer()
+{
+	return static_cast<FPBDCollisionConstraints*>(ConstraintContainer);
+}
+
+inline const FPBDCollisionConstraint& FPBDCollisionConstraintHandle::GetContact() const
+{
+	return *GetConstraint();
+}
+
+inline FPBDCollisionConstraint& FPBDCollisionConstraintHandle::GetContact()
+{
+	return *GetConstraint();
+}
+
+inline ECollisionCCDType FPBDCollisionConstraintHandle::GetCCDType() const
+{
+	return GetContact().GetCCDType();
+}
+
+inline void FPBDCollisionConstraintHandle::SetEnabled(bool InEnabled)
+{
+	GetContact().SetDisabled(!InEnabled);
+}
+
+inline bool FPBDCollisionConstraintHandle::IsEnabled() const
+{
+	return !GetContact().GetDisabled();
+}
+
+inline bool FPBDCollisionConstraintHandle::IsProbe() const
+{
+	return GetContact().GetIsProbe();
+}
+
+inline FVec3 FPBDCollisionConstraintHandle::GetAccumulatedImpulse() const
+{
+	return GetContact().AccumulatedImpulse;
+}
+
+inline FParticlePair FPBDCollisionConstraintHandle::GetConstrainedParticles() const
+{
+	return { GetContact().GetParticle0(), GetContact().GetParticle1() };
+}
+
+inline void FPBDCollisionConstraintHandle::PreGatherInput(const FReal Dt, FPBDIslandSolverData& SolverData)
+{
+	ConcreteContainer()->PreGatherInput(Dt, GetContact(), SolverData);
+}
+
+inline void FPBDCollisionConstraintHandle::GatherInput(FReal Dt, const int32 Particle0Level, const int32 Particle1Level, FPBDIslandSolverData& SolverData)
+{
+	ConcreteContainer()->GatherInput(Dt, GetContact(), Particle0Level, Particle1Level, SolverData);
+}
+
+inline FSolverBody* FPBDCollisionConstraintHandle::GetSolverBody0()
+{
+	return GetContact().GetSolverBody0();
+}
+
+inline FSolverBody* FPBDCollisionConstraintHandle::GetSolverBody1()
+{
+	return GetContact().GetSolverBody1();
+}
+
 }
diff --git a/Engine/Source/Runtime/Experimental/Chaos/Public/Chaos/PBDConstraintRule.h b/Engine/Source/Runtime/Experimental/Chaos/Public/Chaos/PBDConstraintRule.h
--- a/Engine/Source/Runtime/Experimental/Chaos/Public/Chaos/PBDConstraintRule.h
+++ b/Engine/Source/Runtime/Experimental/Chaos/Public/Chaos/PBDConstraintRule.h
@@ -153,20 +153,25 @@ namespace Chaos
 		virtual void SetUseContactGraph(const bool InUseContactGraph) {}
 
 		/** Change enabled state on all constraints associated with the specified particles */
-		virtual void SetConstraintsEnabled(const TSet<TGeometryParticleHandle<FReal, 3>*>& DisabledParticles, bool bInEnabled) 
-		{ 
-			for (TGeometryParticleHandle<FReal, 3>*DisabledParticle : DisabledParticles)
+		inline void SetConstraintsEnabled(FGeometryParticleHandle* ParticleHandle, bool bInEnabled)
+		{
+			for (FConstraintHandle* Constraint : ParticleHandle->ParticleConstraints())
 			{
-				for( FConstraintHandle* Constraint : DisabledParticle->ParticleConstraints() )
+				if (Constraint->IsEnabled() != bInEnabled)
 				{
-					if (Constraint->IsEnabled() != bInEnabled)
-					{
-						Constraint->SetEnabled(bInEnabled);
-					}
+					Constraint->SetEnabled(bInEnabled);
 				}
 			}
 		}
 
+		inline void SetConstraintsEnabled(const TSet<TGeometryParticleHandle<FReal, 3>*>& ParticleHandles, bool bInEnabled)
+		{ 
+			for (TGeometryParticleHandle<FReal, 3>* ParticleHandle : ParticleHandles)
+			{
+				SetConstraintsEnabled(ParticleHandle, bInEnabled);
+			}
+		}
+
 		/** Disconnect all constraints associated with the specified particles */
 		virtual void DisconnectConstraints(const TSet<TGeometryParticleHandle<FReal, 3>*>& RemovedParticles) {  }
 
diff --git a/Engine/Source/Runtime/Experimental/Chaos/Public/Chaos/PBDJointConstraints.h b/Engine/Source/Runtime/Experimental/Chaos/Public/Chaos/PBDJointConstraints.h
--- a/Engine/Source/Runtime/Experimental/Chaos/Public/Chaos/PBDJointConstraints.h
+++ b/Engine/Source/Runtime/Experimental/Chaos/Public/Chaos/PBDJointConstraints.h
@@ -26,7 +26,7 @@ namespace Chaos
 	class FPBDJointSolver;
 	class FPBDIslandSolverData;
 
-	class CHAOS_API FPBDJointConstraintHandle : public TIndexedContainerConstraintHandle<FPBDJointConstraints>
+	class CHAOS_API FPBDJointConstraintHandle final : public TIndexedContainerConstraintHandle<FPBDJointConstraints>
 	{
 	public:
 		using Base = TIndexedContainerConstraintHandle<FPBDJointConstraints>;
@@ -54,7 +54,8 @@ namespace Chaos
 		const FPBDJointSettings& GetJointSettings() const { return GetSettings(); }	//needed for property macros
 
 		void SetSettings(const FPBDJointSettings& Settings);
-		TVec2<FGeometryParticleHandle*> GetConstrainedParticles() const;
+		
+		TVec2<FGeometryParticleHandle*> GetConstrainedParticles() const override;
 
 		void PreGatherInput(const FReal Dt, FPBDIslandSolverData& SolverData);
 		void GatherInput(const FReal Dt, const int32 Particle0Level, const int32 Particle1Level, FPBDIslandSolverData& SolverData);
@@ -108,7 +109,6 @@ namespace Chaos
 
 		using FConstraintContainerHandle = FPBDJointConstraintHandle;
 		using FConstraintHandleAllocator = TConstraintHandleAllocator<FPBDJointConstraints>;
-		using FParticlePair = TVector<TGeometryParticleHandle<FReal, 3>*, 2>;
 		using FVectorPair = TVector<FVec3, 2>;
 		using FTransformPair = TVector<FRigidTransform3, 2>;
 		using FHandles = TArray<FConstraintContainerHandle*>;
diff --git a/Engine/Source/Runtime/Experimental/Chaos/Public/Chaos/PBDNullConstraints.h b/Engine/Source/Runtime/Experimental/Chaos/Public/Chaos/PBDNullConstraints.h
new file mode 100644
--- /dev/null
+++ b/Engine/Source/Runtime/Experimental/Chaos/Public/Chaos/PBDNullConstraints.h
@@ -0,0 +1,166 @@
+// Copyright Epic Games, Inc. All Rights Reserved.
+#pragma once
+
+#include "Chaos/Declares.h"
+#include "Chaos/ConstraintHandle.h"
+#include "Chaos/ParticleHandleFwd.h"
+
+
+// @todo(chaos): These classes should be in the unit testing suite, but we must currently explicitly instantiate the ConstraintRules in the Chaos unit
+// because the template code is in a source file. We need to support custom constraints anyway and the NullConstraint could use that when it exists.
+
+namespace Chaos
+{
+	class FPBDNullConstraintHandle;
+	class FPBDNullConstraints;
+
+	/**
+	 * @brief A dummy constraint used for unit testing
+	*/
+	class FPBDNullConstraint
+	{
+	public:
+		FPBDNullConstraint(const TVec2<FGeometryParticleHandle*>& InConstrainedParticles)
+			: ConstrainedParticles(InConstrainedParticles)
+			, bEnabled(true)
+			, bSleeping(false)
+		{
+		}
+
+		FParticlePair ConstrainedParticles;
+		bool bEnabled;
+		bool bSleeping;
+	};
+
+	/**
+	 * Constraint Container with minimal API required to test the Graph.
+	 */
+	class FPBDNullConstraints : public FPBDIndexedConstraintContainer
+	{
+	public:
+		using FConstraintContainerHandle = FPBDNullConstraintHandle;
+		using FConstraintSolverContainerType = FConstraintSolverContainer;	// @todo(chaos): Add island solver for this constraint type
+
+		FPBDNullConstraints();
+
+		int32 NumConstraints() const { return Constraints.Num(); }
+
+		FPBDNullConstraint& GetConstraint(const int32 ConstraintIndex)
+		{
+			return Constraints[ConstraintIndex];
+		}
+
+		const FPBDNullConstraint& GetConstraint(const int32 ConstraintIndex) const
+		{
+			return Constraints[ConstraintIndex];
+		}
+
+		FPBDNullConstraintHandle* AddConstraint(const TVec2<FGeometryParticleHandle*>& InConstraintedParticles);
+
+		FParticlePair GetConstrainedParticles(int32 ConstraintIndex) const
+		{
+			return Constraints[ConstraintIndex].ConstrainedParticles;
+		}
+
+		TArray<FPBDNullConstraintHandle*>& GetConstraintHandles()
+		{
+			return Handles;
+		}
+
+		const TArray<FPBDNullConstraintHandle*>& GetConstraintHandles() const
+		{
+			return Handles;
+		}
+
+		void PrepareTick() {}
+		void UnprepareTick() {}
+		void UpdatePositionBasedState(const FReal Dt) {}
+
+		// Constraint Rule API
+		void SetNumIslandConstraints(const int32 NumIslandConstraints, FPBDIslandSolverData& SolverData) {}
+
+		// Simple Constraint Rule API
+		void PreGatherInput(const FReal Dt, FPBDIslandSolverData& SolverData) {}
+		void GatherInput(const FReal Dt, FPBDIslandSolverData& SolverData) {}
+		void ScatterOutput(const FReal Dt, FPBDIslandSolverData& SolverData) {}
+		bool ApplyPhase1(const FReal Dt, const int32 It, const int32 NumIts, FPBDIslandSolverData& SolverData) { return true; }
+		bool ApplyPhase2(const FReal Dt, const int32 It, const int32 NumIts, FPBDIslandSolverData& SolverData) { return true; }
+		bool ApplyPhase3(const FReal Dt, const int32 It, const int32 NumIts, FPBDIslandSolverData& SolverData) { return true; }
+
+		// Island Constraint Rule API
+		void PreGatherInput(const FReal Dt, const int32 ConstraintIndex, FPBDIslandSolverData& SolverData) {}
+		void GatherInput(const FReal Dt, const int32 ConstraintIndex, const int32 Particle0Level, const int32 Particle1Level, FPBDIslandSolverData& SolverData) {}
+		bool ApplyPhase1Serial(const FReal Dt, const int32 It, const int32 NumIts, FPBDIslandSolverData& SolverData) { return true; }
+		bool ApplyPhase2Serial(const FReal Dt, const int32 It, const int32 NumIts, FPBDIslandSolverData& SolverData) { return true; }
+		bool ApplyPhase3Serial(const FReal Dt, const int32 It, const int32 NumIts, FPBDIslandSolverData& SolverData) { return true; }
+
+		TArray<FPBDNullConstraint> Constraints;
+		TArray<FPBDNullConstraintHandle*> Handles;
+		TConstraintHandleAllocator<FPBDNullConstraints> HandleAllocator;
+	};
+
+	class FPBDNullConstraintHandle final : public TIndexedContainerConstraintHandle<FPBDNullConstraints>
+	{
+	public:
+		FPBDNullConstraintHandle(FPBDNullConstraints* InConstraintContainer, int32 ConstraintIndex)
+			: TIndexedContainerConstraintHandle<FPBDNullConstraints>(InConstraintContainer, ConstraintIndex)
+		{
+		}
+
+		virtual void SetEnabled(bool bInEnabled)  override
+		{
+			ConcreteContainer()->GetConstraint(GetConstraintIndex()).bEnabled = bInEnabled;
+		}
+
+		virtual bool IsEnabled() const override
+		{
+			return ConcreteContainer()->GetConstraint(GetConstraintIndex()).bEnabled;
+		}
+
+		virtual void SetIsSleeping(bool bInIsSleeping)  override
+		{
+			ConcreteContainer()->GetConstraint(GetConstraintIndex()).bSleeping = bInIsSleeping;
+		}
+
+		virtual bool IsSleeping() const override
+		{
+			return ConcreteContainer()->GetConstraint(GetConstraintIndex()).bSleeping;
+		}
+
+		virtual FParticlePair GetConstrainedParticles() const override
+		{
+			return ConcreteContainer()->GetConstrainedParticles(GetConstraintIndex());
+		}
+
+		void PreGatherInput(const FReal Dt, FPBDIslandSolverData& SolverData)
+		{
+			ConcreteContainer()->PreGatherInput(Dt, GetConstraintIndex(), SolverData);
+		}
+
+		void GatherInput(const FReal Dt, const int32 Particle0Level, const int32 Particle1Level, FPBDIslandSolverData& SolverData)
+		{
+			ConcreteContainer()->GatherInput(Dt, GetConstraintIndex(), Particle0Level, Particle1Level, SolverData);
+		}
+
+		static const FConstraintHandleTypeID& StaticType()
+		{
+			static FConstraintHandleTypeID STypeID(TEXT("FPBDNullConstraintHandle"), &FIndexedConstraintHandle::StaticType());
+			return STypeID;
+		}
+	};
+
+
+	inline FPBDNullConstraints::FPBDNullConstraints()
+		: FPBDIndexedConstraintContainer(FPBDNullConstraintHandle::StaticType())
+	{
+	}
+
+	inline FPBDNullConstraintHandle* FPBDNullConstraints::AddConstraint(const TVec2<FGeometryParticleHandle*>& InConstraintedParticles)
+	{
+		const int32 ConstraintIndex = Constraints.Emplace(FPBDNullConstraint(InConstraintedParticles));
+		const int32 HandleIndex = Handles.Emplace(HandleAllocator.AllocHandle(this, ConstraintIndex));
+		check(ConstraintIndex == HandleIndex);
+		return Handles[HandleIndex];
+	}
+
+}
\ No newline at end of file
diff --git a/Engine/Source/Runtime/Experimental/Chaos/Public/Chaos/PBDPositionConstraints.h b/Engine/Source/Runtime/Experimental/Chaos/Public/Chaos/PBDPositionConstraints.h
--- a/Engine/Source/Runtime/Experimental/Chaos/Public/Chaos/PBDPositionConstraints.h
+++ b/Engine/Source/Runtime/Experimental/Chaos/Public/Chaos/PBDPositionConstraints.h
@@ -12,7 +12,7 @@ namespace Chaos
 	class FPBDPositionConstraints;
 	class FPBDIslandSolverData;
 
-	class FPBDPositionConstraintHandle : public TIndexedContainerConstraintHandle<FPBDPositionConstraints>
+	class FPBDPositionConstraintHandle final : public TIndexedContainerConstraintHandle<FPBDPositionConstraints>
 	{
 	public:
 		using Base = TIndexedContainerConstraintHandle<FPBDPositionConstraints>;
@@ -22,7 +22,8 @@ namespace Chaos
 		FPBDPositionConstraintHandle() {}
 		FPBDPositionConstraintHandle(FConstraintContainer* InConstraintContainer, int32 InConstraintIndex) 
 			: TIndexedContainerConstraintHandle<FPBDPositionConstraints>(InConstraintContainer, InConstraintIndex) {}
-		TVector<FGeometryParticleHandle*, 2> GetConstrainedParticles() const;
+		
+		virtual FParticlePair GetConstrainedParticles() const override;
 
 		void PreGatherInput(const FReal Dt, FPBDIslandSolverData& SolverData);
 		void GatherInput(const FReal Dt, const int32 Particle0Level, const int32 Particle1Level, FPBDIslandSolverData& SolverData);
diff --git a/Engine/Source/Runtime/Experimental/Chaos/Public/Chaos/PBDRigidDynamicSpringConstraints.h b/Engine/Source/Runtime/Experimental/Chaos/Public/Chaos/PBDRigidDynamicSpringConstraints.h
--- a/Engine/Source/Runtime/Experimental/Chaos/Public/Chaos/PBDRigidDynamicSpringConstraints.h
+++ b/Engine/Source/Runtime/Experimental/Chaos/Public/Chaos/PBDRigidDynamicSpringConstraints.h
@@ -12,7 +12,7 @@ namespace Chaos
 	class FPBDRigidDynamicSpringConstraints;
 	class FPBDIslandSolverData;
 
-	class CHAOS_API FPBDRigidDynamicSpringConstraintHandle : public TIndexedContainerConstraintHandle<FPBDRigidDynamicSpringConstraints>
+	class CHAOS_API FPBDRigidDynamicSpringConstraintHandle final : public TIndexedContainerConstraintHandle<FPBDRigidDynamicSpringConstraints>
 	{
 	public:
 		using Base = TIndexedContainerConstraintHandle<FPBDRigidDynamicSpringConstraints>;
@@ -27,7 +27,7 @@ namespace Chaos
 		{
 		}
 
-		TVec2<FGeometryParticleHandle*> GetConstrainedParticles() const;
+		virtual FParticlePair GetConstrainedParticles() const override;
 
 		void PreGatherInput(const FReal Dt, FPBDIslandSolverData& SolverData);
 		void GatherInput(const FReal Dt, const int32 Particle0Level, const int32 Particle1Level, FPBDIslandSolverData& SolverData);
diff --git a/Engine/Source/Runtime/Experimental/Chaos/Public/Chaos/PBDRigidSpringConstraints.h b/Engine/Source/Runtime/Experimental/Chaos/Public/Chaos/PBDRigidSpringConstraints.h
--- a/Engine/Source/Runtime/Experimental/Chaos/Public/Chaos/PBDRigidSpringConstraints.h
+++ b/Engine/Source/Runtime/Experimental/Chaos/Public/Chaos/PBDRigidSpringConstraints.h
@@ -12,7 +12,7 @@ namespace Chaos
 	class FPBDRigidSpringConstraints;
 	class FPBDIslandSolverData;
 
-	class CHAOS_API FPBDRigidSpringConstraintHandle : public TIndexedContainerConstraintHandle<FPBDRigidSpringConstraints>
+	class CHAOS_API FPBDRigidSpringConstraintHandle final : public TIndexedContainerConstraintHandle<FPBDRigidSpringConstraints>
 	{
 	public:
 		using Base = TIndexedContainerConstraintHandle<FPBDRigidSpringConstraints>;
@@ -31,7 +31,7 @@ namespace Chaos
 		const TVector<FVec3, 2>& GetConstraintPositions() const;
 		void SetConstraintPositions(const TVector<FVec3, 2>& ConstraintPositions);
 		
-		TVector<FGeometryParticleHandle*, 2> GetConstrainedParticles() const;
+		virtual FParticlePair GetConstrainedParticles() const override;
 
 		// Get the rest length of the spring
 		FReal GetRestLength() const;
diff --git a/Engine/Source/Runtime/Experimental/Chaos/Public/Chaos/PBDRigidsEvolution.h b/Engine/Source/Runtime/Experimental/Chaos/Public/Chaos/PBDRigidsEvolution.h
--- a/Engine/Source/Runtime/Experimental/Chaos/Public/Chaos/PBDRigidsEvolution.h
+++ b/Engine/Source/Runtime/Experimental/Chaos/Public/Chaos/PBDRigidsEvolution.h
@@ -366,7 +366,7 @@ public:
 	CHAOS_API void EnableParticle(FGeometryParticleHandle* Particle, const FGeometryParticleHandle* ParentParticle)
 	{
 		Particles.EnableParticle(Particle);
-		EnableConstraints({ Particle });
+		EnableConstraints(Particle);
 		ConstraintGraph.EnableParticle(Particle, ParentParticle);
 		DirtyParticle(*Particle);
 	}
@@ -376,14 +376,17 @@ public:
 		RemoveParticleFromAccelerationStructure(*Particle);
 		Particles.DisableParticle(Particle);
 		// disable constraint before disabling the particle in the graph because of dependencies 
-		DisableConstraints({ Particle });
+		DisableConstraints(Particle);
 		ConstraintGraph.DisableParticle(Particle);
-		DisableConstraints(TSet<FGeometryParticleHandle*>({ Particle }));
 	}
 
 	CHAOS_API void InvalidateParticle(FGeometryParticleHandle* Particle)
 	{
+		// Remove the particle from the constraint graph. This should remove all the constraints too
 		ConstraintGraph.DisableParticle(Particle);
+
+		// Destroy all the transient constraints (collisions) because the particle has changed somehow (e.g. new shapes) and they may have cached the previous state
+		DestroyTransientConstraints(Particle);
 	}
 	
 	CHAOS_API void FlushExternalAccelerationQueue(FAccelerationStructure& Acceleration,FPendingSpatialDataQueue& ExternalQueue);
@@ -480,32 +483,9 @@ public:
 
 	CHAOS_API void DisableParticles(const TSet<FGeometryParticleHandle*>& InParticles);
 
-	CHAOS_API void WakeIslands()
-	{
-		TArray<int32> UniqueIslands;
-
-		// there could easily be duplicates, best to remove these since WakeIsland is potentially expensive call
-		int32 IslandIdx = 0;
-		while (!IslandsToWake.IsEmpty())
-		{
-			IslandsToWake.Dequeue(IslandIdx);
-			UniqueIslands.AddUnique(IslandIdx);
-		}
-
-		for (int32 Island : UniqueIslands)
-		{
-			WakeIsland(Island);
-		}
-	}
-
 	CHAOS_API void WakeIsland(const int32 Island)
 	{
 		ConstraintGraph.WakeIsland(Particles, Island);
-		//Update Particles SOAs
-		/*for (auto Particle : ContactGraph.GetIslandParticles(Island))
-		{
-			ActiveIndices.Add(Particle);
-		}*/
 	}
 
 	/** remove a constraint from the constraint graph (see AddConstraintsToConstraintGraph) */
@@ -546,8 +526,10 @@ public:
 
 
 
-	/** Disconnect constraints from a set of particles to be removed (or destroyed) 
-	* this will set the constraints to Enbaled = false and set their respective bodies handles to nullptr
+	/** Disconnect constraints from a set of particles to be destroyed. 
+	* this will set the constraints to Enbaled = false and set their respective bodies handles to nullptr.
+	* Once this is done, the constraints cannot be re-enabled.
+	* @note This only applies to persistent constraints (joints etc), not transient constraints (collisons)
 	*/
 	CHAOS_API void DisconnectConstraints(const TSet<FGeometryParticleHandle*>& RemovedParticles)
 	{
@@ -562,29 +544,43 @@ public:
 		}
 	}
 
-	/** Disconnect constraints from a set of particles to be removed (or destroyed)
-	* this will set the constraints to Enabled = false, but leave connections to the 
-	* particles. 
+	/** Disconnect constraints from a particle to be removed (or destroyed)
+	* this will set the constraints to Enabled = false, but leave connections to the particles to support
+	* re-enabling at a later time.
+	* @note This only applies to persistent constraints (joints etc), not transient constraints (collisons)
 	*/
+	CHAOS_API void DisableConstraints(FGeometryParticleHandle* ParticleHandle)
+	{
+		RemoveConstraintsFromConstraintGraph(ParticleHandle->ParticleConstraints());
+
+		for (FPBDConstraintGraphRule* ConstraintRule : ConstraintRules)
+		{
+			ConstraintRule->SetConstraintsEnabled(ParticleHandle, false);
+		}
+	}
+
 	CHAOS_API void DisableConstraints(const TSet<FGeometryParticleHandle*>& DisabledParticles)
 	{
 		for (FGeometryParticleHandle* ParticleHandle : DisabledParticles)
 		{
-			RemoveConstraintsFromConstraintGraph(ParticleHandle->ParticleConstraints());
+			DisableConstraints(ParticleHandle);
 		}
+	}
 
+	/** Enable constraints from the enabled particles; constraints will only become enabled if their particle end points are valid */
+	CHAOS_API void EnableConstraints(FGeometryParticleHandle* ParticleHandle)
+	{
 		for (FPBDConstraintGraphRule* ConstraintRule : ConstraintRules)
 		{
-			ConstraintRule->SetConstraintsEnabled(DisabledParticles, false);
+			ConstraintRule->SetConstraintsEnabled(ParticleHandle, true);
 		}
 	}
 
-	/** Enable constraints from the enabled particles; constraints will only become enabled if their particle end points are valid */
 	CHAOS_API void EnableConstraints(const TSet<FGeometryParticleHandle*>& EnabledParticles)
 	{
-		for (FPBDConstraintGraphRule* ConstraintRule : ConstraintRules)
+		for (FGeometryParticleHandle* ParticleHandle : EnabledParticles)
 		{
-			ConstraintRule->SetConstraintsEnabled(EnabledParticles, true);
+			EnableConstraints(ParticleHandle);
 		}
 	}
 
@@ -610,6 +606,8 @@ public:
 		}
 	}
 
+	virtual void DestroyTransientConstraints(FGeometryParticleHandle* Particle) {}
+
 	const TParticleView<FPBDRigidClusteredParticles>& GetNonDisabledClusteredView() const { return Particles.GetNonDisabledClusteredView(); }
 
 	CHAOS_API TSerializablePtr<FChaosPhysicsMaterial> GetPhysicsMaterial(const FGeometryParticleHandle* Particle) const { return Particle->AuxilaryValue(PhysicsMaterials); }
@@ -627,7 +625,6 @@ public:
 		Particle->AuxilaryValue(PhysicsMaterials) = InMaterial;
 	}
 
-	CHAOS_API const TArray<FGeometryParticleHandle*>& GetIslandParticles(const int32 Island) const { return ConstraintGraph.GetIslandParticles(Island); }
 	CHAOS_API int32 NumIslands() const { return ConstraintGraph.NumIslands(); }
 
 	void InitializeAccelerationStructures()
@@ -966,7 +963,6 @@ protected:
 
 	// Allows us to tell evolution to stop starting async tasks if we are trying to cleanup solver/evo.
 	bool bCanStartAsyncTasks;
-	TQueue<int32, EQueueMode::Mpsc> IslandsToWake;
 
 	TArray<FUniqueIdx> UniqueIndicesPendingRelease;
 
diff --git a/Engine/Source/Runtime/Experimental/Chaos/Public/Chaos/PBDRigidsEvolutionGBF.h b/Engine/Source/Runtime/Experimental/Chaos/Public/Chaos/PBDRigidsEvolutionGBF.h
--- a/Engine/Source/Runtime/Experimental/Chaos/Public/Chaos/PBDRigidsEvolutionGBF.h
+++ b/Engine/Source/Runtime/Experimental/Chaos/Public/Chaos/PBDRigidsEvolutionGBF.h
@@ -156,6 +156,8 @@ namespace Chaos
 
 		void DestroyParticleCollisionsInAllocator(FGeometryParticleHandle* Particle);
 
+		virtual void DestroyTransientConstraints(FGeometryParticleHandle* Particle) override;
+
 		CHAOS_API inline void EndFrame(FReal Dt)
 		{
 			Particles.GetNonDisabledDynamicView().ParallelFor([&](auto& Particle, int32 Index) {
diff --git a/Engine/Source/Runtime/Experimental/Chaos/Public/Chaos/PBDSuspensionConstraints.h b/Engine/Source/Runtime/Experimental/Chaos/Public/Chaos/PBDSuspensionConstraints.h
--- a/Engine/Source/Runtime/Experimental/Chaos/Public/Chaos/PBDSuspensionConstraints.h
+++ b/Engine/Source/Runtime/Experimental/Chaos/Public/Chaos/PBDSuspensionConstraints.h
@@ -19,7 +19,7 @@ namespace Chaos
 	class FPBDIslandSolverData;
 	class FPBDCollisionSolver;
 
-	class CHAOS_API FPBDSuspensionConstraintHandle : public TIndexedContainerConstraintHandle<FPBDSuspensionConstraints>
+	class CHAOS_API FPBDSuspensionConstraintHandle final : public TIndexedContainerConstraintHandle<FPBDSuspensionConstraints>
 	{
 	public:
 		using Base = TIndexedContainerConstraintHandle<FPBDSuspensionConstraints>;
@@ -33,7 +33,7 @@ namespace Chaos
 
 		void SetSettings(const FPBDSuspensionSettings& Settings);
 
-		TVec2<FGeometryParticleHandle*> GetConstrainedParticles() const;
+		virtual FParticlePair GetConstrainedParticles() const override;
 
 		void PreGatherInput(const FReal Dt, FPBDIslandSolverData& SolverData);
 		void GatherInput(const FReal Dt, const int32 Particle0Level, const int32 Particle1Level, FPBDIslandSolverData& SolverData);
