Bitwise check instead of a strict equality check on EInterchangeTranslatorType, allowing formats to appear in the dialog box when importing from the content browser

diff --git a/Engine/Source/Runtime/Experimental/Interchange/Engine/Private/InterchangeManager.cpp b/Engine/Source/Runtime/Experimental/Interchange/Engine/Private/InterchangeManager.cpp
--- a/Engine/Source/Runtime/Experimental/Interchange/Engine/Private/InterchangeManager.cpp
+++ b/Engine/Source/Runtime/Experimental/Interchange/Engine/Private/InterchangeManager.cpp
@@ -533,7 +533,7 @@ TArray<FString> UInterchangeManager::GetSupportedFormats(const EInterchangeTrans
 	{
 		const UInterchangeTranslatorBase* TranslatorBaseCDO = TranslatorClass->GetDefaultObject<UInterchangeTranslatorBase>();
 
-		if (TranslatorBaseCDO->GetTranslatorType() == ForTranslatorType)
+		if(EnumHasAnyFlags(TranslatorBaseCDO->GetTranslatorType(), ForTranslatorType))
 		{
 			FileExtensions.Append(TranslatorBaseCDO->GetSupportedFormats());
 		}
