Further reduce the geometry collection BLAS build count.
1. Move the build out of the section loop so that we add one BLAS per geometry collection.
2. Create the ray tracing geometry but delay the BLAS build later to dynamic raytracing geometry update.

diff --git a/Engine/Source/Runtime/Experimental/GeometryCollectionEngine/Private/GeometryCollection/GeometryCollectionSceneProxy.cpp b/Engine/Source/Runtime/Experimental/GeometryCollectionEngine/Private/GeometryCollection/GeometryCollectionSceneProxy.cpp
--- a/Engine/Source/Runtime/Experimental/GeometryCollectionEngine/Private/GeometryCollection/GeometryCollectionSceneProxy.cpp
+++ b/Engine/Source/Runtime/Experimental/GeometryCollectionEngine/Private/GeometryCollection/GeometryCollectionSceneProxy.cpp
@@ -1159,25 +1159,25 @@ void FGeometryCollectionSceneProxy::GetDynamicRayTracingInstances(FRayTracingMat
 				}
 #endif
 				//#TODO: bone color, bone selection and render bound?
-				
-				FRWBuffer* VertexBuffer = RayTracingDynamicVertexBuffer.NumBytes > 0 ? &RayTracingDynamicVertexBuffer : nullptr;
-
-				const uint32 VertexCount = Section.MaxVertexIndex + 1;
-				Context.DynamicRayTracingGeometriesToUpdate.Add(
-					FRayTracingDynamicGeometryUpdateParams
-					{
-						RayTracingInstance.Materials,
-						false,
-						VertexCount,
-						VertexCount * (uint32)sizeof(FVector3f),
-						RayTracingGeometry.Initializer.TotalPrimitiveCount,
-						&RayTracingGeometry,
-						VertexBuffer,
-						true
-					}
-				);
 			}
 
+			FRWBuffer* VertexBuffer = RayTracingDynamicVertexBuffer.NumBytes > 0 ? &RayTracingDynamicVertexBuffer : nullptr;
+
+			const uint32 VertexCount = MaxVertexIndex + 1;
+			Context.DynamicRayTracingGeometriesToUpdate.Add(
+				FRayTracingDynamicGeometryUpdateParams
+				{
+					RayTracingInstance.Materials,
+					false,
+					VertexCount,
+					VertexCount * (uint32)sizeof(FVector3f),
+					RayTracingGeometry.Initializer.TotalPrimitiveCount,
+					&RayTracingGeometry,
+					VertexBuffer,
+					true
+				}
+			);
+
 			RayTracingInstance.BuildInstanceMaskAndFlags(GetScene().GetFeatureLevel());
 
 			OutRayTracingInstances.Emplace(RayTracingInstance);
@@ -1217,7 +1217,8 @@ void FGeometryCollectionSceneProxy::UpdatingRayTracingGeometry_RenderingThread(F
 		if (RayTracingGeometry.Initializer.TotalPrimitiveCount > 0)
 		{
 			RayTracingGeometry.Initializer.IndexBuffer = InIndexBuffer->IndexBufferRHI;
-			RayTracingGeometry.UpdateRHI();
+			// Create the ray tracing geometry but delay the acceleration structure build.
+			RayTracingGeometry.CreateRayTracingGeometry(ERTAccelerationStructureBuildPriority::Skip);
 		}
 
 		bGeometryResourceUpdated = false;
