make sure that the constraint ticks before the child if the both parent and child tick functions are the same (e.i. two controls controling the same skeletal mesh)

diff --git a/Engine/Source/Runtime/Experimental/Animation/Constraints/Private/TransformConstraint.cpp b/Engine/Source/Runtime/Experimental/Animation/Constraints/Private/TransformConstraint.cpp
--- a/Engine/Source/Runtime/Experimental/Animation/Constraints/Private/TransformConstraint.cpp
+++ b/Engine/Source/Runtime/Experimental/Animation/Constraints/Private/TransformConstraint.cpp
@@ -125,7 +125,9 @@ void UTickableTransformConstraint::Setup()
 void UTickableTransformConstraint::SetupDependencies()
 {
 	FTickFunction* ParentTickFunction = ParentTRSHandle->IsValid() ? ParentTRSHandle->GetTickFunction() : nullptr;
-	if (ParentTickFunction)
+	FTickFunction* ChildTickFunction = ChildTRSHandle ? ChildTRSHandle->GetTickFunction() : nullptr;
+	
+	if (ParentTickFunction && (ChildTickFunction != ParentTickFunction))
 	{
 		// manage dependencies
 		// force ConstraintTickFunction to tick after InParent does.
@@ -134,10 +136,9 @@ void UTickableTransformConstraint::SetupDependencies()
 	}
 	
 	// TODO also check for cycle dependencies here
-	FTickFunction* ChildTickFunction = ChildTRSHandle ? ChildTRSHandle->GetTickFunction() : nullptr;
-	if (ChildTickFunction && ChildTickFunction != ParentTickFunction)
+	if (ChildTickFunction)
 	{
-		// force InParent to tick after ConstraintTickFunction does.
+		// force InChild to tick after ConstraintTickFunction does.
 		// Note that this might not register anything if the child can't tick (static meshes for instance)
 		ChildTickFunction->AddPrerequisite(this, ConstraintTick);
 	}
@@ -147,6 +148,8 @@ void UTickableTransformConstraint::PostLoad()
 {
 	Super::PostLoad();
 	ConstraintTick.RegisterFunction(GetFunction() );
+
+	SetupDependencies();	
 	RegisterDelegates();
 }
 
@@ -154,6 +157,8 @@ void UTickableTransformConstraint::PostDuplicate(bool bDuplicateForPIE)
 {
 	Super::PostDuplicate(bDuplicateForPIE);
 	ConstraintTick.RegisterFunction(GetFunction());
+
+	SetupDependencies();
 	RegisterDelegates();
 }
 
@@ -846,25 +851,40 @@ UTickableTransformConstraint* FTransformConstraintUtils::CreateAndAddFromActors(
 	const ETransformConstraintType InType,
 	const bool bMaintainOffset)
 {
+	static const TCHAR* ErrorPrefix = TEXT("FTransformConstraintUtils::CreateAndAddFromActors");
+	
 	// SANITY CHECK
 	if (!InWorld || !InParent || !InChild)
 	{
-		UE_LOG(LogTemp, Error, TEXT("FTransformConstraintUtils::CreateAndAddFromActors sanity check failed."));
+		UE_LOG(LogTemp, Error, TEXT("%s sanity check failed."), ErrorPrefix);
 		return nullptr;
 	}
 	
 	UConstraintsManager* ConstraintsManager = UConstraintsManager::Get(InWorld);
 	if (!ConstraintsManager)
 	{
-		UE_LOG(LogTemp, Error, TEXT("FTransformConstraintUtils::CreateAndAddFromActors constraint manager is null."));
+		UE_LOG(LogTemp, Error, TEXT("%s constraint manager is null."), ErrorPrefix);
 		return nullptr;
 	}
 
 	UTransformableHandle* ParentHandle = GetHandle(InParent, ConstraintsManager);
+	if (!ParentHandle)
+	{
+		return nullptr;
+	}
+	
 	UTransformableHandle* ChildHandle = GetHandle(InChild, ConstraintsManager);
+	if (!ChildHandle)
+	{
+		return nullptr;
+	}
 
-	if (!ParentHandle || !ChildHandle)
+	if (ChildHandle->GetHash() == ParentHandle->GetHash())
 	{
+		ChildHandle->MarkAsGarbage();
+		ParentHandle->MarkAsGarbage();
+		
+		UE_LOG(LogTemp, Error, TEXT("%s handles are pointing at the same object."), ErrorPrefix);
 		return nullptr;
 	}
 	
