Fixed RotatorNetSerializer delta serialization bug.

#rnx

diff --git a/Engine/Source/Runtime/Experimental/Iris/Core/Private/Iris/Serialization/RotatorNetSerializers.cpp b/Engine/Source/Runtime/Experimental/Iris/Core/Private/Iris/Serialization/RotatorNetSerializers.cpp
--- a/Engine/Source/Runtime/Experimental/Iris/Core/Private/Iris/Serialization/RotatorNetSerializers.cpp
+++ b/Engine/Source/Runtime/Experimental/Iris/Core/Private/Iris/Serialization/RotatorNetSerializers.cpp
@@ -210,6 +210,13 @@ void FRotatorAsShortNetSerializer::DeserializeDelta(FNetSerializationContext& Co
 		TempValue.Z += DZ;
 	}
 
+	// Reconstruct flags
+	uint32 XYZIsNotZero = 0U;
+	XYZIsNotZero |= (TempValue.X != 0 ? XDiffersMask : 0U);
+	XYZIsNotZero |= (TempValue.Y != 0 ? YDiffersMask : 0U);
+	XYZIsNotZero |= (TempValue.Z != 0 ? ZDiffersMask : 0U);
+	TempValue.XYZIsNotZero = XYZIsNotZero;
+
 	QuantizedType& Target = *reinterpret_cast<QuantizedType*>(Args.Target);
 	Target = TempValue;
 }
@@ -253,8 +260,8 @@ bool FRotatorAsShortNetSerializer::IsEqual(FNetSerializationContext& Context, co
 		const bool bXIsEqual = (SourceValue0.X == SourceValue1.X);
 		const bool bYIsEqual = (SourceValue0.Y == SourceValue1.Y);
 		const bool bZIsEqual = (SourceValue0.Z == SourceValue1.Z);
-		// There's no need to check the flags as they are derived from the X,Y and Z values.
-		return bXIsEqual & bYIsEqual & bZIsEqual;
+		const bool bFlagsAreEqual = (SourceValue0.XYZIsNotZero == SourceValue1.XYZIsNotZero);
+		return bXIsEqual & bYIsEqual & bZIsEqual & bFlagsAreEqual;
 	}
 	else
 	{
@@ -358,6 +365,13 @@ void FRotatorAsByteNetSerializer::DeserializeDelta(FNetSerializationContext& Con
 		TempValue.Z = Reader->ReadBits(8U);
 	}
 
+	// Reconstruct flags
+	uint32 XYZIsNotZero = 0U;
+	XYZIsNotZero |= (TempValue.X != 0 ? XDiffersMask : 0U);
+	XYZIsNotZero |= (TempValue.Y != 0 ? YDiffersMask : 0U);
+	XYZIsNotZero |= (TempValue.Z != 0 ? ZDiffersMask : 0U);
+	TempValue.XYZIsNotZero = XYZIsNotZero;
+
 	QuantizedType& Target = *reinterpret_cast<QuantizedType*>(Args.Target);
 	Target = TempValue;
 }
@@ -401,8 +415,8 @@ bool FRotatorAsByteNetSerializer::IsEqual(FNetSerializationContext& Context, con
 		const bool bXIsEqual = (SourceValue0.X == SourceValue1.X);
 		const bool bYIsEqual = (SourceValue0.Y == SourceValue1.Y);
 		const bool bZIsEqual = (SourceValue0.Z == SourceValue1.Z);
-		// There's no need to check the flags as they are derived from the X,Y and Z values.
-		return bXIsEqual & bYIsEqual & bZIsEqual;
+		const bool bFlagsAreEqual = (SourceValue0.XYZIsNotZero == SourceValue1.XYZIsNotZero);
+		return bXIsEqual & bYIsEqual & bZIsEqual & bFlagsAreEqual;
 	}
 	else
 	{
