Chaos - fix index out of range in Island Manager

diff --git a/Engine/Source/Runtime/Experimental/Chaos/Private/Chaos/Island/IslandGraph.cpp b/Engine/Source/Runtime/Experimental/Chaos/Private/Chaos/Island/IslandGraph.cpp
--- a/Engine/Source/Runtime/Experimental/Chaos/Private/Chaos/Island/IslandGraph.cpp
+++ b/Engine/Source/Runtime/Experimental/Chaos/Private/Chaos/Island/IslandGraph.cpp
@@ -521,8 +521,13 @@ void FIslandGraph<NodeType, EdgeType, IslandType, OwnerType>::UpdateLevels(const
 		{
 			FGraphEdge& GraphEdge = GraphEdges[EdgeIndex];
 
+#if CHAOS_CONSTRAINTHANDLE_DEBUG_ENABLED
+			// Valid edges must have a valid island
+			ensure(!GraphEdge.bValidEdge || GraphIslands.IsValidIndex(GraphEdge.IslandIndex));
+#endif
+
 			// Do nothing if the edge is not coming from the same container or if the island is sleeping 
-			if (GraphEdge.ItemContainer == ContainerId && GraphEdge.LevelIndex == INDEX_NONE && !GraphIslands[GraphEdge.IslandIndex].bIsSleeping)
+			if (GraphEdge.bValidEdge && (GraphEdge.ItemContainer == ContainerId) && (GraphEdge.LevelIndex == INDEX_NONE) && !GraphIslands[GraphEdge.IslandIndex].bIsSleeping)
 			{
 				const int32 OtherIndex = (NodeIndex == GraphEdge.FirstNode) ?
 							GraphEdge.SecondNode : GraphEdge.FirstNode;
@@ -574,7 +579,7 @@ void FIslandGraph<NodeType, EdgeType, IslandType, OwnerType>::ComputeLevels(cons
 	// An isolated island that is only dynamics will not have been processed above, put everything without a level into level zero
 	for(auto& GraphEdge : GraphEdges)
 	{
-		if(GraphEdge.ItemContainer == ContainerId)
+		if(GraphEdge.bValidEdge && (GraphEdge.ItemContainer == ContainerId))
 		{
 			GraphEdge.LevelIndex = FGenericPlatformMath::Max(GraphEdge.LevelIndex, 0);
 		}
@@ -623,9 +628,14 @@ void FIslandGraph<NodeType, EdgeType, IslandType, OwnerType>::UpdateColors(const
 		{
 			FGraphEdge& GraphEdge = GraphEdges[EdgeIndex];
 
+#if CHAOS_CONSTRAINTHANDLE_DEBUG_ENABLED
+			// Valid edges must have a valid island
+			ensure(!GraphEdge.bValidEdge || GraphIslands.IsValidIndex(GraphEdge.IslandIndex));
+#endif
+
 			// Do nothing if the edge is not coming from the same container or if the island is sleeping 
-			if (GraphEdge.ItemContainer == ContainerId && GraphEdge.ColorIndex == INDEX_NONE &&
-				!GraphIslands[GraphEdge.IslandIndex].bIsSleeping && (GraphIslands[GraphEdge.IslandIndex].NumEdges > MinEdges) )
+			if (GraphEdge.bValidEdge && (GraphEdge.ItemContainer == ContainerId) && (GraphEdge.ColorIndex == INDEX_NONE) &&
+				!GraphIslands[GraphEdge.IslandIndex].bIsSleeping && (GraphIslands[GraphEdge.IslandIndex].NumEdges > MinEdges))
 			{
 				// Get the opposite node index for the given edge
 				const int32 OtherIndex = (NodeIndex == GraphEdge.FirstNode) ?
