Chaos - fix joint destroy followng particle disconnect

diff --git a/Engine/Source/Runtime/Experimental/Chaos/Private/Chaos/PBDJointConstraints.cpp b/Engine/Source/Runtime/Experimental/Chaos/Private/Chaos/PBDJointConstraints.cpp
--- a/Engine/Source/Runtime/Experimental/Chaos/Private/Chaos/PBDJointConstraints.cpp
+++ b/Engine/Source/Runtime/Experimental/Chaos/Private/Chaos/PBDJointConstraints.cpp
@@ -482,8 +482,14 @@ namespace Chaos
 		FConstraintContainerHandle* ConstraintHandle = Handles[ConstraintIndex];
 		if (ConstraintHandle != nullptr)
 		{
-			ConstraintParticles[ConstraintIndex][0]->RemoveConstraintHandle(ConstraintHandle);
-			ConstraintParticles[ConstraintIndex][1]->RemoveConstraintHandle(ConstraintHandle);
+			if (ConstraintParticles[ConstraintIndex][0] != nullptr)
+			{
+				ConstraintParticles[ConstraintIndex][0]->RemoveConstraintHandle(ConstraintHandle);
+			}
+			if (ConstraintParticles[ConstraintIndex][1] != nullptr)
+			{
+				ConstraintParticles[ConstraintIndex][1]->RemoveConstraintHandle(ConstraintHandle);
+			}
 
 			// Release the handle for the freed constraint
 			HandleAllocator.FreeHandle(ConstraintHandle);
@@ -503,8 +509,6 @@ namespace Chaos
 		}
 	}
 
-
-
 	
 	void FPBDJointConstraints::DisconnectConstraints(const TSet<TGeometryParticleHandle<FReal, 3>*>& RemovedParticles)
 	{
@@ -530,6 +534,8 @@ namespace Chaos
 					}
 				}
 			}
+
+			RemovedParticle->ParticleConstraints().Reset();
 		}
 	}
 
