Chaoss - fix sphere and capsule CCD against heightfield and trimesh

diff --git a/Engine/Source/Runtime/Experimental/Chaos/Private/Chaos/HeightField.cpp b/Engine/Source/Runtime/Experimental/Chaos/Private/Chaos/HeightField.cpp
--- a/Engine/Source/Runtime/Experimental/Chaos/Private/Chaos/HeightField.cpp
+++ b/Engine/Source/Runtime/Experimental/Chaos/Private/Chaos/HeightField.cpp
@@ -509,7 +509,10 @@ namespace Chaos
 
 				// If we are definitely penetrating by less than IgnorePenetration at T=1, we can ignore this triangle
 				// Checks the query geometry support point along the normal at the sweep end point
-				VectorRegister4Float IgnorePenetrationSimd = MakeVectorRegisterFloatFromDouble(VectorSetFloat1(IgnorePenetration));
+				// NOTE: We are using SupportCore, so we need to add Radius to the early-out distance check. We could just use Support
+				// but this way avoids an extra sqrt. We don't modify IgnorePenetration in the ctor because it gets used below.
+				VectorRegister4Float RadiusV = MakeVectorRegisterFloatFromDouble(VectorSetFloat1(OtherGeom.GetRadius()));
+				VectorRegister4Float IgnorePenetrationSimd = VectorSubtract(MakeVectorRegisterFloatFromDouble(VectorSetFloat1(IgnorePenetration)), RadiusV);
 				VectorRegister4Float TriangleNormalSimd = VectorNormalize(VectorCross(VectorSubtract(B, A), VectorSubtract(C, A)));
 				VectorRegister4Float OtherTriangleNormalSimd = VectorQuaternionInverseRotateVector(RotationSimd, TriangleNormalSimd);
 				VectorRegister4Float OtherTrianglePositionSimd = VectorQuaternionInverseRotateVector(RotationSimd, VectorSubtract(A, EndPointSimd));
diff --git a/Engine/Source/Runtime/Experimental/Chaos/Private/Chaos/TriangleMeshImplicitObject.cpp b/Engine/Source/Runtime/Experimental/Chaos/Private/Chaos/TriangleMeshImplicitObject.cpp
--- a/Engine/Source/Runtime/Experimental/Chaos/Private/Chaos/TriangleMeshImplicitObject.cpp
+++ b/Engine/Source/Runtime/Experimental/Chaos/Private/Chaos/TriangleMeshImplicitObject.cpp
@@ -1046,7 +1046,10 @@ struct FTriangleMeshSweepVisitorCCD
 
 		// If we are definitely penetrating by less than IgnorePenetration at T=1, we can ignore this triangle
 		// Checks the query geometry support point along the normal at the sweep end point
-		VectorRegister4Float IgnorePenetrationSimd = MakeVectorRegisterFloatFromDouble(VectorSetFloat1(IgnorePenetration));
+		// NOTE: We are using SupportCore, so we need to add Radius to the early-out distance check. We could just use Support
+		// but this way avoids an extra sqrt. We don't modify IgnorePenetration in the ctor because it gets used below.
+		VectorRegister4Float RadiusV = MakeVectorRegisterFloatFromDouble(VectorSetFloat1(QueryGeom.GetRadius()));
+		VectorRegister4Float IgnorePenetrationSimd = VectorSubtract(MakeVectorRegisterFloatFromDouble(VectorSetFloat1(IgnorePenetration)), RadiusV);
 		VectorRegister4Float LengthSimd = MakeVectorRegisterFloatFromDouble(VectorSetFloat1(LengthScale * Length));
 		VectorRegister4Float EndPointSimd = VectorAdd(TranslationSimd, VectorMultiply(RayDirSimd, LengthSimd));
 		VectorRegister4Float TriangleNormalSimd = VectorNormalize(VectorCross(VectorSubtract(B, A), VectorSubtract(C, A)));
