Make BodySetup use the TConvexHull path if it is the current default for the ConvexBuilder

diff --git a/Engine/Source/Runtime/Engine/Private/PhysicsEngine/BodySetup.cpp b/Engine/Source/Runtime/Engine/Private/PhysicsEngine/BodySetup.cpp
--- a/Engine/Source/Runtime/Engine/Private/PhysicsEngine/BodySetup.cpp
+++ b/Engine/Source/Runtime/Engine/Private/PhysicsEngine/BodySetup.cpp
@@ -1777,18 +1777,7 @@ TArray<int32> FKConvexElem::GetChaosConvexIndices() const
 			ConvexVertices[VertIndex] = VertexData[VertIndex];
 		}
 
-		TArray<Chaos::TVec3<int32>> Triangles;
-		Chaos::FConvexBuilder::Params BuildParams;
-		BuildParams.HorizonEpsilon = Chaos::FConvexBuilder::SuggestEpsilon(ConvexVertices);
-		Chaos::FConvexBuilder::BuildConvexHull(ConvexVertices, Triangles, BuildParams);
-
-		ResultIndexData.Reserve(Triangles.Num() * 3);
-		for (Chaos::TVec3<int32> Tri : Triangles)
-		{
-			ResultIndexData.Add(Tri[0]);
-			ResultIndexData.Add(Tri[1]);
-			ResultIndexData.Add(Tri[2]);
-		}
+		Chaos::FConvexBuilder::BuildIndices(ConvexVertices, ResultIndexData);
 	}
 
 	return ResultIndexData;
diff --git a/Engine/Source/Runtime/Experimental/Chaos/Private/Chaos/CollisionConvexMesh.cpp b/Engine/Source/Runtime/Experimental/Chaos/Private/Chaos/CollisionConvexMesh.cpp
--- a/Engine/Source/Runtime/Experimental/Chaos/Private/Chaos/CollisionConvexMesh.cpp
+++ b/Engine/Source/Runtime/Experimental/Chaos/Private/Chaos/CollisionConvexMesh.cpp
@@ -35,6 +35,41 @@ namespace Chaos
 		return BuildMethod == EBuildMethod::ConvexHull3;
 	}
 
+	void FConvexBuilder::BuildIndices(const TArray<FVec3Type>& InVertices, TArray<int32>& OutResultIndexData, EBuildMethod BuildMethod)
+	{
+		if (UseConvexHull3(BuildMethod))
+		{
+			UE::Geometry::TConvexHull3<FRealType> HullCompute;
+			if (HullCompute.Solve<FVec3Type>(InVertices))
+			{
+				const TArray<UE::Geometry::FIndex3i>& HullTriangles = HullCompute.GetTriangles();
+				OutResultIndexData.Reserve(HullTriangles.Num() * 3);
+				for (const UE::Geometry::FIndex3i& Tri : HullTriangles)
+				{
+					// Winding is backwards from what Chaos expects
+					OutResultIndexData.Add(Tri[0]);
+					OutResultIndexData.Add(Tri[2]);
+					OutResultIndexData.Add(Tri[1]);
+				}
+			}
+		}
+		else
+		{
+			TArray<Chaos::TVec3<int32>> Triangles;
+			Params BuildParams;
+			BuildParams.HorizonEpsilon = SuggestEpsilon(InVertices);
+			BuildConvexHull(InVertices, Triangles, BuildParams);
+			OutResultIndexData.Reserve(Triangles.Num() * 3);
+			for (Chaos::TVec3<int32> Tri : Triangles)
+			{
+				OutResultIndexData.Add(Tri[0]);
+				OutResultIndexData.Add(Tri[1]);
+				OutResultIndexData.Add(Tri[2]);
+			}
+		}
+
+	}
+
 	void FConvexBuilder::Build(const TArray<FVec3Type>& InVertices, TArray <FPlaneType>& OutPlanes, TArray<TArray<int32>>& OutFaceIndices, TArray<FVec3Type>& OutVertices, FAABB3Type& OutLocalBounds, EBuildMethod BuildMethod)
 	{
 		OutPlanes.Reset();
diff --git a/Engine/Source/Runtime/Experimental/Chaos/Public/Chaos/CollisionConvexMesh.h b/Engine/Source/Runtime/Experimental/Chaos/Public/Chaos/CollisionConvexMesh.h
--- a/Engine/Source/Runtime/Experimental/Chaos/Public/Chaos/CollisionConvexMesh.h
+++ b/Engine/Source/Runtime/Experimental/Chaos/Public/Chaos/CollisionConvexMesh.h
@@ -262,6 +262,8 @@ namespace Chaos
 		static CHAOS_API void Build(const TArray<FVec3Type>& InVertices, TArray<FPlaneType>& OutPlanes, TArray<TArray<int32>>& OutFaceIndices, TArray<FVec3Type>& OutVertices, FAABB3Type& OutLocalBounds, EBuildMethod BuildMethod = EBuildMethod::Default);
 
 		static bool UseConvexHull3(EBuildMethod BuildMethod);
+
+		static CHAOS_API void BuildIndices(const TArray<FVec3Type>& InVertices, TArray<int32>& OutResultIndexData, EBuildMethod BuildMethod = EBuildMethod::Default);
 		
 		static void BuildConvexHull(const TArray<FVec3Type>& InVertices, TArray<TVec3<int32>>& OutIndices, const Params& InParams = Params())
 		{
