New Modular Vehicle system initial pass

diff --git a/Engine/Source/Runtime/Experimental/Chaos/Private/GeometryCollection/ManagedArrayCollection.cpp b/Engine/Source/Runtime/Experimental/Chaos/Private/GeometryCollection/ManagedArrayCollection.cpp
--- a/Engine/Source/Runtime/Experimental/Chaos/Private/GeometryCollection/ManagedArrayCollection.cpp
+++ b/Engine/Source/Runtime/Experimental/Chaos/Private/GeometryCollection/ManagedArrayCollection.cpp
@@ -381,6 +381,12 @@ void FManagedArrayCollection::CopyAttribute(const FManagedArrayCollection& InCol
 	FKeyType SrcKey = FManagedArrayCollection::MakeMapKey(SrcName, Group);
 	FKeyType DestKey = FManagedArrayCollection::MakeMapKey(DestName, Group);
 
+	// fail gracefully
+	if (!(ensureMsgf(HasAttribute(SrcName, Group), TEXT("Missing managed array collection attribute"))))
+	{
+		return;
+	}
+
 	if (!HasAttribute(DestName, Group))
 	{
 		const FValueType& V = InCollection.Map[SrcKey];
diff --git a/Engine/Source/Runtime/Experimental/Chaos/Private/PhysicsProxy/GeometryCollectionPhysicsProxy.cpp b/Engine/Source/Runtime/Experimental/Chaos/Private/PhysicsProxy/GeometryCollectionPhysicsProxy.cpp
--- a/Engine/Source/Runtime/Experimental/Chaos/Private/PhysicsProxy/GeometryCollectionPhysicsProxy.cpp
+++ b/Engine/Source/Runtime/Experimental/Chaos/Private/PhysicsProxy/GeometryCollectionPhysicsProxy.cpp
@@ -543,9 +543,12 @@ void FGeometryCollectionPhysicsProxy::Initialize(Chaos::FPBDRigidsEvolutionBase
 								}
 								else
 								{
-									TUniquePtr<Chaos::FImplicitObject> CopiedChildImplicit = GameThreadCollection.Implicits[ChildIndex]->DeepCopy();
-									FImplicitObjectTransformed* TransformedChildImplicit = new FImplicitObjectTransformed(MoveTemp(CopiedChildImplicit), RelativeShapeTransform);  
-									ChildImplicits.Add(TUniquePtr<FImplicitObjectTransformed>(TransformedChildImplicit));
+								//	if (!(GameThreadCollection.StatusFlags[ChildIndex] & FGeometryCollection::ENodeFlags::FS_IgnoreCollisionInParentCluster))
+									{
+										TUniquePtr<Chaos::FImplicitObject> CopiedChildImplicit = GameThreadCollection.Implicits[ChildIndex]->DeepCopy();
+										FImplicitObjectTransformed* TransformedChildImplicit = new FImplicitObjectTransformed(MoveTemp(CopiedChildImplicit), RelativeShapeTransform);  
+										ChildImplicits.Add(TUniquePtr<FImplicitObjectTransformed>(TransformedChildImplicit));
+									}
 								}
 							}
 						}
diff --git a/Engine/Source/Runtime/Experimental/Chaos/Public/GeometryCollection/GeometryCollection.h b/Engine/Source/Runtime/Experimental/Chaos/Public/GeometryCollection/GeometryCollection.h
--- a/Engine/Source/Runtime/Experimental/Chaos/Public/GeometryCollection/GeometryCollection.h
+++ b/Engine/Source/Runtime/Experimental/Chaos/Public/GeometryCollection/GeometryCollection.h
@@ -117,7 +117,10 @@ public:
 		FS_None = 0,
 
 		// identify nodes that should be removed from the simulation instead of becoming a fractured body
-		FS_RemoveOnFracture = 0x00000004
+		FS_RemoveOnFracture = 0x00000004,
+
+		FS_IgnoreCollisionInParentCluster = 0x00000008
+
 	};
 	//
 	//
@@ -247,7 +250,6 @@ public:
 	/** Connection of leaf geometry */
 	TArray<TArray<int32>> ConnectionGraph();
 
-
 	//
 	//
 	//
diff --git a/Engine/Source/Runtime/Experimental/ChaosVehicles/ChaosVehiclesCore/Private/SimModule/AerofoilModule.cpp b/Engine/Source/Runtime/Experimental/ChaosVehicles/ChaosVehiclesCore/Private/SimModule/AerofoilModule.cpp
new file mode 100644
--- /dev/null
+++ b/Engine/Source/Runtime/Experimental/ChaosVehicles/ChaosVehiclesCore/Private/SimModule/AerofoilModule.cpp
@@ -0,0 +1,145 @@
+// Copyright Epic Games, Inc. All Rights Reserved.
+
+#include "SimModule/AerofoilModule.h"
+#include "SimModule/SimModuleTree.h"
+#include "VehicleUtility.h"
+
+#if VEHICLE_DEBUGGING_ENABLED
+PRAGMA_DISABLE_OPTIMIZATION
+#endif
+
+namespace Chaos
+{
+
+	void FAerofoilSimModule::Simulate(float DeltaTime, const FAllInputs& Inputs, FSimModuleTree& VehicleModuleSystem)
+	{
+		float Altitude = Inputs.VehicleWorldTransform.GetLocation().Z;
+
+		switch (Setup().Type)
+		{
+			case EAerofoil::Elevator:
+				ControlSurfaceAngle = Inputs.Pitch * Setup().MaxControlAngle;
+				break;
+
+			case EAerofoil::Rudder:
+				ControlSurfaceAngle = Inputs.Yaw * Setup().MaxControlAngle;
+				break;
+
+			case EAerofoil::Wing:
+				ControlSurfaceAngle = Inputs.Roll * Setup().MaxControlAngle;
+			break;
+		}
+
+		// needs local velocity at aerofoil to calculate lift and drag
+		FVector LocalForce = GetForce(ModuleLocalVelocity * CmToMScaling(), CmToM(Altitude), DeltaTime);
+		FVector LocalOffset = GetCenterOfLiftOffset() * MToCmScaling();
+
+		AddLocalForceAtPosition(LocalForce * MToCmScaling(), LocalOffset);
+	}
+
+	float FAerofoilSimModule::CalcAngleOfAttackDegrees(const FVector& UpAxis, const FVector& InAirflowVector)
+	{
+		float fMag = FVector::DotProduct(UpAxis, InAirflowVector);
+		return RadToDeg(FMath::Asin(fMag));
+	}
+
+	float FAerofoilSimModule::CalcLiftCoefficient(float InAngleOfAttack, float InControlSurfaceAngle)
+	{
+		float PeakValue = 2.0f; // typically the Coefficient can reach this peak value
+		float TotalAngle = InAngleOfAttack + InControlSurfaceAngle;
+
+		if (FMath::Abs(TotalAngle) > (Setup().StallAngle * 2.0f))
+		{
+			return 0.0f;
+		}
+
+		return FMath::Sin(TotalAngle * (PI * 0.5f) / Setup().StallAngle) * PeakValue;
+	}
+
+	float FAerofoilSimModule::CalcDragCoefficient(float InAngleOfAttack, float InControlSurfaceAngle)
+	{
+		if (InAngleOfAttack > 90.f)
+		{
+			InAngleOfAttack = 180.f - InAngleOfAttack;
+		}
+
+		if (InAngleOfAttack < -90.f)
+		{
+			InAngleOfAttack = -180.f - InAngleOfAttack;
+		}
+
+		float Value = (InAngleOfAttack + InControlSurfaceAngle) / (Setup().StallAngle + FMath::Abs(InControlSurfaceAngle));
+		return (0.05f + Value * Value);
+	}
+
+	float FAerofoilSimModule::CalcCentreOfLift()
+	{
+		// moves backwards past stall angle
+		if (AngleOfAttack > Setup().StallAngle)
+		{
+			return (AngleOfAttack - Setup().StallAngle) * 10.0f + 20.0f;
+		}
+
+		// moves forwards below stall angle
+		return (Setup().StallAngle - AngleOfAttack) * 20.0f / Setup().StallAngle + 20.0f;
+	}
+
+	float FAerofoilSimModule::CalcDynamicPressure(float VelocitySqr, float InAltitude)
+	{
+		float AltitudeMultiplierEffect = 1.0f;
+
+		return AltitudeMultiplierEffect * 0.5f * CurrentAirDensity * VelocitySqr;
+	}
+
+	FVector FAerofoilSimModule::GetCenterOfLiftOffset()
+	{
+		float X = 0.0f;
+
+		if (Setup().Type == EAerofoil::Wing)
+		{
+			X = (CalcCentreOfLift() - 50.0f) / 100.0f;
+		}
+
+		return Setup().Offset + FVector(X, 0.0f, 0.0f);
+	}
+
+	FVector FAerofoilSimModule::GetForce(const FVector& v, float Altitude, float DeltaTime)
+	{
+		FVector Force(0.0f, 0.0f, 0.0f);
+
+		float AirflowMagnitudeSqr = v.SizeSquared();
+
+		// can only generate lift if there is airflow over aerofoil, early out
+		if (FMath::Abs(AirflowMagnitudeSqr) < SMALL_NUMBER)
+		{
+			return Force;
+		}
+
+		// airflow direction in opposite direction to vehicle direction of travel
+		AirflowNormal = -v;
+		AirflowNormal.Normalize();
+
+		// determine angle of attack for control surface
+		AngleOfAttack = CalcAngleOfAttackDegrees(Setup().UpAxis, AirflowNormal);
+
+		// Aerofoil Camber and Control Surface just lumped together
+		float TotalControlAngle = ControlSurfaceAngle + Setup().Camber;
+
+		// dynamic pressure dependent on speed, altitude (air pressure)
+		float Common = Setup().Area * CalcDynamicPressure(AirflowMagnitudeSqr, Altitude);
+
+		// Lift and Drag coefficients are based on the angle of attack and Control Angle
+		float LiftCoef = CalcLiftCoefficient(AngleOfAttack, TotalControlAngle) * Setup().LiftMultiplier;
+		float DragCoef = CalcDragCoefficient(AngleOfAttack, TotalControlAngle) * Setup().DragMultiplier;
+
+		// Combine to create a single force vector
+		Force = Setup().UpAxis * (Common * LiftCoef) + AirflowNormal * (Common * DragCoef);
+
+		return Force;
+	}
+
+} // namespace Chaos
+
+#if VEHICLE_DEBUGGING_ENABLED
+PRAGMA_ENABLE_OPTIMIZATION
+#endif
diff --git a/Engine/Source/Runtime/Experimental/ChaosVehicles/ChaosVehiclesCore/Private/SimModule/AxleModule.cpp b/Engine/Source/Runtime/Experimental/ChaosVehicles/ChaosVehiclesCore/Private/SimModule/AxleModule.cpp
new file mode 100644
--- /dev/null
+++ b/Engine/Source/Runtime/Experimental/ChaosVehicles/ChaosVehiclesCore/Private/SimModule/AxleModule.cpp
@@ -0,0 +1,34 @@
+// Copyright Epic Games, Inc. All Rights Reserved.
+
+#pragma once
+
+#include "SimModule/AxleModule.h"
+#include "SimModule/SimModuleTree.h"
+#include "VehicleUtility.h"
+
+#if VEHICLE_DEBUGGING_ENABLED
+PRAGMA_DISABLE_OPTIMIZATION
+#endif
+
+namespace Chaos
+{
+
+	FAxleSimModule::FAxleSimModule(const FAxleSettings& Settings)
+		: TSimModuleSettings<FAxleSettings>(Settings)
+	{
+
+	}
+
+	void FAxleSimModule::Simulate(float DeltaTime, const FAllInputs& Inputs, FSimModuleTree& VehicleModuleSystem)
+	{
+		float BrakeTorque = 0.0f;
+		TransmitTorque(VehicleModuleSystem, DriveTorque, BrakeTorque);
+		IntegrateAngularVelocity(DeltaTime, Setup().AxleInertia);
+	}
+
+
+} // namespace Chaos
+
+#if VEHICLE_DEBUGGING_ENABLED
+PRAGMA_ENABLE_OPTIMIZATION
+#endif
diff --git a/Engine/Source/Runtime/Experimental/ChaosVehicles/ChaosVehiclesCore/Private/SimModule/ChassisModule.cpp b/Engine/Source/Runtime/Experimental/ChaosVehicles/ChaosVehiclesCore/Private/SimModule/ChassisModule.cpp
new file mode 100644
--- /dev/null
+++ b/Engine/Source/Runtime/Experimental/ChaosVehicles/ChaosVehiclesCore/Private/SimModule/ChassisModule.cpp
@@ -0,0 +1,25 @@
+// Copyright Epic Games, Inc. All Rights Reserved.
+
+#pragma once
+
+#include "SimModule/ChassisModule.h"
+#include "SimModule/SimModuleTree.h"
+
+
+namespace Chaos
+{
+
+	FChassisSimModule::FChassisSimModule(const FChassisSettings& Settings)
+		: TSimModuleSettings<FChassisSettings>(Settings)
+	{
+
+	}
+
+	void FChassisSimModule::Simulate(float DeltaTime, const FAllInputs& Inputs, FSimModuleTree& VehicleModuleSystem)
+	{
+		AddLocalForce(FVector(20000000,0,0));
+
+	}
+
+
+} // namespace Chaos
diff --git a/Engine/Source/Runtime/Experimental/ChaosVehicles/ChaosVehiclesCore/Private/SimModule/ClutchModule.cpp b/Engine/Source/Runtime/Experimental/ChaosVehicles/ChaosVehiclesCore/Private/SimModule/ClutchModule.cpp
new file mode 100644
--- /dev/null
+++ b/Engine/Source/Runtime/Experimental/ChaosVehicles/ChaosVehiclesCore/Private/SimModule/ClutchModule.cpp
@@ -0,0 +1,66 @@
+// Copyright Epic Games, Inc. All Rights Reserved.
+
+#pragma once
+
+#include "SimModule/ClutchModule.h"
+#include "SimModule/SimModuleTree.h"
+#include "VehicleUtility.h"
+
+#if VEHICLE_DEBUGGING_ENABLED
+PRAGMA_DISABLE_OPTIMIZATION
+#endif
+
+namespace Chaos
+{
+
+	FClutchSimModule::FClutchSimModule(const FClutchSettings& Settings) : TSimModuleSettings<FClutchSettings>(Settings)
+		, ClutchValue(0.f)
+	{
+
+	}
+
+	void FClutchSimModule::Simulate(float DeltaTime, const FAllInputs& Inputs, FSimModuleTree& VehicleModuleSystem)
+	{
+		//FTorqueSimModule* Parent = static_cast<FTorqueSimModule*>(GetParent());
+		//FTorqueSimModule* Child = static_cast<FTorqueSimModule*>(GetFirstChild());
+
+		//float EngineSpeed = Parent->GetAngularVelocity();
+		//float TransmissionSpeed = Child->GetAngularVelocity();
+
+		//// difference in speed between the two plates
+		//float AngularVelocityDifference = EngineSpeed - TransmissionSpeed;
+
+		// Inputs.Clutch 0 is engaged/locked, 1 is depressed/open
+		ClutchValue = (1.0f - Inputs.Clutch) * Setup().ClutchStrength;
+
+		FTorqueSimModule* Parent = static_cast<FTorqueSimModule*>(GetParent());
+		FTorqueSimModule* Child = static_cast<FTorqueSimModule*>(GetFirstChild());
+
+		float EngineSpeed = Parent->GetAngularVelocity();
+		float TransmissionSpeed = Child->GetAngularVelocity();
+
+		// difference in speed between the two plates
+		float AngularVelocityDifference = EngineSpeed - TransmissionSpeed;
+
+		Parent->AddAngularVelocity(-AngularVelocityDifference * 0.1f);
+		Child->AddAngularVelocity(AngularVelocityDifference * 0.1f);
+
+		float BrakeTorque = 0.0f;
+		TransmitTorque(VehicleModuleSystem, DriveTorque, BrakeTorque, 1.0f, 1.0f/*ClutchValue*/);
+	}
+
+	bool FClutchSimModule::GetDebugString(FString& StringOut) const
+	{
+		FTorqueSimModule::GetDebugString(StringOut);
+		StringOut += FString::Format(TEXT("Value {0}")
+			, { ClutchValue });
+		return true;
+
+	}
+
+
+} // namespace Chaos
+
+#if VEHICLE_DEBUGGING_ENABLED
+PRAGMA_ENABLE_OPTIMIZATION
+#endif
diff --git a/Engine/Source/Runtime/Experimental/ChaosVehicles/ChaosVehiclesCore/Private/SimModule/DeferredForcesModular.cpp b/Engine/Source/Runtime/Experimental/ChaosVehicles/ChaosVehiclesCore/Private/SimModule/DeferredForcesModular.cpp
new file mode 100644
--- /dev/null
+++ b/Engine/Source/Runtime/Experimental/ChaosVehicles/ChaosVehiclesCore/Private/SimModule/DeferredForcesModular.cpp
@@ -0,0 +1,191 @@
+// Copyright Epic Games, Inc. All Rights Reserved.
+
+#include "SimModule/DeferredForcesModular.h"
+#include "SimModule/SimulationModuleBase.h"
+#include "Chaos/Particle/ParticleUtilities.h"
+#include "Chaos/Utilities.h"
+#include "Chaos/ParticleHandle.h"
+
+#if CHAOS_DEBUG_DRAW
+#include "Chaos/DebugDrawQueue.h"
+#endif
+
+
+FCoreModularVehicleDebugParams GCoreModularVehicleDebugParams;
+
+#if CHAOS_DEBUG_DRAW
+FAutoConsoleVariableRef CVarChaosModularVehiclesShowMass(TEXT("p.ModularVehicle.ShowMass"), GCoreModularVehicleDebugParams.ShowMass, TEXT("Show mass and inertia above objects."));
+FAutoConsoleVariableRef CVarChaosModularVehiclesShowForces(TEXT("p.ModularVehicle.ShowForces"), GCoreModularVehicleDebugParams.ShowForces, TEXT("Show forces on modular vehicles."));
+FAutoConsoleVariableRef CVarChaosModularVehiclesDrawForceScaling(TEXT("p.ModularVehicle.DrawForceScaling"), GCoreModularVehicleDebugParams.DrawForceScaling, TEXT("Set scaling of modular vehicle debug draw forces."));
+#endif
+
+// TODO: passing all these parameters in is horrible, tidy this up now we know what needs to be done
+Chaos::FPBDRigidClusteredParticleHandle* FDeferredForcesModular::GetParticle(TArray<Chaos::FPBDRigidClusteredParticleHandle*>& Particles
+		, TArray<Chaos::FPBDRigidClusteredParticleHandle*>& ClusterParticles
+		, int TransformIndex
+		, const TManagedArray<FTransform>& Transforms
+		, const TManagedArray<FTransform>& CollectionMassToLocal
+		, const TManagedArray<int32>& Parent
+		, FTransform& TransformOut)
+{
+
+	Chaos::FPBDRigidClusteredParticleHandle* SimParticle = Particles[TransformIndex];
+	Chaos::FPBDRigidClusteredParticleHandle* SimClusterParticle = ClusterParticles[TransformIndex];
+
+	if (SimParticle && !SimParticle->Disabled())
+	{
+		TransformOut = FTransform::Identity;
+		return SimParticle;		// apply to broken off child
+	}
+	else if (SimClusterParticle && !SimClusterParticle->Disabled())
+	{
+		TransformOut = CollectionMassToLocal[Parent[TransformIndex]].Inverse() * Transforms[TransformIndex];
+		return SimClusterParticle; // apply to intact cluster
+	}
+
+	return nullptr;
+}
+
+void FDeferredForcesModular::Apply(TArray<Chaos::FPBDRigidClusteredParticleHandle*>& Particles
+		, TArray<Chaos::FPBDRigidClusteredParticleHandle*>& ClusterParticles
+		, const TManagedArray<FTransform>& Transforms
+	, const TManagedArray<FTransform>& CollectionMassToLocal
+	, const TManagedArray<int32>& Parent)
+{
+	FTransform RelativeTransform;
+	for (const FApplyForceData& Data : ApplyForceDatas)
+	{
+		Chaos::FPBDRigidClusteredParticleHandle* RigidHandle = GetParticle(Particles, ClusterParticles, Data.TransformIndex, Transforms, CollectionMassToLocal, Parent, RelativeTransform);
+		if (RigidHandle)
+		{
+			AddForce(RigidHandle, Data, RelativeTransform);
+		}
+	}
+
+	for (const FApplyForceAtPositionData& Data : ApplyForceAtPositionDatas)
+	{
+		Chaos::FPBDRigidClusteredParticleHandle* RigidHandle = GetParticle(Particles, ClusterParticles, Data.TransformIndex, Transforms, CollectionMassToLocal, Parent, RelativeTransform);
+		if (RigidHandle)
+		{
+			AddForceAtPosition(RigidHandle, Data, RelativeTransform);
+		}
+	}
+
+	for (const FAddTorqueInRadiansData& Data : ApplyTorqueDatas)
+	{
+		Chaos::FPBDRigidClusteredParticleHandle* RigidHandle = GetParticle(Particles, ClusterParticles, Data.TransformIndex, Transforms, CollectionMassToLocal, Parent, RelativeTransform);
+		if (RigidHandle)
+		{
+			AddTorque(RigidHandle, Data, RelativeTransform);
+		}
+	}
+
+	ApplyForceDatas.Empty();
+	ApplyForceAtPositionDatas.Empty();
+	ApplyTorqueDatas.Empty();
+}
+
+
+void AddForce_Implementation(Chaos::FPBDRigidClusteredParticleHandle* RigidHandle, const FVector& LocalForce, const FTransform& OffsetTransform)
+{
+	const Chaos::FVec3 WorldCOM = Chaos::FParticleUtilitiesGT::GetCoMWorldPosition(RigidHandle);
+
+	// local to world
+	const FTransform WorldTM(RigidHandle->R(), RigidHandle->X());
+	Chaos::FVec3 Position = WorldTM.TransformPosition(OffsetTransform.GetLocation());
+
+	Chaos::FVec3 Force = WorldTM.TransformVector(LocalForce);
+
+	Chaos::FVec3 Arm = Position - WorldCOM;
+
+	const Chaos::FVec3 WorldTorque = Chaos::FVec3::CrossProduct(Arm, Force);
+	auto* RP = RigidHandle->CastToRigidParticle();
+	if (RP)
+	{
+		RP->AddForce(Force);
+		RP->AddTorque(WorldTorque);
+
+		// #TODO: remove HACK - doing this because vehicles are falling asleep and not waking up when forces are applied
+		RP->SetSleepType(Chaos::ESleepType::NeverSleep);
+
+#if CHAOS_DEBUG_DRAW
+		if (GCoreModularVehicleDebugParams.ShowMass)
+		{
+			FString OutputText = FString::Format(TEXT("{0} {1}"), { RP->M(), RP->I().ToString() });
+			Chaos::FDebugDrawQueue::GetInstance().DrawDebugString(RigidHandle->X() + FVector(0,0,200), OutputText, nullptr, FColor::White, -1, true, 1.0f);
+		}
+#endif
+	}
+
+#if CHAOS_DEBUG_DRAW
+	if (GCoreModularVehicleDebugParams.ShowForces)
+	{
+		// TODO: perhaps render different force types/sources in different colors, i.e. aerofoil, suspension, friction, etc.
+		Chaos::FDebugDrawQueue::GetInstance().DrawDebugLine(Position, Position + Force * GCoreModularVehicleDebugParams.DrawForceScaling, FColor::Blue, false, -1.f, 0, 5.f);
+		Chaos::FDebugDrawQueue::GetInstance().DrawDebugSphere(Position, 5, 8, FColor::White, false, -1.f, 0, 5.f);
+	}
+#endif
+
+}
+
+void FDeferredForcesModular::AddForce(Chaos::FPBDRigidClusteredParticleHandle* RigidHandle, const FApplyForceData& DataIn, const FTransform& OffsetTransform)
+{
+	if (ensure(RigidHandle))
+	{
+		AddForce_Implementation(RigidHandle, DataIn.Force, OffsetTransform);
+	}
+
+}
+
+//void FDeferredForcesModular::AddForce(Chaos::FPBDRigidClusteredParticleHandle* RigidHandle, const FApplyForceData& DataIn)
+//{
+//	check(false);
+//	if (ensure(RigidHandle))
+//	{
+//		Chaos::EObjectStateType ObjectState = RigidHandle->ObjectState();
+//		if (CHAOS_ENSURE(ObjectState == Chaos::EObjectStateType::Dynamic || ObjectState == Chaos::EObjectStateType::Sleeping))
+//		{
+//			if (DataIn.Flags == EForceFlags::AccelChange)
+//			{
+//				const float RigidMass = RigidHandle->M();
+//				const Chaos::FVec3 Acceleration = DataIn.Force * RigidMass;
+//				RigidHandle->AddForce(Acceleration, false);
+//			}
+//			else
+//			{
+//				RigidHandle->AddForce(DataIn.Force, false);
+//			}
+//
+//		}
+//	}
+//}
+
+void FDeferredForcesModular::AddForceAtPosition(Chaos::FPBDRigidClusteredParticleHandle* RigidHandle, const FApplyForceAtPositionData& DataIn, const FTransform& OffsetTransform)
+{
+	if (ensure(RigidHandle))
+	{
+		FTransform ShiftedTransform = OffsetTransform;
+		ShiftedTransform.SetLocation(ShiftedTransform.GetLocation() + DataIn.Position);
+		AddForce_Implementation(RigidHandle, DataIn.Force, ShiftedTransform);
+	}
+}
+
+void FDeferredForcesModular::AddTorque(Chaos::FPBDRigidClusteredParticleHandle* RigidHandle, const FAddTorqueInRadiansData& DataIn, const FTransform& OffsetTransform)
+{	
+	check(false); //#TODO - incorporate OffsetTransform as this is currently wrong
+	//if (ensure(RigidHandle))
+	//{
+	//	Chaos::EObjectStateType ObjectState = RigidHandle->ObjectState();
+	//	if (CHAOS_ENSURE(ObjectState == Chaos::EObjectStateType::Dynamic || ObjectState == Chaos::EObjectStateType::Sleeping))
+	//	{
+	//		if (DataIn.Flags == EForceFlags::AccelChange)
+	//		{
+	//			RigidHandle->AddTorque(Chaos::FParticleUtilitiesXR::GetWorldInertia(RigidHandle) * DataIn.Torque, false);
+	//		}
+	//		else
+	//		{
+	//			RigidHandle->AddTorque(DataIn.Torque, false);
+	//		}
+	//	}
+	//}
+}
diff --git a/Engine/Source/Runtime/Experimental/ChaosVehicles/ChaosVehiclesCore/Private/SimModule/EngineModule.cpp b/Engine/Source/Runtime/Experimental/ChaosVehicles/ChaosVehiclesCore/Private/SimModule/EngineModule.cpp
new file mode 100644
--- /dev/null
+++ b/Engine/Source/Runtime/Experimental/ChaosVehicles/ChaosVehiclesCore/Private/SimModule/EngineModule.cpp
@@ -0,0 +1,80 @@
+// Copyright Epic Games, Inc. All Rights Reserved.
+
+#pragma once
+
+#include "SimModule/EngineModule.h"
+#include "SimModule/SimModuleTree.h"
+#include "VehicleUtility.h"
+
+#if VEHICLE_DEBUGGING_ENABLED
+PRAGMA_DISABLE_OPTIMIZATION
+#endif
+
+namespace Chaos
+{
+
+	void FEngineSimModule::Simulate(float DeltaTime, const FAllInputs& Inputs, FSimModuleTree& VehicleModuleSystem)
+	{
+		if (!EngineStarted)
+		{
+			return;
+		}
+
+		// TODO: Engine braking effect
+		DriveTorque = GetEngineTorque(Inputs.Throttle, GetRPM());
+		
+		float BrakeTorque = 0.f;
+		TransmitTorque(VehicleModuleSystem, DriveTorque, BrakeTorque);
+		IntegrateAngularVelocity(DeltaTime, Setup().EngineInertia);
+
+		// clamp RPM
+		AngularVelocity = FMath::Clamp(AngularVelocity, -MaxEngineSpeed, MaxEngineSpeed);
+
+		// we don't let the engine stall
+		if (AngularVelocity < EngineIdleSpeed)
+		{
+			AngularVelocity = EngineIdleSpeed;
+		}
+
+	}
+
+	bool FEngineSimModule::GetDebugString(FString& StringOut) const
+	{
+		FTorqueSimModule::GetDebugString(StringOut);
+		StringOut += FString::Format(TEXT("Drive {0}, Brake {1}, Load {2} RPM {3}"), { DriveTorque, BrakingTorque, LoadTorque, GetRPM() });
+		return true;
+	}
+
+
+	float FEngineSimModule::GetEngineTorque(float ThrottlePosition, float EngineRPM)
+	{
+		if (EngineStarted)
+		{
+			return ThrottlePosition * GetTorqueFromRPM(EngineRPM);
+		}
+
+		return 0.f;
+	}
+
+	float FEngineSimModule::GetTorqueFromRPM(float RPM, bool LimitToIdle)
+	{
+		//return Setup().MaxTorque; // TODO: Fix this - engine RPM running rampant
+
+		if (!EngineStarted || (FMath::Abs(RPM - Setup().MaxRPM) < 1.0f) || Setup().MaxRPM == 0)
+		{
+			return 0.0f;
+		}
+
+		if (LimitToIdle)
+		{
+			RPM = FMath::Clamp(RPM, (float)Setup().IdleRPM, (float)Setup().MaxRPM);
+		}
+
+		return Setup().TorqueCurve.GetValue(RPM, Setup().MaxRPM, Setup().MaxTorque);
+	}
+
+} // namespace Chaos
+
+#if VEHICLE_DEBUGGING_ENABLED
+PRAGMA_ENABLE_OPTIMIZATION
+#endif
diff --git a/Engine/Source/Runtime/Experimental/ChaosVehicles/ChaosVehiclesCore/Private/SimModule/MotorModule.cpp b/Engine/Source/Runtime/Experimental/ChaosVehicles/ChaosVehiclesCore/Private/SimModule/MotorModule.cpp
new file mode 100644
--- /dev/null
+++ b/Engine/Source/Runtime/Experimental/ChaosVehicles/ChaosVehiclesCore/Private/SimModule/MotorModule.cpp
@@ -0,0 +1,55 @@
+// Copyright Epic Games, Inc. All Rights Reserved.
+
+#pragma once
+
+#include "SimModule/MotorModule.h"
+#include "SimModule/SimModuleTree.h"
+#include "VehicleUtility.h"
+
+#if VEHICLE_DEBUGGING_ENABLED
+PRAGMA_DISABLE_OPTIMIZATION
+#endif
+
+
+namespace Chaos
+{
+
+	FMotorSimModule::FMotorSimModule(const FMotorSettings& Settings)
+		: TSimModuleSettings<FMotorSettings>(Settings)
+	{
+	}
+
+	void FMotorSimModule::Simulate(float DeltaTime, const FAllInputs& Inputs, FSimModuleTree& VehicleModuleSystem)
+	{
+		float MaxRPM = Setup().MaxRPM;
+
+		// protect against divide by zero
+		if (MaxRPM > SMALL_NUMBER)
+		{
+			float NormalizedRPM = GetRPM() / Setup().MaxRPM;
+
+			float NormalizedTorque = 0.0f;
+			if (NormalizedRPM > 0.0f)
+			{
+				NormalizedTorque = 1.0f - FMath::Square(1.0f - 2.0f * NormalizedRPM) * 0.5f;
+			}
+			else
+			{
+				NormalizedTorque = -1.0f - FMath::Square(1.0f + 2.0f * NormalizedRPM);
+			}
+
+			DriveTorque = NormalizedTorque * Setup().MaxTorque * Inputs.Throttle;
+
+			float BrakeTorque = 0.f;
+			TransmitTorque(VehicleModuleSystem, DriveTorque, BrakeTorque);
+			IntegrateAngularVelocity(DeltaTime, Setup().EngineInertia);
+		}
+
+	}
+
+
+} // namespace Chaos
+
+#if VEHICLE_DEBUGGING_ENABLED
+PRAGMA_ENABLE_OPTIMIZATION
+#endif
diff --git a/Engine/Source/Runtime/Experimental/ChaosVehicles/ChaosVehiclesCore/Private/SimModule/SimModuleTree.cpp b/Engine/Source/Runtime/Experimental/ChaosVehicles/ChaosVehiclesCore/Private/SimModule/SimModuleTree.cpp
new file mode 100644
--- /dev/null
+++ b/Engine/Source/Runtime/Experimental/ChaosVehicles/ChaosVehiclesCore/Private/SimModule/SimModuleTree.cpp
@@ -0,0 +1,251 @@
+// Copyright Epic Games, Inc. All Rights Reserved.
+
+#include "SimModule/SimModuleTree.h"
+#include "SimModule/SimulationModuleBase.h"
+#include "Chaos/ParticleHandleFwd.h"
+#include "PhysicsProxy/GeometryCollectionPhysicsProxy.h"
+#include "VehicleUtility.h"
+
+#if VEHICLE_DEBUGGING_ENABLED
+PRAGMA_DISABLE_OPTIMIZATION
+#endif
+
+namespace Chaos
+{
+
+int FSimModuleTree::AddRoot(ISimulationModuleBase* SimModule)
+{
+	return AddNodeBelow(-1, SimModule);
+}
+
+void FSimModuleTree::Reparent(int AtIndex, int ParentIndex)
+{
+	check(AtIndex < SimulationModuleTree.Num());
+	check(ParentIndex < SimulationModuleTree.Num());
+
+	UE_LOG(LogSimulationModule, Log, TEXT("Reparent %s To %s")
+		, *SimulationModuleTree[AtIndex].SimModule->GetDebugName()
+		, *SimulationModuleTree[ParentIndex].SimModule->GetDebugName());
+
+	int OrginalParent = SimulationModuleTree[AtIndex].Parent;
+	if (OrginalParent != ParentIndex)
+	{
+		SimulationModuleTree[AtIndex].Parent = ParentIndex;
+		SimulationModuleTree[ParentIndex].Children.Add(AtIndex);
+
+		SimulationModuleTree[OrginalParent].Children.Remove(AtIndex);
+	}
+}
+
+int FSimModuleTree::AddNodeBelow(int AtIndex, ISimulationModuleBase* SimModule)
+{
+	int NewIndex = GetNextIndex();
+	FSimModuleNode& Node = SimulationModuleTree[NewIndex];
+	SimModule->SetTreeIndex(NewIndex);
+	Node.SimModule = SimModule;
+	Node.Parent = AtIndex;
+	if (AtIndex >= 0)
+	{
+		SimulationModuleTree[AtIndex].Children.Add(NewIndex);
+	}
+	else
+	{
+		Node.Parent = FSimModuleNode::INVALID_INDEX;
+	}
+
+	return NewIndex;
+}
+
+int FSimModuleTree::GetNextIndex()
+{
+	int NewIndex = FSimModuleNode::INVALID_INDEX;
+	if (FreeList.IsEmpty())
+	{
+		NewIndex = SimulationModuleTree.Num();
+		SimulationModuleTree.AddZeroed(1);
+		SimulationModuleTree[NewIndex].Parent = FSimModuleNode::INVALID_INDEX;
+		SimulationModuleTree[NewIndex].SimModule = nullptr;
+	}
+	else
+	{
+		NewIndex = FreeList.Pop();
+	}
+
+	return NewIndex;
+}
+
+int FSimModuleTree::InsertNodeAbove(int AtIndex, ISimulationModuleBase* SimModule)
+{
+	int NewIndex = FSimModuleNode::INVALID_INDEX;
+
+	if (ensure(AtIndex < SimulationModuleTree.Num()))
+	{
+		NewIndex = GetNextIndex();
+		FSimModuleNode& Node = SimulationModuleTree[NewIndex];
+
+		int OriginalParentIdx = SimulationModuleTree[AtIndex].Parent;
+
+		// remove current idx from children of parent & add new index in its place
+		SimulationModuleTree[OriginalParentIdx].Children.Remove(AtIndex);
+		SimulationModuleTree[OriginalParentIdx].Children.Add(NewIndex);
+
+		SimulationModuleTree[AtIndex].Parent = NewIndex;
+
+		Node.SimModule = SimModule;
+		Node.Parent = OriginalParentIdx;	// new node takes parent from existing node
+		Node.Children.Add(AtIndex);			// existing node becomes child of new node
+		SimModule->SetTreeIndex(NewIndex);
+	}
+
+	return NewIndex;
+}
+
+void FSimModuleTree::DeleteNode(int AtIndex)
+{
+	// multiple children might become equal parents?
+	
+	int ParentIndex = SimulationModuleTree[AtIndex].Parent;
+	
+	if (ParentIndex >= 0)
+	{
+		// remove from parents children list
+		SimulationModuleTree[ParentIndex].Children.Remove(AtIndex);
+	}
+
+	// move deleted nodes children to parent and these children need new parent
+	for (int ChildIndex : SimulationModuleTree[AtIndex].Children)
+	{
+		if (ParentIndex >= 0)
+		{
+			SimulationModuleTree[ParentIndex].Children.Add(ChildIndex);
+		}
+		SimulationModuleTree[ChildIndex].Parent = ParentIndex;
+	}	
+
+	SimulationModuleTree[AtIndex].Parent = FSimModuleNode::INVALID_INDEX;
+	SimulationModuleTree[AtIndex].Children.Empty();
+	delete SimulationModuleTree[AtIndex].SimModule;
+	SimulationModuleTree[AtIndex].SimModule = nullptr;
+
+	FreeList.Push(AtIndex);
+
+}
+
+void FSimModuleTree::Simulate(float DeltaTime, FAllInputs& Inputs, FGeometryCollectionPhysicsProxy* PhysicsProxy)
+{
+	if (PhysicsProxy)
+	{
+		UpdateModuleVelocites(PhysicsProxy);
+	}
+
+	TArray<int> RootNodes;
+	GetRootNodes(RootNodes);
+
+	for (int RootIndex : RootNodes)
+	{
+		SimulateNode(DeltaTime, Inputs, RootIndex);
+	}
+}
+
+void FSimModuleTree::SimulateNode(float DeltaTime, FAllInputs& Inputs, int NodeIndex)
+{
+	if (ISimulationModuleBase* Module = AccessSimModule(NodeIndex))
+	{
+		if (Module->IsEnabled())
+		{
+			Module->Simulate(DeltaTime, Inputs, *this);
+		}
+
+		for (int ChildIdx : GetChildren(NodeIndex))
+		{
+			SimulateNode(DeltaTime, Inputs, ChildIdx);
+		}
+	}
+}
+
+
+void FSimModuleTree::DeleteNodesBelow(int AtIndex)
+{
+	if (IsValidNode(AtIndex))
+	{
+		for (int ChildIdx : GetChildren(AtIndex))
+		{
+			DeleteNodesBelow(ChildIdx);
+		}
+
+		delete SimulationModuleTree[AtIndex].SimModule;
+		SimulationModuleTree[AtIndex].SimModule = nullptr;
+		SimulationModuleTree[AtIndex].Children.Empty();
+		SimulationModuleTree[AtIndex].Parent = FSimModuleNode ::INVALID_INDEX;
+
+		FreeList.Push(AtIndex);
+
+	}
+}
+
+
+void FSimModuleTree::GetRootNodes(TArray<int>& RootNodesOut)
+{
+	RootNodesOut.Empty();
+
+	// never assume the root bone is always index 0
+	for (int i = 0; i < SimulationModuleTree.Num(); i++)
+	{
+		if (SimulationModuleTree[i].SimModule != nullptr && SimulationModuleTree[i].Parent == FSimModuleNode::INVALID_INDEX)
+		{
+			RootNodesOut.Add(i);
+		}
+	}
+}
+
+void FSimModuleTree::UpdateModuleVelocites(FGeometryCollectionPhysicsProxy* PhysicsProxy)
+{
+	check(PhysicsProxy);
+
+	const TArray<Chaos::FPBDRigidClusteredParticleHandle*>& Clusters = PhysicsProxy->GetSolverClusterHandles();
+	const TArray<Chaos::FPBDRigidClusteredParticleHandle*>& Particles = PhysicsProxy->GetSolverParticleHandles();
+
+	// capture the velocities at the start of each sim iteration
+	for (int i = 0; i < SimulationModuleTree.Num(); i++)
+	{
+		if (ISimulationModuleBase* Module = SimulationModuleTree[i].SimModule)
+		{
+			Chaos::FPBDRigidClusteredParticleHandle* ParentParticle = Clusters[Module->GetTransformIndex()];
+			const FTransform BodyTransform(ParentParticle->R(), ParentParticle->X());
+
+			if (Module->IsBehaviourType(eSimModuleTypeFlags::Velocity))
+			{
+				Chaos::FPBDRigidClusteredParticleHandle* Particle = nullptr;
+				if (Module->IsClustered())
+				{
+					Particle = Clusters[Module->GetTransformIndex()];
+				}
+				else
+				{
+					Particle = Particles[Module->GetTransformIndex()];				
+				}
+
+				if (Particle)
+				{
+					FVector WheelWorldLocation = BodyTransform.TransformPosition(Module->GetParentRelativeTransform().GetLocation());
+					const Chaos::FVec3 Arm = WheelWorldLocation - Particle->X();
+
+					FVector WorldVelocity = Particle->V() - Chaos::FVec3::CrossProduct(Arm, Particle->W());
+					FVector LocalVelocity = BodyTransform.InverseTransformVector(WorldVelocity);
+
+					Module->SetLocalVelocity(LocalVelocity);
+				}
+			
+			}
+		}
+	}
+
+}
+
+
+} // namespace Chaos
+
+
+#if VEHICLE_DEBUGGING_ENABLED
+PRAGMA_ENABLE_OPTIMIZATION
+#endif
diff --git a/Engine/Source/Runtime/Experimental/ChaosVehicles/ChaosVehiclesCore/Private/SimModule/SimulationModuleBase.cpp b/Engine/Source/Runtime/Experimental/ChaosVehicles/ChaosVehiclesCore/Private/SimModule/SimulationModuleBase.cpp
new file mode 100644
--- /dev/null
+++ b/Engine/Source/Runtime/Experimental/ChaosVehicles/ChaosVehiclesCore/Private/SimModule/SimulationModuleBase.cpp
@@ -0,0 +1,56 @@
+// Copyright Epic Games, Inc. All Rights Reserved.
+
+#include "SimModule/SimulationModuleBase.h"
+#include "SimModule/SimModuleTree.h"
+#include "SimModule/DeferredForcesModular.h"
+
+DEFINE_LOG_CATEGORY(LogSimulationModule);
+
+FCoreModularVehicleDebugParams GModularVehicleDebugParams;
+
+namespace Chaos
+{
+
+void ISimulationModuleBase::AddLocalForceAtPosition(const FVector& Force, const FVector& Position, bool bAllowSubstepping, bool bIsLocalForce)
+{
+	AppliedForce = Force;
+	if (SimModuleTree)
+	{
+		SimModuleTree->AccessDeferredForces().Add(FDeferredForcesModular::FApplyForceAtPositionData(TransformIndex, Force, Position, bAllowSubstepping, bIsLocalForce));
+	}
+}
+
+void ISimulationModuleBase::AddLocalForce(const FVector& Force, bool bAllowSubstepping, bool bIsLocalForce)
+{
+	AppliedForce = Force;
+	if (SimModuleTree)
+	{
+		SimModuleTree->AccessDeferredForces().Add(FDeferredForcesModular::FApplyForceData(TransformIndex, Force, bAllowSubstepping, bIsLocalForce));
+	}
+}
+
+ISimulationModuleBase* ISimulationModuleBase::GetParent()
+{
+	return (SimModuleTree != nullptr) ? SimModuleTree->AccessSimModule(SimModuleTree->GetParent(SimTreeIndex)) : nullptr;
+}
+
+ISimulationModuleBase* ISimulationModuleBase::GetFirstChild()
+{
+	const TSet<int32>& Children = SimModuleTree->GetChildren(SimTreeIndex);
+	for (int ChildIndex : Children)
+	{ 
+		return (SimModuleTree != nullptr) ? SimModuleTree->AccessSimModule(ChildIndex) : nullptr;
+	}
+	return nullptr;
+}
+
+bool ISimulationModuleBase::GetDebugString(FString& StringOut) const 
+{
+	StringOut += FString::Format(TEXT("{0}: Enabled {1}, InCluster {2}, TFormIdx {3}, ")
+		, { GetDebugName(), IsEnabled(), IsClustered(), GetTransformIndex() });
+
+	return true; 
+}
+
+
+} //namespace Chaos
\ No newline at end of file
diff --git a/Engine/Source/Runtime/Experimental/ChaosVehicles/ChaosVehiclesCore/Private/SimModule/SuspensionModule.cpp b/Engine/Source/Runtime/Experimental/ChaosVehicles/ChaosVehiclesCore/Private/SimModule/SuspensionModule.cpp
new file mode 100644
--- /dev/null
+++ b/Engine/Source/Runtime/Experimental/ChaosVehicles/ChaosVehiclesCore/Private/SimModule/SuspensionModule.cpp
@@ -0,0 +1,90 @@
+// Copyright Epic Games, Inc. All Rights Reserved.
+
+#pragma once
+
+#include "SimModule/SuspensionModule.h"
+#include "SimModule/SimModuleTree.h"
+#include "SimModule/WheelModule.h"
+#include "VehicleUtility.h"
+
+#if VEHICLE_DEBUGGING_ENABLED
+PRAGMA_DISABLE_OPTIMIZATION
+#endif
+
+namespace Chaos
+{
+
+	FSuspensionSimModule::FSuspensionSimModule(const FSuspensionSettings& Settings)
+		: TSimModuleSettings<FSuspensionSettings>(Settings)
+		, SpringDisplacement(0.f)
+		, LastDisplacement(0.f)
+		, WheelSimTreeIndex(INVALID_INDEX)
+	{
+		AccessSetup().MaxLength = FMath::Abs(Settings.MaxRaise + Settings.MaxDrop);
+	}
+
+	float FSuspensionSimModule::GetSpringLength() const
+	{
+		return  -(Setup().MaxLength - SpringDisplacement);
+	}
+
+	void FSuspensionSimModule::SetSpringLength(float InLength, float WheelRadius)
+	{
+		float DisplacementInput = InLength;
+		DisplacementInput = FMath::Max(0.f, DisplacementInput);
+		SpringDisplacement = Setup().MaxLength - DisplacementInput;
+	}
+
+	void FSuspensionSimModule::GetWorldRaycastLocation(const FTransform& BodyTransform, float WheelRadius, FSpringTrace& OutTrace)
+	{
+		FVector LocalDirection = Setup().SuspensionAxis;
+		FVector WorldLocation = BodyTransform.TransformPosition(GetParentRelativeTransform().GetLocation()/*Setup().LocalOffset*/);
+		FVector WorldDirection = BodyTransform.TransformVector(LocalDirection);
+
+		OutTrace.Start = WorldLocation - WorldDirection * (Setup().MaxRaise);
+		OutTrace.End = WorldLocation + WorldDirection * (Setup().MaxDrop + WheelRadius);
+		float TraceLength = OutTrace.Start.Z - OutTrace.End.Z;
+	}
+
+	void FSuspensionSimModule::Simulate(float DeltaTime, const FAllInputs& Inputs, FSimModuleTree& VehicleModuleSystem)
+	{
+		float ForceIntoSurface = 0.0f;
+		if (SpringDisplacement > 0)
+		{
+			float Damping = (SpringDisplacement < LastDisplacement) ? Setup().CompressionDamping : Setup().ReboundDamping;
+			float SpringSpeed = (LastDisplacement - SpringDisplacement) / DeltaTime;
+
+			float StiffnessForce = SpringDisplacement * Setup().SpringRate;
+			float DampingForce = SpringSpeed * Damping;
+			float SuspensionForce = StiffnessForce - DampingForce;
+			LastDisplacement = SpringDisplacement;
+
+			if (SuspensionForce > 0)
+			{
+				ForceIntoSurface = SuspensionForce;
+				AddLocalForce(Setup().SuspensionAxis * -SuspensionForce/*, Setup().LocalOffset*/);
+			}
+		}
+
+		// tell wheels how much they are being pressed into the ground
+		if (SimModuleTree)
+		{
+			if (Chaos::ISimulationModuleBase* Module = SimModuleTree->AccessSimModule(WheelSimTreeIndex))
+			{
+				check(Module->GetSimType() == eSimType::Wheel);
+				Chaos::FWheelSimModule* Wheel = static_cast<Chaos::FWheelSimModule*>(Module);
+
+				Wheel->SetForceIntoSurface(ForceIntoSurface);
+			}
+			
+
+		}
+	}
+
+
+
+} // namespace Chaos
+
+#if VEHICLE_DEBUGGING_ENABLED
+PRAGMA_ENABLE_OPTIMIZATION
+#endif
\ No newline at end of file
diff --git a/Engine/Source/Runtime/Experimental/ChaosVehicles/ChaosVehiclesCore/Private/SimModule/ThrusterModule.cpp b/Engine/Source/Runtime/Experimental/ChaosVehicles/ChaosVehiclesCore/Private/SimModule/ThrusterModule.cpp
new file mode 100644
--- /dev/null
+++ b/Engine/Source/Runtime/Experimental/ChaosVehicles/ChaosVehiclesCore/Private/SimModule/ThrusterModule.cpp
@@ -0,0 +1,33 @@
+// Copyright Epic Games, Inc. All Rights Reserved.
+
+#pragma once
+
+#include "SimModule/ThrusterModule.h"
+#include "SimModule/SimModuleTree.h"
+#include "VehicleUtility.h"
+
+#if VEHICLE_DEBUGGING_ENABLED
+PRAGMA_DISABLE_OPTIMIZATION
+#endif
+
+namespace Chaos
+{
+
+	FThrusterSimModule::FThrusterSimModule(const FThrusterSettings& Settings)
+		: TSimModuleSettings<FThrusterSettings>(Settings)
+	{
+
+	}
+
+	void FThrusterSimModule::Simulate(float DeltaTime, const FAllInputs& Inputs, FSimModuleTree& VehicleModuleSystem)
+	{
+		// applies continuous force
+		AddLocalForceAtPosition(FVector(Setup().MaxThrustForce * Inputs.Throttle, 0, 0), Setup().ForceOffset);
+
+	}
+
+} // namespace Chaos
+
+#if VEHICLE_DEBUGGING_ENABLED
+PRAGMA_ENABLE_OPTIMIZATION
+#endif
diff --git a/Engine/Source/Runtime/Experimental/ChaosVehicles/ChaosVehiclesCore/Private/SimModule/TorqueSimModule.cpp b/Engine/Source/Runtime/Experimental/ChaosVehicles/ChaosVehiclesCore/Private/SimModule/TorqueSimModule.cpp
new file mode 100644
--- /dev/null
+++ b/Engine/Source/Runtime/Experimental/ChaosVehicles/ChaosVehiclesCore/Private/SimModule/TorqueSimModule.cpp
@@ -0,0 +1,102 @@
+// Copyright Epic Games, Inc. All Rights Reserved.
+
+
+#include "SimModule/TorqueSimModule.h"
+#include "SimModule/SimModuleTree.h"
+#include "VehicleUtility.h"
+
+#if VEHICLE_DEBUGGING_ENABLED
+PRAGMA_DISABLE_OPTIMIZATION
+#endif
+
+namespace Chaos
+{
+
+void FTorqueSimModule::TransmitTorque(const FSimModuleTree& ModuleTree, float PushedTorque, float BrakeTorque, float GearingRatio /*= 1.0f*/, float ClutchSlip /* = 1.0f*/)
+{
+	if (FMath::Abs(GearingRatio) > SMALL_NUMBER && !ModuleTree.IsEmpty())
+	{
+		int ParentIndex = ModuleTree.GetParent(SimTreeIndex);
+
+		const TSet<int>& Children = ModuleTree.GetChildren(SimTreeIndex);
+
+		LoadTorque = 0; // clear torques prior to summation from child nodes
+		BrakingTorque = BrakeTorque;
+		for (int ChildIndex : Children)
+		{
+			// currently doesn't make sense to transmit from one wheel to next in hierarchy - only because hierarchy is 1 deep at present, this may change in future
+			if (ModuleTree.GetSimModule(SimTreeIndex)->GetSimType() == ModuleTree.GetSimModule(ChildIndex)->GetSimType())
+			{
+				continue;
+			}
+
+			ISimulationModuleBase* Module = ModuleTree.AccessSimModule(ChildIndex);
+			if (FTorqueSimModule* Interface = FTorqueSimModule::CastToTorqueInterface(Module))
+			{
+				// push the torque value down to the children
+				// TODO: We are performing equal torque splitting here is that correct/desired?
+				Interface->SetDriveTorque(PushedTorque * GearingRatio * ClutchSlip / Children.Num());
+
+				// summing up braking value back from all children
+				BrakingTorque += Interface->GetBrakingTorque() * ClutchSlip;
+
+				// diff velocity 
+				float DiffVelocity = (Interface->AngularVelocity * GearingRatio) - AngularVelocity;
+				AngularVelocity += DiffVelocity * ClutchSlip / Children.Num();
+			}
+		}
+
+		BrakingTorque /= GearingRatio;
+	}
+}
+
+/**
+* Integrate angular velocity using specified DeltaTime & Inertia value, Note the inertia should be the combined inertia of all the connected pieces otherwise things will rotate at different rates
+*/
+
+void FTorqueSimModule::IntegrateAngularVelocity(float DeltaTime, float Inertia, float MaxRotationVel)
+{
+	ensure(BrakingTorque >= 0.0f);
+
+	// drive torque taken into account
+	AngularVelocity += (DriveTorque + LoadTorque * DeltaTime) / Inertia;
+
+	// braking resists velocity no matter what way we are spinning
+	// also has check that we are not overshooting and starting to accelerating in the opposite direction
+	if (AngularVelocity > 0.0f)
+	{
+		AngularVelocity -= (BrakingTorque * DeltaTime) / Inertia;
+		if (AngularVelocity < 0.0f)
+		{
+			AngularVelocity = 0.0f;
+		}
+	}
+	else
+	{
+		AngularVelocity += (BrakingTorque * DeltaTime) / Inertia;
+		if (AngularVelocity > 0.0f)
+		{
+			AngularVelocity = 0.0f;
+		}
+
+	}
+	
+	AngularVelocity = FMath::Clamp(AngularVelocity, -MaxRotationVel, MaxRotationVel);
+
+	// angluar position integration
+	AngularPosition += AngularVelocity * DeltaTime;
+
+	int ExcessRotations = (int)(AngularPosition / TWO_PI);
+	AngularPosition -= ExcessRotations * TWO_PI;
+
+	AngularVelocity *= 0.98f; // TEMP
+
+	UE_LOG(LogSimulationModule, Log, TEXT("%s: DriveTorque %4.2f, BrakingTorque %4.2f, LoadTorque %4.2f, Speed %4.2f rad/sec, RPM %4.2f, AngularPosition %4.2f, Inertia %4.2f")
+		, *GetDebugName(), DriveTorque, BrakingTorque, LoadTorque, AngularVelocity, GetRPM(), AngularPosition, Inertia);
+}
+
+} // namespace Chaos
+
+#if VEHICLE_DEBUGGING_ENABLED
+PRAGMA_ENABLE_OPTIMIZATION
+#endif
diff --git a/Engine/Source/Runtime/Experimental/ChaosVehicles/ChaosVehiclesCore/Private/SimModule/TransmissionModule.cpp b/Engine/Source/Runtime/Experimental/ChaosVehicles/ChaosVehiclesCore/Private/SimModule/TransmissionModule.cpp
new file mode 100644
--- /dev/null
+++ b/Engine/Source/Runtime/Experimental/ChaosVehicles/ChaosVehiclesCore/Private/SimModule/TransmissionModule.cpp
@@ -0,0 +1,143 @@
+// Copyright Epic Games, Inc. All Rights Reserved.
+
+#pragma once
+
+#include "SimModule/TransmissionModule.h"
+#include "SimModule/SimModuleTree.h"
+#include "VehicleUtility.h"
+
+#if VEHICLE_DEBUGGING_ENABLED
+PRAGMA_DISABLE_OPTIMIZATION
+#endif
+
+namespace Chaos
+{
+
+	void FTransmissionSimModule::Simulate(float DeltaTime, const FAllInputs& Inputs, FSimModuleTree& VehicleModuleSystem)
+	{
+		if (Setup().TransmissionType == FTransmissionSettings::ETransType::Automatic)
+		{
+			// not currently changing gear, also don't want to change up because the wheels are spinning up due to having no load
+			if (!IsCurrentlyChangingGear() && AllowedToChangeGear)
+			{
+				// in Automatic & currently in neutral and not currently changing gear then automatically change up to 1st
+				if (CurrentGear == 0)
+				{
+					ChangeUp();
+				}
+
+				float EngineRPM = GetRPM();
+
+				if (EngineRPM >= Setup().ChangeUpRPM)
+				{
+					if (CurrentGear > 0)
+					{
+						ChangeUp();
+					}
+					else
+					{
+						ChangeDown();
+					}
+				}
+				else if (EngineRPM <= Setup().ChangeDownRPM && FMath::Abs(CurrentGear) > 1) // don't change down to neutral
+				{
+					if (CurrentGear > 0)
+					{
+						ChangeDown();
+					}
+					else
+					{
+						ChangeUp();
+					}
+				}
+			}
+		}
+		else
+		{
+			if (Inputs.ChangeUp)
+			{
+				ChangeUp();
+			}
+			else if (Inputs.ChangeDown)
+			{
+				ChangeDown();
+			}
+			else if (Inputs.SetGearNumber)
+			{
+				TargetGear = Inputs.SetGearNumber;
+			}
+		}
+
+		if (CurrentGear != TargetGear)
+		{
+			CurrentGearChangeTime -= DeltaTime;
+			if (CurrentGearChangeTime <= 0.f)
+			{
+				CurrentGearChangeTime = 0.f;
+				CurrentGear = TargetGear;
+			}
+		}
+
+		float GearRatio = GetGearRatio(CurrentGear);
+
+		// if there IS a selected gear then connect the parent and child building blocks transmitting their torque
+		if (FMath::Abs(GearRatio) > SMALL_NUMBER)
+		{
+			float BrakeTorque = 0.0f;
+			TransmitTorque(VehicleModuleSystem, DriveTorque, BrakeTorque, GearRatio);
+		}
+
+	}
+
+	float FTransmissionSimModule::GetGearRatio(int32 InGear) const
+	{
+		CorrectGearInputRange(InGear);
+
+		if (InGear > 0) // a forwards gear
+		{
+			return Setup().ForwardRatios[InGear - 1] * Setup().FinalDriveRatio;
+		}
+		else if (InGear < 0) // a reverse gear
+		{
+			return -Setup().ReverseRatios[FMath::Abs(InGear) - 1] * Setup().FinalDriveRatio;
+		}
+		else
+		{
+			return 0.f; // neutral has no ratio
+		}
+	}
+
+	void FTransmissionSimModule::SetGear(int32 InGear, bool Immediate /*= false*/)
+	{
+		CorrectGearInputRange(InGear);
+
+		TargetGear = InGear;
+
+		if (TargetGear != CurrentGear)
+		{
+			if (Immediate || Setup().GearChangeTime == 0.f)
+			{
+				CurrentGear = TargetGear;
+			}
+			else
+			{
+				CurrentGear = 0;	// go through neutral for GearChangeTime time period
+				CurrentGearChangeTime = Setup().GearChangeTime;
+			}
+		}
+	}
+
+	bool FTransmissionSimModule::GetDebugString(FString& StringOut) const
+	{
+		FTorqueSimModule::GetDebugString(StringOut);
+		StringOut += FString::Format(TEXT("CurrentGear {0}, Ratio {1}")
+			, { CurrentGear, GetGearRatio(CurrentGear) });
+		return true;
+	}
+
+
+} // namespace Chaos
+
+#if VEHICLE_DEBUGGING_ENABLED
+PRAGMA_ENABLE_OPTIMIZATION
+#endif
diff --git a/Engine/Source/Runtime/Experimental/ChaosVehicles/ChaosVehiclesCore/Private/SimModule/WheelModule.cpp b/Engine/Source/Runtime/Experimental/ChaosVehicles/ChaosVehiclesCore/Private/SimModule/WheelModule.cpp
new file mode 100644
--- /dev/null
+++ b/Engine/Source/Runtime/Experimental/ChaosVehicles/ChaosVehiclesCore/Private/SimModule/WheelModule.cpp
@@ -0,0 +1,179 @@
+// Copyright Epic Games, Inc. All Rights Reserved.
+
+#pragma once
+
+#include "SimModule/WheelModule.h"
+#include "SimModule/SimModuleTree.h"
+#include "VehicleUtility.h"
+
+#if VEHICLE_DEBUGGING_ENABLED
+PRAGMA_DISABLE_OPTIMIZATION
+#endif
+
+namespace Chaos
+{
+
+	FWheelSimModule::FWheelSimModule(const FWheelSettings& Settings)
+		: TSimModuleSettings<FWheelSettings>(Settings)
+		, BrakeTorque(0.0f)
+		, ForceIntoSurface(0.0f)
+		, SurfaceFriction(1.0f)
+		, SuspensionSimTreeIndex(INVALID_INDEX)
+		, ForceFromFriction(FVector::ZeroVector)
+		, MassPerWheel(500.0f*0.25f)
+		, SteerAngleDegrees(0.0f)
+	{
+
+	}
+
+	void FWheelSimModule::Simulate(float DeltaTime, const FAllInputs& Inputs, FSimModuleTree& VehicleModuleSystem)
+	{	
+		float Re = Setup().Radius;
+		float K = 0.4f;
+		float TorqueScaling = 0.00005f;
+		float TractionControlAndABSScaling = 0.98f;	// how close to perfection is the system working
+
+		float HandbrakeTorque = Setup().HandbrakeEnabled ? Inputs.Handbrake * Setup().HandbrakeTorque : 0.0f;
+		SteerAngleDegrees = Setup().SteeringEnabled ? Inputs.Steering * Setup().MaxSteeringAngle : 0.0f;
+		BrakeTorque = Inputs.Brake * Setup().MaxBrakeTorque + HandbrakeTorque;
+
+		// TODO: think about doing this properly, stops vehicles rolling around on their own too much
+		// kindof an auto handbrake
+		if (Inputs.Brake < SMALL_NUMBER && Inputs.Throttle < SMALL_NUMBER && ModuleLocalVelocity.X < 10.0f)
+		{
+			BrakeTorque = Setup().HandbrakeTorque;
+		}
+
+		FRotator SteeringRotator(0.f, SteerAngleDegrees, 0.f);
+		FVector LocalWheelVelocity = SteeringRotator.UnrotateVector(ModuleLocalVelocity);
+
+		float GroundAngularVelocity = LocalWheelVelocity.X / Re;
+		float Delta = GroundAngularVelocity - AngularVelocity;
+		float Torque = Delta * Setup().WheelInertia / DeltaTime;
+
+		// X is longitudinal direction, Y is lateral
+		float SlipAngle = FVehicleUtility::CalculateSlipAngle(LocalWheelVelocity.Y, LocalWheelVelocity.X);
+
+		float AppliedLinearDriveForce = DriveTorque / Re;
+		float AppliedLinearBrakeForce = BrakeTorque / Re;
+
+		// Longitudinal multiplier now affecting both brake and steering equally
+		float AvailableGrip = ForceIntoSurface * SurfaceFriction * Setup().FrictionMultiplier;
+
+		float FinalLongitudinalForce = 0.f;
+		float FinalLateralForce = 0.f;
+
+		// currently just letting the brake override the throttle
+		bool Braking = (BrakeTorque > FMath::Abs(DriveTorque));
+		bool WheelLocked = false;
+
+		// are we actually touching the ground
+		if (ForceIntoSurface > SMALL_NUMBER)
+		{
+			// ABS limiting brake force to match force from the grip available
+			if (Setup().ABSEnabled && Braking && FMath::Abs(AppliedLinearBrakeForce) > AvailableGrip)
+			{
+				if ((Braking && Setup().ABSEnabled) || (!Braking && Setup().TractionControlEnabled))
+				{
+					float Sign = (AppliedLinearBrakeForce > 0.0f) ? 1.0f : -1.0f;
+					AppliedLinearBrakeForce = AvailableGrip * TractionControlAndABSScaling * Sign;
+				}
+			}
+
+			// Traction control limiting drive force to match force from grip available
+			if (Setup().TractionControlEnabled && !Braking && FMath::Abs(AppliedLinearDriveForce) > AvailableGrip)
+			{
+				float Sign = (AppliedLinearDriveForce > 0.0f) ? 1.0f : -1.0f;
+				AppliedLinearDriveForce = AvailableGrip * TractionControlAndABSScaling * Sign;
+			}
+
+			if (Braking)
+			{
+				// whether the velocity is +ve or -ve when we brake we are slowing the vehicle down
+				// so force is opposing current direction of travel.
+				float ForceRequiredToBringToStop = MassPerWheel * K * (LocalWheelVelocity.X) / DeltaTime;
+				FinalLongitudinalForce = AppliedLinearBrakeForce;
+
+				// check we are not applying more force than required so we end up overshooting 
+				// and accelerating in the opposite direction
+				FinalLongitudinalForce = FMath::Clamp(FinalLongitudinalForce, -FMath::Abs(ForceRequiredToBringToStop), FMath::Abs(ForceRequiredToBringToStop));
+
+				// ensure the brake opposes current direction of travel
+				if (LocalWheelVelocity.X > 0.0f)
+				{
+					FinalLongitudinalForce = -FinalLongitudinalForce;
+				}
+
+			}
+			else
+			{
+				FinalLongitudinalForce = AppliedLinearDriveForce;
+			}
+
+			float ForceRequiredToBringToStop = -(MassPerWheel * K * LocalWheelVelocity.Y) / DeltaTime;
+
+			// use slip angle to generate a sideways force
+			if (Setup().LateralSlipGraph.IsEmpty())
+			{
+				float AngleLimit = FMath::DegreesToRadians(Setup().SlipAngleLimit);
+				float ClippedAngle = FMath::Clamp(SlipAngle, -AngleLimit, AngleLimit);
+				FinalLateralForce = FMath::Abs(SlipAngle) * Setup().CorneringStiffness;
+			}
+			else
+			{
+				FinalLateralForce = Setup().LateralSlipGraph.EvaluateY(FMath::RadiansToDegrees(SlipAngle)) * Setup().LateralSlipGraphMultiplier;
+			}
+
+			if (FinalLateralForce > FMath::Abs(ForceRequiredToBringToStop))
+			{
+				FinalLateralForce = FMath::Abs(ForceRequiredToBringToStop);
+			}
+			if (LocalWheelVelocity.Y > 0.0f)
+			{
+				FinalLateralForce = -FinalLateralForce;
+			}
+
+		}
+
+		ForceFromFriction = FVector::ZeroVector;
+		ForceFromFriction.X = FinalLongitudinalForce;
+		ForceFromFriction.Y = FinalLateralForce;
+
+		AddLocalForce(SteeringRotator.RotateVector(ForceFromFriction));
+		TransmitTorque(VehicleModuleSystem, DriveTorque, BrakeTorque);
+
+		DriveTorque -= AvailableGrip;
+		if (DriveTorque < 0.0f)
+		{
+			DriveTorque = 0.0f;
+		}
+		else
+		{ 
+			DriveTorque *= TorqueScaling;
+		}
+		BrakingTorque -= AvailableGrip;
+		if (BrakingTorque < 0.0f)
+		{
+			BrakingTorque = 0.0f;
+		}
+
+		LoadTorque = Torque;	// torque from wheels moving over terrain
+		IntegrateAngularVelocity(DeltaTime, Setup().WheelInertia, Setup().MaxRotationVel);
+	}
+
+
+	bool FWheelSimModule::GetDebugString(FString& StringOut) const
+	{
+		FTorqueSimModule::GetDebugString(StringOut);
+		StringOut += FString::Format(TEXT("Drive {0}, Brake {1}, Load {2} RPM {3}  AngVel {4} LongitudinalForce {5}")
+			, { DriveTorque, BrakingTorque, LoadTorque, GetRPM(), AngularVelocity, ForceFromFriction.X });
+		return true;
+	}
+
+
+} // namespace Chaos
+
+
+#if VEHICLE_DEBUGGING_ENABLED
+PRAGMA_ENABLE_OPTIMIZATION
+#endif
diff --git a/Engine/Source/Runtime/Experimental/ChaosVehicles/ChaosVehiclesCore/Public/SimModule/AerofoilModule.h b/Engine/Source/Runtime/Experimental/ChaosVehicles/ChaosVehiclesCore/Public/SimModule/AerofoilModule.h
new file mode 100644
--- /dev/null
+++ b/Engine/Source/Runtime/Experimental/ChaosVehicles/ChaosVehiclesCore/Public/SimModule/AerofoilModule.h
@@ -0,0 +1,123 @@
+// Copyright Epic Games, Inc. All Rights Reserved.
+
+#pragma once
+
+#include "SimModule/TorqueSimModule.h"
+
+
+namespace Chaos
+{
+	struct FAllInputs;
+	class FSimModuleTree;
+
+	enum class EAerofoil : uint8
+	{
+		Fixed = 0,
+		Wing,
+		Rudder,
+		Elevator
+	};
+
+	struct CHAOSVEHICLESCORE_API FAerofoilSettings
+	{
+		FAerofoilSettings()
+			: Offset(FVector::ZeroVector)
+			, UpAxis(FVector(0.f, 0.f, 1.f))
+			, Area(5.0f)
+			, Camber(3.0f)
+			, MaxControlAngle(1.f)
+			, StallAngle(16.0f)
+			, MaxCeiling(1E30)
+			, MinCeiling(-1E30)
+			, Type(EAerofoil::Fixed)
+			, LiftMultiplier(1.0f)
+			, DragMultiplier(1.0f)
+		{
+		}
+
+		FVector Offset;
+		FVector UpAxis;
+		float Area;
+		float Camber;
+		float MaxControlAngle;
+		float StallAngle;
+
+		float MaxCeiling;
+		float MinCeiling;
+
+		EAerofoil Type;
+		float LiftMultiplier;
+		float DragMultiplier;
+
+	};
+
+	class CHAOSVEHICLESCORE_API FAerofoilSimModule : public ISimulationModuleBase, public TSimModuleSettings<FAerofoilSettings>
+	{
+	public:
+
+		FAerofoilSimModule(const FAerofoilSettings& Settings) : TSimModuleSettings<FAerofoilSettings>(Settings)
+			, CurrentAirDensity(RealWorldConsts::AirDensity())
+			, AngleOfAttack(0.f)
+			, ControlSurfaceAngle(0.f)
+			, AirflowNormal(FVector::ZeroVector)
+			, AerofoilId(0)
+		{
+
+		}
+
+		virtual ~FAerofoilSimModule() {}
+
+		virtual eSimType GetSimType() const { return eSimType::Aerofoil; }
+
+		virtual const FString GetDebugName() const { return TEXT("Aerofoil"); }
+
+		virtual bool IsBehaviourType(eSimModuleTypeFlags InType) const override { return (InType & Velocity); }
+
+		virtual void Simulate(float DeltaTime, const FAllInputs& Inputs, FSimModuleTree& VehicleModuleSystem) override;
+
+		void SetDensityOfMedium(float InDensity)
+		{
+			CurrentAirDensity = InDensity;
+		}
+
+	protected:
+		void SetControlSurface(float CtrlSurfaceInput)
+		{
+			ControlSurfaceAngle = CtrlSurfaceInput * Setup().MaxControlAngle;
+		}
+
+		FVector GetCenterOfLiftOffset();
+
+		// returns the combined force of lift and drag at an aerofoil in local coordinates
+		// for direct application to the aircrafts rigid body.
+		FVector GetForce(const FVector& v, float Altitude, float DeltaTime);
+
+		/**
+		 * Dynamic air pressure = 0.5 * AirDensity * Vsqr
+		 */
+		float CalcDynamicPressure(float VelocitySqr, float InAltitude);
+
+		/**  Center of lift moves fore/aft based on current AngleOfAttack */
+		float CalcCentreOfLift();
+
+		/** Returns drag coefficient for the current angle of attack of the aerofoil surface */
+		float CalcDragCoefficient(float InAngleOfAttack, float InControlSurfaceAngle);
+
+		/**
+		 * Returns lift coefficient for the current angle of attack of the aerofoil surface
+		 * Cheating by making control surface part of entire aerofoil movement
+		 */
+		float CalcLiftCoefficient(float InAngleOfAttack, float InControlSurfaceAngle);
+
+		/** Angle of attack is the angle between the aerofoil and the airflow vector */
+		float CalcAngleOfAttackDegrees(const FVector& UpAxis, const FVector& InAirflowVector);
+
+		float CurrentAirDensity;
+		float AngleOfAttack;
+		float ControlSurfaceAngle;
+		FVector AirflowNormal;
+		int AerofoilId;
+
+	};
+
+} // namespace Chaos
\ No newline at end of file
diff --git a/Engine/Source/Runtime/Experimental/ChaosVehicles/ChaosVehiclesCore/Public/SimModule/AxleModule.h b/Engine/Source/Runtime/Experimental/ChaosVehicles/ChaosVehiclesCore/Public/SimModule/AxleModule.h
new file mode 100644
--- /dev/null
+++ b/Engine/Source/Runtime/Experimental/ChaosVehicles/ChaosVehiclesCore/Public/SimModule/AxleModule.h
@@ -0,0 +1,46 @@
+// Copyright Epic Games, Inc. All Rights Reserved.
+
+#pragma once
+
+#include "SimModule/TorqueSimModule.h"
+
+
+namespace Chaos
+{
+	struct FAllInputs;
+	class FSimModuleTree;
+
+	/// <summary>
+	/// Axle settings
+	/// </summary>
+	struct CHAOSVEHICLESCORE_API FAxleSettings
+	{
+		FAxleSettings()
+			: AxleInertia(1.f) // TODO: defaults
+		{
+
+		}
+
+		float AxleInertia;
+	};
+
+	/// <summary>
+	/// A vehicle component that transmits torque from one source to another, i.e. from an engine or differential to wheels
+	///
+	/// </summary>
+	class CHAOSVEHICLESCORE_API FAxleSimModule : public FTorqueSimModule, public TSimModuleSettings<FAxleSettings>
+	{
+	public:
+		FAxleSimModule(const FAxleSettings& Settings);
+
+		virtual eSimType GetSimType() const { return eSimType::Axle; }
+
+		virtual const FString GetDebugName() const { return TEXT("Axle"); }
+
+		virtual void Simulate(float DeltaTime, const FAllInputs& Inputs, FSimModuleTree& VehicleModuleSystem) override;
+
+	private:
+	};
+
+
+} // namespace Chaos
diff --git a/Engine/Source/Runtime/Experimental/ChaosVehicles/ChaosVehiclesCore/Public/SimModule/ChassisModule.h b/Engine/Source/Runtime/Experimental/ChaosVehicles/ChaosVehiclesCore/Public/SimModule/ChassisModule.h
new file mode 100644
--- /dev/null
+++ b/Engine/Source/Runtime/Experimental/ChaosVehicles/ChaosVehiclesCore/Public/SimModule/ChassisModule.h
@@ -0,0 +1,46 @@
+// Copyright Epic Games, Inc. All Rights Reserved.
+
+#pragma once
+
+#include "SimModule/SimulationModuleBase.h"
+
+
+namespace Chaos
+{
+	struct FAllInputs;
+	class FSimModuleTree;
+
+	/// <summary>
+	/// Chassis settings
+	/// </summary>
+	struct CHAOSVEHICLESCORE_API FChassisSettings
+	{
+		FChassisSettings()
+		{
+
+		}
+
+	};
+
+	/// <summary>
+	/// A vehicle component that transmits torque from one source to another, i.e. from an engine or differential to wheels
+	///
+	/// </summary>
+	class CHAOSVEHICLESCORE_API FChassisSimModule : public ISimulationModuleBase, public TSimModuleSettings<FChassisSettings>
+	{
+	public:
+		FChassisSimModule(const FChassisSettings& Settings);
+
+		virtual eSimType GetSimType() const { return eSimType::Chassis; }
+
+		virtual const FString GetDebugName() const { return TEXT("Chassis"); }
+
+		virtual bool IsBehaviourType(eSimModuleTypeFlags InType) const override { return (InType & NonFunctional); }
+
+		virtual void Simulate(float DeltaTime, const FAllInputs& Inputs, FSimModuleTree& VehicleModuleSystem) override;
+
+	private:
+	};
+
+
+} // namespace Chaos
diff --git a/Engine/Source/Runtime/Experimental/ChaosVehicles/ChaosVehiclesCore/Public/SimModule/ClutchModule.h b/Engine/Source/Runtime/Experimental/ChaosVehicles/ChaosVehiclesCore/Public/SimModule/ClutchModule.h
new file mode 100644
--- /dev/null
+++ b/Engine/Source/Runtime/Experimental/ChaosVehicles/ChaosVehiclesCore/Public/SimModule/ClutchModule.h
@@ -0,0 +1,53 @@
+// Copyright Epic Games, Inc. All Rights Reserved.
+
+#pragma once
+
+#include "SimModule/TorqueSimModule.h"
+
+
+namespace Chaos
+{
+	struct FAllInputs;
+	class FSimModuleTree;
+
+	struct CHAOSVEHICLESCORE_API FClutchSettings
+	{
+		FClutchSettings()
+			: ClutchStrength(1.f)
+		{
+
+		}
+
+		float ClutchStrength;
+	};
+
+	/// <summary>
+	/// 
+	/// a vehicle component that transmits torque from one source to another through a clutch system, i.e. connect an engine to a transmission
+	///
+	/// Input Controls - Clutch pedal, normalized value 0 to 1 expected
+	/// Other Inputs - 
+	/// Outputs - 
+	/// 
+	/// </summary>
+	class CHAOSVEHICLESCORE_API FClutchSimModule : public FTorqueSimModule, public TSimModuleSettings<FClutchSettings>
+	{
+	public:
+
+		FClutchSimModule(const FClutchSettings& Settings);
+
+		virtual eSimType GetSimType() const { return eSimType::Clutch; }
+
+		virtual const FString GetDebugName() const { return TEXT("Clutch"); }
+
+		virtual bool GetDebugString(FString& StringOut) const override;
+
+		virtual void Simulate(float DeltaTime, const FAllInputs& Inputs, FSimModuleTree& VehicleModuleSystem);
+
+	private:
+
+		float ClutchValue;
+	};
+
+
+} // namespace Chaos
diff --git a/Engine/Source/Runtime/Experimental/ChaosVehicles/ChaosVehiclesCore/Public/SimModule/DeferredForcesModular.h b/Engine/Source/Runtime/Experimental/ChaosVehicles/ChaosVehiclesCore/Public/SimModule/DeferredForcesModular.h
new file mode 100644
--- /dev/null
+++ b/Engine/Source/Runtime/Experimental/ChaosVehicles/ChaosVehiclesCore/Public/SimModule/DeferredForcesModular.h
@@ -0,0 +1,156 @@
+// Copyright Epic Games, Inc. All Rights Reserved.
+
+#pragma once
+
+#include "CoreMinimal.h"
+#include "Chaos/ParticleHandleFwd.h"
+#include "GeometryCollection/ManagedArray.h"
+
+enum class EForceFlags : uint32
+{
+	None = 0, // No flags.
+
+	AllowSubstepping = 1 << 0,
+	AccelChange = 1 << 1,
+	VelChange = 1 << 2,
+	IsLocalForce = 1 << 3,
+
+};
+ENUM_CLASS_FLAGS(EForceFlags);	
+
+
+class CHAOSVEHICLESCORE_API FDeferredForcesModular
+{
+public:
+
+	struct FApplyForceData
+	{
+		FApplyForceData(int TransformIndexIn, const FVector& ForceIn, bool bAllowSubsteppingIn, bool bAccelChangeIn)
+			: TransformIndex(TransformIndexIn)
+			, Force(ForceIn)
+			, Flags(EForceFlags::None)
+		{
+			Flags |= bAllowSubsteppingIn ? EForceFlags::AllowSubstepping : EForceFlags::None;
+			Flags |= bAccelChangeIn ? EForceFlags::AccelChange : EForceFlags::None;
+		}
+
+		int TransformIndex;
+		FVector Force;
+		EForceFlags Flags;
+	};
+
+	struct FApplyForceAtPositionData
+	{
+		FApplyForceAtPositionData(int TransformIndexIn, const FVector& ForceIn, const FVector& PositionIn, bool bAllowSubsteppingIn, bool bIsLocalForceIn)
+			: TransformIndex(TransformIndexIn)
+			, Force(ForceIn)
+			, Position(PositionIn)
+			, Flags(EForceFlags::None)
+		{
+			Flags |= bAllowSubsteppingIn ? EForceFlags::AllowSubstepping : EForceFlags::None;
+			Flags |= bIsLocalForceIn ? EForceFlags::IsLocalForce : EForceFlags::None;
+		}
+
+		int TransformIndex;
+		FVector Force;
+		FVector Position;
+		EForceFlags Flags;
+	};
+
+	struct FAddTorqueInRadiansData
+	{
+		FAddTorqueInRadiansData(int TransformIndexIn, const FVector& TorqueIn, bool bAllowSubsteppingIn, bool bAccelChangeIn)
+			: TransformIndex(TransformIndexIn)
+			, Torque(TorqueIn)
+			, Flags(EForceFlags::None)
+		{
+			Flags |= bAllowSubsteppingIn ? EForceFlags::AllowSubstepping : EForceFlags::None;
+			Flags |= bAccelChangeIn ? EForceFlags::AccelChange : EForceFlags::None;
+		}
+
+		int TransformIndex;
+		FVector Torque;
+		EForceFlags Flags;
+	};
+
+	struct FAddImpulseData
+	{
+		FAddImpulseData(int TransformIndexIn, const FVector& ImpulseIn, const bool bVelChangeIn)
+			: TransformIndex(TransformIndexIn)
+			, Impulse(ImpulseIn)
+			, Flags(EForceFlags::None)
+		{
+			Flags |= bVelChangeIn ? EForceFlags::VelChange : EForceFlags::None;
+		}
+
+		int TransformIndex;
+		FVector Impulse;
+		EForceFlags Flags;
+	};
+
+	struct FAddImpulseAtPositionData
+	{
+		FAddImpulseAtPositionData(int TransformIndexIn, const FVector& ImpulseIn, const FVector& PositionIn)
+			: TransformIndex(TransformIndexIn)
+			, Impulse(ImpulseIn)
+			, Position(PositionIn)
+		{
+
+		}
+
+		int TransformIndex;
+		FVector Impulse;
+		FVector Position;
+	};
+
+	void Add(const FApplyForceData& ApplyForceDataIn)
+	{
+		ApplyForceDatas.Add(ApplyForceDataIn);
+	}
+
+	void Add(const FApplyForceAtPositionData& ApplyForceAtPositionDataIn)
+	{
+		ApplyForceAtPositionDatas.Add(ApplyForceAtPositionDataIn);
+	}
+
+	void Add(const FAddTorqueInRadiansData& ApplyTorqueDataIn)
+	{
+		ApplyTorqueDatas.Add(ApplyTorqueDataIn);
+	}
+
+	void Add(const FAddImpulseData& ApplyImpulseDataIn)
+	{
+		ApplyImpulseDatas.Add(ApplyImpulseDataIn);
+	}
+
+	void Add(const FAddImpulseAtPositionData& ApplyImpulseAtPositionDataIn)
+	{
+		ApplyImpulseAtPositionDatas.Add(ApplyImpulseAtPositionDataIn);
+	}	
+
+	Chaos::FPBDRigidClusteredParticleHandle* GetParticle(TArray<Chaos::FPBDRigidClusteredParticleHandle*>& Particles
+			, TArray<Chaos::FPBDRigidClusteredParticleHandle*>& ClusterParticles
+			, int TransformIndex
+			, const TManagedArray<FTransform>& Transforms
+			, const TManagedArray<FTransform>& CollectionMassToLocal
+			, const TManagedArray<int32>& Parent
+			, FTransform& TransformOut);
+
+	void Apply(TArray<Chaos::FPBDRigidClusteredParticleHandle*>& Particles
+		, TArray<Chaos::FPBDRigidClusteredParticleHandle*>& ClusterParticles
+		, const TManagedArray<FTransform>& Transforms
+		, const TManagedArray<FTransform>& CollectionMassToLocal
+		, const TManagedArray<int32>& Parent);
+
+private:
+
+	void AddForceAtPosition(Chaos::FPBDRigidClusteredParticleHandle* RigidHandle, const FApplyForceAtPositionData& DataIn, const FTransform& OffsetTransform);
+	void AddTorque(Chaos::FPBDRigidClusteredParticleHandle* RigidHandle, const FAddTorqueInRadiansData& DataIn, const FTransform& OffsetTransform);
+	void AddForce(Chaos::FPBDRigidClusteredParticleHandle* RigidHandle, const FApplyForceData& DataIn, const FTransform& OffsetTransform);
+
+	TArray<FApplyForceData> ApplyForceDatas;
+	TArray<FApplyForceAtPositionData> ApplyForceAtPositionDatas;
+	TArray<FAddTorqueInRadiansData> ApplyTorqueDatas;
+	TArray<FAddImpulseData> ApplyImpulseDatas;
+	TArray<FAddImpulseAtPositionData> ApplyImpulseAtPositionDatas;
+};
diff --git a/Engine/Source/Runtime/Experimental/ChaosVehicles/ChaosVehiclesCore/Public/SimModule/EngineModule.h b/Engine/Source/Runtime/Experimental/ChaosVehicles/ChaosVehiclesCore/Public/SimModule/EngineModule.h
new file mode 100644
--- /dev/null
+++ b/Engine/Source/Runtime/Experimental/ChaosVehicles/ChaosVehiclesCore/Public/SimModule/EngineModule.h
@@ -0,0 +1,81 @@
+// Copyright Epic Games, Inc. All Rights Reserved.
+
+#pragma once
+
+#include "SimModule/TorqueSimModule.h"
+
+PRAGMA_DISABLE_OPTIMIZATION
+
+namespace Chaos
+{
+	struct FAllInputs;
+	class FSimModuleTree;
+
+	struct CHAOSVEHICLESCORE_API FEngineSettings
+	{
+		FEngineSettings()
+			: MaxTorque(300.f)
+			, MaxRPM(6000)
+			, IdleRPM(1200)
+			, EngineBrakeEffect(0.2f)
+			, EngineInertia(100000.0f)
+		{
+			TorqueCurve.AddNormalized(0.5f);
+			TorqueCurve.AddNormalized(0.5f);
+			TorqueCurve.AddNormalized(0.5f);
+			TorqueCurve.AddNormalized(0.5f);
+			TorqueCurve.AddNormalized(0.6f);
+			TorqueCurve.AddNormalized(0.7f);
+			TorqueCurve.AddNormalized(0.8f);
+			TorqueCurve.AddNormalized(0.9f);
+			TorqueCurve.AddNormalized(1.0f);
+			TorqueCurve.AddNormalized(0.9f);
+			TorqueCurve.AddNormalized(0.7f);
+			TorqueCurve.AddNormalized(0.5f);
+		}
+
+		FNormalisedGraph TorqueCurve;
+		float MaxTorque;			// [N.m] The peak torque Y value in the normalized torque graph
+		uint16 MaxRPM;				// [RPM] The absolute maximum RPM the engine can theoretically reach (last X value in the normalized torque graph)
+		uint16 IdleRPM; 		// [RPM] The RPM at which the throttle sits when the car is not moving			
+		float EngineBrakeEffect;	// [0..1] How much the engine slows the vehicle when the throttle is released
+
+		float EngineInertia;
+	};
+
+	class CHAOSVEHICLESCORE_API FEngineSimModule : public FTorqueSimModule, public TSimModuleSettings<FEngineSettings>
+	{
+	public:
+
+		FEngineSimModule(const FEngineSettings& Settings) : TSimModuleSettings<FEngineSettings>(Settings)
+			, EngineIdleSpeed(RPMToOmega(Setup().IdleRPM))
+			, MaxEngineSpeed(RPMToOmega(Setup().MaxRPM))
+			, EngineStarted(true)
+		{
+		}
+
+		virtual ~FEngineSimModule() {}
+
+		virtual eSimType GetSimType() const { return eSimType::Engine; }
+
+		virtual const FString GetDebugName() const { return TEXT("Engine"); }
+
+		virtual bool GetDebugString(FString& StringOut) const override;
+
+		virtual void Simulate(float DeltaTime, const FAllInputs& Inputs, FSimModuleTree& VehicleModuleSystem) override;
+
+		inline float GetEngineIdleSpeed() const { return EngineIdleSpeed; }
+		inline float GetEngineTorque(float ThrottlePosition, float EngineRPM);
+		inline float GetTorqueFromRPM(float RPM, bool LimitToIdle = true);
+
+	protected:
+
+		float EngineIdleSpeed;
+		float MaxEngineSpeed;
+		bool EngineStarted;		// is the engine turned off or has it been started
+
+	};
+
+} // namespace Chaos
+
+PRAGMA_ENABLE_OPTIMIZATION
diff --git a/Engine/Source/Runtime/Experimental/ChaosVehicles/ChaosVehiclesCore/Public/SimModule/MotorModule.h b/Engine/Source/Runtime/Experimental/ChaosVehicles/ChaosVehiclesCore/Public/SimModule/MotorModule.h
new file mode 100644
--- /dev/null
+++ b/Engine/Source/Runtime/Experimental/ChaosVehicles/ChaosVehiclesCore/Public/SimModule/MotorModule.h
@@ -0,0 +1,49 @@
+// Copyright Epic Games, Inc. All Rights Reserved.
+
+#pragma once
+
+#include "SimModule/TorqueSimModule.h"
+
+namespace Chaos
+{
+	struct FAllInputs;
+	class FSimModuleTree;
+
+	/// <summary>
+	/// Motor settings
+	/// </summary>
+	struct CHAOSVEHICLESCORE_API FMotorSettings
+	{
+		FMotorSettings()
+			: MaxRPM(1.0f)
+			, MaxTorque(100.0f)
+			, EngineInertia(1.0f)
+		{
+
+		}
+		float MaxRPM;
+		float MaxTorque;
+		float EngineInertia;
+	};
+
+	/// <summary>
+	/// A vehicle component that provides torque output based on a torque control input
+	/// The output torque is based on a square function style curve, zero at 0.0 and MaxRPM, and 1.0 at mid RPM
+	/// </summary>
+	class CHAOSVEHICLESCORE_API FMotorSimModule : public FTorqueSimModule, public TSimModuleSettings<FMotorSettings>
+	{
+	public:
+
+		FMotorSimModule(const FMotorSettings& Settings);
+
+		virtual eSimType GetSimType() const { return eSimType::Motor; }
+
+		virtual const FString GetDebugName() const { return TEXT("Motor"); }
+
+		virtual void Simulate(float DeltaTime, const FAllInputs& Inputs, FSimModuleTree& VehicleModuleSystem) override;
+
+	private:
+
+	};
+
+} // namespace Chaos
diff --git a/Engine/Source/Runtime/Experimental/ChaosVehicles/ChaosVehiclesCore/Public/SimModule/SimModuleTree.h b/Engine/Source/Runtime/Experimental/ChaosVehicles/ChaosVehiclesCore/Public/SimModule/SimModuleTree.h
new file mode 100644
--- /dev/null
+++ b/Engine/Source/Runtime/Experimental/ChaosVehicles/ChaosVehiclesCore/Public/SimModule/SimModuleTree.h
@@ -0,0 +1,94 @@
+// Copyright Epic Games, Inc. All Rights Reserved.
+
+#pragma once
+
+#include "SimModule/DeferredForcesModular.h"
+
+class FGeometryCollectionPhysicsProxy;
+
+namespace Chaos
+{
+	class ISimulationModuleBase;
+	struct FAllInputs;  // #TODO - needs template or bass class
+
+	class CHAOSVEHICLESCORE_API FSimModuleTree
+	{
+	public:
+		struct FSimModuleNode
+		{
+			FSimModuleNode()
+				: SimModule(nullptr)
+				, Parent(INVALID_INDEX)
+			{
+			}
+
+			ISimulationModuleBase* SimModule;	// we own the pointer, or using sim module container?
+			int Parent;
+			TSet<int> Children;
+
+			const static int INVALID_INDEX = -1;
+
+		};
+
+		FSimModuleTree()
+		{
+		}
+
+		~FSimModuleTree()
+		{
+			DeleteNodesBelow(0);
+		}
+
+		void Reset()
+		{
+			DeleteNodesBelow(0);
+		}
+
+		bool IsEmpty() const { return SimulationModuleTree.IsEmpty(); }
+		int GetParent(int Index) const { check(!SimulationModuleTree.IsEmpty()); return SimulationModuleTree[Index].Parent; }
+		const TSet<int>& GetChildren(int Index) const { check(!SimulationModuleTree.IsEmpty()); return SimulationModuleTree[Index].Children; }
+		const ISimulationModuleBase* GetSimModule(int Index) const { return IsValidNode(Index) ? SimulationModuleTree[Index].SimModule : nullptr; }
+		ISimulationModuleBase* AccessSimModule(int Index) const { return IsValidNode(Index) ? SimulationModuleTree[Index].SimModule : nullptr; }
+		bool IsValidNode(int Index) const { return (SimulationModuleTree.IsEmpty() || Index >= SimulationModuleTree.Num()) ? false : true; }
+		int NumActiveNodes() const { return SimulationModuleTree.Num() - FreeList.Num(); }
+		void GetRootNodes(TArray<int>& RootNodesOut);
+		int GetNumNodes() const { return SimulationModuleTree.Num(); }
+
+		int AddRoot(ISimulationModuleBase* SimModule);
+
+		void Reparent(int Index, int ParentIndex);
+
+		int AddNodeBelow(int AtIndex, ISimulationModuleBase* SimModule);
+
+		int InsertNodeAbove(int AtIndex, ISimulationModuleBase* SimModule);
+
+		void DeleteNode(int AtIndex);
+
+		void Simulate(float DeltaTime, FAllInputs& Inputs, FGeometryCollectionPhysicsProxy* PhysicsProxy);
+
+		FDeferredForcesModular& AccessDeferredForces() { return DeferredForces; }
+		const FDeferredForcesModular& GetDeferredForces() const { return DeferredForces; }
+		const TArray<FSimModuleNode>& GetSimulationModuleTree() { return SimulationModuleTree; }
+
+	protected:
+		void SimulateNode(float DeltaTime, FAllInputs& Inputs, int NodeIdx);
+
+		void DeleteNodesBelow(int NodeIdx);
+
+		int GetNextIndex();
+
+		void UpdateModuleVelocites(FGeometryCollectionPhysicsProxy* PhysicsProxy);
+
+	protected:
+
+		TArray<FSimModuleNode> SimulationModuleTree;
+		TArray<int> FreeList;
+
+		FDeferredForcesModular DeferredForces;
+
+	};
+
+
+
+
+} // namespace Chaos
diff --git a/Engine/Source/Runtime/Experimental/ChaosVehicles/ChaosVehiclesCore/Public/SimModule/SimModulesInclude.h b/Engine/Source/Runtime/Experimental/ChaosVehicles/ChaosVehiclesCore/Public/SimModule/SimModulesInclude.h
new file mode 100644
--- /dev/null
+++ b/Engine/Source/Runtime/Experimental/ChaosVehicles/ChaosVehiclesCore/Public/SimModule/SimModulesInclude.h
@@ -0,0 +1,16 @@
+// Copyright Epic Games, Inc. All Rights Reserved.
+
+#pragma once
+
+#include "SimModule/ChassisModule.h"
+#include "SimModule/ThrusterModule.h"
+#include "SimModule/AerofoilModule.h"
+#include "SimModule/MotorModule.h"
+#include "SimModule/AxleModule.h"
+#include "SimModule/EngineModule.h"
+#include "SimModule/SuspensionModule.h"
+#include "SimModule/WheelModule.h"
+#include "SimModule/TransmissionModule.h"
+#include "SimModule/ClutchModule.h"
+#include "SimModule/TorqueSimModule.h"
+#include "SimModule/SimModuleTree.h"
diff --git a/Engine/Source/Runtime/Experimental/ChaosVehicles/ChaosVehiclesCore/Public/SimModule/SimulationModuleBase.h b/Engine/Source/Runtime/Experimental/ChaosVehicles/ChaosVehiclesCore/Public/SimModule/SimulationModuleBase.h
new file mode 100644
--- /dev/null
+++ b/Engine/Source/Runtime/Experimental/ChaosVehicles/ChaosVehiclesCore/Public/SimModule/SimulationModuleBase.h
@@ -0,0 +1,282 @@
+// Copyright Epic Games, Inc. All Rights Reserved.
+
+#pragma once
+
+#include "CoreMinimal.h"
+#include "Logging/LogMacros.h"
+
+
+DECLARE_LOG_CATEGORY_EXTERN(LogSimulationModule, Log, All);
+
+struct CHAOSVEHICLESCORE_API FCoreModularVehicleDebugParams
+{
+	bool ShowMass = false;
+	bool ShowForces = false;
+	float DrawForceScaling = 0.0005f;
+};
+
+namespace Chaos
+{
+	class FSimModuleTree;
+
+	struct CHAOSVEHICLESCORE_API FControlInputs
+	{
+
+		FControlInputs()
+			: Throttle(0)
+			, Brake(0)
+			, Steering(0)
+			, Clutch(0)
+			, Handbrake(0)
+			, Roll(0)
+			, Pitch(0)
+			, Yaw(0)
+			, ChangeUp(false)
+			, ChangeDown(false)
+			, SetGearNumber(0)
+		{
+		}
+
+		float Throttle;
+		float Brake;
+		float Steering;
+		float Clutch;
+		float Handbrake;
+		float Roll;
+		float Pitch;
+		float Yaw;
+		bool ChangeUp;
+		bool ChangeDown;
+		int SetGearNumber;
+	};
+
+	struct CHAOSVEHICLESCORE_API FModuleHitResults
+	{
+		int SimIndex;
+		FVector ImpactPoint;
+		float Distance;
+		bool bBlockingHit;
+	};
+
+	struct CHAOSVEHICLESCORE_API FAllInputs : public FControlInputs
+	{
+		FTransform VehicleWorldTransform;
+		TMap<int32, FModuleHitResults> HitResults;
+	};
+
+	/**
+	 * Code common between all simulation building blocks settings
+	 */
+	template <typename T>
+	class CHAOSVEHICLESCORE_API TSimModuleSettings
+	{
+	public:
+
+		TSimModuleSettings(const T& SetupIn) : SetupData(SetupIn)
+		{
+			SetupData = SetupIn; // deliberate copy for now
+		}
+
+		FORCEINLINE T& AccessSetup()
+		{
+			return (T&)(SetupData);
+		}
+
+		FORCEINLINE const T& Setup() const
+		{
+			return (SetupData);
+		}
+
+	private:
+		T SetupData;
+	};
+
+
+	enum eSimModuleState
+	{
+		Disabled,
+		Enabled
+	};
+
+	enum eSimModuleTypeFlags
+	{
+		NonFunctional	= (1<<0),	// bitmask 1,2,4,8
+		Raycast			= (1<<1),	// requires raycast data
+		TorqueBased		= (1 << 2),	// performs torque calculations
+		Velocity		= (1 << 3),	// requires velocity data
+	};
+
+	enum eSimType
+	{
+		Undefined = 0,
+		Chassis,		// no simulation effect
+		Thruster,		// applies force
+		Aerofoil,		// applied drag and lift forces
+		Wheel,			// a wheel will simply roll if it has no power source
+		Suspension,		// associated with a wheel
+		Axle,			// connects more than one wheel
+		Transmission,	// gears - torque multiplier
+		Engine,			// (torque curve required) power source generates torque for wheel, axle, transmission, clutch
+		Motor,			// (electric?, no torque curve required?) power source generates torque for wheel, axle, transmission, clutch
+		Clutch,			// limits the amount of torque transferred between source and destination allowing for different rotation speeds of connected axles
+		Wing,			// lift and controls aircraft roll
+		Rudder,			// controls aircraft yaw
+		Elevator,		// controls aircraft pitch
+		Propeller,		// generates thrust when connected to a motor/engine
+
+	};
+
+	/**
+	 * Interface base class for all simulation module building blocks
+	 */
+	class CHAOSVEHICLESCORE_API ISimulationModuleBase
+	{
+	public:
+		const static int INVALID_INDEX = -1;
+
+		ISimulationModuleBase()
+			: SimModuleTree(nullptr)
+			, SimTreeIndex(INVALID_INDEX)
+			, StateFlags(Enabled)
+			, TransformIndex(INVALID_INDEX)
+			, ModuleLocalVelocity(FVector::ZeroVector)
+			, bClustered(true)
+			, AppliedForce(FVector::ZeroVector)
+		{}
+		virtual ~ISimulationModuleBase() {}
+
+		/**
+		* Get the friendly name for this module, primarily for logging & debugging module tree
+		 */
+		virtual const FString GetDebugName() const = 0;
+
+		/**
+		* Is Module of a specific behavioral data type
+		 */
+		virtual bool IsBehaviourType(eSimModuleTypeFlags InType) const = 0;
+
+		/**
+		 * The specific simulation type
+		 */
+		virtual eSimType GetSimType() const = 0;
+
+		/**
+		* Is Module active and simulating
+		 */
+		virtual bool IsEnabled() const { return (StateFlags == eSimModuleState::Enabled); }
+
+		/*
+		* Set Module state, if simulating or not
+		*/
+		void SetStateFlags(eSimModuleState StateFlagsIn) { StateFlags = StateFlagsIn; }
+
+		/**
+		 * The main Simulation function that is called from the physics async callback thread
+		 */
+		virtual void Simulate(float DeltaTime, const FAllInputs& Inputs, FSimModuleTree& VehicleModuleSystem) {};
+
+		/**
+		 * Option to draw debug for this module requires CVar p.Chaos.DebugDraw.Enabled 1
+		 */
+		virtual void DrawDebugInfo() {}
+
+		/**
+		 * Option to return debug text for drawing on the HUD in the Game Thread
+		 */
+		virtual bool GetDebugString(FString& StringOut) const;
+
+		/**
+		 * The transform index references the transform collection, mapping the simulation module to the geometry collection data
+		 */
+		void SetTransformIndex(int TransformIndexIn) { TransformIndex = TransformIndexIn; }
+		const int GetTransformIndex() const { return TransformIndex; }
+
+		/**
+		 * The modules own index in the simulation tree array
+		 */
+		void SetTreeIndex(int TreeIndexIn) { SimTreeIndex = TreeIndexIn; }
+		int GetTreeIndex() const { return SimTreeIndex; }
+
+		/**
+		 * Very useful to store the simulation tree pointer in which we are stored, then we can access other modules that we reference through an index
+		 */
+		void SetSimModuleTree(FSimModuleTree* SimModuleTreeIn) { SimModuleTree = SimModuleTreeIn; }
+
+		/**
+		 * Force application function, handles deferred force application and applying the force at the collect location based on whether the GC cluster is intact or fractured
+		 * Note: forces are applied in local coordinates of the module
+		 */
+		void AddLocalForceAtPosition(const FVector& Force, const FVector& Position, bool bAllowSubstepping = true, bool bIsLocalForce = false);
+
+		/**
+		 * Force application function, handles deferred force application and applying the force at the collect location based on whether the GC cluster is intact or fractured
+		 * Note: forces are applied in local coordinates of the module
+		 */
+		void AddLocalForce(const FVector& Force, bool bAllowSubstepping = true, bool bIsLocalForce = false);
+
+		//---
+
+		/**
+		 * Let the module know if it is still clustered or not
+		 */
+		void SetClustered(bool IsClusteredIn) { bClustered = IsClusteredIn; }
+		bool IsClustered() const { return bClustered; }
+
+		/**
+		 * Set the COM relative transform of module when it is clustered, so relative to parent COM
+		 */
+		void SetClusteredTransform(const FTransform& TransformIn) { ClusteredCOMRelativeTransform  = TransformIn; }
+		const FTransform& GetClusteredTransform() const { return ClusteredCOMRelativeTransform; }
+
+		/**
+		 * Set the COM relative transform of module when it is broken off, so relative to itself
+		 */
+		void SetIntactTransform(const FTransform& TransformIn) { IntactCOMRelativeTransform = TransformIn; }
+		const FTransform& GetIntactTransform() const { return IntactCOMRelativeTransform; }
+
+		/**
+		 * The modules transform relative to the simulating body will depend on whether the GC is intact (get the transform relative to intact cluster)
+		 * or fractured (transform relative to fractured part)
+		 */
+		const FTransform& GetParentRelativeTransform() const
+		{
+			if (bClustered)
+			{
+				return GetClusteredTransform();
+			}
+			else
+			{
+				return GetIntactTransform();
+			}
+		}
+
+		/**
+		 * Update the module with its current velocity
+		 */
+		void SetLocalVelocity(const FVector& VelocityIn) { ModuleLocalVelocity = VelocityIn; }
+		const FVector& GetLocalVelocity() const { return ModuleLocalVelocity; }
+
+		ISimulationModuleBase* GetParent();
+		ISimulationModuleBase* GetFirstChild();
+
+		// for headless chaos testing
+		const FVector& GetAppliedForce() { return AppliedForce; }
+
+	protected:
+
+		FSimModuleTree* SimModuleTree;	// A pointer back to the simulation tree where we are stored
+		int SimTreeIndex;	// Index of this SimModule in the FSimModuleTree
+		eSimModuleState StateFlags;	// TODO: make this more like flags
+		int TransformIndex; // Index of this Sim Module's node in Geometry Collection Transform array
+
+		FTransform ClusteredCOMRelativeTransform;
+		FTransform IntactCOMRelativeTransform;
+		FVector ModuleLocalVelocity;
+		bool bClustered;
+
+		// for headless chaos testing
+		FVector AppliedForce;
+
+	};
+
+} // namespace Chaos
diff --git a/Engine/Source/Runtime/Experimental/ChaosVehicles/ChaosVehiclesCore/Public/SimModule/SuspensionModule.h b/Engine/Source/Runtime/Experimental/ChaosVehicles/ChaosVehiclesCore/Public/SimModule/SuspensionModule.h
new file mode 100644
--- /dev/null
+++ b/Engine/Source/Runtime/Experimental/ChaosVehicles/ChaosVehiclesCore/Public/SimModule/SuspensionModule.h
@@ -0,0 +1,104 @@
+// Copyright Epic Games, Inc. All Rights Reserved.
+
+#pragma once
+
+#include "SimModule/SimulationModuleBase.h"
+
+namespace Chaos
+{
+	struct FAllInputs;
+	class FSimModuleTree;
+
+	struct CHAOSVEHICLESCORE_API FSuspensionSettings
+	{
+		FSuspensionSettings()
+			: SuspensionAxis(FVector(0.f, 0.f, -1.f))
+			, RestOffset(FVector::ZeroVector)
+			, MaxRaise(5.f)
+			, MaxDrop(5.f)
+			, MaxLength(0.f)
+			, SpringRate(1.f)
+			, SpringPreload(0.5f)
+			, CompressionDamping(0.9f)
+			, ReboundDamping(0.9f)
+		//	, SwaybarEffect(0.5f)
+		//	, DampingRatio(0.3f)
+		{
+
+		}
+
+		FVector SuspensionAxis;		// local axis, direction of suspension force raycast traces
+		FVector RestOffset;	
+		float MaxRaise;				// distance [cm]
+		float MaxDrop;				// distance [cm]
+		float MaxLength;			// distance [cm]
+
+		float SpringRate;			// spring constant
+		float SpringPreload;		// Amount of Spring force (independent spring movement)
+		float CompressionDamping;	// limit compression speed
+		float ReboundDamping;		// limit rebound speed
+
+	//	float Swaybar;				// Anti-roll bar
+
+	//	float DampingRatio;			// value between (0-no damping) and (1-critical damping)
+	};
+
+	/** Suspension world ray/shape trace start and end positions */
+	struct CHAOSVEHICLESCORE_API FSpringTrace
+	{
+		FVector Start;
+		FVector End;
+
+		FVector TraceDir()
+		{
+			FVector Dir(End - Start);
+			return Dir.FVector::GetSafeNormal();
+		}
+
+		float Length()
+		{
+			FVector Dir(End - Start);
+			return Dir.Size();
+		}
+	};
+
+	class CHAOSVEHICLESCORE_API FSuspensionSimModule : public ISimulationModuleBase, public TSimModuleSettings<FSuspensionSettings>
+	{
+	public:
+
+		FSuspensionSimModule(const FSuspensionSettings& Settings);
+
+		virtual bool IsBehaviourType(eSimModuleTypeFlags InType) const override { return (InType & Raycast); }
+
+		virtual eSimType GetSimType() const { return eSimType::Suspension; }
+
+		virtual const FString GetDebugName() const { return TEXT("Suspension"); }
+
+		float GetSpringLength() const;
+		void SetSpringLength(float InLength, float WheelRadius);
+		void GetWorldRaycastLocation(const FTransform& BodyTransform, float WheelRadius, FSpringTrace& OutTrace);
+
+		virtual void Simulate(float DeltaTime, const FAllInputs& Inputs, FSimModuleTree& VehicleModuleSystem) override;
+
+		void SetLocation(const FVector& LocationIn) { Location = LocationIn; }
+		const FVector GetLocation() const { return /*Setup().RestOffset + */Location; }
+
+		const FVector& GetRestLocation() const { return Setup().RestOffset; }
+
+		void SetWheelSimTreeIndex(int WheelTreeIndexIn) { WheelSimTreeIndex = WheelTreeIndexIn; }
+		int GetWheelSimTreeIndex() const { return WheelSimTreeIndex; }
+
+	private:
+
+		float SpringDisplacement;
+		float LastDisplacement;
+		int WheelSimTreeIndex;
+
+		FVector Location;
+	};
+
+
+} // namespace Chaos
+
+
+PRAGMA_ENABLE_OPTIMIZATION
diff --git a/Engine/Source/Runtime/Experimental/ChaosVehicles/ChaosVehiclesCore/Public/SimModule/ThrusterModule.h b/Engine/Source/Runtime/Experimental/ChaosVehicles/ChaosVehiclesCore/Public/SimModule/ThrusterModule.h
new file mode 100644
--- /dev/null
+++ b/Engine/Source/Runtime/Experimental/ChaosVehicles/ChaosVehiclesCore/Public/SimModule/ThrusterModule.h
@@ -0,0 +1,50 @@
+// Copyright Epic Games, Inc. All Rights Reserved.
+
+#pragma once
+
+#include "SimModule/SimulationModuleBase.h"
+
+
+namespace Chaos
+{
+	struct FAllInputs;
+	class FSimModuleTree;
+
+	/// <summary>
+	/// Thruster settings
+	/// </summary>
+	struct CHAOSVEHICLESCORE_API FThrusterSettings
+	{
+		FThrusterSettings()
+			: MaxThrustForce(0)
+			, ForceOffset(FVector::ZeroVector)
+		{
+
+		}
+
+		float MaxThrustForce;
+		FVector ForceOffset;
+	};
+
+	/// <summary>
+	/// A vehicle component that transmits torque from one source to another, i.e. from an engine or differential to wheels
+	///
+	/// </summary>
+	class CHAOSVEHICLESCORE_API FThrusterSimModule : public ISimulationModuleBase, public TSimModuleSettings<FThrusterSettings>
+	{
+	public:
+		FThrusterSimModule(const FThrusterSettings& Settings);
+
+		virtual eSimType GetSimType() const { return eSimType::Thruster; }
+
+		virtual const FString GetDebugName() const { return TEXT("Thruster"); }
+
+		virtual bool IsBehaviourType(eSimModuleTypeFlags InType) const override { return (InType & NonFunctional); }
+
+		virtual void Simulate(float DeltaTime, const FAllInputs& Inputs, FSimModuleTree& VehicleModuleSystem) override;
+
+	private:
+	};
+
+
+} // namespace Chaos
diff --git a/Engine/Source/Runtime/Experimental/ChaosVehicles/ChaosVehiclesCore/Public/SimModule/TorqueSimModule.h b/Engine/Source/Runtime/Experimental/ChaosVehicles/ChaosVehiclesCore/Public/SimModule/TorqueSimModule.h
new file mode 100644
--- /dev/null
+++ b/Engine/Source/Runtime/Experimental/ChaosVehicles/ChaosVehiclesCore/Public/SimModule/TorqueSimModule.h
@@ -0,0 +1,84 @@
+// Copyright Epic Games, Inc. All Rights Reserved.
+
+#pragma once
+
+#include "SimModule/SimulationModuleBase.h"
+#include "VehicleUtility.h"
+
+namespace Chaos
+{
+
+
+class CHAOSVEHICLESCORE_API FTorqueSimModule : public ISimulationModuleBase
+{
+public:
+	FTorqueSimModule()
+		: DriveTorque(0.f)
+		, LoadTorque(0.f)
+		, BrakingTorque(0.f)
+		, AngularVelocity(0.f)
+		, AngularPosition(0.0f)
+		, CombinedInertia(0.0f)
+	{
+	}
+
+	/**
+	 * Is Module of a specific type - used for casting
+	 */
+	virtual bool IsBehaviourType(eSimModuleTypeFlags InType) const { return (InType & TorqueBased); }
+
+	void SetDriveTorque(float TorqueIn) { DriveTorque = TorqueIn; }
+	float GetDriveTorque() const { return DriveTorque; }
+
+	void SetLoadTorque(float TorqueIn) { LoadTorque = TorqueIn; }
+	float GetLoadTorque() const { return LoadTorque; }
+
+	void SetBrakingTorque(float TorqueIn) { BrakingTorque = TorqueIn; }
+	float GetBrakingTorque() const { return BrakingTorque; }
+
+	void SetAngularVelocity(float AngularVelocityIn) { AngularVelocity = AngularVelocityIn; }
+	float GetAngularVelocity() const { return AngularVelocity; }
+
+	void AddAngularVelocity(float AngularVelocityIn) { AngularVelocity += AngularVelocityIn; }
+
+	void SetAngularPosition(float AngularPositionIn) { AngularPosition = AngularPositionIn; }
+	float GetAngularPosition() const { return AngularPosition; }
+
+	void SetRPM(float InRPM) { AngularVelocity = RPMToOmega(InRPM); }
+	float GetRPM() const { return OmegaToRPM(AngularVelocity); }
+
+	/**
+	 * Transmit torque between this module and its Parent and Children. DriveTorque passed down to children, LoadTorque passed from child to parent
+	 */
+	void TransmitTorque(const FSimModuleTree& BlockSystem, float PushedTorque, float BrakeTorque = 0.f, float GearingRatio = 1.0f, float ClutchSlip = 1.0f);
+
+	/**
+	 * Integrate angular velocity using specified DeltaTime & Inertia value, Note the inertia should be the combined inertia of all the connected pieces otherwise things will rotate at different rates
+	 */
+	void IntegrateAngularVelocity(float DeltaTime, float Inertia, float MaxRotationVel=MAX_FLT);
+
+	/**
+	 * Cast an ISimulationModuleBase to a FTorqueSimModule if compatible class
+	 */
+	static FTorqueSimModule* CastToTorqueInterface(ISimulationModuleBase* SimModule)
+	{
+		if (SimModule && SimModule->IsBehaviourType(eSimModuleTypeFlags::TorqueBased))
+		{
+			return static_cast<FTorqueSimModule*>(SimModule);
+		}
+
+		return nullptr;
+	}
+
+protected:
+
+	float DriveTorque;
+	float LoadTorque;
+	float BrakingTorque;
+	float AngularVelocity;
+	float AngularPosition;
+	float CombinedInertia; // Note: ?? Do we want this
+};
+
+
+} // namespace Chaos
\ No newline at end of file
diff --git a/Engine/Source/Runtime/Experimental/ChaosVehicles/ChaosVehiclesCore/Public/SimModule/TransmissionModule.h b/Engine/Source/Runtime/Experimental/ChaosVehicles/ChaosVehiclesCore/Public/SimModule/TransmissionModule.h
new file mode 100644
--- /dev/null
+++ b/Engine/Source/Runtime/Experimental/ChaosVehicles/ChaosVehiclesCore/Public/SimModule/TransmissionModule.h
@@ -0,0 +1,118 @@
+// Copyright Epic Games, Inc. All Rights Reserved.
+
+#pragma once
+
+#include "SimModule/TorqueSimModule.h"
+
+
+namespace Chaos
+{
+	struct FAllInputs;
+	class FSimModuleTree;
+
+	struct CHAOSVEHICLESCORE_API FTransmissionSettings
+	{
+		enum ETransType : uint8
+		{
+			Manual,
+			Automatic
+		};
+
+		FTransmissionSettings()
+			: FinalDriveRatio(3.f)
+			, ChangeUpRPM(5000)
+			, ChangeDownRPM(2500)
+			, GearChangeTime(0.5f)
+			, TransmissionEfficiency(1.f)
+			, TransmissionType(ETransType::Automatic)
+			, AutoReverse(true)
+		{
+			ForwardRatios.Add(2.85f);
+			ForwardRatios.Add(2.02f);
+			ForwardRatios.Add(1.35f);
+			ForwardRatios.Add(1.0f);
+
+			ReverseRatios.Add(2.86f);
+		}
+
+		TArray<float> ForwardRatios;	// Gear ratios for forward gears
+		TArray<float> ReverseRatios;	// Gear ratios for reverse Gear(s)
+		float FinalDriveRatio;			// Final drive ratio [~4.0]
+
+		uint32 ChangeUpRPM;				// [RPM]
+		uint32 ChangeDownRPM;			// [RPM]
+		float GearChangeTime; 			// [sec]
+
+		float TransmissionEfficiency;	// Loss from friction in the system mean we might run at around 0.94 Efficiency
+
+		ETransType TransmissionType;	// Specify Automatic or Manual transmission
+
+		bool AutoReverse;				// Arcade handling - holding Brake switches into reverse after vehicle has stopped
+	};
+
+	class CHAOSVEHICLESCORE_API FTransmissionSimModule : public FTorqueSimModule, public TSimModuleSettings<FTransmissionSettings>
+	{
+	public:
+
+		FTransmissionSimModule(const FTransmissionSettings& Settings)
+			: TSimModuleSettings<FTransmissionSettings>(Settings)
+			, CurrentGear(1)
+			, TargetGear(1)
+			, CurrentGearChangeTime(0.f)
+			, AllowedToChangeGear(true)
+		{
+		}
+
+		virtual eSimType GetSimType() const { return eSimType::Transmission; }
+
+		virtual const FString GetDebugName() const { return TEXT("Transmission"); }
+
+		virtual bool GetDebugString(FString& StringOut) const override;
+
+		virtual bool IsBehaviourType(eSimModuleTypeFlags InType) const override { return FTorqueSimModule::IsBehaviourType(InType) || (InType & Velocity); }
+
+		virtual void Simulate(float DeltaTime, const FAllInputs& Inputs, FSimModuleTree& VehicleModuleSystem) override;
+
+	private:
+
+		/** set the target gear number to change to, can change gear immediately if specified
+		 *  i.e. rather than waiting for the gear change time to elapse
+		 */
+		void SetGear(int32 InGear, bool Immediate = false);
+
+		/** Get the final combined gear ratio for the specified gear (reverse gears < 0, neutral 0, forward gears > 0) */
+		float GetGearRatio(int32 InGear) const;
+
+		/** set the target gear to one higher than current target, will clamp gear index within rage */
+		void ChangeUp()
+		{
+			SetGear(TargetGear + 1);
+		}
+
+		/** set the target gear to one lower than current target, will clamp gear index within rage */
+		void ChangeDown()
+		{
+			SetGear(TargetGear - 1);
+		}
+
+		/** Are we currently in the middle of a gear change */
+		bool IsCurrentlyChangingGear() const
+		{
+			return CurrentGear != TargetGear;
+		}
+
+		void CorrectGearInputRange(int32& GearIndexInOut) const
+		{
+			GearIndexInOut = FMath::Clamp(GearIndexInOut, -Setup().ReverseRatios.Num(), Setup().ForwardRatios.Num());
+		}
+
+		int32 CurrentGear; // <0 reverse gear(s), 0 neutral, >0 forward gears
+		int32 TargetGear;  // <0 reverse gear(s), 0 neutral, >0 forward gears
+		float CurrentGearChangeTime; // Time to change gear, no power transmitted to the wheels during change
+
+		bool AllowedToChangeGear; // conditions are ok for an automatic gear change
+
+	};
+
+
+} // namespace Chaos
diff --git a/Engine/Source/Runtime/Experimental/ChaosVehicles/ChaosVehiclesCore/Public/SimModule/WheelModule.h b/Engine/Source/Runtime/Experimental/ChaosVehicles/ChaosVehiclesCore/Public/SimModule/WheelModule.h
new file mode 100644
--- /dev/null
+++ b/Engine/Source/Runtime/Experimental/ChaosVehicles/ChaosVehiclesCore/Public/SimModule/WheelModule.h
@@ -0,0 +1,105 @@
+// Copyright Epic Games, Inc. All Rights Reserved.
+
+#pragma once
+
+#include "SimModule/TorqueSimModule.h"
+
+namespace Chaos
+{
+	struct FAllInputs;
+	class FSimModuleTree;
+
+	struct CHAOSVEHICLESCORE_API FWheelSettings
+	{
+		FWheelSettings()
+			: Radius(30.0f)
+			, Width(20.0f)
+			, WheelInertia(100000.0f) // TODO: ?? defaults and proper value ??
+
+			, FrictionMultiplier(3.0f)
+			, LateralSlipGraphMultiplier(1.0f)
+			, CorneringStiffness(1000.0f)
+			, SlipAngleLimit(8.0f)
+
+			, ABSEnabled(true)
+			, TractionControlEnabled(true)
+			, SteeringEnabled(false)
+			, HandbrakeEnabled(false)
+			, MaxSteeringAngle(45)
+			, MaxBrakeTorque(4000)
+			, HandbrakeTorque(3000)
+			, MaxRotationVel(100.0f)
+
+		{
+
+		}
+
+		float Radius;
+		float Width;
+		float WheelInertia;
+
+		float FrictionMultiplier;
+		float LateralSlipGraphMultiplier;
+		float CorneringStiffness;
+		FGraph LateralSlipGraph;
+		float SlipAngleLimit;
+
+		bool ABSEnabled;			// Advanced braking system operational
+		bool TractionControlEnabled;// Straight Line Traction Control
+		bool SteeringEnabled;
+		bool HandbrakeEnabled;
+
+		float MaxSteeringAngle;
+		float MaxBrakeTorque;
+		float HandbrakeTorque;
+
+		float MaxRotationVel;
+
+
+	};
+
+	class CHAOSVEHICLESCORE_API FWheelSimModule : public FTorqueSimModule, public TSimModuleSettings<FWheelSettings>
+	{
+	public:
+
+		FWheelSimModule(const FWheelSettings& Settings);
+
+		virtual eSimType GetSimType() const { return eSimType::Wheel; }
+
+		virtual const FString GetDebugName() const { return TEXT("Wheel"); }
+
+		virtual bool GetDebugString(FString& StringOut) const override;
+
+		virtual void Simulate(float DeltaTime, const FAllInputs& Inputs, FSimModuleTree& VehicleModuleSystem) override;
+
+		virtual bool IsBehaviourType(eSimModuleTypeFlags InType) const override { return (InType & TorqueBased)||(InType & Velocity); }
+
+
+		/** Set the vehicle's speed at the wheels location in local wheel coords */
+		//void SetGroundSpeed(const FVector& VIn)
+		//{
+		//	GroundVelocityVector = VIn;
+		//}
+
+		void SetSuspensionSimTreeIndex(int IndexIn) { SuspensionSimTreeIndex = IndexIn; }
+		int GetSuspensionSimTreeIndex() const { return SuspensionSimTreeIndex; }
+
+		float GetSteerAngleDegrees() const { return SteerAngleDegrees; }
+
+		void SetForceIntoSurface(float ForceIntoSurfaceIn) { ForceIntoSurface = ForceIntoSurfaceIn; }
+
+	private:
+
+		float BrakeTorque;				// [N.m]
+		float ForceIntoSurface;			// [N]
+		float SurfaceFriction;
+		int SuspensionSimTreeIndex;
+
+		FVector ForceFromFriction;
+		float MassPerWheel;
+		float SteerAngleDegrees;
+
+	};
+
+
+} // namespace Chaos
diff --git a/Engine/Source/Runtime/Experimental/GeometryCollectionEngine/Public/GeometryCollection/GeometryCollectionComponent.h b/Engine/Source/Runtime/Experimental/GeometryCollectionEngine/Public/GeometryCollection/GeometryCollectionComponent.h
--- a/Engine/Source/Runtime/Experimental/GeometryCollectionEngine/Public/GeometryCollection/GeometryCollectionComponent.h
+++ b/Engine/Source/Runtime/Experimental/GeometryCollectionEngine/Public/GeometryCollection/GeometryCollectionComponent.h
@@ -867,7 +867,8 @@ public:
 	const TArray<FMatrix>& GetGlobalMatrices() { return GlobalMatrices; }
 
 	const FGeometryDynamicCollection* GetDynamicCollection() const { return DynamicCollection.Get(); }
-	
+	FGeometryDynamicCollection* GetDynamicCollection() { return DynamicCollection.Get(); } // TEMP HACK?
+
 public:
 	UPROPERTY(BlueprintAssignable, Category = "Collision")
 	FOnChaosPhysicsCollision OnChaosPhysicsCollision;
