Revert previous change to this file to Fix Headless Chaos

diff --git a/Engine/Source/Runtime/Experimental/Chaos/Private/GeometryCollection/ManagedArrayCollection.cpp b/Engine/Source/Runtime/Experimental/Chaos/Private/GeometryCollection/ManagedArrayCollection.cpp
--- a/Engine/Source/Runtime/Experimental/Chaos/Private/GeometryCollection/ManagedArrayCollection.cpp
+++ b/Engine/Source/Runtime/Experimental/Chaos/Private/GeometryCollection/ManagedArrayCollection.cpp
@@ -381,12 +381,6 @@ void FManagedArrayCollection::CopyAttribute(const FManagedArrayCollection& InCol
 	FKeyType SrcKey = FManagedArrayCollection::MakeMapKey(SrcName, Group);
 	FKeyType DestKey = FManagedArrayCollection::MakeMapKey(DestName, Group);
 
-	// fail gracefully
-	if (!(ensureMsgf(HasAttribute(SrcName, Group), TEXT("Missing managed array collection attribute"))))
-	{
-		return;
-	}
-
 	if (!HasAttribute(DestName, Group))
 	{
 		const FValueType& V = InCollection.Map[SrcKey];
