Chaos : fix crash with level attribute not being present when copy to physics thread collection happens

diff --git a/Engine/Source/Runtime/Experimental/GeometryCollectionEngine/Private/GeometryCollection/GeometryCollectionComponent.cpp b/Engine/Source/Runtime/Experimental/GeometryCollectionEngine/Private/GeometryCollection/GeometryCollectionComponent.cpp
--- a/Engine/Source/Runtime/Experimental/GeometryCollectionEngine/Private/GeometryCollection/GeometryCollectionComponent.cpp
+++ b/Engine/Source/Runtime/Experimental/GeometryCollectionEngine/Private/GeometryCollection/GeometryCollectionComponent.cpp
@@ -2289,6 +2289,19 @@ void UGeometryCollectionComponent::OnCreatePhysicsState()
 				}
 			}
 
+			// there's a code path where Level is not serialized and InitializeSharedCollisionStructures is not being called,
+			// resulting in the attribute missing and causing a crash in CopyAttribute calls later in FGeometryCollectionPhysicsProxy::Initialize
+			// @todo(chaos) we should better handle computation of dependent attribute like level
+			// @todo(chaos) We should implement a facade for levels, (parent and child included ? )
+			if (!RestCollection->GetGeometryCollection()->HasAttribute("Level", FTransformCollection::TransformGroup))
+			{
+				TManagedArray<int32>& Levels = RestCollection->GetGeometryCollection()->AddAttribute<int32>("Level", FTransformCollection::TransformGroup);
+				for (int32 TransformIndex = 0; TransformIndex < Levels.Num(); ++TransformIndex)
+				{
+					FGeometryCollectionPhysicsProxy::CalculateAndSetLevel(TransformIndex, RestCollection->GetGeometryCollection()->Parent,  Levels);
+				}
+			}
+			
 			// let's copy anchored information if available
 			const Chaos::Facades::FCollectionAnchoringFacade RestCollectionAnchoringFacade(*RestCollection->GetGeometryCollection());
 			Chaos::Facades::FCollectionAnchoringFacade DynamicCollectionAnchoringFacade(*DynamicCollection);
