Chaos : Flesh

Add new MPM grid class and associated tests.

diff --git a/Engine/Source/Runtime/Experimental/Chaos/Private/Chaos/UniformGrid.cpp b/Engine/Source/Runtime/Experimental/Chaos/Private/Chaos/UniformGrid.cpp
--- a/Engine/Source/Runtime/Experimental/Chaos/Private/Chaos/UniformGrid.cpp
+++ b/Engine/Source/Runtime/Experimental/Chaos/Private/Chaos/UniformGrid.cpp
@@ -250,6 +250,135 @@ TVector<T, 3> TUniformGrid<T, 3>::ClampMinusHalf(const TVector<T, 3>& X) const
 	return Result;
 }
 
+template<class T>
+void TMPMGrid<T>::BaseNodeIndex(const TVector<T, 3>& X, TVector<int32, 3>& Index, TVector<T, 3>& Weights) const
+{
+	for (uint32 i = 0; i < 3; ++i)
+		Index[i] = int32(FMath::Floor((X[i] - MMinCorner[i]) / MDx[i]));
+	for (uint32 i = 0; i < 3; ++i)
+		Weights[i] = (X[i] - MMinCorner[i]) / MDx[i] - T(Index[i]);
+}
+
+template<class T>
+int32 TMPMGrid<T>::FlatIndex(const TVector<int32, 3>& Index) const 
+{
+	/*
+	This assumes data will be stored at each index (i.e. numbered by flatIndex(index)). For linear interpolation over
+	the grid we assume that the index[i] \in [0,gridN[i]-1]. For quadratic interpolation, we assume that the
+	index[i] \in [-1,gridN[i]-1].
+
+	For linear interpolation, data is stored at the nodes of the grid, e.g. xi = i dx (i=0,1,...,gridN-1) and
+	for quadratic interpolation, data is stored at the cell centers of the grid, e.g. xi = i dx + dx/2 (i=-1,0,1,...,gridN-1)
+	*/
+
+	if (interp == linear) {
+		return Index[0] * MCells[1] * MCells[2] + Index[1] * MCells[2] + Index[2];
+	}
+	else {
+		return (MCells[1] + 1) * (MCells[2] + 1) * (Index[0] + 1) + (MCells[2] + 1) * (Index[1] + 1) + Index[2] + 1;
+	}
+
+
+	return -1;
+}
+
+template<class T>
+TVector<int32, 3> TMPMGrid<T>::Lin2MultiIndex(const int32 IndexIn) const 
+{
+	/*
+	This version returns an Index instead of taking one in by reference.
+	*/
+	TVector<int32, 3> MultiIndex;
+
+	if (interp == linear) {
+		MultiIndex[0] = IndexIn / (MCells[1] * MCells[2]);
+		MultiIndex[1] = (IndexIn / MCells[2]) % MCells[1];
+		MultiIndex[2] = IndexIn % MCells[2];
+	}
+	else {
+		MultiIndex[0] = IndexIn / ((MCells[1] + 1) * (MCells[2] + 1)) - 1;
+		MultiIndex[1] = (IndexIn / (MCells[2] + 1)) % (MCells[1] + 1) - 1;
+		MultiIndex[2] = IndexIn % (MCells[2] + 1) - 1;
+	}
+
+
+	return MultiIndex;
+}
+
+
+
+template<class T>
+TVector<T, 3> TMPMGrid<T>::Node(int32 FlatIndexIn) const {
+	TVector<T, 3> result;
+	TVector<int32, 3> index = Lin2MultiIndex(FlatIndexIn);
+	
+	if (interp == linear) {
+		for (int32 i = 0; i < 3; ++i) {
+			result[i] = T(index[i]) * MDx[i] + MMinCorner[i];
+		}
+	}
+	else {
+		for (size_t i = 0; i < 3; ++i) {
+			result[i] = (T(index[i]) + T(0.5)) * MDx[i] + MMinCorner[i];
+		}
+	}
+
+	return result;
+}
+
+template<class T>
+int32 TMPMGrid<T>::Loc2GlobIndex(const int32 IndexIn, const TVector<int32, 3>& LocalIndexIn) const
+{
+
+	if (interp == linear) {
+		return IndexIn + MCells[1] * MCells[2] * LocalIndexIn[0] + MCells[2] * LocalIndexIn[1] + LocalIndexIn[2];
+	}
+	else {
+		return IndexIn + (MCells[1] + 1) * (MCells[2] + 1) * (LocalIndexIn[0] - 1) + (MCells[2] + 1) * (LocalIndexIn[1] - 1) + LocalIndexIn[2] - 1;
+	}
+
+	return -1;
+
+}
+
+template<class T>
+void TMPMGrid<T>::UpdateGridFromPositions(const Chaos::TDynamicParticles<T, 3>& InParticles)
+{
+	TVector<T, 3> ParticleMin, ParticleMax;
+	if (InParticles.Size() > 0)
+	{
+		ParticleMin = InParticles.X()[0];
+		ParticleMax = InParticles.X()[0];
+		for (uint32 p = 0; p < InParticles.Size(); p++) {
+			for (int32 alpha = 0; alpha < 3; alpha++) {
+				if (ParticleMin[alpha] > InParticles.X()[p][alpha]) {
+					ParticleMin[alpha] = InParticles.X()[p][alpha];
+				}
+				if (ParticleMax[alpha] < InParticles.X()[p][alpha]) {
+					ParticleMax[alpha] = InParticles.X()[p][alpha];
+				}
+			}
+		}
+		TVector<T, 3> dims = ParticleMax - ParticleMin;
+
+		for (uint32 c = 0; c < 3; c++) {
+			MMinCorner[c] = MDx[c] * T(FMath::CeilToInt((ParticleMin[c] / MDx[c]))) - T(2) * MDx[c];
+			MCells[c] = int32(FMath::CeilToInt((dims[c] / MDx[c]))) + 4;
+			MMaxCorner[c] = MMinCorner[c] + T(MCells[c]) * MDx[c];
+		}
+	}
+}
+
+template<class T>
+void TMPMGrid<T>::SetInterp(InterpType InterpIn) {
+	interp = InterpIn;
+	if (interp == linear)
+		NPerDir = 2;
+	else
+		NPerDir = 3;
+}
+
+
 
 template<class T>
 bool Chaos::TUniformGrid<T, 3>::IsValid(const TVector<int32, 3>& X) const
@@ -269,3 +398,17 @@ template class Chaos::TUniformGrid<Chaos::FReal, 2>;
 
 template FVec3 Chaos::TUniformGridBase<Chaos::FReal, 3>::LinearlyInterpolate<FVec3>(const TArrayND<FVec3, 3>&, const FVec3&) const;
 template CHAOS_API Chaos::FReal Chaos::TUniformGridBase<Chaos::FReal, 3>::LinearlyInterpolate<Chaos::FReal>(const TArrayND<Chaos::FReal, 3>&, const FVec3&) const;
+template int32 TMPMGrid<Chaos::FReal>::Loc2GlobIndex(const int32 IndexIn, const TVector<int32, 3>& LocalIndexIn) const;
+template int32 TMPMGrid<Chaos::FReal>::FlatIndex(const TVector<int32, 3>& Index) const;
+template TVector<int32, 3> TMPMGrid<Chaos::FReal>::Lin2MultiIndex(const int32 IndexIn) const;
+template void TMPMGrid<Chaos::FReal>::SetInterp(InterpType InterpIn);
+template void TMPMGrid<Chaos::FReal>::UpdateGridFromPositions(const Chaos::TDynamicParticles<Chaos::FReal, 3>& InParticles);
+template void TMPMGrid<Chaos::FReal>::BaseNodeIndex(const TVector<Chaos::FReal, 3>& X, TVector<int32, 3>& Index, TVector<Chaos::FReal, 3>& Weights) const;
+template TVector<Chaos::FReal, 3> TMPMGrid<Chaos::FReal>::Node(int32 FlatIndexIn) const;
+template int32 TMPMGrid<Chaos::FRealSingle>::Loc2GlobIndex(const int32 IndexIn, const TVector<int32, 3>& LocalIndexIn) const;
+template int32 TMPMGrid<Chaos::FRealSingle>::FlatIndex(const TVector<int32, 3>& Index) const;
+template TVector<int32, 3> TMPMGrid<Chaos::FRealSingle>::Lin2MultiIndex(const int32 IndexIn) const;
+template void TMPMGrid<Chaos::FRealSingle>::SetInterp(InterpType InterpIn);
+template void TMPMGrid<Chaos::FRealSingle>::UpdateGridFromPositions(const Chaos::TDynamicParticles<Chaos::FRealSingle, 3>& InParticles);
+template void TMPMGrid<Chaos::FRealSingle>::BaseNodeIndex(const TVector<Chaos::FRealSingle, 3>& X, TVector<int32, 3>& Index, TVector<Chaos::FRealSingle, 3>& Weights) const;
+template TVector<Chaos::FRealSingle, 3> TMPMGrid<Chaos::FRealSingle>::Node(int32 FlatIndexIn) const;
\ No newline at end of file
diff --git a/Engine/Source/Runtime/Experimental/Chaos/Public/Chaos/UniformGrid.h b/Engine/Source/Runtime/Experimental/Chaos/Public/Chaos/UniformGrid.h
--- a/Engine/Source/Runtime/Experimental/Chaos/Public/Chaos/UniformGrid.h
+++ b/Engine/Source/Runtime/Experimental/Chaos/Public/Chaos/UniformGrid.h
@@ -3,6 +3,8 @@
 
 #include "Chaos/Vector.h"
 #include "ChaosCheck.h"
+#include "Chaos/ArrayCollectionArray.h"
+#include "Chaos/DynamicParticles.h"
 
 namespace Chaos
 {
@@ -141,8 +143,7 @@ class CHAOS_API TUniformGridBase
 			}
 		}
 		return Result;
-	}
-
+	} 
 #ifdef PLATFORM_COMPILER_CLANG
 	// Disable optimization (-ffast-math) since its currently causing regressions.
 	//		freciprocal-math:
@@ -304,6 +305,92 @@ class CHAOS_API TUniformGrid : public TUniformGridBase<T, d>
 	}
 };
 
+template<class T>
+class CHAOS_API TMPMGrid : public TUniformGridBase<T, 3>
+{
+	using TUniformGridBase<T, 3>::MCells;
+	using TUniformGridBase<T, 3>::MMinCorner;
+	using TUniformGridBase<T, 3>::MMaxCorner;
+	using TUniformGridBase<T, 3>::MDx;
+
+public:
+	using TUniformGridBase<T, 3>::GetNumCells;
+	using TUniformGridBase<T, 3>::Location;
+
+	TMPMGrid() {}
+	TMPMGrid(const TVector<T, 3>& MinCorner, const TVector<T, 3>& MaxCorner, const TVector<int32, 3>& Cells, const uint32 GhostCells = 0)
+		: TUniformGridBase<T, 3>(MinCorner, MaxCorner, Cells, GhostCells) {}
+	TMPMGrid(std::istream& Stream)
+		: TUniformGridBase<T, 3>(Stream) {}
+	~TMPMGrid() {}
+
+	void BaseNodeIndex(const TVector<T, 3>& X, TVector<int32, 3>& Index, TVector<T, 3>& weights) const;
+	
+	inline T Nijk(T w, int32 ii) const {
+		if (interp == linear)
+			return T(1) - w + T(ii) * (T(2) * w - T(1));
+		else
+			return ((T(1.5) * (w * w - w) - T(0.25)) * (T(ii) - 1) + (T(0.5) * w - T(0.25))) * (T(ii) - 1) - w * w + w + T(0.5);
+	}
+
+	inline TVector<int32, 3> Loc2GlobIndex(const TVector<int32, 3>& IndexIn, const TVector<int32, 3>& LocalIndexIn) const {
+		TVector<int32, 3> Result;
+		if (interp == linear) {
+			for (uint32 i = 0; i < 3; ++i)
+				Result[i] = IndexIn[i] + LocalIndexIn[i];
+		}
+		else {
+			for (uint32 i = 0; i < 3; ++i)
+				Result[i] = IndexIn[i] + LocalIndexIn[i] - 1;
+		}
+		return Result;
+	}
+
+	inline int32 Size() const {
+		if (interp == linear)
+			return (MCells[0] * MCells[1] * MCells[2]);
+		else
+			return ((MCells[0] + 1) * (MCells[1] + 1) * (MCells[2] + 1));
+		return -1;
+	}
+
+	int32 Loc2GlobIndex(const int32 IndexIn, const TVector<int32, 3>& LocalIndexIn) const;
+
+	inline TVector<T, 3> Node(const TVector<int32, 3>& IndexIn) const {
+		TVector<T, 3> Result((T)0.);
+
+		if (interp == linear) {
+			for (int32 i = 0; i < 3; ++i) {
+				Result[i] = T(IndexIn[i]) * MDx[i] + MMinCorner[i];
+			}
+		}
+		else {
+			for (size_t i = 0; i < 3; ++i) {
+				Result[i] = (T(IndexIn[i]) + T(0.5)) * MDx[i] + MMinCorner[i];
+			}
+		}
+
+		return Result;
+	}
+
+	TVector<T, 3> Node(int32 FlatIndexIn) const;
+
+	int32 FlatIndex(const TVector<int32, 3>& Index) const;
+	TVector<int32, 3> Lin2MultiIndex(const int32 IndexIn) const;
+	const TVector<int32, 3> GetCells() const { return MCells; }
+	const TVector<T, 3> GetMinCorner() const { return MMinCorner; } 
+	void SetDx(const TVector<T, 3> DxIn) { MDx = DxIn; }
+	void UpdateGridFromPositions(const Chaos::TDynamicParticles<T, 3>& InParticles);
+	
+	enum InterpType { linear = 1, quadratic = 2 };
+	void SetInterp(InterpType InterpIn);
+
+	InterpType interp = linear;
+	uint32 NPerDir = 2;
+
+};
+
+
 template<class T>
 class TUniformGrid<T, 3> : public TUniformGridBase<T, 3>
 {
@@ -338,6 +425,7 @@ class TUniformGrid<T, 3> : public TUniformGridBase<T, 3>
 	bool IsValid(const TVector<int32, 3>& X) const;
 };
 
+
 template <typename T, int d>
 FArchive& operator<<(FArchive& Ar, TUniformGridBase<T, d>& Value)
 {
@@ -348,6 +436,8 @@ FArchive& operator<<(FArchive& Ar, TUniformGridBase<T, d>& Value)
 #if PLATFORM_MAC || PLATFORM_LINUX
 extern template class CHAOS_API Chaos::TUniformGridBase<Chaos::FReal, 3>;
 extern template class CHAOS_API Chaos::TUniformGrid<Chaos::FReal, 3>;
+extern template class CHAOS_API Chaos::TMPMGrid<Chaos::FReal>;
+extern template class CHAOS_API Chaos::TMPMGrid<Chaos::FRealSingle>;
 extern template class CHAOS_API Chaos::TUniformGrid<Chaos::FReal, 2>;
 #endif
 
