ChaosFlesh : Moved solver to Chaos

diff --git a/Engine/Source/Runtime/Experimental/Chaos/Private/Chaos/Deformable/ChaosDeformableSolver.cpp b/Engine/Source/Runtime/Experimental/Chaos/Private/Chaos/Deformable/ChaosDeformableSolver.cpp
new file mode 100644
--- /dev/null
+++ b/Engine/Source/Runtime/Experimental/Chaos/Private/Chaos/Deformable/ChaosDeformableSolver.cpp
@@ -0,0 +1,489 @@
+// Copyright Epic Games, Inc. All Rights Reserved.
+
+#include "Chaos/Deformable/ChaosDeformableSolver.h"
+#include "Chaos/Deformable/ChaosDeformableSolverTypes.h"
+#include "Chaos/Deformable/ChaosDeformableSolverProxy.h"
+
+#include "Chaos/TriangleMesh.h"
+#include "Chaos/PBDAltitudeSpringConstraints.h"
+#include "Chaos/PBDBendingConstraints.h"
+#include "Chaos/PBDSpringConstraints.h"
+#include "Chaos/PBDVolumeConstraint.h"
+#include "Chaos/PBDSpringConstraints.h"
+#include "Chaos/PBDTriangleMeshCollisions.h"
+#include "Chaos/PBDCollisionSpringConstraints.h"
+#include "Chaos/PBDParticles.h"
+#include "Chaos/PBDSoftsSolverParticles.h"
+#include "Chaos/PBDTetConstraints.h"
+#include "Chaos/PerParticleGravity.h"
+#include "Chaos/XPBDCorotatedConstraints.h"
+#include "Chaos/XPBDVolumeConstraints.h"
+#include "Chaos/XPBDCorotatedFiberConstraints.h"
+#include "Chaos/Plane.h"
+#include "Chaos/Utilities.h"
+#include "Chaos/PBDEvolution.h"
+#include "Containers/StringConv.h"
+#include "CoreMinimal.h"
+#include "GeometryCollection/Facades/CollectionKinematicBindingFacade.h"
+#include "Misc/FileHelper.h"
+#include "HAL/PlatformFileManager.h"
+#include "HAL/IConsoleManager.h"
+#include "Misc/Paths.h"
+
+
+DEFINE_LOG_CATEGORY_STATIC(LogChaosDeformableSolver, Log, All);
+namespace Chaos::Softs
+{
+	FCriticalSection FDeformableSolver::PackageMutex;
+
+
+	FDeformableSolver::FDeformableSolver(FDeformableSolverProperties InProp)
+		: Property(InProp)
+	{
+		Reset(Property);
+	}
+
+	void FDeformableSolver::FPhysicsThreadAccess::Reset(const FDeformableSolverProperties& InProps)
+	{
+		Solver.Reset(InProps);
+	}
+	void FDeformableSolver::Reset(const FDeformableSolverProperties& InProps)
+	{
+		Property = InProps;
+		MObjects = TArrayCollectionArray<const UObject*>();
+		FSolverParticles LocalParticlesDummy;
+		FSolverRigidParticles RigidParticles;
+		Evolution.Reset(new FPBDEvolution(MoveTemp(LocalParticlesDummy), MoveTemp(RigidParticles), {}, Property.NumSolverIterations));
+		Evolution->Particles().AddArray(&MObjects);
+		if (Property.bDoSelfCollision)
+		{
+			SurfaceElements.Reset(new TArray<Chaos::TVec3<int32>>());
+			SurfaceTriangleMesh.Reset(new Chaos::FTriangleMesh());
+		}
+		Frame = 0;
+		Time = 0.f;
+	}
+
+	bool FDeformableSolver::FPhysicsThreadAccess::Advance(FSolverReal DeltaTime)
+	{
+		return Solver.Advance(DeltaTime);
+	}
+	bool FDeformableSolver::Advance(FSolverReal DeltaTime)
+	{
+		int32 NumIterations = FMath::Clamp<int32>(Property.NumSolverSubSteps, 0, INT_MAX);
+		if (NumIterations)
+		{
+			FSolverReal SubDeltaTime = DeltaTime / (FSolverReal)NumIterations;
+			if (!FMath::IsNearlyZero(SubDeltaTime))
+			{
+				for (int i = 0; i < NumIterations; ++i)
+				{
+					TickSimulation(SubDeltaTime);
+				}
+				Frame++;
+				return true;
+			}
+		}
+		return false;
+	}
+
+	void FDeformableSolver::FPhysicsThreadAccess::InitializeSimulationObjects()
+	{
+		Solver.InitializeSimulationObjects();
+	}
+	void FDeformableSolver::InitializeSimulationObjects()
+	{
+		for (TUniquePtr<FThreadingProxy>& Proxy : UninitializedProxys)
+		{
+			InitializeSimulationObject(*Proxy);
+			InitializeKinematicState(*Proxy);
+
+			FThreadingProxy::FKey Key = Proxy->GetOwner();
+			Proxies.Add(Key, TUniquePtr<FThreadingProxy>(Proxy.Release()));
+		}
+		UninitializedProxys.SetNum(0, true);
+
+		if (Property.bDoSelfCollision)
+		{
+			InitializeSelfCollisionVariables();
+		}
+		InitializeCollisionBodies();
+	}
+
+	void FDeformableSolver::FPhysicsThreadAccess::InitializeSimulationObject(FThreadingProxy& InProxy)
+	{
+		Solver.InitializeSimulationObject(InProxy);
+	}
+	void FDeformableSolver::InitializeSimulationObject(FThreadingProxy& InProxy)
+	{
+		if (FFleshThreadingProxy* Proxy = InProxy.As<FFleshThreadingProxy>())
+		{
+			if (const FManagedArrayCollection* Dynamic = &Proxy->GetDynamicCollection())
+			{
+				if (const FManagedArrayCollection* Rest = &Proxy->GetRestCollection())
+				{
+					const TManagedArray<FSolverReal>* MassArray = Rest->FindAttribute<FSolverReal>("Mass", FGeometryCollection::VerticesGroup);
+
+					// @todo(chaos) : make user attributes
+					FSolverReal Mass = 100.0;
+
+					const TManagedArray<FVector3f>& Vertex = Rest->GetAttribute<FVector3f>("Vertex", FGeometryCollection::VerticesGroup);
+					const TManagedArray<FIntVector>& Indices = Rest->GetAttribute<FIntVector>("Indices", FGeometryCollection::FacesGroup);
+					const TManagedArray<FVector3f>& DynamicVertex = Dynamic->GetAttribute<FVector3f>("Vertex", FGeometryCollection::VerticesGroup);
+
+					if (uint32 NumParticles = Vertex.Num())
+					{
+						if (uint32 NumSurfaceElements = Indices.Num())
+						{
+							// @todo(chaos) : reduce conversions
+							auto ChaosVert = [](FVector3d V) { return Chaos::FVec3(V.X, V.Y, V.Z); };
+							auto DoubleVert = [](FVector3f V) { return FVector3d(V.X, V.Y, V.Z); };
+							auto ChaosM = [](FSolverReal M, const TManagedArray<float>* AM, int32 Index, int32 Num) { return FSolverReal((AM != nullptr) ? (*AM)[Index] : M / FSolverReal(Num)); };
+							auto ChaosInvM = [](FSolverReal M) { return FSolverReal(FMath::IsNearlyZero(M) ? 0.0 : 1 / M); };
+							auto ChaosTet = [](FIntVector4 V, int32 dp) { return Chaos::TVec4<int32>(dp + V.X, dp + V.Y, dp + V.Z, dp + V.W); };
+							auto ChaosTri = [](FIntVector V, int32 dp) { return  Chaos::TVec3<int32>(dp + V.X, dp + V.Y, dp + V.Z); };
+							auto ChaosIntVec2 = [](int32 A, int32 B) { return  Chaos::TVec2<int32>(A, B); };
+
+							//  Add Simulation Particles Node
+							const FTransform& InitialTransform = Proxy->GetInitialTransform();
+							int32 ParticleStart = Evolution->AddParticleRange(NumParticles, 1, true);
+							for (uint32 vdx = 0; vdx < NumParticles; ++vdx)
+							{
+								int32 SolverParticleIndex = ParticleStart + vdx;
+								Evolution->Particles().X(SolverParticleIndex) = ChaosVert(InitialTransform.TransformPosition(DoubleVert(DynamicVertex[vdx])));
+								Evolution->Particles().V(SolverParticleIndex) = Chaos::FVec3(0.f, 0.f, 0.f);
+								Evolution->Particles().M(SolverParticleIndex) = ChaosM(Mass, MassArray, vdx, NumParticles);
+								Evolution->Particles().InvM(SolverParticleIndex) = ChaosInvM(Evolution->Particles().M(SolverParticleIndex));
+								Evolution->Particles().PAndInvM(SolverParticleIndex).InvM = Evolution->Particles().InvM(SolverParticleIndex);
+								MObjects[SolverParticleIndex] = Proxy->GetOwner();
+							}
+							Proxy->SetSolverParticleRange(ParticleStart, NumParticles);
+
+							if (Property.bEnableKinematics)
+							{
+								typedef Chaos::Facades::FKinematicBindingFacade Kinematics;
+
+								// Add Kinematics Node
+								for (int i = Kinematics::NumKinematicBindings(Rest) - 1; i >= 0; i--)
+								{
+									Kinematics::FBindingKey Key = Kinematics::GetKinematicBindingKey(Rest, i);
+
+									int32 BoneIndex = INDEX_NONE;
+									TArray<int32> BoundVerts;
+									TArray<float> BoundWeights;
+									Kinematics::GetBoneBindings(Rest, Key, BoneIndex, BoundVerts, BoundWeights);
+
+									for (int32 vdx : BoundVerts)
+									{
+										if (BoneIndex != INDEX_NONE)
+										{
+											int32 ParticleIndex = ParticleStart + vdx;
+											Evolution->Particles().InvM(ParticleIndex) = 0.f;
+											Evolution->Particles().PAndInvM(ParticleIndex).InvM = 0.f;
+										}
+									}
+								}
+							}
+
+							const TManagedArray<FIntVector4>& Tetrahedron = Rest->GetAttribute<FIntVector4>("Tetrahedron", "Tetrahedral");
+
+							if (uint32 NumElements = Tetrahedron.Num())
+							{
+								const FIntVector2& Range = Proxy->GetSolverParticleRange();
+
+								// Add Tetrahedral Elements Node
+								TArray<Chaos::TVec4<int32>> Elements;
+								Elements.SetNum(NumElements);
+								for (uint32 edx = 0; edx < NumElements; ++edx)
+								{
+									Elements[edx] = ChaosTet(Tetrahedron[edx], Range[0]);
+								}
+
+
+								FXPBDCorotatedConstraints<FSolverReal, FSolverParticles>* CorotatedConstraint =
+									new FXPBDCorotatedConstraints<FSolverReal, FSolverParticles>(
+										Evolution->Particles(), Elements, /*bRecordMetric = */false, (FSolverReal)100000.0);
+
+								int32 InitIndex = Evolution->AddConstraintInitRange(1, true);
+								Evolution->ConstraintInits()[InitIndex] =
+									[CorotatedConstraint](FSolverParticles& InParticles, const FSolverReal Dt)
+								{
+									CorotatedConstraint->Init();
+								};
+
+
+								int32 ConstraintIndex = Evolution->AddConstraintRuleRange(1, true);
+								Evolution->ConstraintRules()[ConstraintIndex] =
+									[CorotatedConstraint](FSolverParticles& InParticles, const FSolverReal Dt)
+								{
+									CorotatedConstraint->ApplyInParallel(InParticles, Dt);
+
+								};
+								CorotatedConstraints.Add(TUniquePtr<FXPBDCorotatedConstraints<FSolverReal, FSolverParticles>>(CorotatedConstraint));
+							}
+
+							if (Property.bDoSelfCollision)
+							{
+								// Add Surface Elements Node
+								int32 SurfaceElementsOffset = SurfaceElements->Num();
+								SurfaceElements->SetNum(NumSurfaceElements + SurfaceElementsOffset);
+								for (uint32 edx = 0; edx < NumSurfaceElements; ++edx)
+								{
+									(*SurfaceElements)[edx + SurfaceElementsOffset] = ChaosTri(Indices[edx], ParticleStart);
+								}
+
+							}
+
+
+							Time = 0.f;
+							bSimulationInitialized = true;
+						}
+					}
+				}
+			}
+		}
+	}
+
+	void FDeformableSolver::FPhysicsThreadAccess::InitializeCollisionBodies()
+	{
+		return Solver.InitializeCollisionBodies();
+	}
+	void FDeformableSolver::InitializeCollisionBodies()
+	{
+		if (Property.bUseFloor && Evolution->CollisionParticles().Size()==0)
+		{
+			Chaos::FVec3 Position(0.f);
+			Chaos::FVec3 EulerRot(0.f);
+			Evolution->AddCollisionParticleRange(1, 1, true);
+			Evolution->CollisionParticles().X(0) = Position;
+			Evolution->CollisionParticles().R(0) = Chaos::TRotation<Chaos::FReal, 3>::MakeFromEuler(EulerRot);
+			Evolution->CollisionParticles().SetDynamicGeometry(0, MakeUnique<Chaos::TPlane<Chaos::FReal, 3>>(Chaos::FVec3(0.f, 0.f, 0.f), Chaos::FVec3(0.f, 0.f, 1.f)));
+		}
+	}
+
+	void FDeformableSolver::FPhysicsThreadAccess::InitializeKinematicState(FThreadingProxy& InProxy)
+	{
+		Solver.InitializeKinematicState(InProxy);
+	}
+	void FDeformableSolver::InitializeKinematicState(FThreadingProxy& InProxy)
+	{
+		auto MKineticUpdate = [this](FSolverParticles& MParticles, const FSolverReal Dt, const FSolverReal MTime, const int32 Index)
+		{
+			if (0 <= Index && Index < this->MObjects.Num())
+			{
+				if (const UObject* Owner = this->MObjects[Index])
+				{
+					if (const FFleshThreadingProxy* Proxy = Proxies[Owner]->As<FFleshThreadingProxy>())
+					{
+						const FManagedArrayCollection& Rest = Proxy->GetRestCollection();
+						const FTransform& InitialTransform = Proxy->GetInitialTransform();
+						const TManagedArray<FVector3f>& Vertex = Rest.GetAttribute<FVector3f>("Vertex", FGeometryCollection::VerticesGroup);
+						const FIntVector2& Range = Proxy->GetSolverParticleRange();
+
+						// @todo(chaos) : reduce conversions
+						auto ChaosVert = [](FVector3f V) { return Chaos::FVec3(V.X, V.Y, V.Z); };
+
+						MParticles.X(Index) = InitialTransform.TransformPosition(ChaosVert(Vertex[Index - Range[0]]));
+						MParticles.PAndInvM(Index).P = MParticles.X(Index);
+					}
+				}
+			}
+		};
+		Evolution->SetKinematicUpdateFunction(MKineticUpdate);
+	}
+
+	void FDeformableSolver::FPhysicsThreadAccess::InitializeSelfCollisionVariables()
+	{
+		Solver.InitializeSelfCollisionVariables();
+	}
+	void FDeformableSolver::InitializeSelfCollisionVariables()
+	{
+		int32 NumParticles = Evolution->Particles().Size();
+		SurfaceTriangleMesh->Init(*SurfaceElements);
+		TriangleMeshCollisions.Reset(new FPBDTriangleMeshCollisions(
+			0, Evolution->Particles().Size(), *SurfaceTriangleMesh, false, false));
+		TSet<Chaos::TVec2<int32>>* InDisabledCollisionElements = new TSet<Chaos::TVec2<int32>>();
+		for (int32 i = 0; i < (int32)NumParticles; i++)
+		{
+			Chaos::TVec2<int32> LocalEdge = { i, i };
+			InDisabledCollisionElements->Add(LocalEdge);
+		}
+		CollisionSpringConstraint.Reset(new FPBDCollisionSpringConstraints(0, (int32)NumParticles, *SurfaceTriangleMesh, nullptr, MoveTemp(*InDisabledCollisionElements), 1.f, 1.f));
+		int32 InitIndex1 = Evolution->AddConstraintInitRange(1, true);
+		Evolution->ConstraintInits()[InitIndex1] =
+			[this](FSolverParticles& InParticles, const FSolverReal Dt)
+		{
+			this->TriangleMeshCollisions->Init(InParticles);
+			TArray<FPBDTriangleMeshCollisions::FGIAColor> EmptyGIAColors;
+			this->CollisionSpringConstraint->Init(InParticles, TriangleMeshCollisions->GetSpatialHash(), static_cast<TConstArrayView<FPBDTriangleMeshCollisions::FGIAColor>>(EmptyGIAColors), EmptyGIAColors);
+		};
+		int32 ConstraintIndex1 = Evolution->AddConstraintRuleRange(1, true);
+		Evolution->ConstraintRules()[ConstraintIndex1] =
+			[this](FSolverParticles& InParticles, const FSolverReal Dt)
+		{
+			this->CollisionSpringConstraint->Apply(InParticles, Dt);
+		};
+	}
+
+
+	void FDeformableSolver::FPhysicsThreadAccess::TickSimulation(FSolverReal DeltaTime)
+	{
+		Solver.TickSimulation(DeltaTime);
+	}
+	void FDeformableSolver::TickSimulation(FSolverReal DeltaTime)
+	{
+		if (!Property.FixTimeStep)
+		{
+			Evolution->AdvanceOneTimeStep(DeltaTime);
+			Time += DeltaTime;
+		}
+		else
+		{
+			Evolution->AdvanceOneTimeStep(Property.TimeStepSize);
+			Time += Property.TimeStepSize;
+		}
+
+		FOutputDataMap OutputBuffers;
+		for (TPair< FThreadingProxy::FKey, TUniquePtr<FThreadingProxy> > & BaseProxyPair : Proxies)
+		{
+			UpdateOutputState(*BaseProxyPair.Value);
+			if (FFleshThreadingProxy* Proxy = BaseProxyPair.Value->As<FFleshThreadingProxy>())
+			{
+				OutputBuffers.Add(Proxy->GetOwner(), TSharedPtr<FThreadingProxy::FOutputBuffer>(new FFleshThreadingProxy::FFleshOutputBuffer(*Proxy)));
+
+				if (Property.CacheToFile)
+				{
+					// @todo
+					//WriteFrame(Proxy, DeltaTime);
+				}
+			}
+
+		}
+
+		PushPackage(Frame, MoveTemp(OutputBuffers));
+	}
+
+
+	void FDeformableSolver::FPhysicsThreadAccess::PushPackage(int32 InFrame, FOutputDataMap&& InPackage)
+	{
+		Solver.PushPackage(InFrame, MoveTemp(InPackage));
+	}
+	void FDeformableSolver::PushPackage(int32 InFrame, FOutputDataMap&& InPackage)
+	{
+		FScopeLock Lock(&PackageMutex);
+		OutputPackages.Push(TUniquePtr< FOutputPackage >(new FOutputPackage(InFrame, MoveTemp(InPackage))));
+	}
+
+
+	TUniquePtr<FOutputPackage>  FDeformableSolver::FGameThreadAccess::PullPackage()
+	{
+		return Solver.PullPackage();
+	}
+	TUniquePtr<FOutputPackage> FDeformableSolver::PullPackage()
+	{
+		FScopeLock Lock(&PackageMutex);
+		if (OutputPackages.Num())
+			return OutputPackages.Pop();
+		return TUniquePtr<FOutputPackage>(nullptr);
+	}
+
+	void  FDeformableSolver::FGameThreadAccess::AddProxy(TUniquePtr<FThreadingProxy> InObject)
+	{
+		return Solver.AddProxy(TUniquePtr<FThreadingProxy>(InObject.Release()));
+	}
+	void FDeformableSolver::AddProxy(TUniquePtr<FThreadingProxy> InObject)
+	{
+		UninitializedProxys.Add(TUniquePtr< FThreadingProxy>(InObject.Release()));
+	}
+
+	void  FDeformableSolver::FPhysicsThreadAccess::UpdateOutputState(FThreadingProxy& InProxy)
+	{
+		Solver.UpdateOutputState(InProxy);
+	}
+	void FDeformableSolver::UpdateOutputState(FThreadingProxy& InProxy)
+	{
+		if (FFleshThreadingProxy* Proxy = InProxy.As<FFleshThreadingProxy>())
+		{
+
+			const FIntVector2& Range = Proxy->GetSolverParticleRange();
+			if (0 <= Range[0])
+			{
+				// @todo(chaos) : reduce conversions
+				auto UEVertd = [](Chaos::FVec3 V) { return FVector3d(V.X, V.Y, V.Z); };
+				auto UEVertf = [](FVector3d V) { return FVector3f((float)V.X, (float)V.Y, (float)V.Z); };
+
+				TManagedArray<FVector3f>& Position = Proxy->GetDynamicCollection().ModifyAttribute<FVector3f>("Vertex", FGeometryCollection::VerticesGroup);
+				for (int32 vdx = 0; vdx < Position.Num(); vdx++)
+				{
+					Position[vdx] = UEVertf(UEVertd(Evolution->Particles().X(vdx + Range[0])));
+				}
+			}
+		}
+	}
+
+	void  FDeformableSolver::FPhysicsThreadAccess::WriteFrame(FThreadingProxy& InProxy, const FSolverReal DeltaTime)
+	{
+		Solver.WriteFrame(InProxy, DeltaTime);
+	}
+	void FDeformableSolver::WriteFrame(FThreadingProxy& InProxy, const FSolverReal DeltaTime)
+	{
+		if (FFleshThreadingProxy* Proxy = InProxy.As<FFleshThreadingProxy>())
+		{
+			if (const FManagedArrayCollection* Rest = &Proxy->GetRestCollection())
+			{
+				const TManagedArray<FIntVector>& Indices = Rest->GetAttribute<FIntVector>("Indices", FGeometryCollection::FacesGroup);
+
+				WriteTrisGEO(Evolution->Particles(), Indices);
+				FString file = FPaths::ProjectDir();
+				file.Append(TEXT("/HoudiniOutput/DtLog.txt"));
+				if (Frame == 0)
+				{
+					FFileHelper::SaveStringToFile(FString(TEXT("DeltaTime\r\n")), *file);
+				}
+				FFileHelper::SaveStringToFile((FString::SanitizeFloat(DeltaTime) + FString(TEXT("\r\n"))), *file, FFileHelper::EEncodingOptions::AutoDetect, &IFileManager::Get(), EFileWrite::FILEWRITE_Append);
+			}
+		}
+	}
+
+	void  FDeformableSolver::FPhysicsThreadAccess::WriteTrisGEO(const FSolverParticles& Particles, const TManagedArray<FIntVector>& Mesh)
+	{
+		Solver.WriteTrisGEO(Particles, Mesh);
+	}
+	void FDeformableSolver::WriteTrisGEO(const FSolverParticles& Particles, const TManagedArray<FIntVector>& Mesh)
+	{
+		FString file = FPaths::ProjectDir();
+		//file.Append(TEXT("\HoudiniOuput\Test.geo"));
+		file.Append(TEXT("/HoudiniOutput/sim_frame_"));
+		file.Append(FString::FromInt(Frame));
+		file.Append(TEXT(".geo"));
+
+		int32 Np = Particles.Size();
+		int32 NPrims = Mesh.Num();
+
+		// We will use this FileManager to deal with the file.
+		IPlatformFile& FileManager = FPlatformFileManager::Get().GetPlatformFile();
+		FFileHelper::SaveStringToFile(FString(TEXT("PGEOMETRY V5\r\n")), *file);
+		FString HeaderInfo = FString(TEXT("NPoints ")) + FString::FromInt(Np) + FString(TEXT(" NPrims ")) + FString::FromInt(NPrims) + FString(TEXT("\r\n"));
+		FString MoreHeader = FString(TEXT("NPointGroups 0 NPrimGroups 0\r\nNPointAttrib 0 NVertexAttrib 0 NPrimAttrib 0 NAttrib 0\r\n"));
+
+		FFileHelper::SaveStringToFile(HeaderInfo, *file, FFileHelper::EEncodingOptions::AutoDetect, &IFileManager::Get(), EFileWrite::FILEWRITE_Append);
+		FFileHelper::SaveStringToFile(MoreHeader, *file, FFileHelper::EEncodingOptions::AutoDetect, &IFileManager::Get(), EFileWrite::FILEWRITE_Append);
+
+		for (int32 i = 0; i < Np; i++) {
+
+			FString ParticleInfo = FString::SanitizeFloat(Particles.X(i)[0]) + FString(TEXT(" ")) + FString::SanitizeFloat(Particles.X(i)[1]) + FString(TEXT(" ")) + FString::SanitizeFloat(Particles.X(i)[2]) + FString(TEXT(" ")) + FString::FromInt(1) + FString(TEXT("\r\n"));
+			FFileHelper::SaveStringToFile(ParticleInfo, *file, FFileHelper::EEncodingOptions::AutoDetect, &IFileManager::Get(), EFileWrite::FILEWRITE_Append);
+
+		}
+
+		for (int32 i = 0; i < Mesh.Num(); i++) {
+			//outstream << "Poly 3 < ";
+			FString ElementToWrite = FString(TEXT("Poly 3 < ")) + FString::FromInt(Mesh[i][0]) + FString(TEXT(" ")) + FString::FromInt(Mesh[i][1]) + FString(TEXT(" ")) + FString::FromInt(Mesh[i][2]) + FString(TEXT("\r\n"));
+			FFileHelper::SaveStringToFile(ElementToWrite, *file, FFileHelper::EEncodingOptions::AutoDetect, &IFileManager::Get(), EFileWrite::FILEWRITE_Append);
+		}
+
+		FFileHelper::SaveStringToFile(FString(TEXT("beginExtra\n")), *file, FFileHelper::EEncodingOptions::AutoDetect, &IFileManager::Get(), EFileWrite::FILEWRITE_Append);
+		FFileHelper::SaveStringToFile(FString(TEXT("endExtra\n")), *file, FFileHelper::EEncodingOptions::AutoDetect, &IFileManager::Get(), EFileWrite::FILEWRITE_Append);
+
+	}
+}; // Namespace Chaos::Softs
\ No newline at end of file
diff --git a/Engine/Source/Runtime/Experimental/Chaos/Private/GeometryCollection/Facades/CollectionKinematicBindingFacade.cpp b/Engine/Source/Runtime/Experimental/Chaos/Private/GeometryCollection/Facades/CollectionKinematicBindingFacade.cpp
new file mode 100644
--- /dev/null
+++ b/Engine/Source/Runtime/Experimental/Chaos/Private/GeometryCollection/Facades/CollectionKinematicBindingFacade.cpp
@@ -0,0 +1,74 @@
+// Copyright Epic Games, Inc. All Rights Reserved.
+
+/*=============================================================================
+	FleshCollection.cpp: FFleshCollection methods.
+=============================================================================*/
+
+#include "GeometryCollection/Facades/CollectionKinematicBindingFacade.h"
+#include "GeometryCollection/GeometryCollection.h"
+
+namespace Chaos::Facades
+{
+	// Attributes
+	const FName FKinematicBindingFacade::KinematicGroup("Kinematics");
+	const FName FKinematicBindingFacade::KinematicBoneBindingIndex("Binding");
+	const FName FKinematicBindingFacade::KinematicBoneBindingToGroup("BindingToGroup");
+
+
+	void FKinematicBindingFacade::Init()
+	{
+		FName SelectionGroup1 = FName(GeometryCollection::Facades::FSelectionFacade::BoundGroup.ToString() + "_" + FGeometryCollection::VerticesGroup.ToString());
+		FManagedArrayCollection::FConstructionParameters BindingGroupDependency(SelectionGroup1);
+
+		// Kinematic Bindings
+		Collection.AddAttribute<int32>(FKinematicBindingFacade::KinematicBoneBindingIndex, FKinematicBindingFacade::KinematicGroup, BindingGroupDependency);
+		Collection.AddAttribute<FString>(FKinematicBindingFacade::KinematicBoneBindingToGroup, FKinematicBindingFacade::KinematicGroup);
+	}
+
+	FKinematicBindingFacade::FBindingKey FKinematicBindingFacade::SetBoneBindings(FManagedArrayCollection* InCollection, const int32 InBoneIndex, const TArray<int32>& InBoneVerts, const TArray<float>& InBoneWeights)
+	{
+		return GeometryCollection::Facades::FSelectionFacade::AddSelection(InCollection, InBoneIndex, InBoneVerts, InBoneWeights, FGeometryCollection::VerticesGroup);
+	}
+
+	void FKinematicBindingFacade::GetBoneBindings(const FManagedArrayCollection* InCollection, const FKinematicBindingFacade::FBindingKey& Key, int32& OutBoneIndex, TArray<int32>& OutBoneVerts, TArray<float>& OutBoneWeights)
+	{
+		GeometryCollection::Facades::FSelectionFacade::GetSelection(InCollection, Key, OutBoneIndex, OutBoneVerts, OutBoneWeights);
+	}
+
+
+	int32 FKinematicBindingFacade::AddKinematicBinding(FManagedArrayCollection* InCollection, const FKinematicBindingFacade::FBindingKey& Key)
+	{
+		if (InCollection && InCollection->HasGroup(KinematicGroup))
+		{
+			int32 NewIndex = InCollection->AddElements(1, KinematicGroup);
+			TManagedArray<int32>* KinemaitcBoneBinding = InCollection->FindAttribute<int32>(KinematicBoneBindingIndex, KinematicGroup);
+			TManagedArray<FString>* KinemaitcBoneBindingToGroup = InCollection->FindAttribute<FString>(KinematicBoneBindingToGroup, KinematicGroup);
+
+			(*KinemaitcBoneBinding)[NewIndex] = Key.Index;
+			(*KinemaitcBoneBindingToGroup)[NewIndex] = Key.GroupName.ToString();
+
+			return NewIndex;
+		}
+		return INDEX_NONE;
+	}
+
+	FKinematicBindingFacade::FBindingKey FKinematicBindingFacade::GetKinematicBindingKey(const FManagedArrayCollection* InCollection, int Index)
+	{
+		FKinematicBindingFacade::FBindingKey Key;
+		if (InCollection && InCollection->HasGroup(KinematicGroup))
+		{
+			const TManagedArray<int32>* KinemaitcBoneBinding = InCollection->FindAttribute<int32>(KinematicBoneBindingIndex, KinematicGroup);
+			const TManagedArray<FString>* KinemaitcBoneBindingToGroup = InCollection->FindAttribute<FString>(KinematicBoneBindingToGroup, KinematicGroup);
+			if (KinemaitcBoneBinding && KinemaitcBoneBindingToGroup)
+			{
+				if (0 <= Index && Index < InCollection->NumElements(KinematicGroup))
+				{
+					Key.GroupName = FName((*KinemaitcBoneBindingToGroup)[Index]);
+					Key.Index = (*KinemaitcBoneBinding)[Index];
+				}
+			}
+		}
+		return Key;
+	}
+
+}
diff --git a/Engine/Source/Runtime/Experimental/Chaos/Public/Chaos/Deformable/ChaosDeformableSolver.h b/Engine/Source/Runtime/Experimental/Chaos/Public/Chaos/Deformable/ChaosDeformableSolver.h
new file mode 100644
--- /dev/null
+++ b/Engine/Source/Runtime/Experimental/Chaos/Public/Chaos/Deformable/ChaosDeformableSolver.h
@@ -0,0 +1,140 @@
+// Copyright Epic Games, Inc. All Rights Reserved.
+
+#pragma once
+
+#include "Chaos/Core.h"
+#include "Chaos/PBDCollisionSpringConstraints.h"
+#include "Chaos/PBDEvolution.h"
+#include "Chaos/PBDParticles.h"
+#include "Chaos/PBDSoftsSolverParticles.h"
+#include "Chaos/PBDSoftsEvolutionFwd.h"
+#include "Chaos/PBDTriangleMeshCollisions.h"
+#include "Chaos/TriangleMesh.h"
+#include "Chaos/XPBDCorotatedConstraints.h"
+#include "Chaos/Deformable/ChaosDeformableSolverTypes.h"
+#include "Chaos/Deformable/ChaosDeformableSolverProxy.h"
+#include "CoreMinimal.h"
+#include "GeometryCollection/ManagedArrayCollection.h"
+#include "Templates/UniquePtr.h"
+
+namespace Chaos::Softs
+{
+	// @todo(Deformable) : Explore using the chaos ISolverBase, IProxyBase instead. 
+	class CHAOS_API FDeformableSolver
+	{
+		friend class FGameThreadAccess;
+		friend class FPhysicsThreadAccess;
+		friend class FDeformableFun;
+
+	public:
+
+		FDeformableSolver(FDeformableSolverProperties InProp = FDeformableSolverProperties());
+
+		/* Physics Thread Access API */
+		class CHAOS_API FPhysicsThreadAccess
+		{
+		public:
+			FPhysicsThreadAccess(FDeformableSolver& InSolver, const FPhysicsThreadAccessor&) : Solver(InSolver) {}
+
+			/* Simulation Advance */
+			bool Advance(FSolverReal DeltaTime);
+			bool IsInitialized() { return Solver.IsInitialized(); }
+			void Reset(const FDeformableSolverProperties&);
+			void TickSimulation(FSolverReal DeltaTime);
+			void UpdateOutputState(FThreadingProxy&);
+			void PushPackage(int32 Frame, FOutputDataMap&& Package);
+
+			/* Iteration Advance */
+			void InitializeSimulationObjects();
+			void InitializeSimulationObject(FThreadingProxy&);
+			void InitializeCollisionBodies();
+			void InitializeKinematicState(FThreadingProxy&);
+			void InitializeSelfCollisionVariables();
+
+			/*IO Utility*/
+			void WriteTrisGEO(const Softs::FSolverParticles& Particles, const TManagedArray<FIntVector>& Mesh);
+			void WriteFrame(FThreadingProxy& , const FSolverReal DeltaTime);
+
+			const FDeformableSolverProperties& GetProperties() const { return Solver.GetProperties(); }
+
+			FPBDEvolution* GetEvolution() { return Solver.Evolution.Get(); }
+			const FPBDEvolution* GetEvolution() const { return Solver.Evolution.Get(); }
+
+			TArrayCollectionArray<const UObject*>& GetObjectsMap() { return Solver.MObjects; }
+			const TArrayCollectionArray<const UObject*>& GetObjectsMap() const { return Solver.MObjects; }
+
+
+		private:
+			FDeformableSolver& Solver;
+		};
+
+
+		/* Game Thread Access API */
+		class CHAOS_API FGameThreadAccess
+		{
+		public:
+			FGameThreadAccess(FDeformableSolver& InSolver, const FGameThreadAccessor&) : Solver(InSolver) {}
+
+			void AddProxy(TUniquePtr<FThreadingProxy> InObject);
+
+			TUniquePtr<FOutputPackage> PullPackage();
+
+		private:
+			FDeformableSolver& Solver;
+		};
+
+	protected:
+
+		/* Simulation Advance */
+		bool Advance(FSolverReal DeltaTime);
+		bool IsInitialized() { return bSimulationInitialized; }
+		void Reset(const FDeformableSolverProperties&);
+		void TickSimulation(FSolverReal DeltaTime);
+		void UpdateOutputState(FThreadingProxy&);
+		void PushPackage(int32 Frame, FOutputDataMap&& Package);
+
+		/* Iteration Advance */
+		void InitializeSimulationObjects();
+		void InitializeSimulationObject(FThreadingProxy&);
+		void InitializeCollisionBodies();
+		void InitializeKinematicState(FThreadingProxy&);
+		void InitializeSelfCollisionVariables();
+
+		/*IO Utility*/
+		void WriteTrisGEO(const FSolverParticles& Particles, const TManagedArray<FIntVector>& Mesh);
+		void WriteFrame(FThreadingProxy&, const FSolverReal DeltaTime);
+
+		/*Game Thread API*/
+		void AddProxy(TUniquePtr<FThreadingProxy> InObject);
+		TUniquePtr<FOutputPackage> PullPackage();
+
+		const FDeformableSolverProperties& GetProperties() const { return Property; }
+
+	private:
+
+		// connections outside the solver.
+		static FCriticalSection	PackageMutex;
+		TArray< TUniquePtr<FThreadingProxy> > UninitializedProxys;
+		TMap< FThreadingProxy::FKey, TUniquePtr<FThreadingProxy> > Proxies;
+		TArray< TUniquePtr<FOutputPackage>  > OutputPackages;
+
+		// User Configuration
+		FDeformableSolverProperties Property;
+
+
+		// Simulation Variables
+		TUniquePtr<Softs::FPBDEvolution> Evolution;
+		TArray<TUniquePtr<Softs::FXPBDCorotatedConstraints<Softs::FSolverReal, Softs::FSolverParticles>>> CorotatedConstraints;
+		TUniquePtr<Softs::FPBDCollisionSpringConstraints> CollisionSpringConstraint;
+		TUniquePtr<Softs::FPBDTriangleMeshCollisions> TriangleMeshCollisions;
+		TArrayCollectionArray<const UObject*> MObjects;
+		TUniquePtr <TArray<TVec3<int32>>> SurfaceElements;
+		TUniquePtr <FTriangleMesh> SurfaceTriangleMesh;
+
+		FSolverReal Time = 0.f;
+		int32 Frame = 0;
+		bool bSimulationInitialized = false;
+	};
+
+
+}; // namesapce Chaos::Softs
diff --git a/Engine/Source/Runtime/Experimental/Chaos/Public/Chaos/Deformable/ChaosDeformableSolverProxy.h b/Engine/Source/Runtime/Experimental/Chaos/Public/Chaos/Deformable/ChaosDeformableSolverProxy.h
new file mode 100644
--- /dev/null
+++ b/Engine/Source/Runtime/Experimental/Chaos/Public/Chaos/Deformable/ChaosDeformableSolverProxy.h
@@ -0,0 +1,120 @@
+// Copyright Epic Games, Inc. All Rights Reserved.
+
+#pragma once
+#include "CoreMinimal.h"
+#include "UObject/ObjectMacros.h"
+#include "GeometryCollection/GeometryCollection.h"
+#include "GeometryCollection/ManagedArrayCollection.h"
+
+namespace Chaos::Softs
+{
+	class CHAOS_API FThreadingProxy {
+		const UObject* Owner = nullptr;
+		FName TypeName = FName("");
+
+	public:
+		typedef const UObject* FKey;
+
+		FThreadingProxy(const UObject* InOwner = nullptr, FName InTypeName = FName(""))
+			: Owner(InOwner)
+			, TypeName(InTypeName)
+		{}
+		virtual ~FThreadingProxy() {}
+
+		FName BaseTypeName() const { return TypeName; }
+
+		template<class T>
+		T* As() { return T::TypeName().IsEqual(BaseTypeName()) ? (T*)this : nullptr; }
+
+		template<class T>
+		T* As() const { return T::TypeName().IsEqual(BaseTypeName()) ? (T*)this : nullptr; }
+
+		const UObject* GetOwner() const { return Owner; }
+
+
+		class FOutputBuffer
+		{
+			const UObject* Owner = nullptr;
+			FName TypeName = FName("");
+
+		public:
+			FOutputBuffer(const FThreadingProxy& Ref)
+				: Owner(Ref.GetOwner())
+				, TypeName(Ref.BaseTypeName()) {}
+
+			FName BaseTypeName() { return TypeName; }
+
+			template<class T>
+			T* As() { return T::Source::TypeName().IsEqual(BaseTypeName()) ? (T*)this : nullptr; }
+		};
+	};
+
+	class CHAOS_API FFleshThreadingProxy : public FThreadingProxy
+	{
+		typedef FThreadingProxy Super;
+		FTransform InitialTransform = FTransform::Identity;
+		const FManagedArrayCollection& Rest;
+		FManagedArrayCollection Dynamic;
+		FIntVector2 SolverParticleRange = FIntVector2(0, 0);
+
+	public:
+
+		FFleshThreadingProxy(UObject* InOwner,
+			const FTransform& InInitialTransform,
+			const FManagedArrayCollection& InRest,
+			const FManagedArrayCollection& InDynamic)
+			: Super(InOwner, TypeName())
+			, InitialTransform(InInitialTransform)
+			, Rest(InRest)
+			, Dynamic(InDynamic)
+		{}
+		virtual ~FFleshThreadingProxy() {}
+
+		static FName TypeName() { return FName("Flesh"); }
+
+		const FTransform& GetInitialTransform() const
+		{
+			return InitialTransform;
+		}
+
+		void SetSolverParticleRange(int32 InStart, int32 InRange)
+		{
+			SolverParticleRange[0] = InStart;
+			SolverParticleRange[1] = InRange;
+		}
+
+		const FIntVector2& GetSolverParticleRange() const
+		{
+			return SolverParticleRange;
+		}
+
+		FManagedArrayCollection& GetDynamicCollection() {
+			return Dynamic;
+		}
+		const FManagedArrayCollection& GetDynamicCollection() const {
+			return Dynamic;
+		}
+
+		const FManagedArrayCollection& GetRestCollection() const {
+			return Rest;
+		}
+
+		class FFleshOutputBuffer : public FThreadingProxy::FOutputBuffer
+		{
+			typedef FThreadingProxy::FOutputBuffer Super;
+
+		public:
+			typedef FFleshThreadingProxy Source;
+			FFleshOutputBuffer(const Source& Ref)
+				: Super(Ref)
+			{
+				// Ref.Dynamic will have updated solver data 
+				Dynamic.AddAttribute<FVector3f>("Vertex", FGeometryCollection::VerticesGroup);
+				Dynamic.CopyMatchingAttributesFrom(Ref.GetDynamicCollection());
+			}
+
+			FManagedArrayCollection Dynamic;
+		};
+
+	};
+}// namespace Chaos::Softs
\ No newline at end of file
diff --git a/Engine/Source/Runtime/Experimental/Chaos/Public/Chaos/Deformable/ChaosDeformableSolverTypes.h b/Engine/Source/Runtime/Experimental/Chaos/Public/Chaos/Deformable/ChaosDeformableSolverTypes.h
new file mode 100644
--- /dev/null
+++ b/Engine/Source/Runtime/Experimental/Chaos/Public/Chaos/Deformable/ChaosDeformableSolverTypes.h
@@ -0,0 +1,97 @@
+// Copyright Epic Games, Inc. All Rights Reserved.
+
+#pragma once
+
+#include "Chaos/Deformable/ChaosDeformableSolverProxy.h"
+#include "Chaos/PBDSoftsEvolutionFwd.h"
+#include "CoreMinimal.h"
+
+class UDeformableSolverComponent;
+
+namespace Chaos::Softs
+{
+	struct CHAOS_API FDeformableSolverProperties
+	{
+		FDeformableSolverProperties(
+			int32 InNumSolverSubSteps = 2,
+			int32 InNumSolverIterations = 5,
+			bool InFixTimeStep = false,
+			FSolverReal InTimeStepSize = (FSolverReal)0.05,
+			bool InCacheToFile = false,
+			bool InbEnableKinematics = true,
+			bool InbUseFloor = true,
+			bool InbDoSelfCollision = false)
+			: NumSolverSubSteps(InNumSolverSubSteps)
+			, NumSolverIterations(InNumSolverIterations)
+			, FixTimeStep(InFixTimeStep)
+			, TimeStepSize(InTimeStepSize)
+			, CacheToFile(InCacheToFile)
+			, bEnableKinematics(InbEnableKinematics)
+			, bUseFloor(InbUseFloor)
+			, bDoSelfCollision(InbDoSelfCollision)
+		{}
+
+		FDeformableSolverProperties(const FDeformableSolverProperties& InProp)
+			: NumSolverSubSteps(InProp.NumSolverSubSteps)
+			, NumSolverIterations(InProp.NumSolverIterations)
+			, FixTimeStep(InProp.FixTimeStep)
+			, TimeStepSize(InProp.TimeStepSize)
+			, CacheToFile(InProp.CacheToFile)
+			, bEnableKinematics(InProp.bEnableKinematics)
+			, bUseFloor(InProp.bUseFloor)
+			, bDoSelfCollision(InProp.bDoSelfCollision)
+		{}
+
+		int32 NumSolverSubSteps = 5;
+		int32 NumSolverIterations = 5;
+		bool FixTimeStep = false;
+		FSolverReal TimeStepSize = (FSolverReal)0.05;
+		bool CacheToFile = false;
+		bool bEnableKinematics = true;
+		bool bUseFloor = true;
+		bool bDoSelfCollision = false;
+	};
+
+
+	/*Data Transfer*/
+	typedef TSharedPtr<FThreadingProxy::FOutputBuffer> FOutputDataMapValue; // OutputBuffer
+	typedef TMap<FThreadingProxy::FKey, FOutputDataMapValue > FOutputDataMap; // <const UObject*,OutputBuffer>
+
+	struct CHAOS_API FOutputPackage {
+		FOutputPackage()
+		{}
+
+		FOutputPackage(int32 InFrame, FOutputDataMap&& InMap)
+			: Frame(InFrame)
+			, ObjectMap(InMap)
+		{}
+
+		int32 Frame = INDEX_NONE;
+		FOutputDataMap ObjectMap;
+	};
+
+	/* Accessor for the Game Thread*/
+	class CHAOS_API FGameThreadAccessor
+	{
+	public:
+		friend class UDeformableSolverComponent;
+#if PLATFORM_WINDOWS
+	protected:
+#endif
+		FGameThreadAccessor() {}
+	};
+
+
+	/* Accessor for the Physics Thread*/
+	class CHAOS_API FPhysicsThreadAccessor
+	{
+	public:
+		friend class UDeformableSolverComponent;
+#if PLATFORM_WINDOWS
+	protected:
+#endif
+		FPhysicsThreadAccessor() {}
+	};
+
+
+}; // namesapce Chaos::Softs
\ No newline at end of file
diff --git a/Engine/Source/Runtime/Experimental/Chaos/Public/Chaos/PBDTetConstraints.h b/Engine/Source/Runtime/Experimental/Chaos/Public/Chaos/PBDTetConstraints.h
--- a/Engine/Source/Runtime/Experimental/Chaos/Public/Chaos/PBDTetConstraints.h
+++ b/Engine/Source/Runtime/Experimental/Chaos/Public/Chaos/PBDTetConstraints.h
@@ -12,7 +12,7 @@ class FPBDTetConstraints : public FPBDTetConstraintsBase
 	using Base::Constraints;
 
   public:
-	  FPBDTetConstraints(const FSolverParticles& InParticles, TArray<TVec4<int32>>&& InConstraints, const FReal InStiffness = (FReal)1)
+	  FPBDTetConstraints(const FSolverParticles& InParticles, TArray<TVec4<int32>>&& InConstraints, const FSolverReal InStiffness = (FSolverReal)1)
 	    : Base(InParticles, MoveTemp(InConstraints), InStiffness) {}
 	virtual ~FPBDTetConstraints() override {}
 
diff --git a/Engine/Source/Runtime/Experimental/Chaos/Public/Chaos/PBDVolumeConstraintBase.h b/Engine/Source/Runtime/Experimental/Chaos/Public/Chaos/PBDVolumeConstraintBase.h
--- a/Engine/Source/Runtime/Experimental/Chaos/Public/Chaos/PBDVolumeConstraintBase.h
+++ b/Engine/Source/Runtime/Experimental/Chaos/Public/Chaos/PBDVolumeConstraintBase.h
@@ -20,7 +20,7 @@ class FPBDVolumeConstraintBase
 		{
 			Com += InParticles.X(i);
 		}
-		Com /= InParticles.Size();
+		Com /= (FSolverReal)InParticles.Size();
 		RefVolume = (FSolverReal)0.;
 		for (const TVec3<int32>& Constraint : Constraints)
 		{
@@ -59,7 +59,7 @@ class FPBDVolumeConstraintBase
 		{
 			Com += InParticles.P(i);
 		}
-		Com /= InParticles.Size();
+		Com /= (FSolverReal)InParticles.Size();
 		TArray<FSolverVec3> Grads;
 		Grads.SetNum(InParticles.Size());
 		for (FSolverVec3& Elem : Grads)
@@ -95,7 +95,7 @@ class FPBDVolumeConstraintBase
 		{
 			Com += InParticles.P(i);
 		}
-		Com /= InParticles.Size();
+		Com /= (FSolverReal)InParticles.Size();
 		FSolverReal Volume = (FSolverReal)0.;
 		for (const TVec3<int32>& Constraint : Constraints)
 		{
diff --git a/Engine/Source/Runtime/Experimental/Chaos/Public/Chaos/XPBDCorotatedConstraints.h b/Engine/Source/Runtime/Experimental/Chaos/Public/Chaos/XPBDCorotatedConstraints.h
--- a/Engine/Source/Runtime/Experimental/Chaos/Public/Chaos/XPBDCorotatedConstraints.h
+++ b/Engine/Source/Runtime/Experimental/Chaos/Public/Chaos/XPBDCorotatedConstraints.h
@@ -312,7 +312,7 @@ namespace Chaos::Softs
 		
 
 
-		TVec4<TVector<T, 3>> GetDeterminantDelta(const ParticleType& Particles, const T Dt, const int32 ElementIndex, const T Tol = 1e-3) const
+		TVec4<TVector<T, 3>> GetDeterminantDelta(const ParticleType& Particles, const T Dt, const int32 ElementIndex, const T Tol = (T)1e-3) const
 		{
 			TRACE_CPUPROFILER_EVENT_SCOPE(TEXT("STAT_ChaosXPBDCorotatedApplyDet"));
 			//SCOPE_CYCLE_COUNTER(STAT_ChaosXPBDCorotatedDet);
@@ -370,7 +370,7 @@ namespace Chaos::Softs
 
 
 
-		TVec4<TVector<T, 3>> GetPolarDelta(const ParticleType& Particles, const T Dt, const int32 ElementIndex, const T Tol = 1e-3) const
+		TVec4<TVector<T, 3>> GetPolarDelta(const ParticleType& Particles, const T Dt, const int32 ElementIndex, const T Tol = (T)1e-3) const
 		{	
 			TRACE_CPUPROFILER_EVENT_SCOPE(TEXT("STAT_ChaosXPBDCorotatedApplyPolar"));
 			SCOPE_CYCLE_COUNTER(STAT_ChaosXPBDCorotatedPolar);
diff --git a/Engine/Source/Runtime/Experimental/Chaos/Public/GeometryCollection/Facades/CollectionKinematicBindingFacade.h b/Engine/Source/Runtime/Experimental/Chaos/Public/GeometryCollection/Facades/CollectionKinematicBindingFacade.h
new file mode 100644
--- /dev/null
+++ b/Engine/Source/Runtime/Experimental/Chaos/Public/GeometryCollection/Facades/CollectionKinematicBindingFacade.h
@@ -0,0 +1,41 @@
+// Copyright Epic Games, Inc. All Rights Reserved.
+
+#pragma once
+
+#include "CoreMinimal.h"
+#include "GeometryCollection/ManagedArrayCollection.h"
+#include "GeometryCollection/Facades/CollectionSelectionFacade.h"
+
+namespace Chaos::Facades
+{
+	/** Kinematic Facade */
+	class CHAOS_API FKinematicBindingFacade
+	{
+		FManagedArrayCollection& Collection;
+
+	public:
+		typedef GeometryCollection::Facades::FSelectionFacade::FSelectionKey FBindingKey;
+
+		FKinematicBindingFacade(FManagedArrayCollection& InCollection) : Collection(InCollection) {}
+
+		void Init();
+
+		//
+		//  Skeletal Mesh Bone Bindings
+		//
+		static FBindingKey SetBoneBindings(FManagedArrayCollection*, const int32 BoneIndex, const TArray<int32>& Vertices, const TArray<float>& Weights);
+		static void GetBoneBindings(const FManagedArrayCollection*, const FBindingKey& Key, int32& OutBoneIndex, TArray<int32>& OutBoneVerts, TArray<float>& OutBoneWeights);
+
+		//
+		// Kinematics
+		//
+		static const FName KinematicGroup;
+		static const FName KinematicBoneBindingIndex;
+		static const FName KinematicBoneBindingToGroup;
+
+		static int32 AddKinematicBinding(FManagedArrayCollection*, const FBindingKey& Key);
+		static int32 NumKinematicBindings(const FManagedArrayCollection* Collection) { return Collection->NumElements(KinematicGroup); }
+		static FBindingKey GetKinematicBindingKey(const FManagedArrayCollection* InCollection, int Index);
+
+	};
+}
