[Interchange] Fixed issues with packaging Interchange plugins

diff --git a/Engine/Plugins/Experimental/Interchange/Runtime/Config/Game.ini b/Engine/Plugins/Experimental/Interchange/Runtime/Config/Game.ini
new file mode 100644
--- /dev/null
+++ b/Engine/Plugins/Experimental/Interchange/Runtime/Config/Game.ini
@@ -0,0 +1,5 @@
+[/Script/UnrealEd.ProjectPackagingSettings]
++DirectoriesToAlwaysCook=(Path="/Interchange/Functions")
++DirectoriesToAlwaysCook=(Path="/Interchange/gltf")
++DirectoriesToAlwaysCook=(Path="/Interchange/Materials")
++DirectoriesToAlwaysCook=(Path="/Interchange/Utilities")
diff --git a/Engine/Source/Runtime/Experimental/Interchange/Engine/Private/InterchangeManager.cpp b/Engine/Source/Runtime/Experimental/Interchange/Engine/Private/InterchangeManager.cpp
--- a/Engine/Source/Runtime/Experimental/Interchange/Engine/Private/InterchangeManager.cpp
+++ b/Engine/Source/Runtime/Experimental/Interchange/Engine/Private/InterchangeManager.cpp
@@ -507,6 +507,7 @@ bool UInterchangeManager::RegisterFactory(const UClass* FactoryClass)
 
 bool UInterchangeManager::RegisterWriter(const UClass* WriterClass)
 {
+#if WITH_EDITOR
 	if (!WriterClass)
 	{
 		return false;
@@ -522,6 +523,7 @@ bool UInterchangeManager::RegisterWriter(const UClass* WriterClass)
 		return false;
 	}
 	RegisteredWriters.Add(WriterClass, WriterToRegister);
+#endif
 	return true;
 }
 
