Chaos : Vehicle : Fix static analysis error

diff --git a/Engine/Source/Runtime/Experimental/ChaosVehicles/ChaosVehiclesCore/Private/SimModule/SimulationModuleBase.cpp b/Engine/Source/Runtime/Experimental/ChaosVehicles/ChaosVehiclesCore/Private/SimModule/SimulationModuleBase.cpp
--- a/Engine/Source/Runtime/Experimental/ChaosVehicles/ChaosVehiclesCore/Private/SimModule/SimulationModuleBase.cpp
+++ b/Engine/Source/Runtime/Experimental/ChaosVehicles/ChaosVehiclesCore/Private/SimModule/SimulationModuleBase.cpp
@@ -36,10 +36,13 @@ ISimulationModuleBase* ISimulationModuleBase::GetParent()
 
 ISimulationModuleBase* ISimulationModuleBase::GetFirstChild()
 {
-	const TSet<int32>& Children = SimModuleTree->GetChildren(SimTreeIndex);
-	for (int ChildIndex : Children)
-	{ 
-		return (SimModuleTree != nullptr) ? SimModuleTree->AccessSimModule(ChildIndex) : nullptr;
+	if (SimModuleTree)
+	{
+		const TSet<int32>& Children = SimModuleTree->GetChildren(SimTreeIndex);
+		for (const int ChildIndex : Children)
+		{ 
+			return SimModuleTree->AccessSimModule(ChildIndex);
+		}
 	}
 	return nullptr;
 }
