StringNetSerializerUtils robustness improvements.

#rnx

diff --git a/Engine/Source/Runtime/Experimental/Iris/Core/Private/Iris/Serialization/StringNetSerializerUtils.cpp b/Engine/Source/Runtime/Experimental/Iris/Core/Private/Iris/Serialization/StringNetSerializerUtils.cpp
--- a/Engine/Source/Runtime/Experimental/Iris/Core/Private/Iris/Serialization/StringNetSerializerUtils.cpp
+++ b/Engine/Source/Runtime/Experimental/Iris/Core/Private/Iris/Serialization/StringNetSerializerUtils.cpp
@@ -123,7 +123,6 @@ void FStringNetSerializerBase::Dequantize(FNetSerializationContext& Context, con
 
 		if (OutLength == 0 || TempString.GetData()[OutLength - 1] != 0)
 		{
-			// $IRIS TODO: Bad decoding? Report an error?
 			OutLength += 1U;
 		}
 
diff --git a/Engine/Source/Runtime/Experimental/Iris/Core/Private/Iris/Serialization/StringNetSerializerUtils.h b/Engine/Source/Runtime/Experimental/Iris/Core/Private/Iris/Serialization/StringNetSerializerUtils.h
--- a/Engine/Source/Runtime/Experimental/Iris/Core/Private/Iris/Serialization/StringNetSerializerUtils.h
+++ b/Engine/Source/Runtime/Experimental/Iris/Core/Private/Iris/Serialization/StringNetSerializerUtils.h
@@ -136,11 +136,15 @@ public:
 			constexpr WideCharType LastCodePoint = 0x10FFFF;
 			constexpr WideCharType ReplacementCharacter = 0xFFFD;
 
-			check(DestLen >= 3U*SourceLen);
+			// Unless the destination length is safe let's not encode. Encoding would require OOB checks.
+			if (DestLen < 3U*SourceLen)
+			{
+				return false;
+			}
 
 			uint32 DestIt = 0;
 			uint32 SourceIt = 0;
-			for (const uint32 SourceEndIt = SourceLen; SourceIt != SourceEndIt; ++SourceIt)
+			for (const uint32 SourceEndIt = SourceLen; SourceIt < SourceEndIt; ++SourceIt)
 			{
 				WideCharType Char = Source[SourceIt];
 				Char = (Char > LastCodePoint ? ReplacementCharacter : Char);
@@ -175,11 +179,15 @@ public:
 		template<typename T = WideCharType, typename TEnableIf<sizeof(T) == 2, char>::Type CharSize = 2>
 		static bool EncodeImpl(EncodeType* Dest, uint32 DestLen, const WideCharType* Source, uint32 SourceLen, uint32& OutDestLen)
 		{
-			check(DestLen >= 3U*SourceLen);
+			// Unless the destination length is safe let's not encode. Encoding would require OOB checks.
+			if (DestLen < 3U*SourceLen)
+			{
+				return false;
+			}
 
 			uint32 DestIt = 0;
 			uint32 SourceIt = 0;
-			for (const uint32 SourceEndIt = SourceLen; SourceIt != SourceEndIt; ++SourceIt)
+			for (const uint32 SourceEndIt = SourceLen; SourceIt < SourceEndIt; ++SourceIt)
 			{
 				const WideCharType Char = Source[SourceIt];
 				const WideCharType NextChar = Source[FMath::Min(SourceIt + 1U, SourceEndIt - 1U)];
@@ -229,14 +237,17 @@ public:
 				return true;
 			}
 
-			check(DestLen >= GetSafeDecodedBufferLength(SourceLen));
+			if (DestLen < GetSafeDecodedBufferLength(SourceLen))
+			{
+				return false;
+			}
 
 			constexpr WideCharType ReplacementCharacter = 0xFFFD;
 
 			uint32 SurrogateCount = 0;
 			uint32 DestIt = 0;
 			bool bIsErrorDetected = false;
-			for (uint32 SourceIt = 0, SourceEndIt = SourceLen; SourceIt != SourceEndIt; ++SourceIt)
+			for (uint32 SourceIt = 0, SourceEndIt = SourceLen, DestEndIt = DestLen; (SourceIt < SourceEndIt) & (DestIt < DestEndIt) & (!bIsErrorDetected); ++SourceIt)
 			{
 				EncodeType Byte;
 				uint32 Codepoint;
@@ -249,12 +260,23 @@ public:
 				// Optionally decode second byte. Current state will remain unchanged unless the most significant bit in Byte is set.
 				Mask = int8(Byte) >> 7U;
 				SourceIt += 1U & Mask;
+				if (SourceIt >= SourceEndIt)
+				{
+					bIsErrorDetected = true;
+					break;
+				}
+
 				Byte = Source[SourceIt];
 				Codepoint = (Codepoint << (7U & Mask)) | (Byte & 0x7FU);
 
 				// Optionally decode third byte. Current state will remain unchanged unless the most significant bit in Byte is set.
 				Mask = int8(Byte) >> 7U;
 				SourceIt += 1U & Mask;
+				if (SourceIt >= SourceEndIt)
+				{
+					bIsErrorDetected = true;
+					break;
+				}
 				Byte = Source[SourceIt];
 				Codepoint = (Codepoint << (7U & Mask)) | (Byte & 0x7FU);
 
@@ -295,14 +317,17 @@ public:
 				return true;
 			}
 
-			check(DestLen >= GetSafeDecodedBufferLength(SourceLen));
+			if (DestLen < GetSafeDecodedBufferLength(SourceLen))
+			{
+				return false;
+			}
 
 			constexpr WideCharType ReplacementCharacter = 0xFFFD;
 
 			uint32 NeedSurrogatesCount = 0;
 			uint32 DestIt = 0;
 			bool bIsErrorDetected = false;
-			for (uint32 SourceIt = 0, SourceEndIt = SourceLen; SourceIt != SourceEndIt; ++SourceIt)
+			for (uint32 SourceIt = 0, SourceEndIt = SourceLen, DestEndIt = DestLen; (SourceIt < SourceEndIt) & (DestIt < DestEndIt) & (!bIsErrorDetected); ++SourceIt)
 			{
 				EncodeType Byte;
 				uint32 Codepoint;
@@ -315,12 +340,22 @@ public:
 				// Optionally decode second byte. Current state will remain unchanged unless the most significant bit in Byte is set.
 				Mask = int8(Byte) >> 7U;
 				SourceIt += 1U & Mask;
+				if (SourceIt >= SourceEndIt)
+				{
+					bIsErrorDetected = true;
+					break;
+				}
 				Byte = Source[SourceIt];
 				Codepoint = (Codepoint << (7U & Mask)) | (Byte & 0x7FU);
 
 				// Optionally decode third byte. Current state will remain unchanged unless the most significant bit in Byte is set.
 				Mask = int8(Byte) >> 7U;
 				SourceIt += 1U & Mask;
+				if (SourceIt >= SourceEndIt)
+				{
+					bIsErrorDetected = true;
+					break;
+				}
 				Byte = Source[SourceIt];
 				Codepoint = (Codepoint << (7U & Mask)) | (Byte & 0x7FU);
 
@@ -334,6 +369,12 @@ public:
 				// Note that there's no error checking for single surrogates. We just ignore them as the code point will not be in need of surrogates.
 				if (bIsInNeedOfSurrogates)
 				{
+					// One character is safe to write from the loop condition, but two could be one too many.
+					if (DestIt + 1U >= DestLen)
+					{
+						bIsErrorDetected = true;
+						break;
+					}
 					Dest[DestIt + 0] = GetHighSurrogate(Codepoint);
 					Dest[DestIt + 1] = GetLowSurrogate(Codepoint);
 					DestIt += 2U;
