Replication : add internal cluster information to rep clusters

diff --git a/Engine/Source/Runtime/Experimental/GeometryCollectionEngine/Private/GeometryCollection/GeometryCollectionComponent.cpp b/Engine/Source/Runtime/Experimental/GeometryCollectionEngine/Private/GeometryCollection/GeometryCollectionComponent.cpp
--- a/Engine/Source/Runtime/Experimental/GeometryCollectionEngine/Private/GeometryCollection/GeometryCollectionComponent.cpp
+++ b/Engine/Source/Runtime/Experimental/GeometryCollectionEngine/Private/GeometryCollection/GeometryCollectionComponent.cpp
@@ -166,7 +166,7 @@ bool FGeometryCollectionRepData::NetSerialize(FArchive& Ar, class UPackageMap* M
 		Ar << Cluster.AngularVelocity;
 		Ar << Cluster.Rotation;
 		Ar << Cluster.ClusterIdx;
-		Ar << Cluster.ObjectState;
+		Ar << Cluster.ClusterState.Value;
 	}
 
 	return true;
@@ -1471,69 +1471,70 @@ void UGeometryCollectionComponent::UpdateRepData()
 			// Particle can be null if we have embedded geometry 
 			if (Particle)
 			{
-			bool bProcess = true;
-			Processed.Add(Particle);
-			FPBDRigidClusteredParticleHandle* Root = Particle;
-			while (Root->Parent())
-			{
-				Root = Root->Parent();
-
-				//TODO: set avoids n^2, would be nice if clustered particle cached its root
-				if(Processed.Contains(Root))
-				{
-					bProcess = false;
-					break;
-				}
-				else
+				bool bProcess = true;
+				Processed.Add(Particle);
+				FPBDRigidClusteredParticleHandle* Root = Particle;
+				while (Root->Parent())
 				{
-					Processed.Add(Root);
-				}
-			}
-
-			if (bProcess && Root->Disabled() == false && ClustersToRep->Find(Root) == nullptr)
-			{
-				const int32 TransformGroupIdx = PhysicsProxy->GetTransformGroupIndexFromHandle(Root);
-				int32 Level = INDEX_NONE;
-				if (Root->InternalCluster() == false)
-				{
-					ensureMsgf(TransformGroupIdx >= 0, TEXT("Non-internal cluster should always have a group index"));
-					ensureMsgf(TransformGroupIdx < TNumericLimits<uint16>::Max(), TEXT("Trying to replicate GC with more than 65k pieces. We assumed uint16 would suffice"));
-
-					Level = Levels[TransformGroupIdx];
-				}
-				else
-				{
-					// Use internal cluster child's index to compute level.
-					const TArray<FPBDRigidParticleHandle*>& Children = RigidClustering.GetChildrenMap()[Root];
-					const int32 ChildTransformGroupIdx = PhysicsProxy->GetTransformGroupIndexFromHandle(Children[0]);
-					Level = Levels[ChildTransformGroupIdx] - 1;
+					Root = Root->Parent();
+	
+					//TODO: set avoids n^2, would be nice if clustered particle cached its root
+					if(Processed.Contains(Root))
+					{
+						bProcess = false;
+						break;
+					}
+					else
+					{
+						Processed.Add(Root);
+					}
 				}
-
-				if (!bEnableAbandonAfterLevel|| Level < ReplicationAbandonClusterLevel)
+	
+				if (bProcess && Root->Disabled() == false && ClustersToRep->Find(Root) == nullptr)
 				{
-					// not already replicated and not abandoned level, start replicating cluster
-					ClustersToRep->Add(Root);
-					bClustersChanged = true;
-				}
+					int32 TransformGroupIdx = INDEX_NONE;
+					int32 Level = INDEX_NONE;
+					if (Root->InternalCluster() == false)
+					{
+						TransformGroupIdx = PhysicsProxy->GetTransformGroupIndexFromHandle(Root);
+						ensureMsgf(TransformGroupIdx >= 0, TEXT("Non-internal cluster should always have a group index"));
+						ensureMsgf(TransformGroupIdx < TNumericLimits<uint16>::Max(), TEXT("Trying to replicate GC with more than 65k pieces. We assumed uint16 would suffice"));
 
-				if(Root->InternalCluster() == false && bFirstUpdate == false)	//if bFirstUpdate it must be that these are the initial roots of the GC. These did not break off so no need to replicate
-				{
-					//a one off so record it
-					ensureMsgf(TransformGroupIdx >= 0, TEXT("Non-internal cluster should always have a group index"));
-					ensureMsgf(TransformGroupIdx < TNumericLimits<uint16>::Max(), TEXT("Trying to replicate GC with more than 65k pieces. We assumed uint16 would suffice"));
-
-					// Because we cull ClustersToRep with abandoned level, we must make sure we don't add duplicates to one off activated.
-					// TODO: avoid search for entry for perf
-					FGeometryCollectionActivatedCluster OneOffActivated(TransformGroupIdx, Root->V(), Root->W());
-					if(!RepData.OneOffActivated.Contains(OneOffActivated))
+						Level = Levels[TransformGroupIdx];
+					}
+					else
 					{
+						// Use internal cluster child's index to compute level.
+						const TArray<FPBDRigidParticleHandle*>& Children = RigidClustering.GetChildrenMap()[Root];
+						const int32 ChildTransformGroupIdx = PhysicsProxy->GetTransformGroupIndexFromHandle(Children[0]);
+						Level = Levels[ChildTransformGroupIdx] - 1;
+					}
+	
+					if (!bEnableAbandonAfterLevel|| Level < ReplicationAbandonClusterLevel)
+					{
+						// not already replicated and not abandoned level, start replicating cluster
+						ClustersToRep->Add(Root);
 						bClustersChanged = true;
-						RepData.OneOffActivated.Add(OneOffActivated);
+					}
+	
+					if(Root->InternalCluster() == false && bFirstUpdate == false)	//if bFirstUpdate it must be that these are the initial roots of the GC. These did not break off so no need to replicate
+					{
+						//a one off so record it
+						ensureMsgf(TransformGroupIdx >= 0, TEXT("Non-internal cluster should always have a group index"));
+						ensureMsgf(TransformGroupIdx < TNumericLimits<uint16>::Max(), TEXT("Trying to replicate GC with more than 65k pieces. We assumed uint16 would suffice"));
+	
+						// Because we cull ClustersToRep with abandoned level, we must make sure we don't add duplicates to one off activated.
+						// TODO: avoid search for entry for perf
+						FGeometryCollectionActivatedCluster OneOffActivated(TransformGroupIdx, Root->V(), Root->W());
+						if(!RepData.OneOffActivated.Contains(OneOffActivated))
+						{
+							bClustersChanged = true;
+							RepData.OneOffActivated.Add(OneOffActivated);
+						}
 					}
 				}
 			}
 		}
-		}
 		
 		INC_DWORD_STAT_BY(STAT_GCReplicatedFractures, RepData.OneOffActivated.Num());
 
@@ -1557,7 +1558,8 @@ void UGeometryCollectionComponent::UpdateRepData()
 				ClusterRep.Rotation = Cluster->R();
 				ClusterRep.LinearVelocity = Cluster->V();
 				ClusterRep.AngularVelocity = Cluster->W();
-				ClusterRep.ObjectState = static_cast<int8>(Cluster->ObjectState());
+				ClusterRep.ClusterState.SetObjectState(Cluster->ObjectState());
+				ClusterRep.ClusterState.SetInternalCluster(Cluster->InternalCluster());
 				int32 TransformGroupIdx;
 				if(Cluster->InternalCluster())
 				{
@@ -1567,7 +1569,7 @@ void UGeometryCollectionComponent::UpdateRepData()
 				}
 				else
 				{
-					//not internal so we can just use the cluster's ID. On client we'll know based on the parent whether to use this index or the parent
+					// not internal so we can just use the cluster's ID. On client we'll know based on the parent whether to use this index or the parent
 					TransformGroupIdx = PhysicsProxy->GetTransformGroupIndexFromHandle(Cluster);
 				}
 
@@ -1650,6 +1652,7 @@ void UGeometryCollectionComponent::ProcessRepData()
 		// Set initial velocities if not hard snapping
 		if(!bHardSnap)
 		{
+			// TODO: we should get an update cluster position first so that when particles break off they get the right position 
 			// TODO: should we invalidate?
 			OneOff->SetV(ActivatedCluster.InitialLinearVelocity);
 			OneOff->SetW(ActivatedCluster.InitialAngularVelocity);
@@ -1662,15 +1665,23 @@ void UGeometryCollectionComponent::ProcessRepData()
 	{
 		for (const FGeometryCollectionClusterRep& RepCluster : RepData.Clusters)
 		{
-			FPBDRigidParticleHandle* Cluster = PhysicsProxy->GetParticles()[RepCluster.ClusterIdx];
-			if(Cluster->Disabled() == false)
+			if (FPBDRigidParticleHandle* Cluster = PhysicsProxy->GetParticles()[RepCluster.ClusterIdx])
 			{
-				Cluster->SetX(RepCluster.Position);
-				Cluster->SetR(RepCluster.Rotation);
-				Cluster->SetV(RepCluster.LinearVelocity);
-				Cluster->SetW(RepCluster.AngularVelocity);
-				//TODO: snap object state too once we fix interpolation
-				//Solver->GetEvolution()->SetParticleObjectState(Cluster, static_cast<Chaos::EObjectStateType>(RepCluster.ObjectState));
+				if (RepCluster.ClusterState.IsInternalCluster())
+				{
+					// internal cluster do not have an index so we rep data send one of the children's
+					// let's find the parent
+					Cluster = Cluster->CastToClustered()->Parent();
+				}
+				if(Cluster && Cluster->Disabled() == false)
+				{
+					Cluster->SetX(RepCluster.Position);
+					Cluster->SetR(RepCluster.Rotation);
+					Cluster->SetV(RepCluster.LinearVelocity);
+					Cluster->SetW(RepCluster.AngularVelocity);
+					// TODO: snap object state too once we fix interpolation
+					// Solver->GetEvolution()->SetParticleObjectState(Cluster, static_cast<Chaos::EObjectStateType>(RepCluster.ObjectState));
+				}
 			}
 		}
 	}
diff --git a/Engine/Source/Runtime/Experimental/GeometryCollectionEngine/Public/GeometryCollection/GeometryCollectionComponent.h b/Engine/Source/Runtime/Experimental/GeometryCollectionEngine/Public/GeometryCollection/GeometryCollectionComponent.h
--- a/Engine/Source/Runtime/Experimental/GeometryCollectionEngine/Public/GeometryCollection/GeometryCollectionComponent.h
+++ b/Engine/Source/Runtime/Experimental/GeometryCollectionEngine/Public/GeometryCollection/GeometryCollectionComponent.h
@@ -296,16 +296,57 @@ public:
  */
 struct FGeometryCollectionClusterRep
 {
+	struct FClusterState
+	{
+		static constexpr uint8 StateMask = 0b111;
+		static constexpr uint8 StateOffset = 0;
+		static constexpr uint8 InternalClusterMask = 0b1;
+		static constexpr uint8 InternalClusterOffset = 3;
+
+		void SetMaskedValue(uint8 Val, uint8 Mask, uint8 Offset)
+		{
+			Value &= ~Mask;
+			Value |= ((Val & Mask) << Offset);
+		}
+		
+		uint8 GetMaskedValue(uint8 Mask, uint8 Offset) const
+		{
+			return (Value >> Offset) & Mask;
+		}
+		
+		void SetObjectState(Chaos::EObjectStateType State)
+		{
+			SetMaskedValue((uint8)State, StateMask, StateOffset);
+		}
+
+		void SetInternalCluster(bool bInternalCluster)
+		{
+			SetMaskedValue(bInternalCluster?1:0, InternalClusterMask, InternalClusterOffset);
+		}
+		
+		Chaos::EObjectStateType GetObjectState()
+		{
+			return (Chaos::EObjectStateType)GetMaskedValue(StateMask, StateOffset);
+		}
+
+		bool IsInternalCluster() const
+		{
+			return (bool)GetMaskedValue(InternalClusterMask, InternalClusterOffset);
+		}
+		
+		uint8 Value = 0;
+	};
+	
 	FVector_NetQuantize100 Position;
 	FVector_NetQuantize100 LinearVelocity;
 	FVector_NetQuantize100 AngularVelocity;
 	FQuat Rotation;
-	uint16 ClusterIdx;
-	int8 ObjectState;
+	uint16 ClusterIdx;		  // index of the cluster or one of its child if the cluster is internal ( see ClusterState)
+	FClusterState ClusterState;
 
 	bool ClusterChanged(const FGeometryCollectionClusterRep& Other) const
 	{
-		return Other.ObjectState != ObjectState
+		return Other.ClusterState.Value != ClusterState.Value
 			|| Other.ClusterIdx != ClusterIdx
 			|| Other.Position != Position
 			|| Other.LinearVelocity != LinearVelocity
