Chaos : Flesh

Add mpm transfer, new grid based coloring functions and associated tests.

diff --git a/Engine/Source/Runtime/Experimental/Chaos/Private/Chaos/GraphColoring.cpp b/Engine/Source/Runtime/Experimental/Chaos/Private/Chaos/GraphColoring.cpp
--- a/Engine/Source/Runtime/Experimental/Chaos/Private/Chaos/GraphColoring.cpp
+++ b/Engine/Source/Runtime/Experimental/Chaos/Private/Chaos/GraphColoring.cpp
@@ -132,6 +132,78 @@ static bool VerifyGraph(TArray<TArray<int32>> ColorGraph, const TArray<Chaos::TV
 	return true;
 }
 
+
+template<typename T>
+static bool VerifyGraphAllDynamic(TArray<TArray<int32>> ColorGraph, const TArray<Chaos::TVec4<int32>>& Graph, const Chaos::TDynamicParticles<T, 3>& InParticles)
+{
+	for (int32 i = 0; i < ColorGraph.Num(); ++i)
+	{
+		TMap<int32, int32> NodeToColorMap;
+		for (const auto& Edge : ColorGraph[i])
+		{
+			int32 Node1 = Graph[Edge][0];
+			int32 Node2 = Graph[Edge][1];
+			int32 Node3 = Graph[Edge][2];
+			int32 Node4 = Graph[Edge][3];
+			if (NodeToColorMap.Contains(Node1))
+			{
+				UE_LOG(LogChaos, Error, TEXT("Color %d has duplicate Node %d"), i, Node1);
+				return false;
+			}
+			if (NodeToColorMap.Contains(Node2))
+			{
+				UE_LOG(LogChaos, Error, TEXT("Color %d has duplicate Node %d"), i, Node2);
+				return false;
+			}
+			if (NodeToColorMap.Contains(Node3))
+			{
+				UE_LOG(LogChaos, Error, TEXT("Color %d has duplicate Node %d"), i, Node3);
+				return false;
+			}
+			if (NodeToColorMap.Contains(Node4))
+			{
+				UE_LOG(LogChaos, Error, TEXT("Color %d has duplicate Node %d"), i, Node4);
+				return false;
+			}
+
+			NodeToColorMap.Add(Node1, i);
+			NodeToColorMap.Add(Node2, i);
+			NodeToColorMap.Add(Node3, i);
+			NodeToColorMap.Add(Node4, i);
+			
+		}
+	}
+	return true;
+}
+
+//just verify if different element in each subcolor has intersecting nodes
+template <typename T>
+static bool VerifyGridBasedSubColoring(const TArray<TArray<int32>>& ElementsPerColor, const Chaos::TMPMGrid<T>& Grid, const TArray<TArray<int32>>& ConstraintsNodesSet, const TArray<TArray<TArray<int32>>>& ElementsPerSubColors)
+{
+	for (int32 i = 0; i < ElementsPerSubColors.Num(); i++)
+	{
+		for (int32 j = 0; j < ElementsPerSubColors[i].Num(); j++)
+		{
+			TSet<int32> CoveredGridNodes;
+			//first gather all Grid nodes of an element:
+			for (int32 k = 0; k < ElementsPerSubColors[i][j].Num(); k++)
+			{
+				int32 e = ElementsPerSubColors[i][j][k];
+				for (int32 Node:ConstraintsNodesSet[e])
+				{
+					if (CoveredGridNodes.Contains(Node))
+					{
+						return false;
+					}
+					CoveredGridNodes.Emplace(Node);
+				}
+			}
+		}
+	}
+	return true;
+
+}
+
 template<typename T>
 TArray<TArray<int32>> Chaos::FGraphColoring::ComputeGraphColoring(const TArray<Chaos::TVec2<int32>>& Graph, const Chaos::TDynamicParticles<T, 3>& InParticles)
 {
@@ -696,6 +768,335 @@ TArray<TArray<int32>> Chaos::FGraphColoring::ComputeGraphColoring(const TArray<T
 	return ColorGraph;
 }
 
+template<typename T>
+TArray<TArray<int32>> Chaos::FGraphColoring::ComputeGraphColoringAllDynamic(const TArray<TVec4<int32>>& Graph, const Chaos::TDynamicParticles<T, 3>& InParticles)
+{
+	using namespace Chaos;
+
+	TArray<TArray<int32>> ColorGraph;
+	TArray<FGraphNode> Nodes;
+	TArray<FGraphTetEdge> Edges;
+	Nodes.SetNum(InParticles.Size());
+	Edges.SetNum(Graph.Num());
+	int32 MaxColor = -1;
+
+	for (int32 i = 0; i < Graph.Num(); ++i)
+	{
+		const TVec4<int32>& Constraint = Graph[i];
+		Edges[i].FirstNode = Constraint[0];
+		Edges[i].SecondNode = Constraint[1];
+		Edges[i].ThirdNode = Constraint[2];
+		Edges[i].FourthNode = Constraint[3];
+		Nodes[Constraint[0]].Edges.Add(i);
+		Nodes[Constraint[1]].Edges.Add(i);
+		Nodes[Constraint[2]].Edges.Add(i);
+		Nodes[Constraint[3]].Edges.Add(i);
+	}
+
+	TSet<int32> ProcessedNodes;
+	TArray<int32> NodesToProcess;
+
+	for (uint32 ParticleNodeIndex = 0; ParticleNodeIndex < InParticles.Size(); ++ParticleNodeIndex)
+	{
+		//const bool bIsParticleDynamic = InParticles.InvM(ParticleNodeIndex) != (T)0.;
+		if (ProcessedNodes.Contains(ParticleNodeIndex))
+		{
+			continue;
+		}
+
+		NodesToProcess.Add(ParticleNodeIndex);
+
+		while (NodesToProcess.Num())
+		{
+			const int32 NodeIndex = NodesToProcess.Last();
+			FGraphNode& GraphNode = Nodes[NodeIndex];
+
+			NodesToProcess.SetNum(NodesToProcess.Num() - 1, /*bAllowShrinking=*/false);
+			ProcessedNodes.Add(NodeIndex);
+
+			for (const int32 EdgeIndex : GraphNode.Edges)
+			{
+				FGraphTetEdge& GraphEdge = Edges[EdgeIndex];
+
+				// If edge has been colored skip it
+				if (GraphEdge.Color >= 0)
+				{
+					continue;
+				}
+
+				// Get index to the other node on the edge
+				int32 OtherNodeIndex = INDEX_NONE;
+				int32 OtherNodeIndex2 = INDEX_NONE;
+				int32 OtherNodeIndex3 = INDEX_NONE;
+				if (GraphEdge.FirstNode == NodeIndex)
+				{
+					OtherNodeIndex = GraphEdge.SecondNode;
+					OtherNodeIndex2 = GraphEdge.ThirdNode;
+					OtherNodeIndex3 = GraphEdge.FourthNode;
+				}
+				if (GraphEdge.SecondNode == NodeIndex)
+				{
+					OtherNodeIndex = GraphEdge.FirstNode;
+					OtherNodeIndex2 = GraphEdge.ThirdNode;
+					OtherNodeIndex3 = GraphEdge.FourthNode;
+				}
+				if (GraphEdge.ThirdNode == NodeIndex)
+				{
+					OtherNodeIndex = GraphEdge.FirstNode;
+					OtherNodeIndex2 = GraphEdge.SecondNode;
+					OtherNodeIndex3 = GraphEdge.FourthNode;
+				}
+				if (GraphEdge.FourthNode == NodeIndex)
+				{
+					OtherNodeIndex = GraphEdge.FirstNode;
+					OtherNodeIndex2 = GraphEdge.SecondNode;
+					OtherNodeIndex3 = GraphEdge.ThirdNode;
+				}
+
+				// Find next color that is not used already at this node
+				while (GraphNode.UsedColors.Contains(GraphNode.NextColor))
+				{
+					GraphNode.NextColor++;
+				}
+				int32 ColorToUse = GraphNode.NextColor;
+
+				// Exclude colors used by the other node (but still allow this node to use them for other edges)
+				if (OtherNodeIndex != INDEX_NONE)
+				{
+					FGraphNode& OtherNode = Nodes[OtherNodeIndex];
+
+					while (OtherNode.UsedColors.Contains(ColorToUse) || GraphNode.UsedColors.Contains(ColorToUse))
+					{
+						ColorToUse++;
+					}
+					
+				}
+				if (OtherNodeIndex2 != INDEX_NONE)
+				{
+					FGraphNode& OtherNode = Nodes[OtherNodeIndex2];
+					if (OtherNodeIndex == INDEX_NONE)
+					{
+						while (OtherNode.UsedColors.Contains(ColorToUse) || GraphNode.UsedColors.Contains(ColorToUse))
+						{
+							ColorToUse++;
+						}
+					}
+					else
+					{
+						FGraphNode& PrevOtherNode = Nodes[OtherNodeIndex];
+						while (OtherNode.UsedColors.Contains(ColorToUse) || PrevOtherNode.UsedColors.Contains(ColorToUse) || GraphNode.UsedColors.Contains(ColorToUse))
+						{
+							ColorToUse++;
+						}
+					}
+					
+				}
+				if (OtherNodeIndex3 != INDEX_NONE)
+				{
+					FGraphNode& OtherNode = Nodes[OtherNodeIndex3];
+
+					if (OtherNodeIndex == INDEX_NONE && OtherNodeIndex2 == INDEX_NONE)
+					{
+						while (OtherNode.UsedColors.Contains(ColorToUse) || GraphNode.UsedColors.Contains(ColorToUse))
+						{
+							ColorToUse++;
+						}
+					}
+					else if (OtherNodeIndex == INDEX_NONE)
+					{
+						//OtherNodeIndex2 != INDEX_NONE
+						FGraphNode& PrevOtherNode = Nodes[OtherNodeIndex2];
+						while (OtherNode.UsedColors.Contains(ColorToUse) || PrevOtherNode.UsedColors.Contains(ColorToUse) || GraphNode.UsedColors.Contains(ColorToUse))
+						{
+							ColorToUse++;
+						}
+					}
+					else if (OtherNodeIndex2 == INDEX_NONE)
+					{
+						FGraphNode& PrevOtherNode = Nodes[OtherNodeIndex];
+						while (OtherNode.UsedColors.Contains(ColorToUse) || PrevOtherNode.UsedColors.Contains(ColorToUse) || GraphNode.UsedColors.Contains(ColorToUse))
+						{
+							ColorToUse++;
+						}
+					}
+					else
+					{
+						FGraphNode& PrevOtherNode = Nodes[OtherNodeIndex];
+						FGraphNode& PrevPrevOtherNode = Nodes[OtherNodeIndex2];
+						while (OtherNode.UsedColors.Contains(ColorToUse) || PrevOtherNode.UsedColors.Contains(ColorToUse) || PrevPrevOtherNode.UsedColors.Contains(ColorToUse) || GraphNode.UsedColors.Contains(ColorToUse))
+						{
+							ColorToUse++;
+						}
+					}
+					
+				}
+
+				// Assign color and set as used at this node
+				MaxColor = FMath::Max(ColorToUse, MaxColor);
+				GraphNode.UsedColors.Add(ColorToUse);
+				GraphEdge.Color = ColorToUse;
+
+				// Bump color to use next time, but only if we weren't forced to use a different color by the other node
+				if (ColorToUse == GraphNode.NextColor) 
+				{
+					GraphNode.NextColor++;
+				}
+
+				if (ColorGraph.Num() <= MaxColor)
+				{
+					ColorGraph.SetNum(MaxColor + 1);
+				}
+				ColorGraph[GraphEdge.Color].Add(EdgeIndex);
+
+				if (OtherNodeIndex != INDEX_NONE)
+				{
+					FGraphNode& OtherGraphNode = Nodes[OtherNodeIndex];
+
+					// Mark other node as not allowing use of this color
+					OtherGraphNode.UsedColors.Add(GraphEdge.Color);
+					
+					// Queue other node for processing
+					if (!ProcessedNodes.Contains(OtherNodeIndex))
+					{
+						NodesToProcess.Add(OtherNodeIndex);
+					}
+					
+				}
+				if (OtherNodeIndex2 != INDEX_NONE)
+				{
+					FGraphNode& OtherGraphNode = Nodes[OtherNodeIndex2];
+
+					// Mark other node as not allowing use of this color
+					OtherGraphNode.UsedColors.Add(GraphEdge.Color);
+					
+					// Queue other node for processing
+					if (!ProcessedNodes.Contains(OtherNodeIndex2))
+					{
+						NodesToProcess.Add(OtherNodeIndex2);
+					}
+					
+				}
+				if (OtherNodeIndex3 != INDEX_NONE)
+				{
+					FGraphNode& OtherGraphNode = Nodes[OtherNodeIndex3];
+					// Mark other node as not allowing use of this color
+					OtherGraphNode.UsedColors.Add(GraphEdge.Color);
+						
+					// Queue other node for processing
+					if (!ProcessedNodes.Contains(OtherNodeIndex3))
+					{
+						NodesToProcess.Add(OtherNodeIndex3);
+					}
+					
+				}
+			}
+		}
+	}
+
+	checkSlow(VerifyGraphAllDynamic(ColorGraph, Graph, InParticles));
+	return ColorGraph;
+}
+
+
+//// everything in the data should be ordered by elements i.e. somethijng[e] is the data for eth element
+template<typename T>
+void Chaos::ComputeGridBasedGraphSubColoringPointer(const TArray<TArray<int32>>& ElementsPerColor, const TMPMGrid<T>& Grid, TArray<TArray<int32>>* PreviousColoring, const TArray<TArray<int32>>& ConstraintsNodesSet, TArray<TArray<TArray<int32>>>& ElementsPerSubColors) 
+{
+	bool InitialGuess = true;
+	if (!PreviousColoring) {
+		InitialGuess = false;
+		PreviousColoring = new TArray<TArray<int32>>();
+		PreviousColoring->SetNum(ElementsPerColor.Num());
+	}
+	ElementsPerSubColors.SetNum(ElementsPerColor.Num());
+
+	// std::cout << "Num colors: " << ElementsPerColor.size() << std::endl;
+	//PhysicsParallelFor(CellData.Num() / NPerEle, [&](const int32 i)
+	//	{
+	//		int32 CellIndex = NPerEle * i + CurrentLocIndex;
+	//		P2GApplyHelper(InParticles, CellIndex, GridData);
+	//	});
+
+	for (int32 i = 0; i < ElementsPerColor.Num(); i++) {
+		if (!InitialGuess) {
+			PreviousColoring->operator[](i).Init(-1, ElementsPerColor[i].Num());
+		}
+		int32 NumNodes = Grid.Size();
+		TArray<int32> ElementSubColors;
+		ElementSubColors.Init(-1, ElementsPerColor[i].Num());
+		TArray<TSet<int32>*> UsedColors;
+		UsedColors.Init(nullptr, NumNodes);
+		//std::vector<std::unordered_set<int>*> UsedColors(NumNodes, nullptr);
+		for (int32 j = 0; j < ElementsPerColor[i].Num(); j++) {
+			int32 ColorToUse = 0;
+			int32 e = ElementsPerColor[i][j];
+			// check initial guess:
+			if (InitialGuess) {
+				ColorToUse = PreviousColoring->operator[](i)[j];
+				bool ColorFound = false;
+				for (auto node : ConstraintsNodesSet[e]) {
+					if (!UsedColors[node]) {
+						UsedColors[node] = new TSet<int32>();
+					}
+					if (UsedColors[node]->Contains(ColorToUse)) {
+						ColorFound = true;
+						break;
+					}
+				}
+				if (ColorFound) {
+					ColorToUse = 0;
+				}
+				else {
+					for (auto node : ConstraintsNodesSet[e]) {
+						UsedColors[node]->Emplace(ColorToUse);
+					}
+					ElementSubColors[j] = ColorToUse;
+				}
+			}
+			if (ElementSubColors[j] == -1) {
+				while (true) {
+					bool ColorFound = false;
+					for (auto node : ConstraintsNodesSet[e]) {
+						if (!UsedColors[node]) {
+							UsedColors[node] = new TSet<int32>();// new std::unordered_set<int>;
+						}
+						if (UsedColors[node]->Contains(ColorToUse)) {
+							ColorFound = true;
+							break;
+						}
+					}
+					if (!ColorFound) {
+						break;
+					}
+					ColorToUse++;
+				}
+				ElementSubColors[j] = ColorToUse;
+				for (auto node : ConstraintsNodesSet[e]) {
+					UsedColors[node]->Emplace(ColorToUse);
+				}
+			}
+
+			// assign colors to previous guess for next timestep:
+			PreviousColoring->operator[](i)[j] = ColorToUse;
+		}
+
+		for (int ii = 0; ii < UsedColors.Num(); ii++) {
+			delete (UsedColors[ii]);
+		}
+
+		int32 NumColors = FMath::Max<int32>(ElementSubColors);
+
+		//int32 num_colors = *std::max_element(ElementSubColors.begin(), ElementSubColors.end());
+		ElementsPerSubColors[i].SetNum(0);
+		ElementsPerSubColors[i].SetNum(NumColors + 1);
+
+		for (int32 j = 0; j < ElementsPerColor[i].Num(); j++) {
+			ElementsPerSubColors[i][ElementSubColors[j]].Emplace(ElementsPerColor[i][j]);
+		}
+	}
+	checkSlow(VerifyGridBasedSubColoring<T>(ElementsPerColor, Grid, ConstraintsNodesSet, ElementsPerSubColors));
+}
+
 
 template CHAOS_API TArray<TArray<int32>> Chaos::FGraphColoring::ComputeGraphColoring<Chaos::FRealSingle>(const TArray<Chaos::TVector<int32, 2>>&, const Chaos::TDynamicParticles<Chaos::FRealSingle, 3>&);
 template CHAOS_API TArray<TArray<int32>> Chaos::FGraphColoring::ComputeGraphColoring<Chaos::FRealDouble>(const TArray<Chaos::TVector<int32, 2>>&, const Chaos::TDynamicParticles<Chaos::FRealDouble, 3>&);
@@ -703,3 +1104,7 @@ template CHAOS_API TArray<TArray<int32>> Chaos::FGraphColoring::ComputeGraphColo
 template CHAOS_API TArray<TArray<int32>> Chaos::FGraphColoring::ComputeGraphColoring<Chaos::FRealDouble>(const TArray<Chaos::TVector<int32, 3>>&, const Chaos::TDynamicParticles<Chaos::FRealDouble, 3>&);
 template CHAOS_API TArray<TArray<int32>> Chaos::FGraphColoring::ComputeGraphColoring<Chaos::FRealSingle>(const TArray<Chaos::TVector<int32, 4>>&, const Chaos::TDynamicParticles<Chaos::FRealSingle, 3>&);
 template CHAOS_API TArray<TArray<int32>> Chaos::FGraphColoring::ComputeGraphColoring<Chaos::FRealDouble>(const TArray<Chaos::TVector<int32, 4>>&, const Chaos::TDynamicParticles<Chaos::FRealDouble, 3>&);
+template CHAOS_API TArray<TArray<int32>> Chaos::FGraphColoring::ComputeGraphColoringAllDynamic<Chaos::FRealSingle>(const TArray<Chaos::TVector<int32, 4>>&, const Chaos::TDynamicParticles<Chaos::FRealSingle, 3>&);
+template CHAOS_API TArray<TArray<int32>> Chaos::FGraphColoring::ComputeGraphColoringAllDynamic<Chaos::FRealDouble>(const TArray<Chaos::TVector<int32, 4>>&, const Chaos::TDynamicParticles<Chaos::FRealDouble, 3>&);
+template CHAOS_API void Chaos::ComputeGridBasedGraphSubColoringPointer(const TArray<TArray<int32>>& ElementsPerColor, const TMPMGrid<Chaos::FRealSingle>& Grid, TArray<TArray<int32>>* PreviousColoring, const TArray<TArray<int32>>& ConstraintsNodesSet, TArray<TArray<TArray<int32>>>& ElementsPerSubColors);
+template CHAOS_API void Chaos::ComputeGridBasedGraphSubColoringPointer(const TArray<TArray<int32>>& ElementsPerColor, const TMPMGrid<Chaos::FRealDouble>& Grid, TArray<TArray<int32>>* PreviousColoring, const TArray<TArray<int32>>& ConstraintsNodesSet, TArray<TArray<TArray<int32>>>& ElementsPerSubColors);
\ No newline at end of file
diff --git a/Engine/Source/Runtime/Experimental/Chaos/Public/Chaos/GraphColoring.h b/Engine/Source/Runtime/Experimental/Chaos/Public/Chaos/GraphColoring.h
--- a/Engine/Source/Runtime/Experimental/Chaos/Public/Chaos/GraphColoring.h
+++ b/Engine/Source/Runtime/Experimental/Chaos/Public/Chaos/GraphColoring.h
@@ -5,6 +5,7 @@
 #include "Chaos/DynamicParticles.h"
 #include "Chaos/Core.h"
 #include "Chaos/Vector.h"
+#include "Chaos/UniformGrid.h"
 
 namespace Chaos
 {
@@ -66,6 +67,11 @@ class FGraphColoring
 	CHAOS_API static TArray<TArray<int32>> ComputeGraphColoring(const TArray<TVector<int32, 3>>& Graph, const TDynamicParticles<T, 3>& InParticles);
 	template<typename T>
 	CHAOS_API static TArray<TArray<int32>> ComputeGraphColoring(const TArray<TVector<int32, 4>>& Graph, const TDynamicParticles<T, 3>& InParticles);
+	template<typename T>
+	CHAOS_API static TArray<TArray<int32>> ComputeGraphColoringAllDynamic(const TArray<TVec4<int32>>& Graph, const Chaos::TDynamicParticles<T, 3>& InParticles);
 };
 
+template<typename T> 
+CHAOS_API void ComputeGridBasedGraphSubColoringPointer(const TArray<TArray<int32>>& ElementsPerColor, const TMPMGrid<T>& Grid, TArray<TArray<int32>>* PreviousColoring, const TArray<TArray<int32>>& ConstraintsNodesSet, TArray<TArray<TArray<int32>>>& ElementsPerSubColors);
+
 }
diff --git a/Engine/Source/Runtime/Experimental/Chaos/Public/Chaos/MPMTransfer.h b/Engine/Source/Runtime/Experimental/Chaos/Public/Chaos/MPMTransfer.h
new file mode 100644
--- /dev/null
+++ b/Engine/Source/Runtime/Experimental/Chaos/Public/Chaos/MPMTransfer.h
@@ -0,0 +1,238 @@
+// Copyright Epic Games, Inc. All Rights Reserved.
+#pragma once
+
+#include "Chaos/Vector.h"
+#include "ChaosCheck.h"
+#include "Chaos/ArrayCollectionArray.h"
+#include "Chaos/DynamicParticles.h"
+#include "Chaos/UniformGrid.h"
+#include "Algo/Unique.h"
+
+namespace Chaos {
+
+template<class T>
+class TMPMTransfer
+{
+public:
+
+	uint32 NPerSec = 2 * 2;
+	uint32 NPerEle = 2 * 2 * 2;
+	uint32 NTransfer = 3;
+	int32 NumCells;
+
+	//TODO(Yizhou): Think whether mpm transfer should just own the grid
+	//(Is the grid used by the constraint as well?)
+	TMPMGrid<T>& Grid;
+
+	TArray<TVector<T, 3>> Weights;
+	TArray<TVector<int32, 3>> Indices;
+	TArray<TArray<int32>> CellData; //CellData[i] registers which particles are in ith cell
+	//TArray<TArray<int32>> CellColors;
+
+	//meta data for grid based cons
+	TArray<TArray<int32>> ElementGridNodes;
+	TArray<TArray<T>> ElementGridNodeWeights;
+	TArray<TArray<TArray<int32>>> ElementGridNodeIncidentElements;
+
+
+	TMPMTransfer(){}
+	TMPMTransfer(TMPMGrid<T>& _Grid) :Grid(_Grid) 
+	{
+		ensure(Grid.NPerDir == 2 || Grid.NPerDir == 3);
+		NPerSec = Grid.NPerDir * Grid.NPerDir;
+		NPerEle = NPerSec * Grid.NPerDir;
+		//CellColors.SetNum(NPerEle);
+	}
+
+	//template <typename ReorderFunctor, typename SplatFunctor>
+	//This one does an initial splat of momentum and mass to the grid. 
+	//One can add splat functor in the future passibly for other kinds of splat
+	void InitialP2G(const TDynamicParticles<T, 3>& InParticles, TArray<T>& GridData) 
+	{
+		int32 N = InParticles.Size();
+
+		{
+			TRACE_CPUPROFILER_EVENT_SCOPE(TEXT("STAT_ChaosMPMTransferInitialBinning"));
+			Indices.SetNum(N);
+			Weights.SetNum(N);
+
+			/////////////////////////
+			// compute weights and bin
+			/////////////////////////
+
+			//TODO(Yizhou): Do a timing and determine whether the following loop should 
+			//use physics parallel for:
+			for (int32 p = 0; p < N; p++)
+			{
+				Grid.BaseNodeIndex(InParticles.X(p), Indices[p], Weights[p]);
+			}
+
+		}
+		/////////////////////
+		// Computes which particles in the same cell
+		/////////////////////
+		{
+			TRACE_CPUPROFILER_EVENT_SCOPE(TEXT("STAT_ChaosMPMTransferCellMetaCalc"));
+			NumCells = Grid.Size();
+			CellData.SetNum(NumCells);
+			for (int32 c = 0; c < NumCells; c++)
+			{
+				CellData[c].SetNum(0);
+			}
+			for (int32 p = 0; p < N; p++)
+			{
+				CellData[Grid.FlatIndex(Indices[p])].Emplace(p);
+			}
+		}
+
+		/////////////////////
+		// splat data to cells
+		/////////////////////
+		{
+			TRACE_CPUPROFILER_EVENT_SCOPE(TEXT("STAT_ChaosMPMTransferSplatData"));
+			GridData.Init((T)0., CellData.Num() * (NTransfer + 1));
+
+			TVector<int32, 3> GridCells = Grid.GetCells();
+			
+			for (int32 ii = 0; ii < int32(Grid.NPerDir); ii++)
+			{
+				for (int32 jj = 0; jj < int32(Grid.NPerDir); jj++)
+				{
+					for (int32 kk = 0; kk < int32(Grid.NPerDir); kk++)
+					{
+						int32 CurrentLocIndex = ii * NPerSec + jj * Grid.NPerDir + kk;
+						PhysicsParallelFor(GridCells[0] / Grid.NPerDir, [&](const int32 iii)
+							{
+								for (int32 jjj = 0; jjj < int32(GridCells[1] / Grid.NPerDir); jjj++)
+								{
+									for (int32 kkk = 0; kkk < int32(GridCells[2] / Grid.NPerDir); kkk++)
+									{
+										TVector<int32, 3> MultiIndex = { iii * int32(Grid.NPerDir) + ii, jjj * int32(Grid.NPerDir) + jj, kkk * int32(Grid.NPerDir) + kk};
+										int32 CellIndex = Grid.FlatIndex(MultiIndex);
+										P2GApplyHelper(InParticles, CellIndex, GridData);
+									}
+								}
+						});
+
+					}
+				}
+			}
+
+		}
+	
+	}
+
+	//currently only splats mass and momentum
+	void P2GApplyHelper(const TDynamicParticles<T, 3>& InParticles, const int32 CellIndex, TArray<T>& GridData)
+	{
+		//check if valid cell:
+		if (CellIndex < NumCells && CellData[CellIndex].Num() > 0)
+		{
+			for (int32 i = 0; i < CellData[CellIndex].Num(); i++)
+			{
+				int32 p = CellData[CellIndex][i];
+				for (int iii = 0; iii < int32(Grid.NPerDir); iii++)
+				{
+					T Nii = Grid.Nijk(Weights[p][0], iii);
+					for (int jjj = 0; jjj < int32(Grid.NPerDir); jjj++)
+					{
+						T Njj = Grid.Nijk(Weights[p][1], jjj);
+						for (int kkk = 0; kkk < int32(Grid.NPerDir); kkk++)
+						{
+							TVector<int32, 3> LocIndex = { iii, jjj, kkk };
+							TVector<int32, 3> GlobMultiIndex = Grid.Loc2GlobIndex(Indices[p], LocIndex);
+							int32 GlobIndex = Grid.FlatIndex(GlobMultiIndex);
+							T Nkk = Grid.Nijk(Weights[p][2], kkk);
+							T NProd = Nii * Njj * Nkk;
+							GridData[(NTransfer + 1) * GlobIndex] += NProd * InParticles.M(p);
+							for (int32 alpha = 0; alpha < 3; alpha++)
+							{
+								GridData[(NTransfer + 1) * GlobIndex + alpha + 1] += NProd * InParticles.M(p) * InParticles.V(p)[alpha];
+							}
+
+						}
+					}
+				}
+			}
+		}
+	}
+
+	void ComputeElementMetaData(const TArray<TVector<int32, 4>>& InMesh)
+	{
+		ElementGridNodes.SetNum(InMesh.Num());
+		ElementGridNodeWeights.SetNum(InMesh.Num());
+		ElementGridNodeIncidentElements.SetNum(InMesh.Num());
+		//TODO(Yizhou): make following loop parallel for with appropriate bool condition
+		for (int32 e = 0; e < InMesh.Num(); e++)
+		{
+			ElementGridNodes[e].SetNum(NPerEle * 4);
+			ElementGridNodeWeights[e].SetNum(NPerEle * 4);
+			for (int32 ie = 0; ie < 4; ie++)
+			{
+				int32 p = InMesh[e][ie];
+				TVector<int32, 3> Index = Indices[p];
+				for (int32 ii = 0; ii < Grid.NPerDir; ii++)
+				{
+					T Nii = Grid.Nijk(Weights[p][0], ii);
+					for (int32 jj = 0; jj < Grid.NPerDir; jj++)
+					{
+						T Njj = Grid.Nijk(Weights[p][1], jj);
+						for (int32 kk = 0; kk < Grid.NPerDir; kk++)
+						{
+							T Nkk = Grid.Nijk(Weights[p][2], kk);
+							TVector<int32, 3> LocIndex = { ii, jj, kk };
+							TVector<int32, 3> GlobIndex = Grid.Loc2GlobIndex(Index, LocIndex);
+							int32 GlobIndexFlat = Grid.FlatIndex(GlobIndex);
+							ElementGridNodes[e][ie * NPerEle + ii * NPerSec + jj * Grid.NPerDir + kk] = GlobIndexFlat;
+							ElementGridNodeWeights[e][ie * NPerEle + ii * NPerSec + jj * Grid.NPerDir + kk] = Nii * Njj * Nkk;
+						}
+					}
+				}
+			}
+			ComputeIncidentElements(ElementGridNodes[e], ElementGridNodeIncidentElements[e]);
+		}
+	}
+
+	//computes incident elements in serial
+	static void ComputeIncidentElements(const TArray<int32>& ArrayIn, TArray<TArray<int32>>& IncidentElements) 
+	{
+		TArray<int32> Ordering, Ranges;
+		Ordering.SetNum(ArrayIn.Num());
+		Ranges.SetNum(ArrayIn.Num()+1);
+
+		for (int32 i = 0; i < ArrayIn.Num(); ++i) 
+		{
+			Ordering[i] = i;
+			Ranges[i] = i;
+		}
+		
+		//TODO(Yizhou): decide to use sort or heap sort or merge sort. 
+		Ordering.Sort([&ArrayIn](const int32 A, const int32 B) 
+		{
+		return ArrayIn[A] < ArrayIn[B];
+		});
+
+		int32 Last = Algo::Unique(Ranges, [&ArrayIn, &Ordering](int32 i, int32 j) { return ArrayIn[Ordering[i]] == ArrayIn[Ordering[j]]; });
+
+		int32 NumNodes = Last - 1;
+		Ranges[NumNodes] = ArrayIn.Num();
+		Ranges.SetNum(NumNodes + 1);
+
+		IncidentElements.SetNum(Ranges.Num() - 1);
+
+		for (int32 p = 0; p < IncidentElements.Num(); ++p) 
+		{
+			IncidentElements[p].SetNum(Ranges[p + 1] - Ranges[p]);
+			for (int32 e = Ranges[p]; e < Ranges[p + 1]; e++) 
+			{
+				IncidentElements[p][e - Ranges[p]] = Ordering[e];
+			}
+		}
+	}
+
+	
+
+
+};
+
+}  
