Chaos : Flesh

Fix Mac compile issues.

diff --git a/Engine/Source/Runtime/Experimental/Chaos/Private/Chaos/GraphColoring.cpp b/Engine/Source/Runtime/Experimental/Chaos/Private/Chaos/GraphColoring.cpp
--- a/Engine/Source/Runtime/Experimental/Chaos/Private/Chaos/GraphColoring.cpp
+++ b/Engine/Source/Runtime/Experimental/Chaos/Private/Chaos/GraphColoring.cpp
@@ -998,9 +998,8 @@ TArray<TArray<int32>> Chaos::FGraphColoring::ComputeGraphColoringAllDynamic(cons
 }
 
 
-//// everything in the data should be ordered by elements i.e. somethijng[e] is the data for eth element
 template<typename T>
-void Chaos::ComputeGridBasedGraphSubColoringPointer(const TArray<TArray<int32>>& ElementsPerColor, const TMPMGrid<T>& Grid, TArray<TArray<int32>>* PreviousColoring, const TArray<TArray<int32>>& ConstraintsNodesSet, TArray<TArray<TArray<int32>>>& ElementsPerSubColors) 
+void Chaos::ComputeGridBasedGraphSubColoringPointer(const TArray<TArray<int32>>& ElementsPerColor, const TMPMGrid<T>& Grid, const int32 GridSize, TArray<TArray<int32>>* PreviousColoring, const TArray<TArray<int32>>& ConstraintsNodesSet, TArray<TArray<TArray<int32>>>& ElementsPerSubColors) 
 {
 	bool InitialGuess = true;
 	if (!PreviousColoring) {
@@ -1010,18 +1009,12 @@ void Chaos::ComputeGridBasedGraphSubColoringPointer(const TArray<TArray<int32>>&
 	}
 	ElementsPerSubColors.SetNum(ElementsPerColor.Num());
 
-	// std::cout << "Num colors: " << ElementsPerColor.size() << std::endl;
-	//PhysicsParallelFor(CellData.Num() / NPerEle, [&](const int32 i)
-	//	{
-	//		int32 CellIndex = NPerEle * i + CurrentLocIndex;
-	//		P2GApplyHelper(InParticles, CellIndex, GridData);
-	//	});
 
 	for (int32 i = 0; i < ElementsPerColor.Num(); i++) {
 		if (!InitialGuess) {
 			PreviousColoring->operator[](i).Init(-1, ElementsPerColor[i].Num());
 		}
-		int32 NumNodes = Grid.Size();
+		int32 NumNodes = GridSize;
 		TArray<int32> ElementSubColors;
 		ElementSubColors.Init(-1, ElementsPerColor[i].Num());
 		TArray<TSet<int32>*> UsedColors;
@@ -1106,5 +1099,5 @@ template CHAOS_API TArray<TArray<int32>> Chaos::FGraphColoring::ComputeGraphColo
 template CHAOS_API TArray<TArray<int32>> Chaos::FGraphColoring::ComputeGraphColoring<Chaos::FRealDouble>(const TArray<Chaos::TVector<int32, 4>>&, const Chaos::TDynamicParticles<Chaos::FRealDouble, 3>&);
 template CHAOS_API TArray<TArray<int32>> Chaos::FGraphColoring::ComputeGraphColoringAllDynamic<Chaos::FRealSingle>(const TArray<Chaos::TVector<int32, 4>>&, const Chaos::TDynamicParticles<Chaos::FRealSingle, 3>&);
 template CHAOS_API TArray<TArray<int32>> Chaos::FGraphColoring::ComputeGraphColoringAllDynamic<Chaos::FRealDouble>(const TArray<Chaos::TVector<int32, 4>>&, const Chaos::TDynamicParticles<Chaos::FRealDouble, 3>&);
-template CHAOS_API void Chaos::ComputeGridBasedGraphSubColoringPointer(const TArray<TArray<int32>>& ElementsPerColor, const TMPMGrid<Chaos::FRealSingle>& Grid, TArray<TArray<int32>>* PreviousColoring, const TArray<TArray<int32>>& ConstraintsNodesSet, TArray<TArray<TArray<int32>>>& ElementsPerSubColors);
-template CHAOS_API void Chaos::ComputeGridBasedGraphSubColoringPointer(const TArray<TArray<int32>>& ElementsPerColor, const TMPMGrid<Chaos::FRealDouble>& Grid, TArray<TArray<int32>>* PreviousColoring, const TArray<TArray<int32>>& ConstraintsNodesSet, TArray<TArray<TArray<int32>>>& ElementsPerSubColors);
\ No newline at end of file
+template CHAOS_API void Chaos::ComputeGridBasedGraphSubColoringPointer(const TArray<TArray<int32>>& ElementsPerColor, const TMPMGrid<Chaos::FRealSingle>& Grid, const int32 GridSize, TArray<TArray<int32>>* PreviousColoring, const TArray<TArray<int32>>& ConstraintsNodesSet, TArray<TArray<TArray<int32>>>& ElementsPerSubColors);
+template CHAOS_API void Chaos::ComputeGridBasedGraphSubColoringPointer(const TArray<TArray<int32>>& ElementsPerColor, const TMPMGrid<Chaos::FRealDouble>& Grid, const int32 GridSize, TArray<TArray<int32>>* PreviousColoring, const TArray<TArray<int32>>& ConstraintsNodesSet, TArray<TArray<TArray<int32>>>& ElementsPerSubColors);
\ No newline at end of file
diff --git a/Engine/Source/Runtime/Experimental/Chaos/Public/Chaos/GraphColoring.h b/Engine/Source/Runtime/Experimental/Chaos/Public/Chaos/GraphColoring.h
--- a/Engine/Source/Runtime/Experimental/Chaos/Public/Chaos/GraphColoring.h
+++ b/Engine/Source/Runtime/Experimental/Chaos/Public/Chaos/GraphColoring.h
@@ -72,6 +72,6 @@ class FGraphColoring
 };
 
 template<typename T> 
-CHAOS_API void ComputeGridBasedGraphSubColoringPointer(const TArray<TArray<int32>>& ElementsPerColor, const TMPMGrid<T>& Grid, TArray<TArray<int32>>* PreviousColoring, const TArray<TArray<int32>>& ConstraintsNodesSet, TArray<TArray<TArray<int32>>>& ElementsPerSubColors);
+CHAOS_API void ComputeGridBasedGraphSubColoringPointer(const TArray<TArray<int32>>& ElementsPerColor, const TMPMGrid<T>& Grid, const int32 GridSize, TArray<TArray<int32>>* PreviousColoring, const TArray<TArray<int32>>& ConstraintsNodesSet, TArray<TArray<TArray<int32>>>& ElementsPerSubColors);
 
 }
