Dataflow - Caching Headless Support

diff --git a/Engine/Source/Runtime/Experimental/Dataflow/Core/Private/Dataflow/DataflowInputOutput.cpp b/Engine/Source/Runtime/Experimental/Dataflow/Core/Private/Dataflow/DataflowInputOutput.cpp
--- a/Engine/Source/Runtime/Experimental/Dataflow/Core/Private/Dataflow/DataflowInputOutput.cpp
+++ b/Engine/Source/Runtime/Experimental/Dataflow/Core/Private/Dataflow/DataflowInputOutput.cpp
@@ -25,6 +25,7 @@ bool FDataflowInput::AddConnection(FDataflowConnection* InOutput)
 	if (ensure(InOutput->GetType() == this->GetType()))
 	{
 		Connection = (FDataflowOutput*)InOutput;
+		GetOwningNode()->Invalidate();
 		return true;
 	}
 	return false;
@@ -35,6 +36,7 @@ bool FDataflowInput::RemoveConnection(FDataflowConnection* InOutput)
 	if (ensure(Connection == (FDataflowOutput*)InOutput))
 	{
 		Connection = nullptr;
+		GetOwningNode()->Invalidate();
 		return true;
 	}
 	return false;
@@ -62,7 +64,7 @@ const TArray< const FDataflowOutput* > FDataflowInput::GetConnectedOutputs() con
 
 void FDataflowInput::Invalidate()
 {
-	OwningNode->InvalidateOutputs();
+	OwningNode->Invalidate();
 }
 
 //
@@ -120,13 +122,9 @@ bool FDataflowOutput::RemoveConnection(FDataflowConnection* InInput)
 
 void FDataflowOutput::Invalidate()
 {
-	if (CacheKeyValue != UINT_MAX)
+	for (FDataflowConnection* Con : GetConnections())
 	{
-		CacheKeyValue = UINT_MAX;
-		for (FDataflowConnection* Con : GetConnections())
-		{
-			Con->Invalidate();
-		}
+		Con->Invalidate();
 	}
 }
 
diff --git a/Engine/Source/Runtime/Experimental/Dataflow/Core/Private/Dataflow/DataflowNode.cpp b/Engine/Source/Runtime/Experimental/Dataflow/Core/Private/Dataflow/DataflowNode.cpp
--- a/Engine/Source/Runtime/Experimental/Dataflow/Core/Private/Dataflow/DataflowNode.cpp
+++ b/Engine/Source/Runtime/Experimental/Dataflow/Core/Private/Dataflow/DataflowNode.cpp
@@ -193,6 +193,12 @@ TArray<Dataflow::FPin> FDataflowNode::GetPins() const
 	return RetVal;
 }
 
+void FDataflowNode::Invalidate()
+{
+	LastModifiedTimestamp = FDateTime::Now(); // @todo(dataflow) Now() has ms precision; might need higher
+	InvalidateOutputs();
+}
+
 void FDataflowNode::InvalidateOutputs()
 {
 	for (TPair<uint32, FDataflowOutput*> Elem : Outputs)
diff --git a/Engine/Source/Runtime/Experimental/Dataflow/Core/Public/Dataflow/DataflowInputOutput.h b/Engine/Source/Runtime/Experimental/Dataflow/Core/Public/Dataflow/DataflowInputOutput.h
--- a/Engine/Source/Runtime/Experimental/Dataflow/Core/Public/Dataflow/DataflowInputOutput.h
+++ b/Engine/Source/Runtime/Experimental/Dataflow/Core/Public/Dataflow/DataflowInputOutput.h
@@ -125,7 +125,7 @@ public:
 	{
 		if (Property)
 		{
-			Context.SetData(CacheKey(), new Dataflow::ContextCache<T>(Property, new T(InVal)));
+			Context.SetData(CacheKey(), Property, InVal);
 		}
 	}
 
@@ -133,13 +133,12 @@ public:
 	{
 		if (!this->Evaluate<T>(Context))
 		{
-			Context.SetData(CacheKey(), new Dataflow::ContextCache<T>(Property, new T(Default)));
+			Context.SetData(CacheKey(), Property, Default);
 		}
 
-
 		if (Context.HasData(CacheKey()))
 		{
-			return Context.GetDataReference<T>(CacheKey(), Default);
+			return Context.GetData(CacheKey(), Property, Default);
 		}
 
 		return Default;
@@ -157,15 +156,15 @@ const T& FDataflowInput::GetValue(Dataflow::FContext& Context, const T& Default)
 	if (GetConnectedOutputs().Num())
 	{
 		ensure(GetConnectedOutputs().Num() == 1);
-		if (const FDataflowOutput* ConnectionBase = GetConnection())
+		if (const FDataflowOutput* ConnectionOut = GetConnection())
 		{
-			if (!ConnectionBase->Evaluate<T>(Context))
+			if (!ConnectionOut->Evaluate<T>(Context))
 			{
-				Context.SetData(ConnectionBase->CacheKey(), new Dataflow::ContextCache<T>(Property, new T(Default)));
+				Context.SetData(ConnectionOut->CacheKey(), Property, Default);
 			}
-			if (Context.HasData(ConnectionBase->CacheKey()))
+			if (Context.HasData(ConnectionOut->CacheKey()))
 			{
-				const T& data = Context.GetDataReference<T>(ConnectionBase->CacheKey(), Default);
+				const T& data = Context.GetData(ConnectionOut->CacheKey(), Property, Default);
 				return data;
 			}
 		}
@@ -180,17 +179,27 @@ bool FDataflowOutput::Evaluate(Dataflow::FContext& Context) const
  
 	if (OwningNode->bActive)
 	{
+		// check if the cache has a valid version
+		if(Context.HasData(CacheKey(), OwningNode->LastModifiedTimestamp))
+		{
+			return true;
+		}
+		// if not, evaluate
 		OwningNode->Evaluate(Context, this);
+		// Validation
+		if (!Context.HasData(CacheKey()))
+		{
+			ensureMsgf(false, TEXT("Failed to evaluate output (%s:%s)"), *OwningNode->GetName().ToString(), *GetName().ToString());
+			return false;
+		}
+		return true;
 	}
 	else if(const FDataflowInput* PassthroughInput = OwningNode->FindInput(GetPassthroughRealAddress()))
 	{
-		SetValue<T>(PassthroughInput->GetValue<T>(Context, *reinterpret_cast<const T*>(PassthroughInput->RealAddress())), Context);
+		T PassthroughData = PassthroughInput->GetValue<T>(Context, *reinterpret_cast<const T*>(PassthroughInput->RealAddress()));
+		SetValue<T>(PassthroughData, Context);
+		return true;
 	}
  
-	// Validation
-	if (!Context.HasData(CacheKey()))
-	{
-		ensureMsgf(false, TEXT("Failed to evaluate output (%s:%s)"), *OwningNode->GetName().ToString(), *GetName().ToString());
-	}
-	return true;
+	return false;
 }
\ No newline at end of file
diff --git a/Engine/Source/Runtime/Experimental/Dataflow/Core/Public/Dataflow/DataflowNode.h b/Engine/Source/Runtime/Experimental/Dataflow/Core/Public/Dataflow/DataflowNode.h
--- a/Engine/Source/Runtime/Experimental/Dataflow/Core/Public/Dataflow/DataflowNode.h
+++ b/Engine/Source/Runtime/Experimental/Dataflow/Core/Public/Dataflow/DataflowNode.h
@@ -39,6 +39,7 @@ struct DATAFLOWCORE_API FDataflowNode
 
 	FGuid Guid;
 	FName Name;
+	FDateTime LastModifiedTimestamp;
 
 	TMap< int, FDataflowInput* > Inputs;
 	TMap< int, FDataflowOutput* > Outputs;
@@ -49,12 +50,14 @@ struct DATAFLOWCORE_API FDataflowNode
 	FDataflowNode()
 		: Guid(FGuid())
 		, Name("Invalid")
+		, LastModifiedTimestamp(0)
 	{
 	}
 
 	FDataflowNode(const Dataflow::FNodeParameters& Param, FGuid InGuid = FGuid::NewGuid())
 		: Guid(InGuid)
 		, Name(Param.Name)
+		, LastModifiedTimestamp(0)
 	{
 	}
 
@@ -159,6 +162,8 @@ struct DATAFLOWCORE_API FDataflowNode
 		FindOutput(Reference)->template SetValue<T>(Value, Context);
 	}
 
+	void Invalidate();
+
 	void InvalidateOutputs();
 
 	bool ValidateConnections();
diff --git a/Engine/Source/Runtime/Experimental/Dataflow/Core/Public/Dataflow/DataflowNodeParameters.h b/Engine/Source/Runtime/Experimental/Dataflow/Core/Public/Dataflow/DataflowNodeParameters.h
--- a/Engine/Source/Runtime/Experimental/Dataflow/Core/Public/Dataflow/DataflowNodeParameters.h
+++ b/Engine/Source/Runtime/Experimental/Dataflow/Core/Public/Dataflow/DataflowNodeParameters.h
@@ -3,56 +3,70 @@
 #pragma once
 
 #include "CoreMinimal.h"
+#include "UObject/UnrealType.h"
+#include "Templates/UniquePtr.h"
 
 class  UDataflow;
 
 namespace Dataflow
 {
-	struct ContextCacheBase 
+	struct FContextCacheBase 
 	{
-		ContextCacheBase(FProperty* InProperty = nullptr, void* InData = nullptr)
+		FContextCacheBase(FProperty* InProperty = nullptr, FDateTime InTimestamp = 0)
 			: Property(InProperty)
-			, Data(InData)
+			, Timestamp(InTimestamp)
 		{}
-		virtual ~ContextCacheBase() {}
+		virtual ~FContextCacheBase() {}
 
+		template<typename T>
+		const T& GetTypedData(const FProperty* PropertyIn) const;
+		
 		FProperty* Property = nullptr;
-		void* Data = nullptr;
+		FDateTime Timestamp = FDateTime(0);
 	};
 
 	template<class T>
-	struct ContextCache : public ContextCacheBase 
+	struct FContextCache : public FContextCacheBase 
 	{
-		ContextCache(FProperty* InProperty = nullptr, void* InData = nullptr)
-			: ContextCacheBase(InProperty,InData)
+		FContextCache(FProperty* InProperty, const T& InData, FDateTime Timestamp)
+			: FContextCacheBase(InProperty, Timestamp)
+			, Data(InData)
 		{}
-		virtual ~ContextCache() 
-		{ 
-			delete (T*)Data; 
-		}
+
+		FContextCache(FProperty* InProperty, T&& InData, FDateTime Timestamp)
+			: FContextCacheBase(InProperty, Timestamp)
+			, Data(InData)
+		{}
+		
+		const T Data;
 	};
 
+	template<class T>
+	const T& FContextCacheBase::GetTypedData(const FProperty* PropertyIn) const
+	{
+		check(PropertyIn);
+		// check(PropertyIn->IsA<T>()); // @todo(dataflow) compile error for non-class T; find alternatives
+		check(Property->SameType(PropertyIn));
+		return static_cast<const FContextCache<T>&>(*this).Data;
+	}
+	
 	class DATAFLOWCORE_API FContext
 	{
-		TMap<int64, ContextCacheBase*> DataStore;
+		TMap<int64, TUniquePtr<FContextCacheBase>> DataStore;
+		
+		FContext(FContext&&) = default;
+		FContext& operator=(FContext&&) = default;
+		
+		FContext(const FContext&) = delete;
+		FContext& operator=(const FContext&) = delete;
 
 	public:
 
-
-
 		FContext(float InTime, FString InType = FString(""))
 			: Timestamp(InTime)
 			, Type(StaticType().Append(InType))
 		{}
 
-		~FContext()
-		{
-			for (TTuple<int64, ContextCacheBase*> Elem : DataStore)
-			{
-				delete Elem.Value;
-			}
-		}
-
 		float Timestamp = 0.f;
 		FString Type;
 		static FString StaticType() { return "FContext"; }
@@ -72,51 +86,36 @@ namespace Dataflow
 			return nullptr;
 		}
 
-		void SetData(size_t Key, ContextCacheBase* InData)
+		template<typename T>
+		void SetData(size_t Key, FProperty* Property, const T& Value)
 		{
 			int64 IntKey = (int64)Key;
-			if (!DataStore.Contains(IntKey))
-			{
-				DataStore.Add(IntKey, InData);
-			}
-			else
-			{
-				DataStore[IntKey] = InData;
-			}
+			TUniquePtr<FContextCache<T>> DataStoreEntry = MakeUnique<FContextCache<T>>(Property, Value, FDateTime::Now());
+			DataStore.Emplace(IntKey, MoveTemp(DataStoreEntry));
 		}
 
-		template<class T>
-		bool GetData(size_t Key, T& OutData)
+		template<typename T>
+		void SetData(size_t Key, FProperty* Property, T&& Value)
 		{
 			int64 IntKey = (int64)Key;
-
-			if (DataStore.Contains(IntKey))
-			{
-				// @todo(dataflow) : type check
-				OutData = *(T*)DataStore[IntKey]->Data;
-				return true;
-			}
-			return false;
+			TUniquePtr<FContextCache<T>> DataStoreEntry = MakeUnique<FContextCache<T>>(Property, Forward<T>(Value), FDateTime::Now());
+			DataStore.Emplace(IntKey, MoveTemp(DataStoreEntry));
 		}
-
+		
 		template<class T>
-		const T& GetDataReference(size_t Key, const T& Default)
+		const T& GetData(size_t Key, FProperty* Property, const T& Default = T())
 		{
-			int64 IntKey = (int64)Key;
-
-			if (DataStore.Contains(Key))
+			if (TUniquePtr<FContextCacheBase>* Cache = DataStore.Find(Key))
 			{
-				// @todo(dataflow) : type check
-				return *(T*)DataStore[Key]->Data;
+				return (*Cache)->GetTypedData<T>(Property);
 			}
 			return Default;
 		}
 
-		bool HasData(size_t Key)
+		bool HasData(size_t Key, FDateTime StoredAfter = -1)
 		{
 			int64 IntKey = (int64)Key;
-
-			return DataStore.Contains(IntKey);
+			return DataStore.Contains(IntKey) && DataStore[Key]->Timestamp >= StoredAfter;
 		}
 	};
 
