Guard against Levels attribute not being present on geometry collection in various fracture tools/utilities

diff --git a/Engine/Plugins/Experimental/ChaosEditor/Source/FractureEditor/Private/FractureToolAutoCluster.cpp b/Engine/Plugins/Experimental/ChaosEditor/Source/FractureEditor/Private/FractureToolAutoCluster.cpp
--- a/Engine/Plugins/Experimental/ChaosEditor/Source/FractureEditor/Private/FractureToolAutoCluster.cpp
+++ b/Engine/Plugins/Experimental/ChaosEditor/Source/FractureEditor/Private/FractureToolAutoCluster.cpp
@@ -288,6 +288,10 @@ void FVoronoiPartitioner::GenerateConnectivity(const FGeometryCollection* Geomet
 	
 	Connectivity.SetNum(TransformIndices.Num());
 
+	if (!ensure(GeometryCollection->HasAttribute("Level", FGeometryCollection::TransformGroup)))
+	{
+		return;
+	}
 	const TManagedArray<int32>& Levels = GeometryCollection->GetAttribute<int32>("Level", FGeometryCollection::TransformGroup);
 
 	for (int32 Index = 0; Index < TransformIndices.Num(); ++Index)
diff --git a/Engine/Plugins/Experimental/ChaosEditor/Source/FractureEditor/Private/FractureToolClusterMagnet.cpp b/Engine/Plugins/Experimental/ChaosEditor/Source/FractureEditor/Private/FractureToolClusterMagnet.cpp
--- a/Engine/Plugins/Experimental/ChaosEditor/Source/FractureEditor/Private/FractureToolClusterMagnet.cpp
+++ b/Engine/Plugins/Experimental/ChaosEditor/Source/FractureEditor/Private/FractureToolClusterMagnet.cpp
@@ -86,6 +86,10 @@ void UFractureToolClusterMagnet::Execute(TWeakPtr<FFractureEditorModeToolkit> In
 			FGeometryCollectionEdit Edit(Context.GetGeometryCollectionComponent(), GeometryCollection::EEditUpdate::RestPhysicsDynamic);
 
 			const TManagedArray<TSet<int32>>& Children = Context.GetGeometryCollection()->Children;
+			if (!Context.GetGeometryCollection()->HasAttribute("Level", FTransformCollection::TransformGroup))
+			{
+				FGeometryCollectionClusteringUtility::UpdateHierarchyLevelOfChildren(Edit.GetRestCollection()->GetGeometryCollection().Get(), -1);
+			}
 			const TManagedArray<int32>& Levels = Context.GetGeometryCollection()->GetAttribute<int32>("Level", FGeometryCollection::TransformGroup);
 			int32 StartTransformCount = Children.Num();
 			
diff --git a/Engine/Plugins/Experimental/ChaosEditor/Source/FractureEditor/Private/FractureToolClustering.cpp b/Engine/Plugins/Experimental/ChaosEditor/Source/FractureEditor/Private/FractureToolClustering.cpp
--- a/Engine/Plugins/Experimental/ChaosEditor/Source/FractureEditor/Private/FractureToolClustering.cpp
+++ b/Engine/Plugins/Experimental/ChaosEditor/Source/FractureEditor/Private/FractureToolClustering.cpp
@@ -46,6 +46,10 @@ void UFractureToolFlattenAll::Execute(TWeakPtr<FFractureEditorModeToolkit> InToo
 		{
 			FGeometryCollectionEdit Edit(Context.GetGeometryCollectionComponent(), GeometryCollection::EEditUpdate::RestPhysicsDynamic);
 
+			if (!Context.GetGeometryCollection()->HasAttribute("Level", FGeometryCollection::TransformGroup))
+			{
+				FGeometryCollectionClusteringUtility::UpdateHierarchyLevelOfChildren(Edit.GetRestCollection()->GetGeometryCollection().Get(), -1);
+			}
 			const TManagedArray<int32>& Levels = Context.GetGeometryCollection()->GetAttribute<int32>("Level", FGeometryCollection::TransformGroup);
 
 			Context.ConvertSelectionToClusterNodes();
@@ -188,6 +192,10 @@ void UFractureToolUncluster::Execute(TWeakPtr<FFractureEditorModeToolkit> InTool
 			FGeometryCollectionEdit Edit(Context.GetGeometryCollectionComponent(), GeometryCollection::EEditUpdate::RestPhysicsDynamic);
 
 			const TManagedArray<TSet<int32>>& Children = Context.GetGeometryCollection()->GetAttribute<TSet<int32>>("Children", FGeometryCollection::TransformGroup);
+			if (!Context.GetGeometryCollection()->HasAttribute("Level", FGeometryCollection::TransformGroup))
+			{
+				FGeometryCollectionClusteringUtility::UpdateHierarchyLevelOfChildren(Edit.GetRestCollection()->GetGeometryCollection().Get(), -1);
+			}
 			const TManagedArray<int32>& Levels = Context.GetGeometryCollection()->GetAttribute<int32>("Level", FGeometryCollection::TransformGroup);
 
 			Context.ConvertSelectionToClusterNodes();
diff --git a/Engine/Plugins/Experimental/ChaosEditor/Source/FractureEditor/Private/SGeometryCollectionOutliner.cpp b/Engine/Plugins/Experimental/ChaosEditor/Source/FractureEditor/Private/SGeometryCollectionOutliner.cpp
--- a/Engine/Plugins/Experimental/ChaosEditor/Source/FractureEditor/Private/SGeometryCollectionOutliner.cpp
+++ b/Engine/Plugins/Experimental/ChaosEditor/Source/FractureEditor/Private/SGeometryCollectionOutliner.cpp
@@ -760,7 +760,7 @@ bool FGeometryCollectionTreeItemComponent::FilterBoneIndex(int32 BoneIndex) cons
 			// We only display cluster nodes deeper than the view level.
 			UFractureSettings* FractureSettings = GetMutableDefault<UFractureSettings>();
 
-			if (FractureSettings->FractureLevel >= 0)
+			if (FractureSettings->FractureLevel >= 0 && Collection->HasAttribute("Level", FTransformCollection::TransformGroup))
 			{
 				const TManagedArray<int32>& Level = Collection->GetAttribute<int32>("Level", FTransformCollection::TransformGroup);
 				int32 BoneLevel = Level[BoneIndex];
diff --git a/Engine/Source/Runtime/Experimental/Chaos/Private/GeometryCollection/GeometryCollectionAlgo.cpp b/Engine/Source/Runtime/Experimental/Chaos/Private/GeometryCollection/GeometryCollectionAlgo.cpp
--- a/Engine/Source/Runtime/Experimental/Chaos/Private/GeometryCollection/GeometryCollectionAlgo.cpp
+++ b/Engine/Source/Runtime/Experimental/Chaos/Private/GeometryCollection/GeometryCollectionAlgo.cpp
@@ -512,6 +512,10 @@ namespace GeometryCollectionAlgo
 		if (Level > 0)
 		{
 			const TManagedArray<int32>& Parents = GeometryCollection->Parent;
+			if (!ensure(GeometryCollection->HasAttribute("Level", FGeometryCollection::TransformGroup)))
+			{
+				return;
+			}
 			const TManagedArray<int32>& Levels = GeometryCollection->GetAttribute<int32>("Level", FGeometryCollection::TransformGroup);
 
 			TArray<FTransform> Transforms;
diff --git a/Engine/Source/Runtime/Experimental/Chaos/Private/GeometryCollection/GeometryCollectionClusteringUtility.cpp b/Engine/Source/Runtime/Experimental/Chaos/Private/GeometryCollection/GeometryCollectionClusteringUtility.cpp
--- a/Engine/Source/Runtime/Experimental/Chaos/Private/GeometryCollection/GeometryCollectionClusteringUtility.cpp
+++ b/Engine/Source/Runtime/Experimental/Chaos/Private/GeometryCollection/GeometryCollectionClusteringUtility.cpp
@@ -370,6 +370,10 @@ void FGeometryCollectionClusteringUtility::RenameBone(FGeometryCollection* Geome
 
 int32 FGeometryCollectionClusteringUtility::PickBestNodeToMergeTo(const FGeometryCollection* GeometryCollection, const TArray<int32>& SourceElements)
 {
+	if (!ensure(GeometryCollection->HasAttribute("Level", FGeometryCollection::TransformGroup)))
+	{
+		return -1;
+	}
 	const TManagedArray<int32>& Levels = GeometryCollection->GetAttribute<int32>("Level", FGeometryCollection::TransformGroup);
 	const TManagedArray<TSet<int32>>& Children = GeometryCollection->Children;
 
@@ -466,6 +470,10 @@ void FGeometryCollectionClusteringUtility::GetClusteredBonesWithCommonParent(con
 void FGeometryCollectionClusteringUtility::GetChildBonesFromLevel(const FGeometryCollection* GeometryCollection, int32 SourceBone, int32 Level, TArray<int32>& BonesOut)
 {
 	check(GeometryCollection);
+	if (!ensure(GeometryCollection->HasAttribute("Level", FGeometryCollection::TransformGroup)))
+	{
+		return;
+	}
 	const TManagedArray<int32>& Levels = GeometryCollection->GetAttribute<int32>("Level", FGeometryCollection::TransformGroup);
 	const TManagedArray<int32>& Parents = GeometryCollection->Parent;
 	const TManagedArray<TSet<int32>>& Children = GeometryCollection->Children;
@@ -489,7 +497,11 @@ void FGeometryCollectionClusteringUtility::GetChildBonesFromLevel(const FGeometr
 void FGeometryCollectionClusteringUtility::GetBonesToLevel(const FGeometryCollection* GeometryCollection, int32 Level, TArray<int32>& BonesOut, bool bOnlyClusteredOrRigid, bool bSkipFiltered)
 {
 	check(GeometryCollection);
-
+	
+	if (!ensure(GeometryCollection->HasAttribute("Level", FGeometryCollection::TransformGroup)))
+	{
+		return;
+	}
 	const TManagedArray<int32>& Levels = GeometryCollection->GetAttribute<int32>("Level", FGeometryCollection::TransformGroup);
 	const TManagedArray<int32>& SimType = GeometryCollection->SimulationType;
 
@@ -526,6 +538,10 @@ void FGeometryCollectionClusteringUtility::GetChildBonesAtLevel(const FGeometryC
 	}
 	else
 	{
+		if (!ensure(GeometryCollection->HasAttribute("Level", FGeometryCollection::TransformGroup)))
+		{
+			return;
+		}
 		const TManagedArray<int32>& Levels = GeometryCollection->GetAttribute<int32>("Level", FGeometryCollection::TransformGroup);
 		const TManagedArray<TSet<int32>>& Children = GeometryCollection->Children;
 		if (Levels[SourceBone] == Level)
@@ -556,6 +572,10 @@ int32 FGeometryCollectionClusteringUtility::GetParentOfBoneAtSpecifiedLevel(cons
 {
 	check(GeometryCollection);
 	const TManagedArray<int32>& Parents = GeometryCollection->Parent;
+	if (!ensure(GeometryCollection->HasAttribute("Level", FGeometryCollection::TransformGroup)))
+	{
+		return -1;
+	}
 	const TManagedArray<int32>& Levels = GeometryCollection->GetAttribute<int32>("Level", FGeometryCollection::TransformGroup);
 	const TManagedArray<int32>& SimTypes = GeometryCollection->SimulationType;
 
@@ -672,6 +692,10 @@ void FGeometryCollectionClusteringUtility::CollapseLevelHierarchy(int8 Level, FG
 {
 	check(GeometryCollection);
 
+	if (!GeometryCollection->HasAttribute("Level", FGeometryCollection::TransformGroup))
+	{
+		UpdateHierarchyLevelOfChildren(GeometryCollection, -1);
+	}
 	const TManagedArray<int32>& Levels = GeometryCollection->GetAttribute<int32>("Level", FGeometryCollection::TransformGroup);
 
 	TArray<int32> Elements;
@@ -713,6 +737,10 @@ void FGeometryCollectionClusteringUtility::CollapseLevelHierarchy(int8 Level, FG
 void FGeometryCollectionClusteringUtility::CollapseSelectedHierarchy(int8 Level, const TArray<int32>& SelectedBones, FGeometryCollection* GeometryCollection)
 {
 	check(GeometryCollection);
+	if (!GeometryCollection->HasAttribute("Level", FGeometryCollection::TransformGroup))
+	{
+		UpdateHierarchyLevelOfChildren(GeometryCollection, -1);
+	}
 	const TManagedArray<int32>& Levels = GeometryCollection->GetAttribute<int32>("Level", FGeometryCollection::TransformGroup);
 	const TManagedArray<TSet<int32>>& Children = GeometryCollection->Children;
 
diff --git a/Engine/Source/Runtime/Experimental/GeometryCollectionEngine/Private/GeometryCollection/GeometryCollectionComponent.cpp b/Engine/Source/Runtime/Experimental/GeometryCollectionEngine/Private/GeometryCollection/GeometryCollectionComponent.cpp
--- a/Engine/Source/Runtime/Experimental/GeometryCollectionEngine/Private/GeometryCollection/GeometryCollectionComponent.cpp
+++ b/Engine/Source/Runtime/Experimental/GeometryCollectionEngine/Private/GeometryCollection/GeometryCollectionComponent.cpp
@@ -2754,7 +2754,7 @@ int32 FScopedColorEdit::GetMaxSelectedLevel(bool bOnlyRigid) const
 {
 	int32 MaxSelectedLevel = -1;
 	const UGeometryCollection* GeometryCollection = Component->GetRestCollection();
-	if (GeometryCollection)
+	if (GeometryCollection && GeometryCollection->GetGeometryCollection()->HasAttribute("Level", FGeometryCollection::TransformGroup))
 	{
 		TSharedPtr<FGeometryCollection, ESPMode::ThreadSafe> GeometryCollectionPtr = GeometryCollection->GetGeometryCollection();
 		const TManagedArray<int32>& Levels = GeometryCollectionPtr->GetAttribute<int32>("Level", FGeometryCollection::TransformGroup);
@@ -2777,7 +2777,7 @@ bool FScopedColorEdit::IsSelectionValidAtLevel(int32 TargetLevel) const
 		return true;
 	}
 	const UGeometryCollection* GeometryCollection = Component->GetRestCollection();
-	if (GeometryCollection)
+	if (GeometryCollection && GeometryCollection->GetGeometryCollection()->HasAttribute("Level", FGeometryCollection::TransformGroup))
 	{
 		TSharedPtr<FGeometryCollection, ESPMode::ThreadSafe> GeometryCollectionPtr = GeometryCollection->GetGeometryCollection();
 		const TManagedArray<int32>& Levels = GeometryCollectionPtr->GetAttribute<int32>("Level", FGeometryCollection::TransformGroup);
@@ -2811,7 +2811,7 @@ void FScopedColorEdit::FilterSelectionToLevel(bool bPreferLowestOnly)
 	const UGeometryCollection* GeometryCollection = Component->GetRestCollection();
 	int32 ViewLevel = GetViewLevel();
 	bool bNeedsFiltering = ViewLevel >= 0 || bPreferLowestOnly;
-	if (GeometryCollection && Component->SelectedBones.Num() > 0 && bNeedsFiltering)
+	if (GeometryCollection && Component->SelectedBones.Num() > 0 && bNeedsFiltering && GeometryCollection->GetGeometryCollection()->HasAttribute("Level", FGeometryCollection::TransformGroup))
 	{
 		TSharedPtr<FGeometryCollection, ESPMode::ThreadSafe> GeometryCollectionPtr = GeometryCollection->GetGeometryCollection();
 
@@ -2887,11 +2887,11 @@ void FScopedColorEdit::SelectBones(GeometryCollection::ESelectionMode SelectionM
 			TArray<int32> BonesToSelect;
 			FGeometryCollectionClusteringUtility::GetBonesToLevel(GeometryCollectionPtr.Get(), GetViewLevel(), BonesToSelect, true, true);
 			const TManagedArray<int32>& SimType = GeometryCollectionPtr->SimulationType;
-			const TManagedArray<int32>& Levels = GeometryCollectionPtr->GetAttribute<int32>("Level", FGeometryCollection::TransformGroup);
+			const TManagedArray<int32>* Levels = GeometryCollectionPtr->FindAttributeTyped<int32>("Level", FGeometryCollection::TransformGroup);
 			BonesToSelect.SetNum(Algo::RemoveIf(BonesToSelect, [&](int32 BoneIdx)
 				{
 					return SimType[BoneIdx] != FGeometryCollection::ESimulationTypes::FST_Rigid
-						|| (ViewLevel != -1 && Levels[BoneIdx] != ViewLevel);
+						|| (ViewLevel != -1 && Levels && (*Levels)[BoneIdx] != ViewLevel);
 				}));
 			AppendSelectedBones(BonesToSelect);
 		}
@@ -2904,11 +2904,11 @@ void FScopedColorEdit::SelectBones(GeometryCollection::ESelectionMode SelectionM
 			TArray<int32> BonesToSelect;
 			FGeometryCollectionClusteringUtility::GetBonesToLevel(GeometryCollectionPtr.Get(), ViewLevel, BonesToSelect, true, true);
 			const TManagedArray<int32>& SimType = GeometryCollectionPtr->SimulationType;
-			const TManagedArray<int32>& Levels = GeometryCollectionPtr->GetAttribute<int32>("Level", FGeometryCollection::TransformGroup);
+			const TManagedArray<int32>* Levels = GeometryCollectionPtr->FindAttributeTyped<int32>("Level", FGeometryCollection::TransformGroup);
 			BonesToSelect.SetNum(Algo::RemoveIf(BonesToSelect, [&](int32 BoneIdx)
 			{
 				return SimType[BoneIdx] != FGeometryCollection::ESimulationTypes::FST_Clustered
-					|| (ViewLevel != -1 && Levels[BoneIdx] != ViewLevel);
+					|| (ViewLevel != -1 && Levels && (*Levels)[BoneIdx] != ViewLevel);
 			}));
 			AppendSelectedBones(BonesToSelect);
 		}
