Chaos Modular Vehicles
* Reduce spurious logging
* Fix force direction for locally transformed components
*Added debug color for diffeent categories of force debug render
* Wheels now know they are not touching the ground
* Thruster now has force axis/direction independent of component transform

diff --git a/Engine/Source/Runtime/Experimental/ChaosVehicles/ChaosVehiclesCore/Private/SimModule/AerofoilModule.cpp b/Engine/Source/Runtime/Experimental/ChaosVehicles/ChaosVehiclesCore/Private/SimModule/AerofoilModule.cpp
--- a/Engine/Source/Runtime/Experimental/ChaosVehicles/ChaosVehiclesCore/Private/SimModule/AerofoilModule.cpp
+++ b/Engine/Source/Runtime/Experimental/ChaosVehicles/ChaosVehiclesCore/Private/SimModule/AerofoilModule.cpp
@@ -34,7 +34,7 @@ namespace Chaos
 		FVector LocalForce = GetForce(ModuleLocalVelocity * CmToMScaling(), CmToM(Altitude), DeltaTime);
 		FVector LocalOffset = GetCenterOfLiftOffset() * MToCmScaling();
 
-		AddLocalForceAtPosition(LocalForce * MToCmScaling(), LocalOffset);
+		AddLocalForceAtPosition(LocalForce * MToCmScaling(), LocalOffset, true, false, false, FColor::Yellow);
 	}
 
 	float FAerofoilSimModule::CalcAngleOfAttackDegrees(const FVector& UpAxis, const FVector& InAirflowVector)
diff --git a/Engine/Source/Runtime/Experimental/ChaosVehicles/ChaosVehiclesCore/Private/SimModule/DeferredForcesModular.cpp b/Engine/Source/Runtime/Experimental/ChaosVehicles/ChaosVehiclesCore/Private/SimModule/DeferredForcesModular.cpp
--- a/Engine/Source/Runtime/Experimental/ChaosVehicles/ChaosVehiclesCore/Private/SimModule/DeferredForcesModular.cpp
+++ b/Engine/Source/Runtime/Experimental/ChaosVehicles/ChaosVehiclesCore/Private/SimModule/DeferredForcesModular.cpp
@@ -39,7 +39,9 @@ Chaos::FPBDRigidClusteredParticleHandle* FDeferredForcesModular::GetParticle(TAr
 	}
 	else if (SimClusterParticle && !SimClusterParticle->Disabled())
 	{
-		TransformOut = CollectionMassToLocal[Parent[TransformIndex]].Inverse() * Transforms[TransformIndex];
+		TransformOut.SetLocation(CollectionMassToLocal[Parent[TransformIndex]].InverseTransformPosition(Transforms[TransformIndex].GetLocation()));
+		TransformOut.SetRotation(Transforms[TransformIndex].GetRotation());
+
 		return SimClusterParticle; // apply to intact cluster
 	}
 
@@ -86,15 +88,20 @@ void FDeferredForcesModular::Apply(TArray<Chaos::FPBDRigidClusteredParticleHandl
 }
 
 
-void AddForce_Implementation(Chaos::FPBDRigidClusteredParticleHandle* RigidHandle, const FVector& LocalForce, const FTransform& OffsetTransform)
+void AddForce_Implementation(Chaos::FPBDRigidClusteredParticleHandle* RigidHandle, const FVector& LocalForce, const FTransform& OffsetTransform, bool bLevelSlope, const FColor& DebugColor)
 {
 	const Chaos::FVec3 WorldCOM = Chaos::FParticleUtilitiesGT::GetCoMWorldPosition(RigidHandle);
 
 	// local to world
 	const FTransform WorldTM(RigidHandle->R(), RigidHandle->X());
 	Chaos::FVec3 Position = WorldTM.TransformPosition(OffsetTransform.GetLocation());
+	Chaos::FVec3 Force = WorldTM.TransformVector(OffsetTransform.TransformVector(LocalForce));
 
-	Chaos::FVec3 Force = WorldTM.TransformVector(LocalForce);
+	if (bLevelSlope && WorldTM.GetUnitAxis(EAxis::Z).Z > GCoreModularVehicleDebugParams.LevelSlopeThreshold)
+	{
+		Force.X = 0.0f;
+		Force.Y = 0.0f;
+	}
 
 	Chaos::FVec3 Arm = Position - WorldCOM;
 
@@ -102,6 +109,7 @@ void AddForce_Implementation(Chaos::FPBDRigidClusteredParticleHandle* RigidHandl
 	auto* RP = RigidHandle->CastToRigidParticle();
 	if (RP)
 	{
+
 		RP->AddForce(Force);
 		RP->AddTorque(WorldTorque);
 
@@ -121,7 +129,7 @@ void AddForce_Implementation(Chaos::FPBDRigidClusteredParticleHandle* RigidHandl
 	if (GCoreModularVehicleDebugParams.ShowForces)
 	{
 		// TODO: perhaps render different force types/sources in different colors, i.e. aerofoil, suspension, friction, etc.
-		Chaos::FDebugDrawQueue::GetInstance().DrawDebugLine(Position, Position + Force * GCoreModularVehicleDebugParams.DrawForceScaling, FColor::Blue, false, -1.f, 0, 5.f);
+		Chaos::FDebugDrawQueue::GetInstance().DrawDebugLine(Position, Position + Force * GCoreModularVehicleDebugParams.DrawForceScaling, DebugColor, false, -1.f, 0, 5.f);
 		Chaos::FDebugDrawQueue::GetInstance().DrawDebugSphere(Position, 5, 8, FColor::White, false, -1.f, 0, 5.f);
 	}
 #endif
@@ -132,7 +140,7 @@ void FDeferredForcesModular::AddForce(Chaos::FPBDRigidClusteredParticleHandle* R
 {
 	if (ensure(RigidHandle))
 	{
-		AddForce_Implementation(RigidHandle, DataIn.Force, OffsetTransform);
+		AddForce_Implementation(RigidHandle, DataIn.Force, OffsetTransform, (DataIn.Flags & EForceFlags::LevelSlope) == EForceFlags::LevelSlope, DataIn.DebugColor);
 	}
 
 }
@@ -166,7 +174,7 @@ void FDeferredForcesModular::AddForceAtPosition(Chaos::FPBDRigidClusteredParticl
 	{
 		FTransform ShiftedTransform = OffsetTransform;
 		ShiftedTransform.SetLocation(ShiftedTransform.GetLocation() + DataIn.Position);
-		AddForce_Implementation(RigidHandle, DataIn.Force, ShiftedTransform);
+		AddForce_Implementation(RigidHandle, DataIn.Force, ShiftedTransform, (DataIn.Flags & EForceFlags::LevelSlope) == EForceFlags::LevelSlope, DataIn.DebugColor);
 	}
 }
 
@@ -189,3 +197,4 @@ void FDeferredForcesModular::AddTorque(Chaos::FPBDRigidClusteredParticleHandle*
 	//	}
 	//}
 }
+
diff --git a/Engine/Source/Runtime/Experimental/ChaosVehicles/ChaosVehiclesCore/Private/SimModule/SimModuleTree.cpp b/Engine/Source/Runtime/Experimental/ChaosVehicles/ChaosVehiclesCore/Private/SimModule/SimModuleTree.cpp
--- a/Engine/Source/Runtime/Experimental/ChaosVehicles/ChaosVehiclesCore/Private/SimModule/SimModuleTree.cpp
+++ b/Engine/Source/Runtime/Experimental/ChaosVehicles/ChaosVehiclesCore/Private/SimModule/SimModuleTree.cpp
@@ -227,11 +227,12 @@ void FSimModuleTree::UpdateModuleVelocites(FGeometryCollectionPhysicsProxy* Phys
 
 				if (Particle)
 				{
-					FVector WheelWorldLocation = BodyTransform.TransformPosition(Module->GetParentRelativeTransform().GetLocation());
-					const Chaos::FVec3 Arm = WheelWorldLocation - Particle->X();
+					FVector WorldLocation = BodyTransform.TransformPosition(Module->GetParentRelativeTransform().GetLocation());
+					const Chaos::FVec3 Arm = WorldLocation - Particle->X();
 
 					FVector WorldVelocity = Particle->V() - Chaos::FVec3::CrossProduct(Arm, Particle->W());
 					FVector LocalVelocity = BodyTransform.InverseTransformVector(WorldVelocity);
+					LocalVelocity = Module->GetClusteredTransform().InverseTransformVector(LocalVelocity);
 
 					Module->SetLocalVelocity(LocalVelocity);
 				}
diff --git a/Engine/Source/Runtime/Experimental/ChaosVehicles/ChaosVehiclesCore/Private/SimModule/SimulationModuleBase.cpp b/Engine/Source/Runtime/Experimental/ChaosVehicles/ChaosVehiclesCore/Private/SimModule/SimulationModuleBase.cpp
--- a/Engine/Source/Runtime/Experimental/ChaosVehicles/ChaosVehiclesCore/Private/SimModule/SimulationModuleBase.cpp
+++ b/Engine/Source/Runtime/Experimental/ChaosVehicles/ChaosVehiclesCore/Private/SimModule/SimulationModuleBase.cpp
@@ -11,21 +11,21 @@ FCoreModularVehicleDebugParams GModularVehicleDebugParams;
 namespace Chaos
 {
 
-void ISimulationModuleBase::AddLocalForceAtPosition(const FVector& Force, const FVector& Position, bool bAllowSubstepping, bool bIsLocalForce)
+void ISimulationModuleBase::AddLocalForceAtPosition(const FVector& Force, const FVector& Position, bool bAllowSubstepping, bool bIsLocalForce, bool bLevelSlope, const FColor& DebugColorIn)
 {
 	AppliedForce = Force;
 	if (SimModuleTree)
 	{
-		SimModuleTree->AccessDeferredForces().Add(FDeferredForcesModular::FApplyForceAtPositionData(TransformIndex, Force, Position, bAllowSubstepping, bIsLocalForce));
+		SimModuleTree->AccessDeferredForces().Add(FDeferredForcesModular::FApplyForceAtPositionData(TransformIndex, Force, Position, bAllowSubstepping, bIsLocalForce, bLevelSlope, DebugColorIn));
 	}
 }
 
-void ISimulationModuleBase::AddLocalForce(const FVector& Force, bool bAllowSubstepping, bool bIsLocalForce)
+void ISimulationModuleBase::AddLocalForce(const FVector& Force, bool bAllowSubstepping, bool bIsLocalForce, bool bLevelSlope, const FColor& DebugColorIn)
 {
 	AppliedForce = Force;
 	if (SimModuleTree)
 	{
-		SimModuleTree->AccessDeferredForces().Add(FDeferredForcesModular::FApplyForceData(TransformIndex, Force, bAllowSubstepping, bIsLocalForce));
+		SimModuleTree->AccessDeferredForces().Add(FDeferredForcesModular::FApplyForceData(TransformIndex, Force, bAllowSubstepping, bIsLocalForce, bLevelSlope, DebugColorIn));
 	}
 }
 
diff --git a/Engine/Source/Runtime/Experimental/ChaosVehicles/ChaosVehiclesCore/Private/SimModule/SuspensionModule.cpp b/Engine/Source/Runtime/Experimental/ChaosVehicles/ChaosVehiclesCore/Private/SimModule/SuspensionModule.cpp
--- a/Engine/Source/Runtime/Experimental/ChaosVehicles/ChaosVehiclesCore/Private/SimModule/SuspensionModule.cpp
+++ b/Engine/Source/Runtime/Experimental/ChaosVehicles/ChaosVehiclesCore/Private/SimModule/SuspensionModule.cpp
@@ -60,7 +60,7 @@ namespace Chaos
 			if (SuspensionForce > 0)
 			{
 				ForceIntoSurface = SuspensionForce;
-				AddLocalForce(Setup().SuspensionAxis * -SuspensionForce/*, Setup().LocalOffset*/);
+				AddLocalForce(Setup().SuspensionAxis * -SuspensionForce, true, false, true, FColor::Green);
 			}
 		}
 
diff --git a/Engine/Source/Runtime/Experimental/ChaosVehicles/ChaosVehiclesCore/Private/SimModule/ThrusterModule.cpp b/Engine/Source/Runtime/Experimental/ChaosVehicles/ChaosVehiclesCore/Private/SimModule/ThrusterModule.cpp
--- a/Engine/Source/Runtime/Experimental/ChaosVehicles/ChaosVehiclesCore/Private/SimModule/ThrusterModule.cpp
+++ b/Engine/Source/Runtime/Experimental/ChaosVehicles/ChaosVehiclesCore/Private/SimModule/ThrusterModule.cpp
@@ -20,7 +20,8 @@ namespace Chaos
 	void FThrusterSimModule::Simulate(float DeltaTime, const FAllInputs& Inputs, FSimModuleTree& VehicleModuleSystem)
 	{
 		// applies continuous force
-		AddLocalForceAtPosition(FVector(Setup().MaxThrustForce * Inputs.Throttle, 0, 0), Setup().ForceOffset);
+		FVector Force = Setup().ForceAxis * Setup().MaxThrustForce * Inputs.Throttle;
+		AddLocalForceAtPosition(Force, Setup().ForceOffset, true, false, false, FColor::Magenta);
 
 	}
 
diff --git a/Engine/Source/Runtime/Experimental/ChaosVehicles/ChaosVehiclesCore/Private/SimModule/WheelModule.cpp b/Engine/Source/Runtime/Experimental/ChaosVehicles/ChaosVehiclesCore/Private/SimModule/WheelModule.cpp
--- a/Engine/Source/Runtime/Experimental/ChaosVehicles/ChaosVehiclesCore/Private/SimModule/WheelModule.cpp
+++ b/Engine/Source/Runtime/Experimental/ChaosVehicles/ChaosVehiclesCore/Private/SimModule/WheelModule.cpp
@@ -34,6 +34,7 @@ namespace Chaos
 		float HandbrakeTorque = Setup().HandbrakeEnabled ? Inputs.Handbrake * Setup().HandbrakeTorque : 0.0f;
 		SteerAngleDegrees = Setup().SteeringEnabled ? Inputs.Steering * Setup().MaxSteeringAngle : 0.0f;
 		BrakeTorque = Inputs.Brake * Setup().MaxBrakeTorque + HandbrakeTorque;
+		LoadTorque = 0.0f;
 
 		// TODO: think about doing this properly, stops vehicles rolling around on their own too much
 		// kindof an auto handbrake
@@ -42,120 +43,126 @@ namespace Chaos
 			BrakeTorque = Setup().HandbrakeTorque;
 		}
 
-		FRotator SteeringRotator(0.f, SteerAngleDegrees, 0.f);
-		FVector LocalWheelVelocity = SteeringRotator.UnrotateVector(ModuleLocalVelocity);
+		bool bTouchingGround = ForceIntoSurface > SMALL_NUMBER;
 
-		float GroundAngularVelocity = LocalWheelVelocity.X / Re;
-		float Delta = GroundAngularVelocity - AngularVelocity;
-		float Torque = Delta * Setup().WheelInertia / DeltaTime;
+		if (bTouchingGround)
+		{
+			FRotator SteeringRotator(0.f, SteerAngleDegrees, 0.f);
+			FVector LocalWheelVelocity = SteeringRotator.UnrotateVector(ModuleLocalVelocity);
 
-		// X is longitudinal direction, Y is lateral
-		float SlipAngle = FVehicleUtility::CalculateSlipAngle(LocalWheelVelocity.Y, LocalWheelVelocity.X);
+			float GroundAngularVelocity = LocalWheelVelocity.X / Re;
+			float Delta = GroundAngularVelocity - AngularVelocity;
+			float TorqueFromGroundInteraction = Delta * Setup().WheelInertia / DeltaTime; // torque from wheels moving over terrain
 
-		float AppliedLinearDriveForce = DriveTorque / Re;
-		float AppliedLinearBrakeForce = BrakeTorque / Re;
+			// X is longitudinal direction, Y is lateral
+			float SlipAngle = FVehicleUtility::CalculateSlipAngle(LocalWheelVelocity.Y, LocalWheelVelocity.X);
 
-		// Longitudinal multiplier now affecting both brake and steering equally
-		float AvailableGrip = ForceIntoSurface * SurfaceFriction * Setup().FrictionMultiplier;
+			float AppliedLinearDriveForce = DriveTorque / Re;
+			float AppliedLinearBrakeForce = BrakeTorque / Re;
 
-		float FinalLongitudinalForce = 0.f;
-		float FinalLateralForce = 0.f;
+			// Longitudinal multiplier now affecting both brake and steering equally
+			float AvailableGrip = ForceIntoSurface * SurfaceFriction * Setup().FrictionMultiplier;
 
-		// currently just letting the brake override the throttle
-		bool Braking = (BrakeTorque > FMath::Abs(DriveTorque));
-		bool WheelLocked = false;
+			float FinalLongitudinalForce = 0.f;
+			float FinalLateralForce = 0.f;
 
-		// are we actually touching the ground
-		if (ForceIntoSurface > SMALL_NUMBER)
-		{
-			// ABS limiting brake force to match force from the grip available
-			if (Setup().ABSEnabled && Braking && FMath::Abs(AppliedLinearBrakeForce) > AvailableGrip)
+			// currently just letting the brake override the throttle
+			bool Braking = (BrakeTorque > FMath::Abs(DriveTorque));
+			bool WheelLocked = false;
+
+			// are we actually touching the ground
+			if (ForceIntoSurface > SMALL_NUMBER)
 			{
-				if ((Braking && Setup().ABSEnabled) || (!Braking && Setup().TractionControlEnabled))
+				// ABS limiting brake force to match force from the grip available
+				if (Setup().ABSEnabled && Braking && FMath::Abs(AppliedLinearBrakeForce) > AvailableGrip)
 				{
-					float Sign = (AppliedLinearBrakeForce > 0.0f) ? 1.0f : -1.0f;
-					AppliedLinearBrakeForce = AvailableGrip * TractionControlAndABSScaling * Sign;
+					if ((Braking && Setup().ABSEnabled) || (!Braking && Setup().TractionControlEnabled))
+					{
+						float Sign = (AppliedLinearBrakeForce > 0.0f) ? 1.0f : -1.0f;
+						AppliedLinearBrakeForce = AvailableGrip * TractionControlAndABSScaling * Sign;
+					}
 				}
-			}
 
-			// Traction control limiting drive force to match force from grip available
-			if (Setup().TractionControlEnabled && !Braking && FMath::Abs(AppliedLinearDriveForce) > AvailableGrip)
-			{
-				float Sign = (AppliedLinearDriveForce > 0.0f) ? 1.0f : -1.0f;
-				AppliedLinearDriveForce = AvailableGrip * TractionControlAndABSScaling * Sign;
-			}
+				// Traction control limiting drive force to match force from grip available
+				if (Setup().TractionControlEnabled && !Braking && FMath::Abs(AppliedLinearDriveForce) > AvailableGrip)
+				{
+					float Sign = (AppliedLinearDriveForce > 0.0f) ? 1.0f : -1.0f;
+					AppliedLinearDriveForce = AvailableGrip * TractionControlAndABSScaling * Sign;
+				}
 
-			if (Braking)
-			{
-				// whether the velocity is +ve or -ve when we brake we are slowing the vehicle down
-				// so force is opposing current direction of travel.
-				float ForceRequiredToBringToStop = MassPerWheel * K * (LocalWheelVelocity.X) / DeltaTime;
-				FinalLongitudinalForce = AppliedLinearBrakeForce;
+				if (Braking)
+				{
+					// whether the velocity is +ve or -ve when we brake we are slowing the vehicle down
+					// so force is opposing current direction of travel.
+					float ForceRequiredToBringToStop = MassPerWheel * K * (LocalWheelVelocity.X) / DeltaTime;
+					FinalLongitudinalForce = AppliedLinearBrakeForce;
 
-				// check we are not applying more force than required so we end up overshooting 
-				// and accelerating in the opposite direction
-				FinalLongitudinalForce = FMath::Clamp(FinalLongitudinalForce, -FMath::Abs(ForceRequiredToBringToStop), FMath::Abs(ForceRequiredToBringToStop));
+					// check we are not applying more force than required so we end up overshooting 
+					// and accelerating in the opposite direction
+					FinalLongitudinalForce = FMath::Clamp(FinalLongitudinalForce, -FMath::Abs(ForceRequiredToBringToStop), FMath::Abs(ForceRequiredToBringToStop));
 
-				// ensure the brake opposes current direction of travel
-				if (LocalWheelVelocity.X > 0.0f)
+					// ensure the brake opposes current direction of travel
+					if (LocalWheelVelocity.X > 0.0f)
+					{
+						FinalLongitudinalForce = -FinalLongitudinalForce;
+					}
+
+				}
+				else
 				{
-					FinalLongitudinalForce = -FinalLongitudinalForce;
+					FinalLongitudinalForce = AppliedLinearDriveForce;
+				}
+
+				float ForceRequiredToBringToStop = -(MassPerWheel * K * LocalWheelVelocity.Y) / DeltaTime;
+
+				// use slip angle to generate a sideways force
+				if (Setup().LateralSlipGraph.IsEmpty())
+				{
+					float AngleLimit = FMath::DegreesToRadians(Setup().SlipAngleLimit);
+					float ClippedAngle = FMath::Clamp(SlipAngle, -AngleLimit, AngleLimit);
+					FinalLateralForce = FMath::Abs(SlipAngle) * Setup().CorneringStiffness;
+				}
+				else
+				{
+					FinalLateralForce = Setup().LateralSlipGraph.EvaluateY(FMath::RadiansToDegrees(SlipAngle)) * Setup().LateralSlipGraphMultiplier;
+				}
+
+				if (FinalLateralForce > FMath::Abs(ForceRequiredToBringToStop))
+				{
+					FinalLateralForce = FMath::Abs(ForceRequiredToBringToStop);
+				}
+				if (LocalWheelVelocity.Y > 0.0f)
+				{
+					FinalLateralForce = -FinalLateralForce;
 				}
 
 			}
-			else
-			{
-				FinalLongitudinalForce = AppliedLinearDriveForce;
-			}
 
-			float ForceRequiredToBringToStop = -(MassPerWheel * K * LocalWheelVelocity.Y) / DeltaTime;
+			ForceFromFriction = FVector::ZeroVector;
+			ForceFromFriction.X = FinalLongitudinalForce;
+			ForceFromFriction.Y = FinalLateralForce;
+
+			AddLocalForce(SteeringRotator.RotateVector(ForceFromFriction));
+			TransmitTorque(VehicleModuleSystem, DriveTorque, BrakeTorque);
 
-			// use slip angle to generate a sideways force
-			if (Setup().LateralSlipGraph.IsEmpty())
+			DriveTorque -= AvailableGrip;
+			if (DriveTorque < 0.0f)
 			{
-				float AngleLimit = FMath::DegreesToRadians(Setup().SlipAngleLimit);
-				float ClippedAngle = FMath::Clamp(SlipAngle, -AngleLimit, AngleLimit);
-				FinalLateralForce = FMath::Abs(SlipAngle) * Setup().CorneringStiffness;
+				DriveTorque = 0.0f;
 			}
 			else
-			{
-				FinalLateralForce = Setup().LateralSlipGraph.EvaluateY(FMath::RadiansToDegrees(SlipAngle)) * Setup().LateralSlipGraphMultiplier;
+			{ 
+				DriveTorque *= TorqueScaling;
 			}
-
-			if (FinalLateralForce > FMath::Abs(ForceRequiredToBringToStop))
+			BrakingTorque -= AvailableGrip;
+			if (BrakingTorque < 0.0f)
 			{
-				FinalLateralForce = FMath::Abs(ForceRequiredToBringToStop);
+				BrakingTorque = 0.0f;
 			}
-			if (LocalWheelVelocity.Y > 0.0f)
-			{
-				FinalLateralForce = -FinalLateralForce;
-			}
-
-		}
-
-		ForceFromFriction = FVector::ZeroVector;
-		ForceFromFriction.X = FinalLongitudinalForce;
-		ForceFromFriction.Y = FinalLateralForce;
-
-		AddLocalForce(SteeringRotator.RotateVector(ForceFromFriction));
-		TransmitTorque(VehicleModuleSystem, DriveTorque, BrakeTorque);
 
-		DriveTorque -= AvailableGrip;
-		if (DriveTorque < 0.0f)
-		{
-			DriveTorque = 0.0f;
-		}
-		else
-		{ 
-			DriveTorque *= TorqueScaling;
-		}
-		BrakingTorque -= AvailableGrip;
-		if (BrakingTorque < 0.0f)
-		{
-			BrakingTorque = 0.0f;
+			LoadTorque = TorqueFromGroundInteraction;
 		}
 
-		LoadTorque = Torque;	// torque from wheels moving over terrain
 		IntegrateAngularVelocity(DeltaTime, Setup().WheelInertia, Setup().MaxRotationVel);
 	}
 
diff --git a/Engine/Source/Runtime/Experimental/ChaosVehicles/ChaosVehiclesCore/Public/SimModule/DeferredForcesModular.h b/Engine/Source/Runtime/Experimental/ChaosVehicles/ChaosVehiclesCore/Public/SimModule/DeferredForcesModular.h
--- a/Engine/Source/Runtime/Experimental/ChaosVehicles/ChaosVehiclesCore/Public/SimModule/DeferredForcesModular.h
+++ b/Engine/Source/Runtime/Experimental/ChaosVehicles/ChaosVehiclesCore/Public/SimModule/DeferredForcesModular.h
@@ -14,6 +14,7 @@ enum class EForceFlags : uint32
 	AccelChange = 1 << 1,
 	VelChange = 1 << 2,
 	IsLocalForce = 1 << 3,
+	LevelSlope = 1 << 4
 
 };
 ENUM_CLASS_FLAGS(EForceFlags);	
@@ -25,36 +26,42 @@ public:
 
 	struct FApplyForceData
 	{
-		FApplyForceData(int TransformIndexIn, const FVector& ForceIn, bool bAllowSubsteppingIn, bool bAccelChangeIn)
+		FApplyForceData(int TransformIndexIn, const FVector& ForceIn, bool bAllowSubsteppingIn, bool bAccelChangeIn, bool bLevelSlope, const FColor& ColorIn)
 			: TransformIndex(TransformIndexIn)
 			, Force(ForceIn)
 			, Flags(EForceFlags::None)
+			, DebugColor(ColorIn)
 		{
 			Flags |= bAllowSubsteppingIn ? EForceFlags::AllowSubstepping : EForceFlags::None;
 			Flags |= bAccelChangeIn ? EForceFlags::AccelChange : EForceFlags::None;
+			Flags |= bLevelSlope ? EForceFlags::LevelSlope : EForceFlags::None;
 		}
 
 		int TransformIndex;
 		FVector Force;
 		EForceFlags Flags;
+		FColor DebugColor;
 	};
 
 	struct FApplyForceAtPositionData
 	{
-		FApplyForceAtPositionData(int TransformIndexIn, const FVector& ForceIn, const FVector& PositionIn, bool bAllowSubsteppingIn, bool bIsLocalForceIn)
+		FApplyForceAtPositionData(int TransformIndexIn, const FVector& ForceIn, const FVector& PositionIn, bool bAllowSubsteppingIn, bool bIsLocalForceIn, bool bLevelSlope, const FColor& ColorIn)
 			: TransformIndex(TransformIndexIn)
 			, Force(ForceIn)
 			, Position(PositionIn)
 			, Flags(EForceFlags::None)
+			, DebugColor(ColorIn)
 		{
 			Flags |= bAllowSubsteppingIn ? EForceFlags::AllowSubstepping : EForceFlags::None;
 			Flags |= bIsLocalForceIn ? EForceFlags::IsLocalForce : EForceFlags::None;
+			Flags |= bLevelSlope ? EForceFlags::LevelSlope : EForceFlags::None;
 		}
 
 		int TransformIndex;
 		FVector Force;
 		FVector Position;
 		EForceFlags Flags;
+		FColor DebugColor;
 	};
 
 	struct FAddTorqueInRadiansData
diff --git a/Engine/Source/Runtime/Experimental/ChaosVehicles/ChaosVehiclesCore/Public/SimModule/SimulationModuleBase.h b/Engine/Source/Runtime/Experimental/ChaosVehicles/ChaosVehiclesCore/Public/SimModule/SimulationModuleBase.h
--- a/Engine/Source/Runtime/Experimental/ChaosVehicles/ChaosVehiclesCore/Public/SimModule/SimulationModuleBase.h
+++ b/Engine/Source/Runtime/Experimental/ChaosVehicles/ChaosVehiclesCore/Public/SimModule/SimulationModuleBase.h
@@ -6,13 +6,14 @@
 #include "Logging/LogMacros.h"
 
 
-DECLARE_LOG_CATEGORY_EXTERN(LogSimulationModule, Log, All);
+DECLARE_LOG_CATEGORY_EXTERN(LogSimulationModule, Warning, All);
 
 struct CHAOSVEHICLESCORE_API FCoreModularVehicleDebugParams
 {
 	bool ShowMass = false;
 	bool ShowForces = false;
 	float DrawForceScaling = 0.0005f;
+	float LevelSlopeThreshold = 0.96f; // ~16 degrees
 };
 
 namespace Chaos
@@ -206,13 +207,13 @@ namespace Chaos
 		 * Force application function, handles deferred force application and applying the force at the collect location based on whether the GC cluster is intact or fractured
 		 * Note: forces are applied in local coordinates of the module
 		 */
-		void AddLocalForceAtPosition(const FVector& Force, const FVector& Position, bool bAllowSubstepping = true, bool bIsLocalForce = false);
+		void AddLocalForceAtPosition(const FVector& Force, const FVector& Position, bool bAllowSubstepping = true, bool bIsLocalForce = false, bool bLevelSlope = false, const FColor& DebugColorIn = FColor::Blue);
 
 		/**
 		 * Force application function, handles deferred force application and applying the force at the collect location based on whether the GC cluster is intact or fractured
 		 * Note: forces are applied in local coordinates of the module
 		 */
-		void AddLocalForce(const FVector& Force, bool bAllowSubstepping = true, bool bIsLocalForce = false);
+		void AddLocalForce(const FVector& Force, bool bAllowSubstepping = true, bool bIsLocalForce = false, bool bLevelSlope = false, const FColor& DebugColorIn = FColor::Blue);
 
 		//---
 
diff --git a/Engine/Source/Runtime/Experimental/ChaosVehicles/ChaosVehiclesCore/Public/SimModule/ThrusterModule.h b/Engine/Source/Runtime/Experimental/ChaosVehicles/ChaosVehiclesCore/Public/SimModule/ThrusterModule.h
--- a/Engine/Source/Runtime/Experimental/ChaosVehicles/ChaosVehiclesCore/Public/SimModule/ThrusterModule.h
+++ b/Engine/Source/Runtime/Experimental/ChaosVehicles/ChaosVehiclesCore/Public/SimModule/ThrusterModule.h
@@ -17,12 +17,14 @@ namespace Chaos
 	{
 		FThrusterSettings()
 			: MaxThrustForce(0)
+			, ForceAxis(FVector(1.0f, 0.0f, 0.0f))
 			, ForceOffset(FVector::ZeroVector)
 		{
 
 		}
 
 		float MaxThrustForce;
+		FVector ForceAxis;
 		FVector ForceOffset;
 	};
 
