Chaos - fix island manager using invalid island index

diff --git a/Engine/Source/Runtime/Experimental/Chaos/Private/Chaos/Island/IslandGraph.cpp b/Engine/Source/Runtime/Experimental/Chaos/Private/Chaos/Island/IslandGraph.cpp
--- a/Engine/Source/Runtime/Experimental/Chaos/Private/Chaos/Island/IslandGraph.cpp
+++ b/Engine/Source/Runtime/Experimental/Chaos/Private/Chaos/Island/IslandGraph.cpp
@@ -180,7 +180,11 @@ void FIslandGraph<NodeType, EdgeType, IslandType, OwnerType>::RemoveNode(const N
 					if (GraphNode.NodeEdges.IsValidIndex(NodeEdgeIndex))
 					{
 						const int32 GraphEdgeIndex = GraphNode.NodeEdges[NodeEdgeIndex];
-						GraphIslands[GraphEdges[GraphEdgeIndex].IslandIndex].bIsPersistent = false;
+						const int32 GraphIslandIndex = GraphEdges[GraphEdgeIndex].IslandIndex;
+						if (GraphIslandIndex != INDEX_NONE)
+						{
+							GraphIslands[GraphIslandIndex].bIsPersistent = false;
+						}
 						RemoveEdge(GraphEdgeIndex);
 					}
 				}
@@ -884,7 +888,11 @@ void FIslandGraph<NodeType, EdgeType, IslandType, OwnerType>::UpdateGraph()
 					// @todo(chaos): move NodeIslands array management into FIslandGraph if possible
 					for (auto& EdgeIndex : GraphNode.NodeEdges)
 					{
-						GraphIslands[GraphEdges[EdgeIndex].IslandIndex].bIsSleeping = false;
+						const int32 EdgeIslandIndex = GraphEdges[EdgeIndex].IslandIndex;
+						if (EdgeIslandIndex != INDEX_NONE)
+						{
+							GraphIslands[EdgeIslandIndex].bIsSleeping = false;
+						}
 					}
 				}
 			}
