Chaos : fix call to AddDirtyProxy made on the wrong thread potentially causing crashes

diff --git a/Engine/Source/Runtime/Experimental/Chaos/Private/PhysicsProxy/GeometryCollectionPhysicsProxy.cpp b/Engine/Source/Runtime/Experimental/Chaos/Private/PhysicsProxy/GeometryCollectionPhysicsProxy.cpp
--- a/Engine/Source/Runtime/Experimental/Chaos/Private/PhysicsProxy/GeometryCollectionPhysicsProxy.cpp
+++ b/Engine/Source/Runtime/Experimental/Chaos/Private/PhysicsProxy/GeometryCollectionPhysicsProxy.cpp
@@ -1976,10 +1976,7 @@ void FGeometryCollectionPhysicsProxy::SetWorldTransform_External(const FTransfor
 
 	if (Chaos::FPhysicsSolver* RBDSolver = GetSolver<Chaos::FPhysicsSolver>())
 	{
-		RBDSolver->EnqueueCommandImmediate([this, RBDSolver]()
-		{
-			RBDSolver->AddDirtyProxy(this);
-		});
+		RBDSolver->AddDirtyProxy(this);
 	}
 }
 
