Chaos : geometry collection : now expose sockets from bone names (experimental )

diff --git a/Engine/Source/Runtime/Experimental/Chaos/Private/PhysicsProxy/GeometryCollectionPhysicsProxy.cpp b/Engine/Source/Runtime/Experimental/Chaos/Private/PhysicsProxy/GeometryCollectionPhysicsProxy.cpp
--- a/Engine/Source/Runtime/Experimental/Chaos/Private/PhysicsProxy/GeometryCollectionPhysicsProxy.cpp
+++ b/Engine/Source/Runtime/Experimental/Chaos/Private/PhysicsProxy/GeometryCollectionPhysicsProxy.cpp
@@ -2627,6 +2627,8 @@ bool FGeometryCollectionPhysicsProxy::PullFromPhysicsState(const Chaos::FDirtyGe
 		}
 	}
 
+	PostPhysicsSyncCallback();
+	
 	return true;
 }
 
diff --git a/Engine/Source/Runtime/Experimental/Chaos/Public/PhysicsProxy/GeometryCollectionPhysicsProxy.h b/Engine/Source/Runtime/Experimental/Chaos/Public/PhysicsProxy/GeometryCollectionPhysicsProxy.h
--- a/Engine/Source/Runtime/Experimental/Chaos/Public/PhysicsProxy/GeometryCollectionPhysicsProxy.h
+++ b/Engine/Source/Runtime/Experimental/Chaos/Public/PhysicsProxy/GeometryCollectionPhysicsProxy.h
@@ -396,6 +396,11 @@ public:
 	 * Uninitialized level array should be correct size and defaulted to zeros.
 	 */
 	static int32 CalculateAndSetLevel(int32 TransformGroupIdx, const TManagedArray<int32>& Parent, TManagedArray<int32>& Levels);
+
+	void SetPostPhysicsSyncCallback(TFunction<void()> Callback)
+	{
+		PostPhysicsSyncCallback = Callback;
+	}
 	
 protected:
 	/**
@@ -504,6 +509,9 @@ private:
 	FRecordedTransformTrack RecordedTracks;
 #endif
 
+	// called after we sync the physics thread data ( called on the game thread )
+	TFunction<void()> PostPhysicsSyncCallback;
+	
 	// Per object collision fraction.
 	float CollisionParticlesPerObjectFraction;
 
diff --git a/Engine/Source/Runtime/Experimental/GeometryCollectionEngine/Private/GeometryCollection/GeometryCollectionComponent.cpp b/Engine/Source/Runtime/Experimental/GeometryCollectionEngine/Private/GeometryCollection/GeometryCollectionComponent.cpp
--- a/Engine/Source/Runtime/Experimental/GeometryCollectionEngine/Private/GeometryCollection/GeometryCollectionComponent.cpp
+++ b/Engine/Source/Runtime/Experimental/GeometryCollectionEngine/Private/GeometryCollection/GeometryCollectionComponent.cpp
@@ -2062,6 +2062,100 @@ void UGeometryCollectionComponent::OnUpdateTransform(EUpdateTransformFlags Updat
 	}
 }
 
+bool UGeometryCollectionComponent::HasAnySockets() const
+{
+	if (RestCollection && RestCollection->GetGeometryCollection())
+	{
+		return (RestCollection->GetGeometryCollection()->BoneName.Num() > 0);
+	}
+	return false;
+}
+
+bool UGeometryCollectionComponent::DoesSocketExist(FName InSocketName) const
+{
+	if (RestCollection && RestCollection->GetGeometryCollection())
+	{
+		return RestCollection->GetGeometryCollection()->BoneName.Contains(InSocketName.ToString());
+	}
+	return false;
+}
+
+FTransform UGeometryCollectionComponent::GetSocketTransform(FName InSocketName, ERelativeTransformSpace TransformSpace) const
+{
+	if (RestCollection && RestCollection->GetGeometryCollection())
+	{
+		const TSharedPtr<FGeometryCollection, ESPMode::ThreadSafe> Collection = RestCollection->GetGeometryCollection();
+		const int32 TransformIndex = Collection->BoneName.Find(InSocketName.ToString());
+		if (TransformIndex != INDEX_NONE)
+		{
+			
+			if (GlobalMatrices.IsValidIndex(TransformIndex))
+			{
+				const FTransform BoneComponentSpaceTransform = FTransform(GlobalMatrices[TransformIndex]);
+				switch (TransformSpace)
+				{
+				case RTS_World:
+					return BoneComponentSpaceTransform * GetComponentTransform();
+
+				case RTS_Actor:
+					{
+						if (const AActor* Actor = GetOwner())
+						{
+							const FTransform SocketWorldSpaceTransform = BoneComponentSpaceTransform * GetComponentTransform();
+							return SocketWorldSpaceTransform.GetRelativeTransform(Actor->GetTransform());
+						}
+						break;
+					}
+
+				case RTS_Component:
+					return BoneComponentSpaceTransform;
+					
+				case RTS_ParentBoneSpace:
+					{
+						const int32 ParentTransformIndex = Collection->Parent[TransformIndex];
+						FTransform ParentComponentSpaceTransform{ FTransform::Identity };
+						if (GlobalMatrices.IsValidIndex(ParentTransformIndex))
+						{
+							ParentComponentSpaceTransform = FTransform(GlobalMatrices[ParentTransformIndex]);
+						}
+						return BoneComponentSpaceTransform.GetRelativeTransform(ParentComponentSpaceTransform);
+					}
+
+				default:
+					check(false);
+				}
+			}
+		}
+	}
+	return Super::GetSocketTransform(InSocketName, TransformSpace);
+}
+
+void UGeometryCollectionComponent::QuerySupportedSockets(TArray<FComponentSocketDescription>& OutSockets) const
+{
+	if (RestCollection && RestCollection->GetGeometryCollection())
+	{
+		for (const FString& BoneName: RestCollection->GetGeometryCollection()->BoneName)
+		{
+			FComponentSocketDescription& Desc = OutSockets.AddZeroed_GetRef();
+			Desc.Name = *BoneName;
+			Desc.Type = EComponentSocketType::Type::Bone;
+		}
+	}
+}
+
+void UGeometryCollectionComponent::UpdateAttachedChildrenTransform() const
+{
+	// todo(chaos) : find a way to only update that of transform have changed
+	// right now this does not work properly because the dirty flags may not be updated at the right time
+	//if (PhysicsProxy && PhysicsProxy->IsGTCollectionDirty())
+	{
+	 	for (const TObjectPtr<USceneComponent>& AttachedChild: this->GetAttachChildren())
+	 	{
+	 		AttachedChild->UpdateComponentToWorld();
+	 	}
+	}
+}
+
 void UGeometryCollectionComponent::TickComponent(float DeltaTime, enum ELevelTick TickType, FActorComponentTickFunction *ThisTickFunction)
 {
 	//UE_LOG(UGCC_LOG, Log, TEXT("GeometryCollectionComponent[%p]::TickComponent()"), this);
@@ -2093,6 +2187,16 @@ void UGeometryCollectionComponent::TickComponent(float DeltaTime, enum ELevelTic
 			IncrementSleepTimer(DeltaTime);
 			IncrementBreakTimer(DeltaTime);
 
+			// todo(chaos) : find a way to only update that of transform have changed
+			// right now this does not work properly because the dirty flags may not be updated at the right time
+			//if (PhysicsProxy && PhysicsProxy->IsGTCollectionDirty())
+			// {
+			// 	for (const TObjectPtr<USceneComponent>& AttachedChild: this->GetAttachChildren())
+			// 	{
+			// 		AttachedChild->UpdateComponentToWorld();
+			// 	}
+			// }
+			
 			if (RestCollection->HasVisibleGeometry() || DynamicCollection->IsDirty())
 			{
 				// #todo review: When we've made changes to ISMC, we need to move this function call to SetRenderDynamicData_Concurrent
@@ -2107,7 +2211,7 @@ void UGeometryCollectionComponent::TickComponent(float DeltaTime, enum ELevelTic
 				MarkRenderTransformDirty();
 				MarkRenderDynamicDataDirty();
 				bRenderStateDirty = false;
-				
+
 				const UWorld* MyWorld = GetWorld();
 				if (MyWorld && MyWorld->IsGameWorld())
 				{
@@ -2406,6 +2510,7 @@ void UGeometryCollectionComponent::RegisterAndInitializePhysicsProxy()
 	}
 #endif
 	PhysicsProxy = new FGeometryCollectionPhysicsProxy(this, *DynamicCollection, SimulationParameters, InitialSimFilter, InitialQueryFilter, CollectorGuid);
+	PhysicsProxy->SetPostPhysicsSyncCallback([this]() { UpdateAttachedChildrenTransform(); }); 
 	FPhysScene_Chaos* Scene = GetInnerChaosScene();
 	Scene->AddObject(this, PhysicsProxy);
 
diff --git a/Engine/Source/Runtime/Experimental/GeometryCollectionEngine/Public/GeometryCollection/GeometryCollectionComponent.h b/Engine/Source/Runtime/Experimental/GeometryCollectionEngine/Public/GeometryCollection/GeometryCollectionComponent.h
--- a/Engine/Source/Runtime/Experimental/GeometryCollectionEngine/Public/GeometryCollection/GeometryCollectionComponent.h
+++ b/Engine/Source/Runtime/Experimental/GeometryCollectionEngine/Public/GeometryCollection/GeometryCollectionComponent.h
@@ -467,7 +467,11 @@ public:
 	virtual FBoxSphereBounds CalcBounds(const FTransform& LocalToWorld) const override;
 	virtual FBoxSphereBounds CalcLocalBounds() const { return LocalBounds; }
 
-	virtual bool HasAnySockets() const override { return false; }
+	virtual bool HasAnySockets() const override;
+	virtual bool DoesSocketExist(FName InSocketName) const override;
+	virtual FTransform GetSocketTransform(FName InSocketName, ERelativeTransformSpace TransformSpace = RTS_World) const override;
+	virtual void QuerySupportedSockets(TArray<FComponentSocketDescription>& OutSockets) const override;
+	
 	virtual void TickComponent(float DeltaTime, enum ELevelTick TickType, FActorComponentTickFunction *ThisTickFunction) override;
 	virtual void AsyncPhysicsTickComponent(float DeltaTime, float SimTime) override;
 	//~ Begin USceneComponent Interface.
@@ -1086,6 +1090,8 @@ private:
 	void UpdateDecay(int32 TransformIdx, float UpdatedDecay, bool UseClusterCrumbling, bool HasDynamicInternalClusterParent, FGeometryCollectionDecayContext& ContextInOut);
 	void ProcessRepData();
 
+	void UpdateAttachedChildrenTransform() const;
+	
 	void BuildInitialFilterData();
 
 	/** The clusters we need to replicate */
