Iris CSV stats.

#rnx

diff --git a/Engine/Source/Runtime/Engine/Classes/Engine/NetDriver.h b/Engine/Source/Runtime/Engine/Classes/Engine/NetDriver.h
--- a/Engine/Source/Runtime/Engine/Classes/Engine/NetDriver.h
+++ b/Engine/Source/Runtime/Engine/Classes/Engine/NetDriver.h
@@ -1941,6 +1941,7 @@ private:
 	void InitIrisSettings(FName NewDriverName);
 	void SetReplicationSystem(UReplicationSystem* ReplicationSystem);
 	void UpdateReplicationViews() const;
+	void SendClientMoveAdjustments();
 #endif
 
 	UPROPERTY(transient)
diff --git a/Engine/Source/Runtime/Engine/Private/Net/Experimental/Iris/DataStreamChannel.cpp b/Engine/Source/Runtime/Engine/Private/Net/Experimental/Iris/DataStreamChannel.cpp
--- a/Engine/Source/Runtime/Engine/Private/Net/Experimental/Iris/DataStreamChannel.cpp
+++ b/Engine/Source/Runtime/Engine/Private/Net/Experimental/Iris/DataStreamChannel.cpp
@@ -16,6 +16,7 @@
 #include "Iris/Serialization/NetSerializationContext.h"
 #include "Iris/Core/IrisProfiler.h"
 #include "Iris/Core/IrisMemoryTracker.h"
+#include "ProfilingDebugging/CsvProfiler.h"
 #endif // UE_WITH_IRIS
 
 namespace UE::Net::Private
@@ -162,6 +163,10 @@ void UDataStreamChannel::Tick()
 		return;
 	}
 
+#if UE_NET_IRIS_CSV_STATS
+	CSV_SCOPED_TIMING_STAT_EXCLUSIVE(UDataStreamChannel_Tick);
+#endif
+
 	IRIS_PROFILER_SCOPE(UDataStreamChannel_Tick);
 	LLM_SCOPE_BYTAG(Iris);
 
diff --git a/Engine/Source/Runtime/Engine/Private/NetDriver.cpp b/Engine/Source/Runtime/Engine/Private/NetDriver.cpp
--- a/Engine/Source/Runtime/Engine/Private/NetDriver.cpp
+++ b/Engine/Source/Runtime/Engine/Private/NetDriver.cpp
@@ -815,15 +815,20 @@ void UNetDriver::TickFlush(float DeltaSeconds)
 	{
 		// Update all clients.
 #if WITH_SERVER_CODE
-		int32 Updated = ServerReplicateActors( DeltaSeconds );
+		int32 Updated = 0;
 
 #if UE_WITH_IRIS
 		if (ReplicationSystem)
 		{
 			UpdateReplicationViews();
+			SendClientMoveAdjustments();
 			ReplicationSystem->PreSendUpdate(DeltaSeconds);
 		}
+		else
 #endif // UE_WITH_IRIS
+		{
+			Updated = ServerReplicateActors(DeltaSeconds);
+		}
 
 		static int32 LastUpdateCount = 0;
 		// Only log the zero replicated actors once after replicating an actor
@@ -6069,6 +6074,35 @@ void UNetDriver::UpdateReplicationViews() const
 		ReplicationSystem->SetReplicationView(Conn->GetConnectionId(), ReplicationView);
 	}
 }
+
+void UNetDriver::SendClientMoveAdjustments()
+{
+	for (UNetConnection* Connection : ClientConnections)
+	{
+		if (Connection == nullptr || Connection->ViewTarget == nullptr)
+		{
+			continue;
+		}
+
+		if (APlayerController* PC = Connection->PlayerController)
+		{
+			PC->SendClientAdjustment();
+		}
+
+		for (UNetConnection* ChildConnection : Connection->Children)
+		{
+			if (ChildConnection == nullptr)
+			{
+				continue;
+			}
+
+			if (APlayerController* PC = ChildConnection->PlayerController)
+			{
+				PC->SendClientAdjustment();
+			}
+		}
+	}
+}
 #endif // UE_WITH_IRIS
 
 #if NET_DEBUG_RELEVANT_ACTORS
diff --git a/Engine/Source/Runtime/Experimental/Iris/Core/Private/Iris/ReplicationSystem/AttachmentReplication.cpp b/Engine/Source/Runtime/Experimental/Iris/Core/Private/Iris/ReplicationSystem/AttachmentReplication.cpp
--- a/Engine/Source/Runtime/Experimental/Iris/Core/Private/Iris/ReplicationSystem/AttachmentReplication.cpp
+++ b/Engine/Source/Runtime/Experimental/Iris/Core/Private/Iris/ReplicationSystem/AttachmentReplication.cpp
@@ -51,6 +51,11 @@ public:
 		return PreQueue.Num() > 0 || ReliableQueue.HasUnsentBlobs();
 	}
 
+	bool CanSendBlobs() const
+	{
+		return ReliableQueue.HasUnsentBlobs();
+	}
+
 	bool IsSafeToDestroy() const
 	{
 		return ReliableQueue.IsSafeToDestroy();
@@ -204,7 +209,7 @@ bool FNetObjectAttachmentSendQueue::Enqueue(TArrayView<const TRefCountPtr<FNetBl
 	return true;
 }
 
-void FNetObjectAttachmentSendQueue::DropUnreliable(bool &bOutHasUnsent)
+void FNetObjectAttachmentSendQueue::DropUnreliable(bool& bOutHasUnsent)
 {
 	UnreliableQueue.Empty();
 	bOutHasUnsent = (ReliableQueue != nullptr && ReliableQueue->HasUnsentBlobs());
@@ -241,26 +246,36 @@ void FNetObjectAttachmentSendQueue::SetUnreliableQueueCapacity(uint32 QueueCapac
 	UnreliableQueue.PopNoCheck(DropCount);
 }
 
-void FNetObjectAttachmentSendQueue::Serialize(FNetSerializationContext& Context, FNetHandle NetHandle, FNetObjectAttachmentSendQueue::ReplicationRecord& OutRecord, bool& bOutHasUnsentAttachments)
+EAttachmentWriteStatus FNetObjectAttachmentSendQueue::Serialize(FNetSerializationContext& Context, FNetHandle NetHandle, FNetObjectAttachmentSendQueue::ReplicationRecord& OutRecord, bool& bOutHasUnsentAttachments)
 {
 	FNetBitStreamWriter& Writer = *Context.GetBitStreamWriter();
 
 	FInternalRecord ReplicationRecord;
+	// This count is the total number of unsent reliable blobs. If the reliable window is full no blobs can be sentuntil some have been acked.
 	const uint32 UnsentReliableCount = (ReliableQueue != nullptr ? ReliableQueue->GetUnsentBlobCount() : 0U);
-	const bool bHasReliableAttachments = UnsentReliableCount > 0;
+	const bool bCanSendReliableAttachments = UnsentReliableCount > 0 && ReliableQueue->CanSendBlobs();
 	const bool bHasUnreliableAttachments = !UnreliableQueue.IsEmpty();
-	if (!(bHasReliableAttachments || bHasUnreliableAttachments))
+	if (!(bCanSendReliableAttachments || bHasUnreliableAttachments))
 	{
 		// Ideally we shouldn't get here, but we can handle it. Important to overflow to report a soft error.
 		Writer.DoOverflow();
 		OutRecord = ReplicationRecord.CombinedRecord;
 		bOutHasUnsentAttachments = false;
-		return;
+
+		if (UnsentReliableCount > 0)
+		{
+			// We want to send reliable blobs but the window is full.
+			return EAttachmentWriteStatus::ReliableWindowFull;
+		}
+		else
+		{
+			return EAttachmentWriteStatus::NoAttachments;
+		}
 	}
 
 	UE_NET_TRACE_SCOPE(RPCs, Writer, Context.GetTraceCollector(), ENetTraceVerbosity::Trace);
 
-	Writer.WriteBool(bHasReliableAttachments);
+	Writer.WriteBool(bCanSendReliableAttachments);
 	const uint32 HasUnreliableAttachmentsWritePos = Writer.GetPosBits();
 	Writer.WriteBool(bHasUnreliableAttachments);
 
@@ -268,11 +283,11 @@ void FNetObjectAttachmentSendQueue::Serialize(FNetSerializationContext& Context,
 	{
 		OutRecord = ReplicationRecord.CombinedRecord;
 		bOutHasUnsentAttachments = true;
-		return;
+		return EAttachmentWriteStatus::BitstreamOverflow;
 	}
 
 	uint32 SerializedReliableCount = 0;
-	if (bHasReliableAttachments)
+	if (bCanSendReliableAttachments)
 	{
 		UE_NET_TRACE_SCOPE(Reliable, Writer, Context.GetTraceCollector(), ENetTraceVerbosity::Trace);
 		SerializedReliableCount = SerializeReliable(Context, NetHandle, ReplicationRecord.ReliableRecord);
@@ -282,10 +297,13 @@ void FNetObjectAttachmentSendQueue::Serialize(FNetSerializationContext& Context,
 			Writer.DoOverflow();
 			OutRecord = ReplicationRecord.CombinedRecord;
 			bOutHasUnsentAttachments = true;
-			return;
+			return EAttachmentWriteStatus::BitstreamOverflow;
 		}
 	}
 
+	// Assume success and change when necessary.
+	EAttachmentWriteStatus WriteStatus = EAttachmentWriteStatus::Success;
+
 	uint32 SerializedUnreliableCount = 0;
 	if (bHasUnreliableAttachments)
 	{
@@ -297,8 +315,10 @@ void FNetObjectAttachmentSendQueue::Serialize(FNetSerializationContext& Context,
 			if (SerializedReliableCount == 0)
 			{
 				Writer.DoOverflow();
+				// Even if we couldn't send reliable attachments due to the window being full we still probably wouldn't have fit any, so return the overflow status.
+				WriteStatus = EAttachmentWriteStatus::BitstreamOverflow;
 			}
-			// Patch up information saying there were unreliable attachments because we couldn't fit any in the bitstream
+			// Patch up information previously saying there were unreliable attachments because we couldn't fit any in the bitstream
 			else
 			{
 				FNetBitStreamWriteScope WriteScope(Writer, HasUnreliableAttachmentsWritePos);
@@ -307,8 +327,18 @@ void FNetObjectAttachmentSendQueue::Serialize(FNetSerializationContext& Context,
 		}
 	}
 
+	if (WriteStatus == EAttachmentWriteStatus::Success)
+	{
+		// Override WriteStatus with ReliableWindowFull in case we fit one or more unreliable attachments but weren't able to write any reliable ones.
+		if (UnsentReliableCount > 0 && !bCanSendReliableAttachments)
+		{
+			WriteStatus = EAttachmentWriteStatus::ReliableWindowFull;
+		}
+	}
+
 	OutRecord = ReplicationRecord.CombinedRecord;
 	bOutHasUnsentAttachments = (SerializedReliableCount < UnsentReliableCount) || (SerializedUnreliableCount < UnreliableQueue.Count());
+	return WriteStatus;
 }
 
 uint32 FNetObjectAttachmentSendQueue::SerializeReliable(FNetSerializationContext& Context, FNetHandle NetHandle, FReliableNetBlobQueue::ReplicationRecord& OutRecord)
@@ -477,7 +507,7 @@ void FNetObjectAttachmentsWriter::DropUnreliableAttachments(ENetObjectAttachment
 	}
 }
 
-void FNetObjectAttachmentsWriter::Serialize(FNetSerializationContext& Context, ENetObjectAttachmentType Type, uint32 ObjectIndex, const FNetHandle NetHandle,  FNetObjectAttachmentsWriter::ReplicationRecord& OutRecord, bool& bOutHasUnsentAttachments)
+EAttachmentWriteStatus FNetObjectAttachmentsWriter::Serialize(FNetSerializationContext& Context, ENetObjectAttachmentType Type, uint32 ObjectIndex, const FNetHandle NetHandle,  FNetObjectAttachmentsWriter::ReplicationRecord& OutRecord, bool& bOutHasUnsentAttachments)
 {
 	FNetObjectAttachmentSendQueue* Queue = GetQueue(Type, ObjectIndex);
 	// If this ensure fires we have bad logic for keeping track of whether there are attachments or not
@@ -485,10 +515,10 @@ void FNetObjectAttachmentsWriter::Serialize(FNetSerializationContext& Context, E
 	{
 		OutRecord = 0;
 		bOutHasUnsentAttachments = false;
-		return;
+		return EAttachmentWriteStatus::NoAttachments;
 	}
 
-	Queue->Serialize(Context, NetHandle, OutRecord, bOutHasUnsentAttachments);
+	return Queue->Serialize(Context, NetHandle, OutRecord, bOutHasUnsentAttachments);
 }
 
 void FNetObjectAttachmentsWriter::CommitReplicationRecord(ENetObjectAttachmentType Type, uint32 ObjectIndex, FNetObjectAttachmentsWriter::ReplicationRecord Record)
diff --git a/Engine/Source/Runtime/Experimental/Iris/Core/Private/Iris/ReplicationSystem/AttachmentReplication.h b/Engine/Source/Runtime/Experimental/Iris/Core/Private/Iris/ReplicationSystem/AttachmentReplication.h
--- a/Engine/Source/Runtime/Experimental/Iris/Core/Private/Iris/ReplicationSystem/AttachmentReplication.h
+++ b/Engine/Source/Runtime/Experimental/Iris/Core/Private/Iris/ReplicationSystem/AttachmentReplication.h
@@ -30,6 +30,18 @@ enum class ENetObjectAttachmentType : uint32
 	InternalCount
 };
 
+enum class EAttachmentWriteStatus : unsigned
+{
+	/** At least one attachment could be sent. */
+	Success,
+	/** There were no attachments to write. */
+	NoAttachments,
+	/** Due to the reliable window size being full no reliable attachments could be sent. */
+	ReliableWindowFull,
+	/** Writing caused bitstream overflow. */
+	BitstreamOverflow,
+};
+
 class FNetObjectAttachmentSendQueue
 {
 public:
@@ -71,7 +83,7 @@ private:
 		FInternalRecord() { CombinedRecord = InvalidReplicationRecord; };
 	};
 
-	void Serialize(FNetSerializationContext& Context, FNetHandle NetHandle, ReplicationRecord& OutRecord, bool& bOutHasUnprocessedAttachments);
+	EAttachmentWriteStatus Serialize(FNetSerializationContext& Context, FNetHandle NetHandle, ReplicationRecord& OutRecord, bool& bOutHasUnprocessedAttachments);
 	uint32 SerializeReliable(FNetSerializationContext& Context, FNetHandle NetHandle, FReliableNetBlobQueue::ReplicationRecord& OutRecord);
 	uint32 SerializeUnreliable(FNetSerializationContext& Context, FNetHandle NetHandle, uint32& OutRecord);
 
@@ -102,7 +114,7 @@ public:
 	void DropAllAttachments(ENetObjectAttachmentType Type, uint32 ObjectIndex);
 	void DropUnreliableAttachments(ENetObjectAttachmentType Type, uint32 ObjectIndex, bool& bOutHasUnsentAttachments);
 
-	void Serialize(FNetSerializationContext& Context, ENetObjectAttachmentType Type, uint32 ObjectIndex, FNetHandle NetHandle, ReplicationRecord& OutRecord, bool& bOutHasUnsentAttachments);
+	EAttachmentWriteStatus Serialize(FNetSerializationContext& Context, ENetObjectAttachmentType Type, uint32 ObjectIndex, FNetHandle NetHandle, ReplicationRecord& OutRecord, bool& bOutHasUnsentAttachments);
 
 	void CommitReplicationRecord(ENetObjectAttachmentType Type, uint32 ObjectIndex, ReplicationRecord Record);
 
diff --git a/Engine/Source/Runtime/Experimental/Iris/Core/Private/Iris/ReplicationSystem/NetBlob/NetRPC.cpp b/Engine/Source/Runtime/Experimental/Iris/Core/Private/Iris/ReplicationSystem/NetBlob/NetRPC.cpp
--- a/Engine/Source/Runtime/Experimental/Iris/Core/Private/Iris/ReplicationSystem/NetBlob/NetRPC.cpp
+++ b/Engine/Source/Runtime/Experimental/Iris/Core/Private/Iris/ReplicationSystem/NetBlob/NetRPC.cpp
@@ -19,6 +19,7 @@
 #include "Iris/Serialization/NetBitStreamUtil.h"
 #include "Iris/Core/IrisLog.h"
 #include "Containers/ArrayView.h"
+#include "ProfilingDebugging/CsvProfiler.h"
 
 DEFINE_LOG_CATEGORY_STATIC(LogIrisRpc, Log, All);
 
@@ -412,6 +413,10 @@ FNetRPC* FNetRPC::Create(UReplicationSystem* ReplicationSystem, const FNetBlobCr
 
 void FNetRPC::CallFunction(FNetSerializationContext& Context)
 {
+#if UE_NET_IRIS_CSV_STATS
+	CSV_SCOPED_TIMING_STAT_EXCLUSIVE(HandleRPC);
+#endif
+
 	// Check whether we are ok with calling the function
 	UObject* Object = ObjectPtr.Get();
 
diff --git a/Engine/Source/Runtime/Experimental/Iris/Core/Private/Iris/ReplicationSystem/ReplicationSystem.cpp b/Engine/Source/Runtime/Experimental/Iris/Core/Private/Iris/ReplicationSystem/ReplicationSystem.cpp
--- a/Engine/Source/Runtime/Experimental/Iris/Core/Private/Iris/ReplicationSystem/ReplicationSystem.cpp
+++ b/Engine/Source/Runtime/Experimental/Iris/Core/Private/Iris/ReplicationSystem/ReplicationSystem.cpp
@@ -22,6 +22,7 @@
 #include "Iris/Serialization/NetBitStreamWriter.h"
 #include "Iris/Serialization/NetSerializer.h"
 #include "HAL/IConsoleManager.h"
+#include "ProfilingDebugging/CsvProfiler.h"
 #include "UObject/UObjectGlobals.h"
 
 namespace ReplicationSystemCVars
@@ -494,6 +495,14 @@ void UReplicationSystem::PreSendUpdate(float DeltaSeconds)
 	// $IRIS TODO. There may be some throttling of connections to tick that we should take into account.
 	const UE::Net::FNetBitArrayView& ReplicatingConnections = MakeNetBitArrayView(Impl->ReplicationSystemInternal.GetConnections().GetValidConnections());
 
+#if UE_NET_IRIS_CSV_STATS && CSV_PROFILER
+	{
+		UE::Net::FNetSendStats& SendStats = InternalSys.GetSendStats();
+		SendStats.Reset();
+		SendStats.SetNumberOfReplicatingConnections(ReplicatingConnections.CountSetBits());
+	}
+#endif
+
 	if (bAllowObjectReplication)
 	{
 		UE_NET_TRACE_FRAME_STATSCOUNTER(GetId(), ReplicationSystem.ReplicatedObjectCount, InternalSys.GetNetHandleManager().GetActiveObjectCount(), ENetTraceVerbosity::Verbose);
@@ -578,6 +587,13 @@ void UReplicationSystem::PostSendUpdate()
 		FDeltaCompressionBaselineManagerPostSendUpdateParams UpdateParams;
 		InternalSys.GetDeltaCompressionBaselineManager().PostSendUpdate(UpdateParams);
 	}
+
+#if UE_NET_IRIS_CSV_STATS && CSV_PROFILER
+	{
+		UE::Net::FNetSendStats& SendStats = InternalSys.GetSendStats();
+		SendStats.ReportCsvStats();
+	}
+#endif
 }
 
 void UReplicationSystem::PostGarbageCollection()
diff --git a/Engine/Source/Runtime/Experimental/Iris/Core/Private/Iris/ReplicationSystem/ReplicationSystemInternal.h b/Engine/Source/Runtime/Experimental/Iris/Core/Private/Iris/ReplicationSystem/ReplicationSystemInternal.h
--- a/Engine/Source/Runtime/Experimental/Iris/Core/Private/Iris/ReplicationSystem/ReplicationSystemInternal.h
+++ b/Engine/Source/Runtime/Experimental/Iris/Core/Private/Iris/ReplicationSystem/ReplicationSystemInternal.h
@@ -23,6 +23,7 @@
 #include "Iris/ReplicationSystem/StringTokenStore.h"
 #include "Iris/ReplicationSystem/WorldLocations.h"
 #include "Iris/ReplicationState/ReplicationStateDescriptorRegistry.h"
+#include "Iris/Stats/NetStats.h"
 
 namespace UE::Net::Private
 {
@@ -87,6 +88,11 @@ public:
 	FDeltaCompressionBaselineManager& GetDeltaCompressionBaselineManager() { return DeltaCompressionBaselineManager; }
 	FDeltaCompressionBaselineInvalidationTracker& GetDeltaCompressionBaselineInvalidationTracker() { return DeltaCompressionBaselineInvalidationTracker; }
 
+	FNetSendStats& GetSendStats()
+	{ 
+		return SendStats;
+	}
+
 private:
 	FReplicationProtocolManager ReplicationProtocolManager;
 	FNetHandleManager NetHandleManager;
@@ -107,6 +113,7 @@ private:
 	FWorldLocations WorldLocations;
 	FDeltaCompressionBaselineManager DeltaCompressionBaselineManager;
 	FDeltaCompressionBaselineInvalidationTracker DeltaCompressionBaselineInvalidationTracker;
+	FNetSendStats SendStats;
 	uint32 Id;
 };
 
diff --git a/Engine/Source/Runtime/Experimental/Iris/Core/Private/Iris/ReplicationSystem/ReplicationWriter.cpp b/Engine/Source/Runtime/Experimental/Iris/Core/Private/Iris/ReplicationSystem/ReplicationWriter.cpp
--- a/Engine/Source/Runtime/Experimental/Iris/Core/Private/Iris/ReplicationSystem/ReplicationWriter.cpp
+++ b/Engine/Source/Runtime/Experimental/Iris/Core/Private/Iris/ReplicationSystem/ReplicationWriter.cpp
@@ -5,7 +5,6 @@
 #include "Iris/IrisConfigInternal.h"
 #include "Iris/Core/IrisLog.h"
 #include "Iris/Core/IrisProfiler.h"
-#include "Iris/DataStream/DataStream.h"
 #include "Net/Core/Trace/NetTrace.h"
 #include "Iris/PacketControl/PacketNotification.h"
 #include "Iris/ReplicationSystem/ReplicationSystemTypes.h"
@@ -21,11 +20,12 @@
 #include "Iris/Serialization/NetReferenceCollector.h"
 #include "Iris/ReplicationSystem/NetTokenStoreState.h"
 #include "Iris/Serialization/InternalNetSerializationContext.h"
-#include "Iris/Serialization/NetBitStreamWriter.h"
 #include "Iris/Serialization/NetBitStreamUtil.h"
 #include "Iris/Serialization/NetSerializer.h"
 #include "Iris/Serialization/NetExportContext.h"
+#include "Iris/Stats/NetStats.h"
 #include "HAL/IConsoleManager.h"
+#include "ProfilingDebugging/CsvProfiler.h"
 #include <algorithm>
 
 #if UE_NET_ENABLE_REPLICATIONWRITER_LOG
@@ -51,6 +51,24 @@ static FAutoConsoleVariableRef CVarWarnAboutDroppedAttachmentsToObjectsNotInScop
 
 static const FName NetError_ObjectStateTooLarge("Object state is too large to be split.");
 
+/** Helper class for timing various operations. */
+class FIrisStatsTimer
+{
+public:
+	FIrisStatsTimer()
+	: StartCycle(FPlatformTime::Cycles64()) 
+	{
+	}
+
+	double GetSeconds() const
+	{
+		return FPlatformTime::ToSeconds64(FPlatformTime::Cycles64() - StartCycle);
+	}
+
+private:
+	uint64 StartCycle;
+};
+
 const TCHAR* FReplicationWriter::LexToString(const EReplicatedObjectState State)
 {
 	static const TCHAR* Names[] = {
@@ -151,6 +169,9 @@ FReplicationWriter::FHugeObjectContext::FHugeObjectContext()
 : SendStatus(EHugeObjectSendStatus::Idle)
 , InternalIndex(0)
 , DebugName(CreatePersistentNetDebugName(TEXT("HugeObjectState"), UE_ARRAY_COUNT(TEXT("HugeObjectState"))))
+, StartSendingTime(0)
+, EndSendingTime(0)
+, StartStallTime(0)
 {
 #if UE_NET_TRACE_ENABLED
 	DebugName->DebugNameId = FNetTrace::TraceName(DebugName->Name);
@@ -1746,7 +1767,19 @@ FReplicationWriter::EWriteObjectStatus FReplicationWriter::WriteObjectInBatch(FN
 					AttachmentWriter.WriteBool(BatchEntry.AttachmentType == ENetObjectAttachmentType::HugeObject);
 				}
 
-				Attachments.Serialize(AttachmentContext, BatchEntry.AttachmentType, InternalIndex, NetHandle, BatchEntry.AttachmentRecord, BatchEntry.bHasUnsentAttachments);
+				const EAttachmentWriteStatus AttachmentWriteStatus = Attachments.Serialize(AttachmentContext, BatchEntry.AttachmentType, InternalIndex, NetHandle, BatchEntry.AttachmentRecord, BatchEntry.bHasUnsentAttachments);
+				if (BatchEntry.AttachmentType == ENetObjectAttachmentType::HugeObject)
+				{
+					if (AttachmentWriteStatus == EAttachmentWriteStatus::ReliableWindowFull)
+					{
+						HugeObjectContext.StartStallTime = FPlatformTime::Cycles64();
+					}
+					else
+					{
+						// Clear stall time now that we were theoretically able to send something.
+						HugeObjectContext.StartStallTime = 0;
+					}
+				}
 
 				// If we didn't manage to fit any attachments then clear the HasAttachments bool in the packet
 				if (AttachmentWriter.GetPosBits() == 0 || AttachmentWriter.IsOverflown())
@@ -1935,6 +1968,7 @@ int FReplicationWriter::PrepareAndSendHugeObjectPayload(FNetSerializationContext
 
 	HugeObjectContext.InternalIndex = InternalIndex;
 	HugeObjectContext.SendStatus = EHugeObjectSendStatus::Sending;
+	HugeObjectContext.StartSendingTime = FPlatformTime::Cycles64();
 	// Store batch record for later processing once the whole state is acked.
 	HugeObjectHeader.ObjectCount = HandleObjectBatchSuccess(BatchInfo, HugeObjectContext.BatchRecord);
 
@@ -1998,7 +2032,12 @@ int FReplicationWriter::PrepareAndSendHugeObjectPayload(FNetSerializationContext
 		CommitBatchRecord(BatchRecord);
 
 		// If all chunks did not make it into the packet (expected) mark the context so that we can try to send the huge object in the next packet if we are allowed
-		WriteContext.bHasHugeObjectToSend = Attachments.HasUnsentAttachments(ENetObjectAttachmentType::HugeObject, ObjectIndexForOOBAttachment);
+		const bool bHasHugeObjectToSend = Attachments.HasUnsentAttachments(ENetObjectAttachmentType::HugeObject, ObjectIndexForOOBAttachment);
+		WriteContext.bHasHugeObjectToSend = bHasHugeObjectToSend;
+		if (!bHasHugeObjectToSend)
+		{
+			HugeObjectContext.EndSendingTime = FPlatformTime::Cycles64();
+		}
 
 		return 1;
 	}
@@ -2152,7 +2191,13 @@ uint32 FReplicationWriter::WriteOOBAttachments(FNetSerializationContext& Context
 			return WrittenObjectCount;
 		}
 
-		WriteContext.bHasHugeObjectToSend = Attachments.HasUnsentAttachments(ENetObjectAttachmentType::HugeObject, ObjectIndexForOOBAttachment);
+		const bool bHasHugeObjectToSend = Attachments.HasUnsentAttachments(ENetObjectAttachmentType::HugeObject, ObjectIndexForOOBAttachment);
+		WriteContext.bHasHugeObjectToSend = bHasHugeObjectToSend;
+		if (!bHasHugeObjectToSend)
+		{
+			HugeObjectContext.EndSendingTime = FPlatformTime::Cycles64();
+		}
+
 		WrittenObjectCount += Result;
 	}
 
@@ -2188,7 +2233,16 @@ uint32 FReplicationWriter::WriteObjects(FNetSerializationContext& Context)
 	{
 		if (!this->WriteContext.ObjectsWrittenThisTick.GetBit(InternalIndex) && this->CanSendObject(InternalIndex))
 		{
+#if UE_NET_IRIS_CSV_STATS && CSV_PROFILER
+			FIrisStatsTimer Timer;
+#endif
 			const int32 Result = this->WriteObjectBatch(Context, InternalIndex, WriteObjectFlag_State | WriteObjectFlag_Attachments);
+#if UE_NET_IRIS_CSV_STATS && CSV_PROFILER
+			if (Result <= 0)
+			{
+				this->WriteContext.Stats.AddReplicationWasteTime(Timer.GetSeconds());
+			}
+#endif
 			if (Result >= 0)
 			{
 				WrittenObjectCount += Result;
@@ -2343,6 +2397,13 @@ int FReplicationWriter::HandleObjectBatchSuccess(const FBatchInfo& BatchInfo, FR
 			// Mark this object as written this tick to avoid sending it multiple times
 			WriteContext.ObjectsWrittenThisTick.SetBit(BatchObjectInfo.InternalIndex);
 		}
+
+#if UE_NET_IRIS_CSV_STATS && CSV_PROFILER
+		if (BatchObjectInfo.bSentState & (Info.LastAckedBaselineIndex != FDeltaCompressionBaselineManager::InvalidBaselineIndex))
+		{
+			++WriteContext.DeltaCompressedObjectCount;
+		}
+#endif
 	}
 
 	return WrittenObjectCount;
@@ -2422,9 +2483,11 @@ UDataStream::EWriteResult FReplicationWriter::BeginWrite()
 	WriteContext.bHasDestroyedObjectsToSend = bHasDestroyedObjectsToSend;
 	WriteContext.bHasHugeObjectToSend = bHasUnsentHugeObject;
 	WriteContext.bHasOOBAttachmentsToSend = bHasUnsentOOBAttachments;
-	WriteContext.CurrentIndex = 0u;
+	WriteContext.CurrentIndex = 0U;
+	WriteContext.DeltaCompressedObjectCount = 0U;
 	WriteContext.FailedToWriteSmallObjectCount = 0U;
-	WriteContext.SortedObjectCount = 0u;
+	WriteContext.DeltaCompressedObjectCount = 0U;
+	WriteContext.SortedObjectCount = 0U;
 
 	// Reset dependent object array
 	WriteContext.DependentObjectsPendingSend.Reset();
@@ -2439,6 +2502,14 @@ UDataStream::EWriteResult FReplicationWriter::BeginWrite()
 	WriteContext.ScheduledObjectInfos = reinterpret_cast<FScheduleObjectInfo*>(FMemory::Malloc(sizeof(FScheduleObjectInfo) * Parameters.MaxActiveReplicatedObjectCount));
 	WriteContext.ScheduledObjectCount = ScheduleObjects(WriteContext.ScheduledObjectInfos);
 
+	// Init CSV stats
+#if UE_NET_IRIS_CSV_STATS && CSV_PROFILER
+	{
+		FNetSendStats& Stats = WriteContext.Stats;
+		Stats.Reset();
+	}
+#endif
+
 	WriteContext.bIsValid = true;
 
 	return UDataStream::EWriteResult::HasMoreData;
@@ -2450,6 +2521,33 @@ void FReplicationWriter::EndWrite()
 
 	if (WriteContext.bIsValid)
 	{
+#if UE_NET_IRIS_CSV_STATS && CSV_PROFILER
+		// Update stats
+		{
+			FNetSendStats& Stats = WriteContext.Stats;
+
+			Stats.SetNumberOfObjectsScheduledForReplication(WriteContext.ScheduledObjectCount);
+			Stats.SetNumberOfReplicatedObjects(WriteContext.CurrentIndex - WriteContext.FailedToWriteSmallObjectCount);
+			Stats.SetNumberOfDeltaCompressedReplicatedObjects(WriteContext.DeltaCompressedObjectCount);
+			if (HugeObjectContext.SendStatus == EHugeObjectSendStatus::Sending)
+			{
+				Stats.SetNumberOfActiveHugeObjects(1U);
+
+				if (HugeObjectContext.EndSendingTime != 0)
+				{
+					Stats.AddHugeObjectWaitingTime(FPlatformTime::ToSeconds64(FPlatformTime::Cycles64() - HugeObjectContext.EndSendingTime));
+				}
+				if (HugeObjectContext.StartStallTime != 0)
+				{
+					Stats.AddHugeObjectStallTime(FPlatformTime::ToSeconds64(FPlatformTime::Cycles64() - HugeObjectContext.StartStallTime));
+				}
+			}
+
+			FNetSendStats& TotalStats = Parameters.ReplicationSystem->GetReplicationSystemInternal()->GetSendStats();
+			TotalStats.Accumulate(Stats);
+		}
+#endif
+
  #if UE_NET_ENABLE_REPLICATIONWRITER_LOG
 		// See if we failed to write any objects
 		if (const int32 NumPendingDependentObjects = WriteContext.DependentObjectsPendingSend.Num())
@@ -2674,6 +2772,9 @@ void FReplicationWriter::ClearHugeObjectContext(FHugeObjectContext& Context) con
 	Context.SendStatus = EHugeObjectSendStatus::Idle;
 	Context.BatchRecord = FBatchRecord();
 	Context.BatchExports.Reset();
+	Context.StartSendingTime = 0;
+	Context.EndSendingTime = 0;
+	Context.StartStallTime = 0;
 }
 
 bool FReplicationWriter::CollectAndWriteExports(FNetSerializationContext& Context, uint8* RESTRICT InternalBuffer, const FReplicationProtocol* Protocol) const
diff --git a/Engine/Source/Runtime/Experimental/Iris/Core/Private/Iris/ReplicationSystem/ReplicationWriter.h b/Engine/Source/Runtime/Experimental/Iris/Core/Private/Iris/ReplicationSystem/ReplicationWriter.h
--- a/Engine/Source/Runtime/Experimental/Iris/Core/Private/Iris/ReplicationSystem/ReplicationWriter.h
+++ b/Engine/Source/Runtime/Experimental/Iris/Core/Private/Iris/ReplicationSystem/ReplicationWriter.h
@@ -13,6 +13,7 @@
 #include "Iris/ReplicationSystem/NetHandle.h"
 #include "Iris/ReplicationSystem/NetExports.h"
 #include "Iris/ReplicationSystem/ObjectReferenceCache.h"
+#include "Iris/Stats/NetStats.h"
 #include "Containers/Array.h"
 
 // Forward declaration
@@ -199,12 +200,19 @@ private:
 		FNetBitArray ObjectsWrittenThisTick;
 		// DependentObjets that we should try to write this packet batch, aka. allow overcommit if we have pending DependentObjects when the packet is full
 		TArray<uint32, TInlineAllocator<32>> DependentObjectsPendingSend;
-		FScheduleObjectInfo* ScheduledObjectInfos;				// Scheduled objects
+		// Scheduled objects
+		FScheduleObjectInfo* ScheduledObjectInfos;
 		uint32 ScheduledObjectCount;
 
-		uint32 SortedObjectCount;								// We do partial sorting so we need to track how many object we have sorted
-		uint32 CurrentIndex;									// If we allow multiple packets to be written during the same update this is where we should continue writing
+		// For performance sake we do partial sorting so we need to track how many object we have sorted.
+		uint32 SortedObjectCount;
+		// The index into the scheduled objects array to attempt to replicate next.
+		uint32 CurrentIndex;
 
+		// The number of replicated objects that were serialized with delta compression.
+		uint32 DeltaCompressedObjectCount;
+
+		// How many objects that were attempted to be replicated but which ultimately didn't fit in the packet.
 		uint32 FailedToWriteSmallObjectCount;
 
 		uint32 bHasDestroyedObjectsToSend : 1;
@@ -212,6 +220,8 @@ private:
 		uint32 bHasHugeObjectToSend : 1;
 		uint32 bHasOOBAttachmentsToSend : 1;
 		uint32 bIsValid : 1;
+
+		FNetSendStats Stats;
 	};
 
 	struct FBatchObjectInfo
@@ -269,6 +279,12 @@ private:
 		FBatchRecord BatchRecord;
 		FNetExportContext::FBatchExports BatchExports;
 		const FNetDebugName* DebugName;
+		// Cycle counter for when the huge object context went from idle to sending.
+		uint64 StartSendingTime;
+		// Cycle counter for when the last part of huge object was sent.
+		uint64 EndSendingTime;
+		// Cycle counter for when it was detected that no more parts of the huge object could be sent until some of the first parts have been acked.
+		uint64 StartStallTime;
 	};
 
 	enum EWriteObjectFlag : unsigned
diff --git a/Engine/Source/Runtime/Experimental/Iris/Core/Public/Iris/IrisConfig.h b/Engine/Source/Runtime/Experimental/Iris/Core/Public/Iris/IrisConfig.h
--- a/Engine/Source/Runtime/Experimental/Iris/Core/Public/Iris/IrisConfig.h
+++ b/Engine/Source/Runtime/Experimental/Iris/Core/Public/Iris/IrisConfig.h
@@ -2,7 +2,7 @@
 
 #pragma once
 
-#include "CoreTypes.h"
+#include "HAL/Platform.h"
 
 namespace UE::Net
 {
@@ -30,4 +30,21 @@ IRISCORE_API void SetUseIrisReplication(bool EnableIrisReplication);
 #endif
 #endif
 
+/** CSV stats. Please check if CSV_PROFILER is enabled too if executing non-trivial code for stats. */
+#ifndef UE_NET_IRIS_CSV_STATS
+#	define UE_NET_IRIS_CSV_STATS 1
+#endif
+
+/** Verbose CSV stats is not recommended in shipping builds due to the expected CPU overhead. */
+#if UE_NET_IRIS_CSV_STATS
+#	ifndef UE_NET_IRIS_VERBOSE_CSV_STATS
+#		define UE_NET_IRIS_VERBOSE_CSV_STATS !UE_BUILD_SHIPPING
+#	endif
+#else
+// Force disable UE_NET_IRIS_VERBOSE_CSV_STATS if UE_NET_IRIS_CSV_STATS isn't enabled.
+#	ifdef UE_NET_IRIS_VERBOSE_CSV_STATS
+#		undef UE_NET_IRIS_VERBOSE_CSV_STATS
+#	endif
+#	define UE_NET_IRIS_VERBOSE_CSV_STATS 0
+#endif
 
diff --git a/Engine/Source/Runtime/Experimental/Iris/Core/Public/Iris/Stats/NetStats.cpp b/Engine/Source/Runtime/Experimental/Iris/Core/Public/Iris/Stats/NetStats.cpp
new file mode 100644
--- /dev/null
+++ b/Engine/Source/Runtime/Experimental/Iris/Core/Public/Iris/Stats/NetStats.cpp
@@ -0,0 +1,71 @@
+// Copyright Epic Games, Inc. All Rights Reserved.
+
+#include "Iris/Stats/NetStats.h"
+#include "ProfilingDebugging/CsvProfiler.h"
+#include "Misc/ScopeLock.h"
+
+namespace UE::Net
+{
+
+// Enable Iris category by default on servers
+CSV_DEFINE_CATEGORY(Iris, WITH_SERVER_CODE);
+
+void FNetSendStats::Accumulate(const FNetSendStats& Other)
+{
+	FScopeLock Lock(&CS);
+
+	Stats.ScheduledForReplicationObjectCount += Other.Stats.ScheduledForReplicationObjectCount;
+	Stats.ReplicatedObjectCount += Other.Stats.ReplicatedObjectCount;
+	Stats.DeltaCompressedObjectCount += Other.Stats.DeltaCompressedObjectCount;
+	Stats.ReplicationWasteObjectCount += Other.Stats.ReplicationWasteObjectCount;
+	Stats.ActiveHugeObjectCount += Other.Stats.ActiveHugeObjectCount;
+	Stats.HugeObjectsWaitingForAckCount += Other.Stats.HugeObjectsWaitingForAckCount;
+	Stats.HugeObjectsStallingCount += Other.Stats.HugeObjectsStallingCount;
+
+	Stats.ReplicationWasteTimeInSeconds += Other.Stats.ReplicationWasteTimeInSeconds;
+	Stats.HugeObjectWaitingForAckTimeInSeconds += Other.Stats.HugeObjectWaitingForAckTimeInSeconds;
+	Stats.HugeObjectStallingTimeInSeconds += Other.Stats.HugeObjectStallingTimeInSeconds;
+}
+
+void FNetSendStats::Reset()
+{
+	FScopeLock Lock(&CS);
+	Stats = FStats();
+}
+
+void FNetSendStats::ReportCsvStats()
+{
+#if CSV_PROFILER
+	FScopeLock Lock(&CS);
+
+	// Calculate connection averages for some stats
+	if (Stats.ReplicatingConnectionCount > 0)
+	{
+		const float ConnectionCountFloat = float(Stats.ReplicatingConnectionCount);
+		CSV_CUSTOM_STAT(Iris, AvgScheduledForReplicationObjectCount, Stats.ScheduledForReplicationObjectCount/ConnectionCountFloat, ECsvCustomStatOp::Set);
+		CSV_CUSTOM_STAT(Iris, AvgReplicatedObjectCount, Stats.ReplicatedObjectCount/ConnectionCountFloat, ECsvCustomStatOp::Set);
+		CSV_CUSTOM_STAT(Iris, AvgReplicationWasteObjectCount, Stats.ReplicationWasteObjectCount/ConnectionCountFloat, ECsvCustomStatOp::Set);
+		CSV_CUSTOM_STAT(Iris, AvgDeltaCompressedObjectCount, Stats.DeltaCompressedObjectCount/ConnectionCountFloat, ECsvCustomStatOp::Set);
+	}
+	else
+	{
+		CSV_CUSTOM_STAT(Iris, AvgScheduledForReplicationObjectCount, 0, ECsvCustomStatOp::Set);
+		CSV_CUSTOM_STAT(Iris, AvgReplicatedObjectCount, 0, ECsvCustomStatOp::Set);
+		CSV_CUSTOM_STAT(Iris, AvgReplicationWasteObjectCount, 0, ECsvCustomStatOp::Set);
+		CSV_CUSTOM_STAT(Iris, AvgDeltaCompressedObjectCount, 0, ECsvCustomStatOp::Set);
+	}
+
+	// Huge object counts
+	CSV_CUSTOM_STAT(Iris, ActiveHugeObjectCount, Stats.ActiveHugeObjectCount, ECsvCustomStatOp::Set);
+	CSV_CUSTOM_STAT(Iris, HugeObjectsWaitingForAckCount, Stats.HugeObjectsWaitingForAckCount, ECsvCustomStatOp::Set);
+	CSV_CUSTOM_STAT(Iris, HugeObjectsStallingCount, Stats.HugeObjectsStallingCount, ECsvCustomStatOp::Set);
+
+	CSV_CUSTOM_STAT(Iris, ReplicatingConnectionCount, Stats.ReplicatingConnectionCount, ECsvCustomStatOp::Set);
+
+	CSV_CUSTOM_STAT(Iris, ReplicationWasteTimeMilliseconds, Stats.ReplicationWasteTimeInSeconds*1000.0, ECsvCustomStatOp::Set);
+	CSV_CUSTOM_STAT(Iris, HugeObjectWaitingForAckTimeInSeconds, Stats.HugeObjectWaitingForAckTimeInSeconds, ECsvCustomStatOp::Set);
+	CSV_CUSTOM_STAT(Iris, HugeObjectStallingTimeInSeconds, Stats.HugeObjectStallingTimeInSeconds, ECsvCustomStatOp::Set);
+#endif
+}
+
+}
diff --git a/Engine/Source/Runtime/Experimental/Iris/Core/Public/Iris/Stats/NetStats.h b/Engine/Source/Runtime/Experimental/Iris/Core/Public/Iris/Stats/NetStats.h
new file mode 100644
--- /dev/null
+++ b/Engine/Source/Runtime/Experimental/Iris/Core/Public/Iris/Stats/NetStats.h
@@ -0,0 +1,129 @@
+// Copyright Epic Games, Inc. All Rights Reserved.
+
+#pragma once
+
+#include "HAL/Platform.h"
+#include "HAL/CriticalSection.h"
+
+namespace UE::Net
+{
+
+/**
+ * Send stats for Iris replication reported to the CSV profiler. Mostly of interest on the server side due to the server authoritative network model.
+ * Its intended use is to do thread local tracking to an instance and then use the Accumulate function for thread safe updating of the ReplicationSystem owned instance.
+ * The ReplicationSystem owned instance is the one reporting to the CSV profiler.
+ */
+class FNetSendStats
+{
+public:
+	FNetSendStats() = default;
+	FNetSendStats(const FNetSendStats&) = delete;
+	FNetSendStats& operator=(const FNetSendStats&) = delete;
+
+	/** Set number of objects scheduled for replication. */
+	void SetNumberOfObjectsScheduledForReplication(uint32 Count);
+
+	/** Set number of replicated objects. */
+	void SetNumberOfReplicatedObjects(uint32 Count);
+
+	/** Get the number of replicated objects. */
+	uint32 GetNumberOfReplicatedObjects() const;
+
+	/** Set number of replicated objects that used delta compression. */
+	void SetNumberOfDeltaCompressedReplicatedObjects(uint32 Count);
+
+	/** Set the number of huge objects in sending or waiting to be acked. */
+	void SetNumberOfActiveHugeObjects(uint32 Count);
+
+	/** Add time in seconds waiting for completely sent huge object to be acked. */
+	void AddHugeObjectWaitingTime(double Seconds);
+
+	/** Add time in seconds waiting to be able to continue sending huge object. */
+	void AddHugeObjectStallTime(double Seconds);
+
+	/** Add time spent on trying to replicate an object in the packet but that ended up not fitting. */
+	void AddReplicationWasteTime(double Seconds);
+
+	/** Add stats from another instance. */
+	IRISCORE_API void Accumulate(const FNetSendStats& Stats);
+
+	/** Reset stats. */
+	IRISCORE_API void Reset();
+
+	/** Report the stats to the CSV profiler. Does nothing if CSV profiler support is compiled out. */
+	IRISCORE_API void ReportCsvStats();
+
+	/** Set number of replicating connections. */
+	void SetNumberOfReplicatingConnections(uint32 Count);
+
+private:
+	// Helper struct to facilitate reset of stats.
+	struct FStats
+	{
+		double HugeObjectWaitingForAckTimeInSeconds = 0;
+		double HugeObjectStallingTimeInSeconds = 0;
+		double ReplicationWasteTimeInSeconds = 0;
+
+		int32 ScheduledForReplicationObjectCount = 0;
+		int32 ReplicatedObjectCount = 0;
+		int32 DeltaCompressedObjectCount = 0;
+		int32 ReplicationWasteObjectCount = 0;
+		int32 ActiveHugeObjectCount = 0;
+		int32 HugeObjectsWaitingForAckCount = 0;
+		int32 HugeObjectsStallingCount = 0;
+		int32 ReplicatingConnectionCount = 0;
+	};
+
+	FCriticalSection CS;
+	FStats Stats;
+};
+
+inline void FNetSendStats::SetNumberOfObjectsScheduledForReplication(uint32 Count)
+{
+	Stats.ScheduledForReplicationObjectCount = Count;
+}
+
+inline void FNetSendStats::SetNumberOfReplicatedObjects(uint32 Count)
+{
+	Stats.ReplicatedObjectCount = Count;
+}
+
+inline uint32 FNetSendStats::GetNumberOfReplicatedObjects() const
+{
+	return Stats.ReplicatedObjectCount;
+}
+
+inline void FNetSendStats::SetNumberOfDeltaCompressedReplicatedObjects(uint32 Count)
+{
+	Stats.DeltaCompressedObjectCount = Count;
+}
+
+inline void FNetSendStats::SetNumberOfActiveHugeObjects(uint32 Count)
+{
+	Stats.ActiveHugeObjectCount = static_cast<int32>(Count);
+}
+
+inline void FNetSendStats::AddHugeObjectWaitingTime(double Seconds)
+{
+	++Stats.HugeObjectsWaitingForAckCount;
+	Stats.HugeObjectWaitingForAckTimeInSeconds += Seconds;
+}
+
+inline void FNetSendStats::AddHugeObjectStallTime(double Seconds)
+{
+	++Stats.HugeObjectsStallingCount;
+	Stats.HugeObjectStallingTimeInSeconds += Seconds;
+}
+
+inline void FNetSendStats::AddReplicationWasteTime(double Seconds)
+{
+	++Stats.ReplicationWasteObjectCount;
+	Stats.ReplicationWasteTimeInSeconds += Seconds;
+}
+
+inline void FNetSendStats::SetNumberOfReplicatingConnections(uint32 Count)
+{
+	Stats.ReplicatingConnectionCount = Count;
+}
+
+}
