Modular Vehicles: Implement Surface Friction

diff --git a/Engine/Source/Runtime/Experimental/ChaosVehicles/ChaosVehiclesCore/Public/SimModule/WheelModule.h b/Engine/Source/Runtime/Experimental/ChaosVehicles/ChaosVehiclesCore/Public/SimModule/WheelModule.h
--- a/Engine/Source/Runtime/Experimental/ChaosVehicles/ChaosVehiclesCore/Public/SimModule/WheelModule.h
+++ b/Engine/Source/Runtime/Experimental/ChaosVehicles/ChaosVehiclesCore/Public/SimModule/WheelModule.h
@@ -74,13 +74,6 @@ namespace Chaos
 
 		virtual bool IsBehaviourType(eSimModuleTypeFlags InType) const override { return (InType & TorqueBased)||(InType & Velocity); }
 
-
-		/** Set the vehicle's speed at the wheels location in local wheel coords */
-		//void SetGroundSpeed(const FVector& VIn)
-		//{
-		//	GroundVelocityVector = VIn;
-		//}
-
 		void SetSuspensionSimTreeIndex(int IndexIn) { SuspensionSimTreeIndex = IndexIn; }
 		int GetSuspensionSimTreeIndex() const { return SuspensionSimTreeIndex; }
 
@@ -88,6 +81,8 @@ namespace Chaos
 
 		void SetForceIntoSurface(float ForceIntoSurfaceIn) { ForceIntoSurface = ForceIntoSurfaceIn; }
 
+		void SetSurfaceFriction(float FrictionIn) { SurfaceFriction = FrictionIn; }
+
 	private:
 
 		float BrakeTorque;				// [N.m]
