Dataflow - Cache's timestamp now uses FPlatformTime::Cycles64()

diff --git a/Engine/Source/Runtime/Experimental/Dataflow/Core/Private/Dataflow/DataflowNode.cpp b/Engine/Source/Runtime/Experimental/Dataflow/Core/Private/Dataflow/DataflowNode.cpp
--- a/Engine/Source/Runtime/Experimental/Dataflow/Core/Private/Dataflow/DataflowNode.cpp
+++ b/Engine/Source/Runtime/Experimental/Dataflow/Core/Private/Dataflow/DataflowNode.cpp
@@ -195,7 +195,7 @@ TArray<Dataflow::FPin> FDataflowNode::GetPins() const
 
 void FDataflowNode::Invalidate()
 {
-	LastModifiedTimestamp = FDateTime::Now(); // @todo(dataflow) Now() has ms precision; might need higher
+	LastModifiedTimestamp = FPlatformTime::Cycles64();
 	InvalidateOutputs();
 }
 
diff --git a/Engine/Source/Runtime/Experimental/Dataflow/Core/Public/Dataflow/DataflowNode.h b/Engine/Source/Runtime/Experimental/Dataflow/Core/Public/Dataflow/DataflowNode.h
--- a/Engine/Source/Runtime/Experimental/Dataflow/Core/Public/Dataflow/DataflowNode.h
+++ b/Engine/Source/Runtime/Experimental/Dataflow/Core/Public/Dataflow/DataflowNode.h
@@ -39,7 +39,7 @@ struct DATAFLOWCORE_API FDataflowNode
 
 	FGuid Guid;
 	FName Name;
-	FDateTime LastModifiedTimestamp;
+	uint64 LastModifiedTimestamp;
 
 	TMap< int, FDataflowInput* > Inputs;
 	TMap< int, FDataflowOutput* > Outputs;
diff --git a/Engine/Source/Runtime/Experimental/Dataflow/Core/Public/Dataflow/DataflowNodeParameters.h b/Engine/Source/Runtime/Experimental/Dataflow/Core/Public/Dataflow/DataflowNodeParameters.h
--- a/Engine/Source/Runtime/Experimental/Dataflow/Core/Public/Dataflow/DataflowNodeParameters.h
+++ b/Engine/Source/Runtime/Experimental/Dataflow/Core/Public/Dataflow/DataflowNodeParameters.h
@@ -12,7 +12,7 @@ namespace Dataflow
 {
 	struct FContextCacheBase 
 	{
-		FContextCacheBase(FProperty* InProperty = nullptr, FDateTime InTimestamp = 0)
+		FContextCacheBase(FProperty* InProperty = nullptr, uint64 InTimestamp = 0)
 			: Property(InProperty)
 			, Timestamp(InTimestamp)
 		{}
@@ -22,18 +22,18 @@ namespace Dataflow
 		const T& GetTypedData(const FProperty* PropertyIn) const;
 		
 		FProperty* Property = nullptr;
-		FDateTime Timestamp = FDateTime(0);
+		uint64 Timestamp = 0;
 	};
 
 	template<class T>
 	struct FContextCache : public FContextCacheBase 
 	{
-		FContextCache(FProperty* InProperty, const T& InData, FDateTime Timestamp)
+		FContextCache(FProperty* InProperty, const T& InData, uint64 Timestamp)
 			: FContextCacheBase(InProperty, Timestamp)
 			, Data(InData)
 		{}
 
-		FContextCache(FProperty* InProperty, T&& InData, FDateTime Timestamp)
+		FContextCache(FProperty* InProperty, T&& InData, uint64 Timestamp)
 			: FContextCacheBase(InProperty, Timestamp)
 			, Data(InData)
 		{}
@@ -90,7 +90,7 @@ namespace Dataflow
 		void SetData(size_t Key, FProperty* Property, const T& Value)
 		{
 			int64 IntKey = (int64)Key;
-			TUniquePtr<FContextCache<T>> DataStoreEntry = MakeUnique<FContextCache<T>>(Property, Value, FDateTime::Now());
+			TUniquePtr<FContextCache<T>> DataStoreEntry = MakeUnique<FContextCache<T>>(Property, Value, FPlatformTime::Cycles64());
 			DataStore.Emplace(IntKey, MoveTemp(DataStoreEntry));
 		}
 
@@ -98,7 +98,7 @@ namespace Dataflow
 		void SetData(size_t Key, FProperty* Property, T&& Value)
 		{
 			int64 IntKey = (int64)Key;
-			TUniquePtr<FContextCache<T>> DataStoreEntry = MakeUnique<FContextCache<T>>(Property, Forward<T>(Value), FDateTime::Now());
+			TUniquePtr<FContextCache<T>> DataStoreEntry = MakeUnique<FContextCache<T>>(Property, Forward<T>(Value), FPlatformTime::Cycles64());
 			DataStore.Emplace(IntKey, MoveTemp(DataStoreEntry));
 		}
 		
@@ -112,7 +112,7 @@ namespace Dataflow
 			return Default;
 		}
 
-		bool HasData(size_t Key, FDateTime StoredAfter = -1)
+		bool HasData(size_t Key, uint64 StoredAfter = 0)
 		{
 			int64 IntKey = (int64)Key;
 			return DataStore.Contains(IntKey) && DataStore[Key]->Timestamp >= StoredAfter;
