Chaos: Fixed incorrect joint constraint termination leading to memory leaks and dangling joints that slow down the solver.

diff --git a/Engine/Source/Runtime/Experimental/Chaos/Private/PhysicsProxy/JointConstraintProxy.cpp b/Engine/Source/Runtime/Experimental/Chaos/Private/PhysicsProxy/JointConstraintProxy.cpp
--- a/Engine/Source/Runtime/Experimental/Chaos/Private/PhysicsProxy/JointConstraintProxy.cpp
+++ b/Engine/Source/Runtime/Experimental/Chaos/Private/PhysicsProxy/JointConstraintProxy.cpp
@@ -107,14 +107,28 @@ void FJointConstraintPhysicsProxy::DestroyOnPhysicsThread(FPBDRigidsSolver* InSo
 {
 	if (Constraint_PT)
 	{
+		TVec2<FGeometryParticleHandle*> Particles = Constraint_PT->GetConstrainedParticles();
+		
+		// Ensure that our connected particles are aware that this constraint no longer exists
+		if(Particles[0])
+		{
+			Particles[0]->RemoveConstraintHandle(Constraint_PT);
+		}
+
+		if(Particles[1])
+		{
+			Particles[1]->RemoveConstraintHandle(Constraint_PT);
+		}
+
 		// @todo(chaos): clean up constraint management
 		if(Constraint_PT->IsInConstraintGraph())
 		{
 			InSolver->GetEvolution()->RemoveConstraintFromConstraintGraph(Constraint_PT);
 		}
 
-		auto& JointConstraints = InSolver->GetJointConstraints();
+		FPBDRigidsSolver::FJointConstraints& JointConstraints = InSolver->GetJointConstraints();
 		JointConstraints.RemoveConstraint(Constraint_PT->GetConstraintIndex());
+		Constraint_PT = nullptr;
 	}
 }
 
diff --git a/Engine/Source/Runtime/Experimental/Chaos/Public/Chaos/ConstraintHandle.h b/Engine/Source/Runtime/Experimental/Chaos/Public/Chaos/ConstraintHandle.h
--- a/Engine/Source/Runtime/Experimental/Chaos/Public/Chaos/ConstraintHandle.h
+++ b/Engine/Source/Runtime/Experimental/Chaos/Public/Chaos/ConstraintHandle.h
@@ -111,8 +111,7 @@ namespace Chaos
 
 		virtual bool IsValid() const
 		{
-			// @todo(chaos): why does IsValid() also check IsEnabled()?
-			return (ConstraintContainer != nullptr) && IsEnabled();
+			return (ConstraintContainer != nullptr);
 		}
 
 		FPBDConstraintContainer* GetContainer()
