Optimization: Low resolution height field

diff --git a/Engine/Source/Runtime/Core/Public/UObject/UE5MainStreamObjectVersion.h b/Engine/Source/Runtime/Core/Public/UObject/UE5MainStreamObjectVersion.h
--- a/Engine/Source/Runtime/Core/Public/UObject/UE5MainStreamObjectVersion.h
+++ b/Engine/Source/Runtime/Core/Public/UObject/UE5MainStreamObjectVersion.h
@@ -262,6 +262,9 @@ struct CORE_API FUE5MainStreamObjectVersion
 		// Updated image media source path resolution logic
 		ImgMediaPathResolutionWithEngineOrProjectTokens,
 
+		// Add low resolution data in Height Field
+		AddLowResolutionHeightField,
+
 		// -----<new versions can be added above this line>-------------------------------------------------
 		VersionPlusOne,
 		LatestVersion = VersionPlusOne - 1
diff --git a/Engine/Source/Runtime/Experimental/Chaos/Private/Chaos/HeightField.cpp b/Engine/Source/Runtime/Experimental/Chaos/Private/Chaos/HeightField.cpp
--- a/Engine/Source/Runtime/Experimental/Chaos/Private/Chaos/HeightField.cpp
+++ b/Engine/Source/Runtime/Experimental/Chaos/Private/Chaos/HeightField.cpp
@@ -659,6 +659,9 @@ namespace Chaos
 		}
 		OutBounds.Thicken(UE_KINDA_SMALL_NUMBER);
 
+		
+		OutData.BuildLowResolutionData();
+
 		if(bHaveMaterials)
 		{
 			if(bOnlyDefaultMaterial)
@@ -758,6 +761,8 @@ namespace Chaos
 				}
 			}
 		}
+
+		OutData.BuildLowResolutionData();
 	}
 
 	FHeightField::FHeightField(TArray<FReal>&& Height, TArray<uint8>&& InMaterialIndices, int32 NumRows, int32 NumCols, const FVec3& InScale)
@@ -1085,19 +1090,9 @@ namespace Chaos
 			const bool bCanWalk = SumPlaneAxis > UE_SMALL_NUMBER;
 			if (bCanWalk)
 			{
-				const FVec3 NextStartOri = NextStart;
 				const VectorRegister4Float CurrentLengthSimd = VectorSet1(static_cast<FRealSingle>(CurrentLength));
-				const FReal BoundsMinZ = CachedBounds.Min().Z;
-				const FReal BoundsMaxZ = CachedBounds.Max().Z;
-				const FReal Length2D = FMath::Sqrt(DirScaled[0] * DirScaled[0] + DirScaled[1] * DirScaled[1]);
-				const FReal FastInc2D = 5.0;  // Length in cell dimension
-				const FReal FastIncScaled = FastInc2D / Length2D;
-
-				const FVec3 InspectionStepVector = DirScaled * FastIncScaled * ScaledDx;
-				FReal NextStartZ = FMath::Clamp(NextStart.Z, BoundsMinZ, BoundsMaxZ); // Numerical errors can violate this condition, so force it here.
-				const FReal FastIncScaled3D = InspectionStepVector.Length();
-				const FReal CellsToInspectZ = FastIncScaled3D * Dir.Z;
-				FReal DistanceProcessed = 0.0;		
+				const int32 LowResInc = Visitor.GeomData->LowResInc;
+
 				while (true)
 				{
 					if (!FlatGrid.IsValid(CellIdx))
@@ -1105,30 +1100,64 @@ namespace Chaos
 						return false;
 					}
 
-					FVec3 NewNextStart = (InspectionStepVector + NextStart);
-					TVec2<int32> NewNextStartInt = TVec2<int32>(static_cast<int32>(NewNextStart[0] / Scale2D[0]), static_cast<int32>(NewNextStart[1] / Scale2D[1]));
-					TVec2<int32> NextStartInt = TVec2<int32>(static_cast<int32>(NextStart[0] / Scale2D[0]), static_cast<int32>(NextStart[1] / Scale2D[1]));
-
-					TVec2<int32> DiffInt = NewNextStartInt - NextStartInt;
-					TVec2<int32> AddedCellIdx = TVec2<int32>(static_cast<int32>(DiffInt[0]), static_cast<int32>(DiffInt[1]));
-
 					FAABBVectorized Bounds;
-					Visitor.GeomData->GetBoundsScaled(CellIdx, AddedCellIdx, Bounds);
+					Visitor.GeomData->GetLowResBoundsScaled(CellIdx, Bounds);
 					Bounds.Thicken(Visitor.ThicknessSimd);
 					if (Bounds.RaycastFast(Visitor.StartPointSimd, Visitor.InvDirSimd, Visitor.Parallel, CurrentLengthSimd))
 					{
-						NextStart = NextStartOri + Dir * DistanceProcessed;
+						TVec2<int32> NextStartInt = TVec2<int32>(static_cast<int32>(NextStart[0] / Scale2D[0]), static_cast<int32>(NextStart[1] / Scale2D[1]));
+						CellIdx = FlatGrid.Cell(NextStartInt);
 						break;
 					}
 
-					NextStart = NewNextStart;
-					CellIdx = FlatGrid.Cell(NewNextStartInt);
-					if (DistanceProcessed > CurrentLength || NextStartZ < BoundsMinZ || NextStartZ > BoundsMaxZ)
+					const FVec2 ScaledCellCenter2D = ScaledMin + FVec2(static_cast<FReal>(CellIdx[0] / LowResInc) + 0.5f, static_cast<FReal>(CellIdx[1] / LowResInc) + 0.5f) * (ScaledDx2D*(FReal)LowResInc);
+					const FVec3 ScaledCellCenter(ScaledCellCenter2D[0], ScaledCellCenter2D[1], ZMidPoint);
+
+					FReal Times[3];
+					FReal BestTime = CurrentLength;
+					bool bTerminate = true;
+					for (int Axis = 0; Axis < 3; ++Axis)
+					{
+						if (!bParallel[Axis])
+						{
+							const FReal CrossPoint = (Dir[Axis] * ScaleSign[Axis]) > 0 ? ScaledCellCenter[Axis] + ScaledDx[Axis] * (FReal)LowResInc / 2 : ScaledCellCenter[Axis] - ScaledDx[Axis] * (FReal)LowResInc / 2;
+							const FReal Distance = CrossPoint - NextStart[Axis];	//note: CellCenter already has /2, we probably want to use the corner instead
+							const FReal Time = Distance * InvDir[Axis];
+							Times[Axis] = Time;
+							if (Time < BestTime)
+							{
+								bTerminate = false;	//found at least one plane to pass through
+								BestTime = Time;
+							}
+						}
+						else
+						{
+							Times[Axis] = TNumericLimits<FReal>::Max();
+						}
+					}
+
+					if (bTerminate)
+					{
+						return false;
+					}
+
+					const TVec2<int32> PrevIdx = CellIdx;
+
+					for (int Axis = 0; Axis < 2; ++Axis)
+					{
+						CellIdx[Axis] += (Times[Axis] <= BestTime) ? ((Dir[Axis] * ScaleSign[Axis]) > 0 ? LowResInc : -LowResInc) : 0;
+						if (CellIdx[Axis] < 0 || CellIdx[Axis] >= FlatGrid.Counts()[Axis])
+						{
+							return false;
+						}
+					}
+
+					if (PrevIdx == CellIdx)
 					{
+						//crossed on z plane which means no longer in heightfield bounds
 						return false;
 					}
-					DistanceProcessed += FastIncScaled3D;
-					NextStartZ += CellsToInspectZ;					
+					NextStart = NextStart + Dir * BestTime;
 				}
 			}
 
diff --git a/Engine/Source/Runtime/Experimental/Chaos/Public/Chaos/HeightField.h b/Engine/Source/Runtime/Experimental/Chaos/Public/Chaos/HeightField.h
--- a/Engine/Source/Runtime/Experimental/Chaos/Public/Chaos/HeightField.h
+++ b/Engine/Source/Runtime/Experimental/Chaos/Public/Chaos/HeightField.h
@@ -270,11 +270,26 @@ namespace Chaos
 			// Range of the chosen type (unsigned so Min is always 0)
 			static constexpr int32 StorageRange = TNumericLimits<StorageType>::Max();
 
+			static constexpr int32 LowResInc = 5;
+
 			// Heights in the chosen format. final placement of the vertex will be at
 			// MinValue + Heights[Index] * HeightPerUnit
 			// With HeightPerUnit being the range of the min/max FReal values of
 			// the heightfield divided by the range of StorageType
 			TArray<StorageType> Heights;
+
+			struct MinMaxHeights
+			{
+				StorageType Min;
+				StorageType Max;
+
+				void Serialize(FChaosArchive& Ar)
+				{
+					Ar << Min;
+					Ar << Max;
+				}
+			};
+			TArray<MinMaxHeights> LowResolutionHeights;
 			TArray<uint8> MaterialIndices;
 			FVec3 Scale;
 			VectorRegister4Float ScaleSimd;
@@ -284,6 +299,7 @@ namespace Chaos
 			uint16 NumCols;
 			FReal Range;
 			FReal HeightPerUnit;
+			uint16 NumColsLowRes;
 
 			constexpr FReal GetCellWidth() const
 			{
@@ -415,7 +431,7 @@ namespace Chaos
 				{
 					for (int IndexX = FirstIndexX; IndexX <= LastIndexX; IndexX++)
 					{
-						const int32 Index = IndexY * (NumCols - 1) + IndexX + IndexY;
+						const int32 Index = IndexY * NumCols + IndexX;
 						check(Index < Heights.Num());
 						const FRealSingle CurrHeight = Heights[Index];
 						MinHeight = FMath::Min<FRealSingle>(CurrHeight, MinHeight);
@@ -433,6 +449,29 @@ namespace Chaos
 											MakeVectorRegisterFloat(static_cast<FRealSingle>(LastIndexX), static_cast<FRealSingle>(LastIndexY), MaxHeight, 0.0f));
 			}
 
+			FORCEINLINE void GetLowResBounds(TVec2<int32> CellIdx, FAABBVectorized& OutBounds) const
+			{
+				int32 FirstIndexX = CellIdx[0];
+				int32 FirstIndexY = CellIdx[1];
+
+				const int32 FirstIndexLowResX = FirstIndexX / LowResInc;
+				const int32 FirstIndexLowResY = FirstIndexY / LowResInc;
+
+				const int32 Index = FirstIndexLowResY * (NumColsLowRes - 1) + FirstIndexLowResX + FirstIndexLowResY;
+				check(Index < LowResolutionHeights.Num());
+
+				FRealSingle MinHeight = LowResolutionHeights[Index].Min;
+				FRealSingle MaxHeight = LowResolutionHeights[Index].Max;
+
+				const FRealSingle MinValueSingle = static_cast<FRealSingle>(MinValue);
+				const FRealSingle HeightPerUnitSingle = static_cast<FRealSingle>(HeightPerUnit);
+				MinHeight = MinValueSingle + MinHeight * HeightPerUnitSingle;
+				MaxHeight = MinValueSingle + MaxHeight * HeightPerUnitSingle;
+
+				OutBounds = FAABBVectorized(MakeVectorRegisterFloat(static_cast<FRealSingle>(FirstIndexLowResX*LowResInc), static_cast<FRealSingle>(FirstIndexLowResY*LowResInc), MinHeight, 0.0f),
+					MakeVectorRegisterFloat(static_cast<FRealSingle>((FirstIndexLowResX+1) * LowResInc), static_cast<FRealSingle>((FirstIndexLowResY + 1) * LowResInc), MaxHeight, 0.0f));
+			}
+
 			FORCEINLINE void GetPointsScaled(int32 Index, FVec3 OutPts[4]) const
 			{
 				GetPoints(Index, OutPts);
@@ -484,6 +523,18 @@ namespace Chaos
 				OutBounds = FAABBVectorized(Min, Max);
 			}
 
+			FORCEINLINE void GetLowResBoundsScaled(TVec2<int32> CellIdx, FAABBVectorized& OutBounds) const
+			{
+				GetLowResBounds(CellIdx, OutBounds);
+
+				VectorRegister4Float P0 = VectorMultiply(OutBounds.GetMin(), ScaleSimd);
+				VectorRegister4Float P1 = VectorMultiply(OutBounds.GetMax(), ScaleSimd);
+				VectorRegister4Float Min = VectorMin(P0, P1);
+				VectorRegister4Float Max = VectorMax(P0, P1);
+
+				OutBounds = FAABBVectorized(Min, Max);
+			}
+
 			FORCEINLINE FReal GetMinHeight() const
 			{
 				return MinValue;
@@ -548,6 +599,47 @@ namespace Chaos
 				{
 					Ar << MaterialIndices;
 				}
+
+				Ar.UsingCustomVersion(FUE5MainStreamObjectVersion::GUID);
+				if (Ar.CustomVer(FUE5MainStreamObjectVersion::GUID) >= FUE5MainStreamObjectVersion::AddLowResolutionHeightField)
+				{
+					Ar << LowResolutionHeights;
+					Ar << NumColsLowRes;
+				}
+				else
+				{
+					BuildLowResolutionData();
+				}
+			}
+
+			// A height bounds for grids of cells of size "LowResInc" square, for faster culling
+			void BuildLowResolutionData()
+			{
+				const int32 NumRowsLowRes = FMath::CeilToInt(FRealSingle(NumRows) / LowResInc);
+				NumColsLowRes = uint16(FMath::CeilToInt(FRealSingle(NumCols) / LowResInc));
+				const int32 NumLowResHeights = NumRowsLowRes * NumColsLowRes;
+				LowResolutionHeights.SetNum(NumLowResHeights);
+
+				for (int32 RowIdxLowRes = 0; RowIdxLowRes < NumRowsLowRes; ++RowIdxLowRes)
+				{
+					for (int32 ColIdxLowRes = 0; ColIdxLowRes < NumColsLowRes; ++ColIdxLowRes)
+					{
+						FDataType::StorageType MinHeight = std::numeric_limits<FDataType::StorageType>::max();
+						FDataType::StorageType MaxHeight = std::numeric_limits<FDataType::StorageType>::min();
+						for (int32 RowIdx = RowIdxLowRes * int32(LowResInc); RowIdx < NumRows && RowIdx <= (RowIdxLowRes + 1) * int32(LowResInc); ++RowIdx)
+						{
+							for (int32 ColIdx = ColIdxLowRes * int32(LowResInc); ColIdx < NumCols && ColIdx <= (ColIdxLowRes + 1) * int32(LowResInc); ++ColIdx)
+							{
+								const int32 HeightIndex = RowIdx * NumCols + ColIdx;
+								MaxHeight = FMath::Max<FDataType::StorageType>(Heights[HeightIndex], MaxHeight);
+								MinHeight = FMath::Min<FDataType::StorageType>(Heights[HeightIndex], MinHeight);
+							}
+						}
+						const int32 HeightLowResIndex = RowIdxLowRes * NumColsLowRes + ColIdxLowRes;
+						LowResolutionHeights[HeightLowResIndex].Max = MaxHeight;
+						LowResolutionHeights[HeightLowResIndex].Min = MinHeight;
+					}
+				}
 			}
 		};
 
@@ -824,4 +916,9 @@ namespace Chaos
 		bool GJKContactPointImp(const GeomType& QueryGeom, const FRigidTransform3& QueryTM, const FReal Thickness, FVec3& ContactLocation, FVec3& ContactNormal, FReal& ContactPhi, int32& ContactFaceIndex) const;
 	};
 
-}
+	FORCEINLINE FChaosArchive& operator<<(FChaosArchive& Ar, FHeightField::FDataType::MinMaxHeights& Value)
+	{
+		Value.Serialize(Ar);
+		return Ar;
+	}
+}
\ No newline at end of file
diff --git a/Engine/Source/Runtime/Experimental/ChaosCore/Public/Chaos/Core.h b/Engine/Source/Runtime/Experimental/ChaosCore/Public/Chaos/Core.h
--- a/Engine/Source/Runtime/Experimental/ChaosCore/Public/Chaos/Core.h
+++ b/Engine/Source/Runtime/Experimental/ChaosCore/Public/Chaos/Core.h
@@ -51,7 +51,7 @@ namespace Chaos
 	// NOTE: if you get a merge conflict on the GUID, you must replace it with a new GUID - do not accept the source or target
 	// or you will likely get DDC version conflicts resulting in crashes during load.
 	// Core version string for Chaos data. Any DDC builder dependent on Chaos for serialization should depend on this version
-	static const TCHAR* const ChaosVersionGUID = TEXT("747bf2ef-49be-4f73-b211-74f68a243085");
+	static const TCHAR* const ChaosVersionGUID = TEXT("08C803C2-E2A7-4A85-AAA6-3353713AEE1D");
 
 	inline FString GetChaosVersionStringInner()
 	{
