add xpbd grid based constraints, add more functions into mpm transfer. Add debuggin option in deformable solver. Add more tests

Chaos : Flesh :
Added xpbd grid based constraint and associated tests. Added more functionalities into mpm transfer. Added debugging functions.

diff --git a/Engine/Source/Runtime/Experimental/Chaos/Private/Chaos/Deformable/ChaosDeformableSolver.cpp b/Engine/Source/Runtime/Experimental/Chaos/Private/Chaos/Deformable/ChaosDeformableSolver.cpp
--- a/Engine/Source/Runtime/Experimental/Chaos/Private/Chaos/Deformable/ChaosDeformableSolver.cpp
+++ b/Engine/Source/Runtime/Experimental/Chaos/Private/Chaos/Deformable/ChaosDeformableSolver.cpp
@@ -55,11 +55,19 @@ namespace Chaos::Softs
 		FSolverRigidParticles RigidParticles;
 		Evolution.Reset(new FPBDEvolution(MoveTemp(LocalParticlesDummy), MoveTemp(RigidParticles), {}, Property.NumSolverIterations));
 		Evolution->Particles().AddArray(&MObjects);
-		if (Property.bDoSelfCollision)
+		if (Property.bDoSelfCollision || Property.CacheToFile)
 		{
 			SurfaceElements.Reset(new TArray<Chaos::TVec3<int32>>());
+		}
+		
+		if (Property.bDoSelfCollision)
+		{
 			SurfaceTriangleMesh.Reset(new Chaos::FTriangleMesh());
 		}
+		if (Property.bUseGridBasedConstraints)
+		{
+			AllElements.Reset(new TArray<Chaos::TVec4<int32>>());
+		}
 		Frame = 0;
 		Time = 0.f;
 	}
@@ -101,12 +109,20 @@ namespace Chaos::Softs
 			FThreadingProxy::FKey Key = Proxy->GetOwner();
 			Proxies.Add(Key, TUniquePtr<FThreadingProxy>(Proxy.Release()));
 		}
-		UninitializedProxys.SetNum(0, true);
-
-		if (Property.bDoSelfCollision)
+		
+		if (UninitializedProxys.Num() != 0)
 		{
-			InitializeSelfCollisionVariables();
+			if (Property.bDoSelfCollision)
+			{
+				InitializeSelfCollisionVariables();
+			}
+
+			if (Property.bUseGridBasedConstraints)
+			{
+				InitializeGridBasedConstraintVariables();
+			}
 		}
+		UninitializedProxys.SetNum(0, true);
 		InitializeCollisionBodies();
 	}
 
@@ -199,6 +215,17 @@ namespace Chaos::Softs
 									Elements[edx] = ChaosTet(Tetrahedron[edx], Range[0]);
 								}
 
+								if (Property.bUseGridBasedConstraints)
+								{
+									int32 ElementsOffset = AllElements->Num();
+									AllElements->SetNum(ElementsOffset + NumElements);
+									for (uint32 edx = 0; edx < NumElements; ++edx)
+									{
+										(*AllElements)[edx + ElementsOffset] = ChaosTet(Tetrahedron[edx], Range[0]);
+									}
+
+								}
+
 
 								FXPBDCorotatedConstraints<FSolverReal, FSolverParticles>* CorotatedConstraint =
 									new FXPBDCorotatedConstraints<FSolverReal, FSolverParticles>(
@@ -222,7 +249,7 @@ namespace Chaos::Softs
 								CorotatedConstraints.Add(TUniquePtr<FXPBDCorotatedConstraints<FSolverReal, FSolverParticles>>(CorotatedConstraint));
 							}
 
-							if (Property.bDoSelfCollision)
+							if (Property.bDoSelfCollision || Property.CacheToFile)
 							{
 								// Add Surface Elements Node
 								int32 SurfaceElementsOffset = SurfaceElements->Num();
@@ -325,6 +352,27 @@ namespace Chaos::Softs
 		};
 	}
 
+	void FDeformableSolver::InitializeGridBasedConstraintVariables()
+	{
+		GridBasedCorotatedConstraint.Reset(new Chaos::Softs::FXPBDGridBasedCorotatedConstraints<FSolverReal, FSolverParticles>(
+			Evolution->Particles(), *AllElements, Property.GridDx, /*bRecordMetric = */false, (Chaos::Softs::FSolverReal).1, (Chaos::Softs::FSolverReal).01, (Chaos::Softs::FSolverReal).4, (Chaos::Softs::FSolverReal)1000.0));
+		Evolution->ResetConstraintRules();
+		int32 InitIndex1 = Evolution->AddConstraintInitRange(1, true);
+		Evolution->ConstraintInits()[InitIndex1] =
+			[this](FSolverParticles& InParticles, const FSolverReal Dt)
+		{
+			this->GridBasedCorotatedConstraint->Init(InParticles, Dt);
+		};
+		int32 ConstraintIndex1 = Evolution->AddConstraintRuleRange(1, true);
+		Evolution->ConstraintRules()[ConstraintIndex1] =
+			[this](FSolverParticles& InParticles, const FSolverReal Dt)
+		{
+			this->GridBasedCorotatedConstraint->ApplyInParallel(InParticles, Dt);
+		};
+	
+	
+	}
+
 
 	void FDeformableSolver::FPhysicsThreadAccess::TickSimulation(FSolverReal DeltaTime)
 	{
@@ -343,6 +391,7 @@ namespace Chaos::Softs
 			Time += Property.TimeStepSize;
 		}
 
+
 		FOutputDataMap OutputBuffers;
 		for (TPair< FThreadingProxy::FKey, TUniquePtr<FThreadingProxy> > & BaseProxyPair : Proxies)
 		{
@@ -432,7 +481,7 @@ namespace Chaos::Softs
 			{
 				const TManagedArray<FIntVector>& Indices = Rest->GetAttribute<FIntVector>("Indices", FGeometryCollection::FacesGroup);
 
-				WriteTrisGEO(Evolution->Particles(), Indices);
+				WriteTrisGEO(Evolution->Particles(), *SurfaceElements);
 				FString file = FPaths::ProjectDir();
 				file.Append(TEXT("/HoudiniOutput/DtLog.txt"));
 				if (Frame == 0)
@@ -444,11 +493,11 @@ namespace Chaos::Softs
 		}
 	}
 
-	void  FDeformableSolver::FPhysicsThreadAccess::WriteTrisGEO(const FSolverParticles& Particles, const TManagedArray<FIntVector>& Mesh)
+	void  FDeformableSolver::FPhysicsThreadAccess::WriteTrisGEO(const FSolverParticles& Particles, const TArray<TVec3<int32>>& Mesh)
 	{
 		Solver.WriteTrisGEO(Particles, Mesh);
 	}
-	void FDeformableSolver::WriteTrisGEO(const FSolverParticles& Particles, const TManagedArray<FIntVector>& Mesh)
+	void FDeformableSolver::WriteTrisGEO(const FSolverParticles& Particles, const TArray<TVec3<int32>>& Mesh)
 	{
 		FString file = FPaths::ProjectDir();
 		//file.Append(TEXT("\HoudiniOuput\Test.geo"));
diff --git a/Engine/Source/Runtime/Experimental/Chaos/Private/Chaos/GraphColoring.cpp b/Engine/Source/Runtime/Experimental/Chaos/Private/Chaos/GraphColoring.cpp
--- a/Engine/Source/Runtime/Experimental/Chaos/Private/Chaos/GraphColoring.cpp
+++ b/Engine/Source/Runtime/Experimental/Chaos/Private/Chaos/GraphColoring.cpp
@@ -2,6 +2,7 @@
 #include "Chaos/GraphColoring.h"
 #include "Chaos/Array.h"
 #include "ChaosLog.h"
+#include "Chaos/Framework/Parallel.h"
 
 template<typename T>
 static bool VerifyGraph(TArray<TArray<int32>> ColorGraph, const TArray<Chaos::TVec2<int32>>& Graph, const Chaos::TDynamicParticles<T, 3>& InParticles)
@@ -999,7 +1000,7 @@ TArray<TArray<int32>> Chaos::FGraphColoring::ComputeGraphColoringAllDynamic(cons
 
 
 template<typename T>
-void Chaos::ComputeGridBasedGraphSubColoringPointer(const TArray<TArray<int32>>& ElementsPerColor, const TMPMGrid<T>& Grid, const int32 GridSize, TArray<TArray<int32>>* PreviousColoring, const TArray<TArray<int32>>& ConstraintsNodesSet, TArray<TArray<TArray<int32>>>& ElementsPerSubColors) 
+void Chaos::ComputeGridBasedGraphSubColoringPointer(const TArray<TArray<int32>>& ElementsPerColor, const TMPMGrid<T>& Grid, const int32 GridSize, TArray<TArray<int32>>*& PreviousColoring, const TArray<TArray<int32>>& ConstraintsNodesSet, TArray<TArray<TArray<int32>>>& ElementsPerSubColors) 
 {
 	bool InitialGuess = true;
 	if (!PreviousColoring) {
@@ -1009,84 +1010,86 @@ void Chaos::ComputeGridBasedGraphSubColoringPointer(const TArray<TArray<int32>>&
 	}
 	ElementsPerSubColors.SetNum(ElementsPerColor.Num());
 
-
-	for (int32 i = 0; i < ElementsPerColor.Num(); i++) {
-		if (!InitialGuess) {
-			PreviousColoring->operator[](i).Init(-1, ElementsPerColor[i].Num());
-		}
-		int32 NumNodes = GridSize;
-		TArray<int32> ElementSubColors;
-		ElementSubColors.Init(-1, ElementsPerColor[i].Num());
-		TArray<TSet<int32>*> UsedColors;
-		UsedColors.Init(nullptr, NumNodes);
-		//std::vector<std::unordered_set<int>*> UsedColors(NumNodes, nullptr);
-		for (int32 j = 0; j < ElementsPerColor[i].Num(); j++) {
-			int32 ColorToUse = 0;
-			int32 e = ElementsPerColor[i][j];
-			// check initial guess:
-			if (InitialGuess) {
-				ColorToUse = PreviousColoring->operator[](i)[j];
-				bool ColorFound = false;
-				for (auto node : ConstraintsNodesSet[e]) {
-					if (!UsedColors[node]) {
-						UsedColors[node] = new TSet<int32>();
-					}
-					if (UsedColors[node]->Contains(ColorToUse)) {
-						ColorFound = true;
-						break;
-					}
-				}
-				if (ColorFound) {
-					ColorToUse = 0;
-				}
-				else {
-					for (auto node : ConstraintsNodesSet[e]) {
-						UsedColors[node]->Emplace(ColorToUse);
-					}
-					ElementSubColors[j] = ColorToUse;
-				}
+	//for (int32 i = 0; i < ElementsPerColor.Num(); i++) {
+	PhysicsParallelFor(ElementsPerColor.Num(), [&](const int32 i)
+		{
+			if (!InitialGuess) {
+				PreviousColoring->operator[](i).Init(-1, ElementsPerColor[i].Num());
 			}
-			if (ElementSubColors[j] == -1) {
-				while (true) {
+			int32 NumNodes = GridSize;
+			TArray<int32> ElementSubColors;
+			ElementSubColors.Init(-1, ElementsPerColor[i].Num());
+			TArray<TSet<int32>*> UsedColors;
+			UsedColors.Init(nullptr, NumNodes);
+			//std::vector<std::unordered_set<int>*> UsedColors(NumNodes, nullptr);
+			for (int32 j = 0; j < ElementsPerColor[i].Num(); j++) {
+				int32 ColorToUse = 0;
+				int32 e = ElementsPerColor[i][j];
+				// check initial guess:
+				if (InitialGuess) {
+					ColorToUse = PreviousColoring->operator[](i)[j];
 					bool ColorFound = false;
 					for (auto node : ConstraintsNodesSet[e]) {
 						if (!UsedColors[node]) {
-							UsedColors[node] = new TSet<int32>();// new std::unordered_set<int>;
+							UsedColors[node] = new TSet<int32>();
 						}
 						if (UsedColors[node]->Contains(ColorToUse)) {
 							ColorFound = true;
 							break;
 						}
 					}
-					if (!ColorFound) {
-						break;
+					if (ColorFound) {
+						ColorToUse = 0;
+					}
+					else {
+						for (auto node : ConstraintsNodesSet[e]) {
+							UsedColors[node]->Emplace(ColorToUse);
+						}
+						ElementSubColors[j] = ColorToUse;
 					}
-					ColorToUse++;
 				}
-				ElementSubColors[j] = ColorToUse;
-				for (auto node : ConstraintsNodesSet[e]) {
-					UsedColors[node]->Emplace(ColorToUse);
+				if (ElementSubColors[j] == -1) {
+					while (true) {
+						bool ColorFound = false;
+						for (auto node : ConstraintsNodesSet[e]) {
+							if (!UsedColors[node]) {
+								UsedColors[node] = new TSet<int32>();// new std::unordered_set<int>;
+							}
+							if (UsedColors[node]->Contains(ColorToUse)) {
+								ColorFound = true;
+								break;
+							}
+						}
+						if (!ColorFound) {
+							break;
+						}
+						ColorToUse++;
+					}
+					ElementSubColors[j] = ColorToUse;
+					for (auto node : ConstraintsNodesSet[e]) {
+						UsedColors[node]->Emplace(ColorToUse);
+					}
 				}
-			}
 
-			// assign colors to previous guess for next timestep:
-			PreviousColoring->operator[](i)[j] = ColorToUse;
-		}
+				// assign colors to previous guess for next timestep:
+				PreviousColoring->operator[](i)[j] = ColorToUse;
+			}
 
-		for (int ii = 0; ii < UsedColors.Num(); ii++) {
-			delete (UsedColors[ii]);
-		}
+			for (int ii = 0; ii < UsedColors.Num(); ii++) {
+				delete (UsedColors[ii]);
+			}
 
-		int32 NumColors = FMath::Max<int32>(ElementSubColors);
+			int32 NumColors = FMath::Max<int32>(ElementSubColors);
 
-		//int32 num_colors = *std::max_element(ElementSubColors.begin(), ElementSubColors.end());
-		ElementsPerSubColors[i].SetNum(0);
-		ElementsPerSubColors[i].SetNum(NumColors + 1);
+			//int32 num_colors = *std::max_element(ElementSubColors.begin(), ElementSubColors.end());
+			ElementsPerSubColors[i].SetNum(0);
+			ElementsPerSubColors[i].SetNum(NumColors + 1);
 
-		for (int32 j = 0; j < ElementsPerColor[i].Num(); j++) {
-			ElementsPerSubColors[i][ElementSubColors[j]].Emplace(ElementsPerColor[i][j]);
-		}
-	}
+			for (int32 j = 0; j < ElementsPerColor[i].Num(); j++) {
+				ElementsPerSubColors[i][ElementSubColors[j]].Emplace(ElementsPerColor[i][j]);
+			}
+		}, ElementsPerColor.Num() < 20);
+	
 	checkSlow(VerifyGridBasedSubColoring<T>(ElementsPerColor, Grid, ConstraintsNodesSet, ElementsPerSubColors));
 }
 
@@ -1099,5 +1102,5 @@ template CHAOS_API TArray<TArray<int32>> Chaos::FGraphColoring::ComputeGraphColo
 template CHAOS_API TArray<TArray<int32>> Chaos::FGraphColoring::ComputeGraphColoring<Chaos::FRealDouble>(const TArray<Chaos::TVector<int32, 4>>&, const Chaos::TDynamicParticles<Chaos::FRealDouble, 3>&);
 template CHAOS_API TArray<TArray<int32>> Chaos::FGraphColoring::ComputeGraphColoringAllDynamic<Chaos::FRealSingle>(const TArray<Chaos::TVector<int32, 4>>&, const Chaos::TDynamicParticles<Chaos::FRealSingle, 3>&);
 template CHAOS_API TArray<TArray<int32>> Chaos::FGraphColoring::ComputeGraphColoringAllDynamic<Chaos::FRealDouble>(const TArray<Chaos::TVector<int32, 4>>&, const Chaos::TDynamicParticles<Chaos::FRealDouble, 3>&);
-template CHAOS_API void Chaos::ComputeGridBasedGraphSubColoringPointer(const TArray<TArray<int32>>& ElementsPerColor, const TMPMGrid<Chaos::FRealSingle>& Grid, const int32 GridSize, TArray<TArray<int32>>* PreviousColoring, const TArray<TArray<int32>>& ConstraintsNodesSet, TArray<TArray<TArray<int32>>>& ElementsPerSubColors);
-template CHAOS_API void Chaos::ComputeGridBasedGraphSubColoringPointer(const TArray<TArray<int32>>& ElementsPerColor, const TMPMGrid<Chaos::FRealDouble>& Grid, const int32 GridSize, TArray<TArray<int32>>* PreviousColoring, const TArray<TArray<int32>>& ConstraintsNodesSet, TArray<TArray<TArray<int32>>>& ElementsPerSubColors);
\ No newline at end of file
+template CHAOS_API void Chaos::ComputeGridBasedGraphSubColoringPointer(const TArray<TArray<int32>>& ElementsPerColor, const TMPMGrid<Chaos::FRealSingle>& Grid, const int32 GridSize, TArray<TArray<int32>>*& PreviousColoring, const TArray<TArray<int32>>& ConstraintsNodesSet, TArray<TArray<TArray<int32>>>& ElementsPerSubColors);
+template CHAOS_API void Chaos::ComputeGridBasedGraphSubColoringPointer(const TArray<TArray<int32>>& ElementsPerColor, const TMPMGrid<Chaos::FRealDouble>& Grid, const int32 GridSize, TArray<TArray<int32>>*& PreviousColoring, const TArray<TArray<int32>>& ConstraintsNodesSet, TArray<TArray<TArray<int32>>>& ElementsPerSubColors);
\ No newline at end of file
diff --git a/Engine/Source/Runtime/Experimental/Chaos/Private/Chaos/UniformGrid.cpp b/Engine/Source/Runtime/Experimental/Chaos/Private/Chaos/UniformGrid.cpp
--- a/Engine/Source/Runtime/Experimental/Chaos/Private/Chaos/UniformGrid.cpp
+++ b/Engine/Source/Runtime/Experimental/Chaos/Private/Chaos/UniformGrid.cpp
@@ -390,25 +390,15 @@ bool Chaos::TUniformGrid<T, 3>::IsValid(const TVector<int32, 3>& X) const
 template class CHAOS_API Chaos::TUniformGridBase<Chaos::FReal, 3>;
 template class CHAOS_API Chaos::TUniformGrid<Chaos::FReal, 3>;
 template class CHAOS_API Chaos::TUniformGrid<Chaos::FReal, 2>;
+template class CHAOS_API Chaos::TMPMGrid<Chaos::FRealSingle>;
+template class CHAOS_API Chaos::TMPMGrid<Chaos::FReal>;
 #else
 template class Chaos::TUniformGridBase<Chaos::FReal, 3>;
 template class Chaos::TUniformGrid<Chaos::FReal, 3>;
 template class Chaos::TUniformGrid<Chaos::FReal, 2>;
+template class Chaos::TMPMGrid<Chaos::FRealSingle>;
+template class Chaos::TMPMGrid<Chaos::FReal>;
 #endif
 
 template FVec3 Chaos::TUniformGridBase<Chaos::FReal, 3>::LinearlyInterpolate<FVec3>(const TArrayND<FVec3, 3>&, const FVec3&) const;
-template CHAOS_API Chaos::FReal Chaos::TUniformGridBase<Chaos::FReal, 3>::LinearlyInterpolate<Chaos::FReal>(const TArrayND<Chaos::FReal, 3>&, const FVec3&) const;
-template int32 TMPMGrid<Chaos::FReal>::Loc2GlobIndex(const int32 IndexIn, const TVector<int32, 3>& LocalIndexIn) const;
-template int32 TMPMGrid<Chaos::FReal>::FlatIndex(const TVector<int32, 3>& Index) const;
-template TVector<int32, 3> TMPMGrid<Chaos::FReal>::Lin2MultiIndex(const int32 IndexIn) const;
-template void TMPMGrid<Chaos::FReal>::SetInterp(InterpType InterpIn);
-template void TMPMGrid<Chaos::FReal>::UpdateGridFromPositions(const Chaos::TDynamicParticles<Chaos::FReal, 3>& InParticles);
-template void TMPMGrid<Chaos::FReal>::BaseNodeIndex(const TVector<Chaos::FReal, 3>& X, TVector<int32, 3>& Index, TVector<Chaos::FReal, 3>& Weights) const;
-template TVector<Chaos::FReal, 3> TMPMGrid<Chaos::FReal>::Node(int32 FlatIndexIn) const;
-template int32 TMPMGrid<Chaos::FRealSingle>::Loc2GlobIndex(const int32 IndexIn, const TVector<int32, 3>& LocalIndexIn) const;
-template int32 TMPMGrid<Chaos::FRealSingle>::FlatIndex(const TVector<int32, 3>& Index) const;
-template TVector<int32, 3> TMPMGrid<Chaos::FRealSingle>::Lin2MultiIndex(const int32 IndexIn) const;
-template void TMPMGrid<Chaos::FRealSingle>::SetInterp(InterpType InterpIn);
-template void TMPMGrid<Chaos::FRealSingle>::UpdateGridFromPositions(const Chaos::TDynamicParticles<Chaos::FRealSingle, 3>& InParticles);
-template void TMPMGrid<Chaos::FRealSingle>::BaseNodeIndex(const TVector<Chaos::FRealSingle, 3>& X, TVector<int32, 3>& Index, TVector<Chaos::FRealSingle, 3>& Weights) const;
-template TVector<Chaos::FRealSingle, 3> TMPMGrid<Chaos::FRealSingle>::Node(int32 FlatIndexIn) const;
\ No newline at end of file
+template CHAOS_API Chaos::FReal Chaos::TUniformGridBase<Chaos::FReal, 3>::LinearlyInterpolate<Chaos::FReal>(const TArrayND<Chaos::FReal, 3>&, const FVec3&) const;
\ No newline at end of file
diff --git a/Engine/Source/Runtime/Experimental/Chaos/Public/Chaos/Deformable/ChaosDeformableSolver.h b/Engine/Source/Runtime/Experimental/Chaos/Public/Chaos/Deformable/ChaosDeformableSolver.h
--- a/Engine/Source/Runtime/Experimental/Chaos/Public/Chaos/Deformable/ChaosDeformableSolver.h
+++ b/Engine/Source/Runtime/Experimental/Chaos/Public/Chaos/Deformable/ChaosDeformableSolver.h
@@ -11,6 +11,7 @@
 #include "Chaos/PBDTriangleMeshCollisions.h"
 #include "Chaos/TriangleMesh.h"
 #include "Chaos/XPBDCorotatedConstraints.h"
+#include "Chaos/XPBDGridBasedCorotatedConstraints.h"
 #include "Chaos/Deformable/ChaosDeformableSolverTypes.h"
 #include "Chaos/Deformable/ChaosDeformableSolverProxy.h"
 #include "CoreMinimal.h"
@@ -52,7 +53,7 @@ namespace Chaos::Softs
 			void InitializeSelfCollisionVariables();
 
 			/*IO Utility*/
-			void WriteTrisGEO(const Softs::FSolverParticles& Particles, const TManagedArray<FIntVector>& Mesh);
+			void WriteTrisGEO(const Softs::FSolverParticles& Particles, const TArray<TVec3<int32>>& Mesh);
 			void WriteFrame(FThreadingProxy& , const FSolverReal DeltaTime);
 
 			const FDeformableSolverProperties& GetProperties() const { return Solver.GetProperties(); }
@@ -99,9 +100,10 @@ namespace Chaos::Softs
 		void InitializeCollisionBodies();
 		void InitializeKinematicState(FThreadingProxy&);
 		void InitializeSelfCollisionVariables();
+		void InitializeGridBasedConstraintVariables();
 
 		/*IO Utility*/
-		void WriteTrisGEO(const FSolverParticles& Particles, const TManagedArray<FIntVector>& Mesh);
+		void WriteTrisGEO(const FSolverParticles& Particles, const TArray<TVec3<int32>>& Mesh);
 		void WriteFrame(FThreadingProxy&, const FSolverReal DeltaTime);
 
 		/*Game Thread API*/
@@ -125,10 +127,12 @@ namespace Chaos::Softs
 		// Simulation Variables
 		TUniquePtr<Softs::FPBDEvolution> Evolution;
 		TArray<TUniquePtr<Softs::FXPBDCorotatedConstraints<Softs::FSolverReal, Softs::FSolverParticles>>> CorotatedConstraints;
+		TUniquePtr<Softs::FXPBDGridBasedCorotatedConstraints<Softs::FSolverReal, Softs::FSolverParticles>> GridBasedCorotatedConstraint;
 		TUniquePtr<Softs::FPBDCollisionSpringConstraints> CollisionSpringConstraint;
 		TUniquePtr<Softs::FPBDTriangleMeshCollisions> TriangleMeshCollisions;
 		TArrayCollectionArray<const UObject*> MObjects;
 		TUniquePtr <TArray<TVec3<int32>>> SurfaceElements;
+		TUniquePtr <TArray<Chaos::TVec4<int32>>> AllElements;
 		TUniquePtr <FTriangleMesh> SurfaceTriangleMesh;
 
 		FSolverReal Time = 0.f;
diff --git a/Engine/Source/Runtime/Experimental/Chaos/Public/Chaos/Deformable/ChaosDeformableSolverTypes.h b/Engine/Source/Runtime/Experimental/Chaos/Public/Chaos/Deformable/ChaosDeformableSolverTypes.h
--- a/Engine/Source/Runtime/Experimental/Chaos/Public/Chaos/Deformable/ChaosDeformableSolverTypes.h
+++ b/Engine/Source/Runtime/Experimental/Chaos/Public/Chaos/Deformable/ChaosDeformableSolverTypes.h
@@ -20,7 +20,9 @@ namespace Chaos::Softs
 			bool InCacheToFile = false,
 			bool InbEnableKinematics = true,
 			bool InbUseFloor = true,
-			bool InbDoSelfCollision = false)
+			bool InbDoSelfCollision = false,
+			bool InbUseGridBasedConstraints = false,
+			FSolverReal InGridDx = (FSolverReal)1. )
 			: NumSolverSubSteps(InNumSolverSubSteps)
 			, NumSolverIterations(InNumSolverIterations)
 			, FixTimeStep(InFixTimeStep)
@@ -29,6 +31,8 @@ namespace Chaos::Softs
 			, bEnableKinematics(InbEnableKinematics)
 			, bUseFloor(InbUseFloor)
 			, bDoSelfCollision(InbDoSelfCollision)
+			, bUseGridBasedConstraints(InbUseGridBasedConstraints)
+			, GridDx(InGridDx)
 		{}
 
 		FDeformableSolverProperties(const FDeformableSolverProperties& InProp)
@@ -40,6 +44,8 @@ namespace Chaos::Softs
 			, bEnableKinematics(InProp.bEnableKinematics)
 			, bUseFloor(InProp.bUseFloor)
 			, bDoSelfCollision(InProp.bDoSelfCollision)
+			, bUseGridBasedConstraints(InProp.bUseGridBasedConstraints)
+			, GridDx(InProp.GridDx)
 		{}
 
 		int32 NumSolverSubSteps = 5;
@@ -50,6 +56,8 @@ namespace Chaos::Softs
 		bool bEnableKinematics = true;
 		bool bUseFloor = true;
 		bool bDoSelfCollision = false;
+		bool bUseGridBasedConstraints = false;
+		FSolverReal GridDx = (FSolverReal)1.;
 	};
 
 
diff --git a/Engine/Source/Runtime/Experimental/Chaos/Public/Chaos/GraphColoring.h b/Engine/Source/Runtime/Experimental/Chaos/Public/Chaos/GraphColoring.h
--- a/Engine/Source/Runtime/Experimental/Chaos/Public/Chaos/GraphColoring.h
+++ b/Engine/Source/Runtime/Experimental/Chaos/Public/Chaos/GraphColoring.h
@@ -72,6 +72,6 @@ class FGraphColoring
 };
 
 template<typename T> 
-CHAOS_API void ComputeGridBasedGraphSubColoringPointer(const TArray<TArray<int32>>& ElementsPerColor, const TMPMGrid<T>& Grid, const int32 GridSize, TArray<TArray<int32>>* PreviousColoring, const TArray<TArray<int32>>& ConstraintsNodesSet, TArray<TArray<TArray<int32>>>& ElementsPerSubColors);
+CHAOS_API void ComputeGridBasedGraphSubColoringPointer(const TArray<TArray<int32>>& ElementsPerColor, const TMPMGrid<T>& Grid, const int32 GridSize, TArray<TArray<int32>>*& PreviousColoring, const TArray<TArray<int32>>& ConstraintsNodesSet, TArray<TArray<TArray<int32>>>& ElementsPerSubColors);
 
 }
diff --git a/Engine/Source/Runtime/Experimental/Chaos/Public/Chaos/MPMTransfer.h b/Engine/Source/Runtime/Experimental/Chaos/Public/Chaos/MPMTransfer.h
--- a/Engine/Source/Runtime/Experimental/Chaos/Public/Chaos/MPMTransfer.h
+++ b/Engine/Source/Runtime/Experimental/Chaos/Public/Chaos/MPMTransfer.h
@@ -33,6 +33,7 @@ public:
 	TArray<TArray<int32>> ElementGridNodes;
 	TArray<TArray<T>> ElementGridNodeWeights;
 	TArray<TArray<TArray<int32>>> ElementGridNodeIncidentElements;
+	TArray<TArray<int32>> ElementGridNodesSet;
 
 
 	TMPMTransfer(){}
@@ -102,17 +103,20 @@ public:
 					{
 						int32 CurrentLocIndex = ii * NPerSec + jj * Grid.NPerDir + kk;
 						PhysicsParallelFor(GridCells[0] / Grid.NPerDir, [&](const int32 iii)
-							{
+							//{
+						//for (int32 iii = 0; iii < int32(GridCells[0] / Grid.NPerDir); iii++)
+						{
 								for (int32 jjj = 0; jjj < int32(GridCells[1] / Grid.NPerDir); jjj++)
 								{
 									for (int32 kkk = 0; kkk < int32(GridCells[2] / Grid.NPerDir); kkk++)
 									{
-										TVector<int32, 3> MultiIndex = { iii * int32(Grid.NPerDir) + ii, jjj * int32(Grid.NPerDir) + jj, kkk * int32(Grid.NPerDir) + kk};
+										TVector<int32, 3> MultiIndex = { iii * int32(Grid.NPerDir) + ii, jjj * int32(Grid.NPerDir) + jj, kkk * int32(Grid.NPerDir) + kk };
 										int32 CellIndex = Grid.FlatIndex(MultiIndex);
 										P2GApplyHelper(InParticles, CellIndex, GridData);
 									}
 								}
-						});
+							});
+						//}
 
 					}
 				}
@@ -145,6 +149,10 @@ public:
 							T Nkk = Grid.Nijk(Weights[p][2], kkk);
 							T NProd = Nii * Njj * Nkk;
 							GridData[(NTransfer + 1) * GlobIndex] += NProd * InParticles.M(p);
+							//if (GlobIndex == 13472)
+							//{
+							//	std::cout << "";
+							//}
 							for (int32 alpha = 0; alpha < 3; alpha++)
 							{
 								GridData[(NTransfer + 1) * GlobIndex + alpha + 1] += NProd * InParticles.M(p) * InParticles.V(p)[alpha];
@@ -162,6 +170,7 @@ public:
 		ElementGridNodes.SetNum(InMesh.Num());
 		ElementGridNodeWeights.SetNum(InMesh.Num());
 		ElementGridNodeIncidentElements.SetNum(InMesh.Num());
+		ElementGridNodesSet.SetNum(InMesh.Num());
 		//TODO(Yizhou): make following loop parallel for with appropriate bool condition
 		for (int32 e = 0; e < InMesh.Num(); e++)
 		{
@@ -171,13 +180,13 @@ public:
 			{
 				int32 p = InMesh[e][ie];
 				TVector<int32, 3> Index = Indices[p];
-				for (int32 ii = 0; ii < Grid.NPerDir; ii++)
+				for (int32 ii = 0; ii < int32(Grid.NPerDir); ii++)
 				{
 					T Nii = Grid.Nijk(Weights[p][0], ii);
-					for (int32 jj = 0; jj < Grid.NPerDir; jj++)
+					for (int32 jj = 0; jj < int32(Grid.NPerDir); jj++)
 					{
 						T Njj = Grid.Nijk(Weights[p][1], jj);
-						for (int32 kk = 0; kk < Grid.NPerDir; kk++)
+						for (int32 kk = 0; kk < int32(Grid.NPerDir); kk++)
 						{
 							T Nkk = Grid.Nijk(Weights[p][2], kk);
 							TVector<int32, 3> LocIndex = { ii, jj, kk };
@@ -190,6 +199,11 @@ public:
 				}
 			}
 			ComputeIncidentElements(ElementGridNodes[e], ElementGridNodeIncidentElements[e]);
+			ElementGridNodesSet[e].SetNum(ElementGridNodeIncidentElements[e].Num());
+			for (int32 kk = 0; kk < ElementGridNodeIncidentElements[e].Num(); kk++)
+			{
+				ElementGridNodesSet[e][kk] = ElementGridNodes[e][ElementGridNodeIncidentElements[e][kk][0]];
+			}
 		}
 	}
 
@@ -230,8 +244,71 @@ public:
 		}
 	}
 
-	
+	void ComputeGridPositions(const TArray<T>& GridData, const T Dt, TArray<TVector<T, 3>>& GridPositions)
+	{
+		GridPositions.SetNum(Grid.Size());
+		for (int32 i = 0; i < Grid.Size(); i++)
+		{
+			TVector<T,3> XOld = Grid.Node(i);
+			T Mass = GridData[(NTransfer + 1) * i];
+			if (Mass == 0) {
+				GridPositions[i] = XOld;
+			}
+			else {
+				for (int32 Alpha = 0; Alpha < int32(NTransfer); Alpha++)
+					GridPositions[i][Alpha] = Dt * GridData[(NTransfer + 1) * i + Alpha + 1] / Mass + XOld[Alpha];
+			}
+		}
+	}
+
 
+	TVec4<TVector<T, 3>> SparseG2P(const TArray<TVector<T, 3>>& GridPositions, const int32 ElementIndex)
+	{
+		TVec4<TVector<T, 3>> Result(TVector<T, 3>((T)0.));
+
+		for (int32 i = 0; i < 4; i++) {
+			for (int32 ii = 0; ii < int32(NPerEle); ii++) {
+				int32 FlatIndex = ElementGridNodes[ElementIndex][NPerEle * i + ii];
+				for (int32 Alpha = 0; Alpha < 3; Alpha++) {
+					Result[i][Alpha] += ElementGridNodeWeights[ElementIndex][NPerEle * i + ii] * GridPositions[FlatIndex][Alpha];
+				}
+			}
+		}
+		
+		return Result;
+	}
+
+
+	TArray<TVector<T, 3>> SparseP2G(const TVec4<TVector<T, 3>>& InGradient, const int32 ElementIndex)
+	{
+		
+		TArray<TVector<T, 3>> GridGradient;
+		GridGradient.Init(TVector<T, 3>((T)0.), 4 * NPerEle);
+		for (int32 i = 0; i < 4; i++) {
+			for (int32 ii = 0; ii < int32(NPerEle); ii++) {
+				for (int32 alpha = 0; alpha < 3; alpha++) {
+					GridGradient[NPerEle * i + ii][alpha] += ElementGridNodeWeights[ElementIndex][NPerEle * i + ii] * InGradient[i][alpha];
+				}
+			}
+		}
+
+		TArray<TVector<T, 3>> GridGradientCompact;
+		GridGradientCompact.Init(TVector<T, 3>((T)0.), ElementGridNodeIncidentElements[ElementIndex].Num());
+
+		for (int32 node = 0; node < ElementGridNodeIncidentElements[ElementIndex].Num(); node++) 
+		{
+			for (int32 node_indices = 0; node_indices < ElementGridNodeIncidentElements[ElementIndex][node].Num(); node_indices++) 
+			{
+				for (int32 alpha = 0; alpha < 3; alpha++) 
+				{
+					GridGradientCompact[node][alpha] += GridGradient[ElementGridNodeIncidentElements[ElementIndex][node][node_indices]][alpha];
+				}
+			}
+		}
+
+		return GridGradient;
+	
+	}
 
 };
 
diff --git a/Engine/Source/Runtime/Experimental/Chaos/Public/Chaos/UniformGrid.h b/Engine/Source/Runtime/Experimental/Chaos/Public/Chaos/UniformGrid.h
--- a/Engine/Source/Runtime/Experimental/Chaos/Public/Chaos/UniformGrid.h
+++ b/Engine/Source/Runtime/Experimental/Chaos/Public/Chaos/UniformGrid.h
@@ -320,6 +320,8 @@ public:
 	TMPMGrid() {}
 	TMPMGrid(const TVector<T, 3>& MinCorner, const TVector<T, 3>& MaxCorner, const TVector<int32, 3>& Cells, const uint32 GhostCells = 0)
 		: TUniformGridBase<T, 3>(MinCorner, MaxCorner, Cells, GhostCells) {}
+	TMPMGrid(const int32 GridN) { for (int32 i = 0; i < 3; i++) { MDx[i] = (T)1. / (T)GridN; } }
+	TMPMGrid(const T GridDx) { for (int32 i = 0; i < 3; i++) { MDx[i] = GridDx; } }
 	TMPMGrid(std::istream& Stream)
 		: TUniformGridBase<T, 3>(Stream) {}
 	~TMPMGrid() {}
diff --git a/Engine/Source/Runtime/Experimental/Chaos/Public/Chaos/XPBDCorotatedConstraints.h b/Engine/Source/Runtime/Experimental/Chaos/Public/Chaos/XPBDCorotatedConstraints.h
--- a/Engine/Source/Runtime/Experimental/Chaos/Public/Chaos/XPBDCorotatedConstraints.h
+++ b/Engine/Source/Runtime/Experimental/Chaos/Public/Chaos/XPBDCorotatedConstraints.h
@@ -62,6 +62,39 @@ namespace Chaos::Softs
 			InitColor(InParticles);
 		}
 
+		FXPBDCorotatedConstraints(
+			const ParticleType& InParticles,
+			const TArray<TVector<int32, 4>>& InMesh,
+			const T GridN = (T).1,
+			const T& EMesh = (T)10.0,
+			const T& NuMesh = (T).3
+		)
+			: MeshConstraints(InMesh)
+		{
+			LambdaArray.Init((T)0., 2 * MeshConstraints.Num());
+			DmInverse.Init((T)0., 9 * MeshConstraints.Num());
+			Measure.Init((T)0., MeshConstraints.Num());
+			Lambda = EMesh * NuMesh / (((T)1. + NuMesh) * ((T)1. - (T)2. * NuMesh));
+			Mu = EMesh / ((T)2. * ((T)1. + NuMesh));
+			for (int e = 0; e < InMesh.Num(); e++)
+			{
+				PMatrix<T, 3, 3> Dm = DsInit(e, InParticles);
+				PMatrix<T, 3, 3> DmInv = Dm.Inverse();
+				for (int r = 0; r < 3; r++) {
+					for (int c = 0; c < 3; c++) {
+						DmInverse[(3 * 3) * e + 3 * r + c] = DmInv.GetAt(r, c);
+					}
+				}
+
+				Measure[e] = Dm.Determinant() / (T)6.;
+
+				if (Measure[e] < (T)0.)
+				{
+					Measure[e] = -Measure[e];
+				}
+			}
+		}
+
 		virtual ~FXPBDCorotatedConstraints() {}
 
 		PMatrix<T, 3, 3> DsInit(const int e, const ParticleType& InParticles) const {
@@ -395,8 +428,8 @@ namespace Chaos::Softs
 				return TVec4<TVector<T, 3>>(TVector<T, 3>((T)0.));
 			}
 
-			TVector<T, 81> dRdF((T)0.);
-			Chaos::dRdFCorotated(Fe, dRdF);
+			//TVector<T, 81> dRdF((T)0.);
+			//Chaos::dRdFCorotated(Fe, dRdF);
 
 			PMatrix<T, 3, 3> DmInvT = ElementDmInv(ElementIndex).GetTransposed();
 
diff --git a/Engine/Source/Runtime/Experimental/Chaos/Public/Chaos/XPBDGridBasedCorotatedConstraints.h b/Engine/Source/Runtime/Experimental/Chaos/Public/Chaos/XPBDGridBasedCorotatedConstraints.h
new file mode 100644
--- /dev/null
+++ b/Engine/Source/Runtime/Experimental/Chaos/Public/Chaos/XPBDGridBasedCorotatedConstraints.h
@@ -0,0 +1,490 @@
+// Copyright Epic Games, Inc. All Rights Reserved.
+#pragma once
+
+#include "Chaos/PBDSpringConstraintsBase.h"
+#include "ChaosStats.h"
+#include "Chaos/ImplicitQRSVD.h"
+#include "Chaos/GraphColoring.h"
+#include "Chaos/Framework/Parallel.h"
+#include "Chaos/MPMTransfer.h"
+#include "Misc/FileHelper.h"
+#include "HAL/PlatformFileManager.h"
+#include "Misc/Paths.h"
+
+//DECLARE_CYCLE_STAT(TEXT("Chaos XPBD Corotated Constraint"), STAT_XPBD_Spring, STATGROUP_Chaos);
+
+namespace Chaos::Softs
+{
+	template <typename T, typename ParticleType>
+	class FXPBDGridBasedCorotatedConstraints : public FXPBDCorotatedConstraints <T, ParticleType>
+	{
+
+		typedef FXPBDCorotatedConstraints<T, ParticleType> Base;
+		using Base::MeshConstraints;
+		using Base::LambdaArray;
+		using Base::Measure;
+		using Base::DmInverse;
+		using Base::Mu;
+		using Base::Lambda;
+
+	public:
+		//this one only accepts tetmesh input and mesh
+		FXPBDGridBasedCorotatedConstraints(
+			const ParticleType& InParticles,
+			const TArray<TVector<int32, 4>>& InMesh,
+			const T GridDx = (T).1,
+			const bool bRecordMetricIn = true,
+			const T MaxDtIn = (T).1,
+			const T MinDtIn = (T).01,
+			const T CFLCoeffIn = (T).4,
+			const T& EMesh = (T)1000.0,
+			const T& NuMesh = (T).3
+		)
+			: Base(InParticles, InMesh, GridDx, EMesh, NuMesh), MaxDt(MaxDtIn), MinDt(MinDtIn), CFLCoeff(CFLCoeffIn), MPMGrid(GridDx), MPMTransfer(MPMGrid), PreviousColoring(nullptr)
+		{
+			//TMPMGrid<T> MPMGrid(GridN);
+			MPMGrid.UpdateGridFromPositions(InParticles);
+			//TMPMTransfer<T> MPMTrasnfer(MPMGrid);
+			LambdaArray.Init((T)0., 2 * MeshConstraints.Num());
+			DmInverse.Init((T)0., 9 * MeshConstraints.Num());
+			Measure.Init((T)0., MeshConstraints.Num());
+			Lambda = EMesh * NuMesh / (((T)1. + NuMesh) * ((T)1. - (T)2. * NuMesh));
+			Mu = EMesh / ((T)2. * ((T)1. + NuMesh));
+			for (int e = 0; e < InMesh.Num(); e++)
+			{
+				PMatrix<T, 3, 3> Dm = Base::DsInit(e, InParticles);
+				PMatrix<T, 3, 3> DmInv = Dm.Inverse();
+				for (int r = 0; r < 3; r++) {
+					for (int c = 0; c < 3; c++) {
+						DmInverse[(3 * 3) * e + 3 * r + c] = DmInv.GetAt(r, c);
+					}
+				}
+
+				Measure[e] = Dm.Determinant() / (T)6.;
+
+				if (Measure[e] < (T)0.)
+				{
+					Measure[e] = -Measure[e];
+				}
+
+			}
+
+			ElementsPerColor = FGraphColoring::ComputeGraphColoringAllDynamic(MeshConstraints, InParticles);
+			Init(InParticles, (T)0.);
+			WriteGridNodes();
+		}
+
+		void WriteGridNodes()
+		{
+			FString file = FPaths::ProjectDir();
+			//file.Append(TEXT("\HoudiniOuput\Test.geo"));
+			file.Append(TEXT("/HoudiniOutput/GridData.geo"));
+
+			TArray<TVector<T, 3>> ActivatedGridNodes;
+
+			for (int32 i = 0; i < MPMGrid.Size(); i++)
+			{
+				if (GridData[i * (MPMTransfer.NTransfer+1)] > T(0.))
+				{
+					TVector<T, 3> GridNode = MPMGrid.Node(i);
+					ActivatedGridNodes.Emplace(GridNode);
+				}
+			}
+
+			int32 Np = ActivatedGridNodes.Num();
+
+			// We will use this FileManager to deal with the file.
+			IPlatformFile& FileManager = FPlatformFileManager::Get().GetPlatformFile();
+			FFileHelper::SaveStringToFile(FString(TEXT("PGEOMETRY V5\r\n")), *file);
+			FString HeaderInfo = FString(TEXT("NPoints ")) + FString::FromInt(Np) + FString(TEXT(" NPrims ")) + FString::FromInt(0) + FString(TEXT("\r\n"));
+			FString MoreHeader = FString(TEXT("NPointGroups 0 NPrimGroups 0\r\nNPointAttrib 0 NVertexAttrib 0 NPrimAttrib 0 NAttrib 0\r\n"));
+
+			FFileHelper::SaveStringToFile(HeaderInfo, *file, FFileHelper::EEncodingOptions::AutoDetect, &IFileManager::Get(), EFileWrite::FILEWRITE_Append);
+			FFileHelper::SaveStringToFile(MoreHeader, *file, FFileHelper::EEncodingOptions::AutoDetect, &IFileManager::Get(), EFileWrite::FILEWRITE_Append);
+
+			for (int32 i = 0; i < Np; i++) {
+
+				FString ParticleInfo = FString::SanitizeFloat(ActivatedGridNodes[i][0]) + FString(TEXT(" ")) + FString::SanitizeFloat(ActivatedGridNodes[i][1]) + FString(TEXT(" ")) + FString::SanitizeFloat(ActivatedGridNodes[i][2]) + FString(TEXT(" ")) + FString::FromInt(1) + FString(TEXT("\r\n"));
+				FFileHelper::SaveStringToFile(ParticleInfo, *file, FFileHelper::EEncodingOptions::AutoDetect, &IFileManager::Get(), EFileWrite::FILEWRITE_Append);
+
+			}
+
+
+			FFileHelper::SaveStringToFile(FString(TEXT("beginExtra\n")), *file, FFileHelper::EEncodingOptions::AutoDetect, &IFileManager::Get(), EFileWrite::FILEWRITE_Append);
+			FFileHelper::SaveStringToFile(FString(TEXT("endExtra\n")), *file, FFileHelper::EEncodingOptions::AutoDetect, &IFileManager::Get(), EFileWrite::FILEWRITE_Append);
+
+		}
+
+		void Init(const ParticleType& InParticles, const T Dt) 
+		{
+			for (T& Lambdas : LambdaArray) { Lambdas = (T)0.; }
+			{
+				TRACE_CPUPROFILER_EVENT_SCOPE(TEXT("STAT_ChaosXPBDGridBasedCorotatedUpdateGrid"));
+				MPMTransfer.Grid.UpdateGridFromPositions(InParticles);
+			}
+			{
+				TRACE_CPUPROFILER_EVENT_SCOPE(TEXT("STAT_ChaosXPBDGridBasedCorotatedInitialP2G"));
+				MPMTransfer.InitialP2G(InParticles, GridData);
+			}
+			{
+				TRACE_CPUPROFILER_EVENT_SCOPE(TEXT("STAT_ChaosXPBDGridBasedCorotatedMetaDataComputation"));
+				MPMTransfer.ComputeElementMetaData(MeshConstraints);
+			}
+			{
+				TRACE_CPUPROFILER_EVENT_SCOPE(TEXT("STAT_ChaosXPBDGridBasedCorotatedSubcoloring"));
+				ComputeGridBasedGraphSubColoringPointer(ElementsPerColor, MPMTransfer.Grid, MPMTransfer.Grid.Size(), PreviousColoring, MPMTransfer.ElementGridNodesSet, ElementsPerSubColors);
+			}
+			{
+				TRACE_CPUPROFILER_EVENT_SCOPE(TEXT("STAT_ChaosXPBDGridBasedCorotatedComputeGridPositions"));
+				MPMTransfer.ComputeGridPositions(GridData, Dt, GridPositions);
+			}
+
+			//if (!InitialGridDataWritten)
+			//{
+			//	WriteGridNodes();
+			//}
+			//MPMTransfer.Grid.UpdateGridFromPositions(InParticles);
+			//MPMTransfer.InitialP2G(InParticles, GridData);
+			//MPMTransfer.ComputeElementMetaData(MeshConstraints);
+			//ComputeGridBasedGraphSubColoringPointer(ElementsPerColor, MPMTransfer.Grid, MPMTransfer.Grid.Size(), PreviousColoring, MPMTransfer.ElementGridNodesSet, ElementsPerSubColors);
+			//MPMTransfer.ComputeGridPositions(GridData, Dt, GridPositions);
+		}
+
+		T ComputeCFLDt(const ParticleType& InParticles)
+		{
+			T VMax = (T)-1.;
+			for (int32 p = 0; p < InParticles.Size(); p++)
+			{
+				T VMag = InParticles.V(p).Size();
+				if (VMag > VMax)
+				{
+					VMax = VMag;
+				}
+			}
+
+			if (VMax > 1e5)
+			{
+				return MinDt;
+			}
+			if (VMax < 1e-5)
+			{
+				return MaxDt;
+			}
+
+			T CFLDt = MPMGrid.MDx[0] * CFLCoeff / VMax;
+
+			if (CFLDt > MaxDt)
+			{
+				return MaxDt;
+			}
+
+			if (CFLDt < MinDt)
+			{
+				return MinDt;
+			}
+
+			return CFLDt;
+
+		}
+
+		const TArray<TArray<int32>> GetElementsPerColor () const { return ElementsPerColor; }
+		const TArray<TArray<int32>> GetPreviousColoring () const { return *PreviousColoring; }
+		const TArray<TArray<TArray<int32>>> GetElementsPerSubColors() const { return ElementsPerSubColors; }
+		const TMPMTransfer<T> GetMPMTransfer() const { return MPMTransfer; }
+
+		void ApplyPolar(ParticleType& Particles, const T Dt, const int32 ElementIndex)  
+		{
+			//step 1: sparse g2p:
+			TVec4<TVector<T, 3>> XPosFromGrid = MPMTransfer.SparseG2P(GridPositions, ElementIndex);
+
+			//step 2: gradient and constraint computation:
+			T C1 = (T)0.;
+			TVec4<TVector<T, 3>> PolarGradient = GetPolarGradient(XPosFromGrid, ElementIndex, C1);
+			
+			if (C1 == (T)0.)
+			{
+				return;
+			}
+
+			//step 3 : sparse p2g:
+			TArray<TVector<T, 3>> GridGradient = MPMTransfer.SparseP2G(PolarGradient, ElementIndex);
+
+			//step 4: computation for delta:
+			T AlphaTilde = (T)1. / (Dt * Dt * Mu * Measure[ElementIndex]);
+
+			T DeltaLambda = -C1 - AlphaTilde * LambdaArray[2 * ElementIndex + 0];
+			T Denom = AlphaTilde;
+
+			for (int32 k = 0; k < MPMTransfer.ElementGridNodeIncidentElements[ElementIndex].Num(); k++) {
+				T Mass = GridData[(MPMTransfer.NTransfer + 1) * MPMTransfer.ElementGridNodes[ElementIndex][MPMTransfer.ElementGridNodeIncidentElements[ElementIndex][k][0]]];
+				if (Mass != 0) {
+					for (int32 alpha = 0; alpha < 3; alpha++) {
+						Denom += GridGradient[k][alpha] * GridGradient[k][alpha] / Mass;
+					}
+				}
+			}
+
+			DeltaLambda /= Denom;
+			LambdaArray[2 * ElementIndex + 0] += DeltaLambda;
+
+			for (int32 k = 0; k < MPMTransfer.ElementGridNodeIncidentElements[ElementIndex].Num(); k++) {
+				int32 CurrentNode = MPMTransfer.ElementGridNodes[ElementIndex][MPMTransfer.ElementGridNodeIncidentElements[ElementIndex][k][0]];
+				T Mass = GridData[(MPMTransfer.NTransfer + 1) * CurrentNode];
+				if (Mass != 0) {
+					for (int32 alpha = 0; alpha < 3; alpha++) {
+						GridPositions[CurrentNode][alpha] += DeltaLambda * GridGradient[k][alpha] / Mass;
+					}
+				}
+			}
+		}
+
+		void ApplyDet(ParticleType& Particles, const T Dt, const int32 ElementIndex) 
+		{
+			//TArray<TVector<T, 3>> XPosFromGrid;
+			//XPosFromGrid.SetNum(4);
+			//step 1: sparse g2p:
+			TVec4<TVector<T, 3>> XPosFromGrid = MPMTransfer.SparseG2P(GridPositions, ElementIndex);
+
+			//step 2: gradient and constraint computation:
+			T C2 = (T)0.;
+			TVec4<TVector<T, 3>> DetGradient = GetDetGradient(XPosFromGrid, ElementIndex, C2);
+
+			if (C2 == (T)0.)
+			{
+				return;
+			}
+
+			//step 3 : sparse p2g:
+			TArray<TVector<T, 3>> GridGradient = MPMTransfer.SparseP2G(DetGradient, ElementIndex);
+
+			//step 4: computation for delta:
+			T AlphaTilde = (T)2. / (Dt * Dt * Lambda * Measure[ElementIndex]);
+
+			T DeltaLambda = -C2 - AlphaTilde * LambdaArray[2 * ElementIndex + 1];
+			T Denom = AlphaTilde;
+
+			for (int32 k = 0; k < MPMTransfer.ElementGridNodeIncidentElements[ElementIndex].Num(); k++) {
+				T Mass = GridData[(MPMTransfer.NTransfer + 1) * MPMTransfer.ElementGridNodes[ElementIndex][MPMTransfer.ElementGridNodeIncidentElements[ElementIndex][k][0]]];
+				if (Mass != 0) {
+					for (int32 alpha = 0; alpha < 3; alpha++) {
+						Denom += GridGradient[k][alpha] * GridGradient[k][alpha] / Mass;
+					}
+				}
+			}
+
+			DeltaLambda /= Denom;
+			LambdaArray[2 * ElementIndex + 1] += DeltaLambda;
+
+			for (int32 k = 0; k < MPMTransfer.ElementGridNodeIncidentElements[ElementIndex].Num(); k++) {
+				int32 CurrentNode = MPMTransfer.ElementGridNodes[ElementIndex][MPMTransfer.ElementGridNodeIncidentElements[ElementIndex][k][0]];
+				T Mass = GridData[(MPMTransfer.NTransfer + 1) * CurrentNode];
+				if (Mass != 0) {
+					for (int32 alpha = 0; alpha < 3; alpha++) {
+						GridPositions[CurrentNode][alpha] += DeltaLambda * GridGradient[k][alpha] / Mass;
+					}
+				}
+			}
+		}
+
+		void ApplyInParallel(ParticleType& Particles, const T Dt) 
+		{
+			for (int32 color = 0; color < ElementsPerSubColors.Num(); color++)
+			{
+				for (int32 subcolor = 0; subcolor < ElementsPerSubColors[color].Num(); subcolor++)
+				{
+					PhysicsParallelFor(ElementsPerSubColors[color][subcolor].Num(), [&](const int32 ColorIndex)
+						{
+							const int32 ElementIndex = ElementsPerSubColors[color][subcolor][ColorIndex];
+							ApplyPolar(Particles, Dt, ElementIndex);
+							ApplyDet(Particles, Dt, ElementIndex);
+						}, ElementsPerSubColors[color][subcolor].Num() < 30);
+				}
+			}
+			//final G2P to compute particle positions:
+			FinalG2P(Particles);
+		}
+
+	private:
+
+		PMatrix<T, 3, 3> Ds(const TVec4<TVector<T, 3>>& LocalPositions) const {
+			PMatrix<T, 3, 3> Result((T)0.);
+			for (int i = 0; i < 3; i++) {
+				for (int c = 0; c < 3; c++) {
+					Result.SetAt(c, i, LocalPositions[i+1][c] - LocalPositions[0][c]);
+				}
+			}
+			return Result;
+		}
+
+
+		PMatrix<T, 3, 3> F(const TVec4<TVector<T, 3>>& LocalPositions, const int32 ElementIndex) const 
+		{
+			return Base::ElementDmInv(ElementIndex) * Ds(LocalPositions);
+		}
+
+		TVec4<TVector<T, 3>> GetPolarGradient(const TVec4<TVector<T, 3>>& LocalPositions, const int32 ElementIndex, T& Constraint, const T Tol = T(1e-4)) const
+		{
+			PMatrix<T, 3, 3> Fe = F(LocalPositions, ElementIndex);
+			PMatrix<T, 3, 3> Re((T)0.), Se((T)0.);
+
+			Chaos::PolarDecomposition(Fe, Re, Se);
+
+			T C1 = (T)0.;
+			for (int i = 0; i < 3; i++)
+			{
+				for (int j = 0; j < 3; j++)
+				{
+					C1 += FMath::Square((Fe - Re).GetAt(i, j));
+				}
+			}
+			C1 = FMath::Sqrt(C1);
+
+			
+			if (C1 < Tol)
+			{
+				Constraint = (T)0.;
+				return TVec4<TVector<T, 3>>(TVector<T, 3>((T)0.));
+			}
+			Constraint = C1;
+
+			PMatrix<T, 3, 3> DmInvT = Base::ElementDmInv(ElementIndex).GetTransposed();
+			TVector<T, 81> dRdF((T)0.);
+			Chaos::dRdFCorotated(Fe, dRdF);
+			TVec4<TVector<T, 3>> dC1(TVector<T, 3>((T)0.));
+			//dC1 = dC1dF * dFdX
+			for (int alpha = 0; alpha < 3; alpha++) {
+				for (int l = 0; l < 3; l++) {
+
+					dC1[0][alpha] += (DmInvT * Re).GetAt(alpha, l) - (DmInvT * Fe).GetAt(alpha, l);
+
+				}
+			}
+			for (int ie = 0; ie < 3; ie++) {
+				for (int alpha = 0; alpha < 3; alpha++) {
+					dC1[ie + 1][alpha] = (DmInvT * Fe).GetAt(alpha, ie) - (DmInvT * Re).GetAt(alpha, ie);
+				}
+			}
+			//it's really ie-1 here
+			for (int ie = 0; ie < 3; ie++) {
+				for (int alpha = 0; alpha < 3; alpha++) {
+					for (int m = 0; m < 3; m++) {
+						for (int n = 0; n < 3; n++) {
+							for (int j = 0; j < 3; j++) {
+								dC1[ie + 1][alpha] -= (Fe.GetAt(m, n) - Re.GetAt(m, n)) * dRdF[9 * (alpha * 3 + j) + m * 3 + n] * DmInvT.GetAt(j, ie);
+							}
+						}
+					}
+				}
+			}
+			for (int alpha = 0; alpha < 3; alpha++) {
+				for (int m = 0; m < 3; m++) {
+					for (int n = 0; n < 3; n++) {
+						for (int l = 0; l < 3; l++) {
+							for (int j = 0; j < 3; j++) {
+								dC1[0][alpha] += (Fe.GetAt(m, n) - Re.GetAt(m, n)) * dRdF[9 * (alpha * 3 + j) + m * 3 + n] * DmInvT.GetAt(j, l);
+							}
+						}
+					}
+				}
+			}
+
+			if (C1 != 0)
+			{
+				for (int i = 0; i < 4; i++)
+				{
+					for (int j = 0; j < 3; j++)
+					{
+						dC1[i][j] /= C1;
+					}
+				}
+			}
+			return dC1;
+		}
+
+		TVec4<TVector<T, 3>> GetDetGradient(const TVec4<TVector<T, 3>>& LocalPositions, const int32 ElementIndex, T& Constraint, const T Tol = T(1e-4)) const
+		{
+			const PMatrix<T, 3, 3> Fe = F(LocalPositions, ElementIndex);
+			PMatrix<T, 3, 3> DmInvT = Base::ElementDmInv(ElementIndex).GetTransposed();
+
+			T J = Fe.Determinant();
+			if (J - 1 < Tol)
+			{
+				Constraint = (T)0.;
+				return TVec4<TVector<T, 3>>(TVector<T, 3>((T)0.));
+			}
+			Constraint = J - 1;
+
+			TVec4<TVector<T, 3>> dC2(TVector<T, 3>((T)0.));
+
+			PMatrix<T, 3, 3> JFinvT((T)0.);
+			JFinvT.SetAt(0, 0, Fe.GetAt(1, 1) * Fe.GetAt(2, 2) - Fe.GetAt(2, 1) * Fe.GetAt(1, 2));
+			JFinvT.SetAt(0, 1, Fe.GetAt(2, 0) * Fe.GetAt(1, 2) - Fe.GetAt(1, 0) * Fe.GetAt(2, 2));
+			JFinvT.SetAt(0, 2, Fe.GetAt(1, 0) * Fe.GetAt(2, 1) - Fe.GetAt(2, 0) * Fe.GetAt(1, 1));
+			JFinvT.SetAt(1, 0, Fe.GetAt(2, 1) * Fe.GetAt(0, 2) - Fe.GetAt(0, 1) * Fe.GetAt(2, 2));
+			JFinvT.SetAt(1, 1, Fe.GetAt(0, 0) * Fe.GetAt(2, 2) - Fe.GetAt(2, 0) * Fe.GetAt(0, 2));
+			JFinvT.SetAt(1, 2, Fe.GetAt(2, 0) * Fe.GetAt(0, 1) - Fe.GetAt(0, 0) * Fe.GetAt(2, 1));
+			JFinvT.SetAt(2, 0, Fe.GetAt(0, 1) * Fe.GetAt(1, 2) - Fe.GetAt(1, 1) * Fe.GetAt(0, 2));
+			JFinvT.SetAt(2, 1, Fe.GetAt(1, 0) * Fe.GetAt(0, 2) - Fe.GetAt(0, 0) * Fe.GetAt(1, 2));
+			JFinvT.SetAt(2, 2, Fe.GetAt(0, 0) * Fe.GetAt(1, 1) - Fe.GetAt(1, 0) * Fe.GetAt(0, 1));
+
+			PMatrix<T, 3, 3> JinvTDmInvT = DmInvT * JFinvT;
+
+			for (int32 ie = 0; ie < 3; ie++) {
+				for (int32 alpha = 0; alpha < 3; alpha++) {
+					dC2[ie + 1][alpha] = JinvTDmInvT.GetAt(alpha, ie);
+				}
+			}
+			for (int32 alpha = 0; alpha < 3; alpha++) {
+				for (int32 l = 0; l < 3; l++) {
+					dC2[0][alpha] -= JinvTDmInvT.GetAt(alpha, l);
+				}
+			}
+			return dC2;
+		}
+
+		//interpolate particle positions from grid positions
+		void FinalG2P(ParticleType& Particles)
+		{
+			TVec4<TVector<T, 3>> Result(TVector<T, 3>((T)0.));
+
+			PhysicsParallelFor(Particles.Size(), [&](const int32 p)
+				//for (int32 p = 0; p < int32(Particles.Size()); p++) 
+				{
+					//TVector<int32, 3> BaseIndex = MPMTransfer.Indices[p];
+					Particles.P(p) = TVec3<T>((T)0.);
+					for (int32 ii = 0; ii < int32(MPMGrid.NPerDir); ii++)
+					{
+						T Nii = MPMGrid.Nijk(MPMTransfer.Weights[p][0], ii);
+						for (int32 jj = 0; jj < int32(MPMGrid.NPerDir); jj++)
+						{
+							T Njj = MPMGrid.Nijk(MPMTransfer.Weights[p][1], jj);
+							for (int32 kk = 0; kk < int32(MPMGrid.NPerDir); kk++)
+							{
+								T Nkk = MPMGrid.Nijk(MPMTransfer.Weights[p][2], kk);
+								TVector<int32, 3> LocIndex = { ii, jj, kk };
+								TVector<int32, 3> GlobIndex = MPMGrid.Loc2GlobIndex(MPMTransfer.Indices[p], LocIndex);
+								int32 GlobIndexFlat = MPMGrid.FlatIndex(GlobIndex);
+								Particles.P(p) += Nii * Njj * Nkk * GridPositions[GlobIndexFlat];
+							}
+						}
+					}
+				}, Particles.Size() < 50);
+
+		}
+
+	protected:
+		T MaxDt;
+		T MinDt;
+		T CFLCoeff;
+		TMPMGrid<T> MPMGrid;
+		TMPMTransfer<T> MPMTransfer;
+		TArray<TArray<int32>> ElementsPerColor;
+		TArray<TArray<int32>>* PreviousColoring;
+		TArray<TArray<TArray<int32>>> ElementsPerSubColors;
+		TArray<T> GridData;
+		TArray<TVector<T, 3>> GridPositions;
+		bool InitialGridDataWritten = false;
+	};
+
+
+}
\ No newline at end of file
