Fix bug in Simplex GJK
UDN 00466794

diff --git a/Engine/Source/Runtime/Experimental/Chaos/Public/Chaos/SimplexVectorized.h b/Engine/Source/Runtime/Experimental/Chaos/Public/Chaos/SimplexVectorized.h
--- a/Engine/Source/Runtime/Experimental/Chaos/Public/Chaos/SimplexVectorized.h
+++ b/Engine/Source/Runtime/Experimental/Chaos/Public/Chaos/SimplexVectorized.h
@@ -377,10 +377,10 @@ namespace Chaos
 			{
 				VectorRegister4Float OutBarycentricVectors[4];
 				const VectorRegister4Float  InvDetM = VectorDivide(GlobalVectorConstants::FloatOne, DetM);
-				OutBarycentricVectors[0] = VectorDivide(Cofactors[0], InvDetM);
-				OutBarycentricVectors[1] = VectorDivide(Cofactors[1], InvDetM);
-				OutBarycentricVectors[2] = VectorDivide(Cofactors[2], InvDetM);
-				OutBarycentricVectors[3] = VectorDivide(Cofactors[3], InvDetM);
+				OutBarycentricVectors[0] = VectorMultiply(Cofactors[0], InvDetM);
+				OutBarycentricVectors[1] = VectorMultiply(Cofactors[1], InvDetM);
+				OutBarycentricVectors[2] = VectorMultiply(Cofactors[2], InvDetM);
+				OutBarycentricVectors[3] = VectorMultiply(Cofactors[3], InvDetM);
 				// a0	b0	a1	b1
 				const VectorRegister4Float OutBarycentric0101 = VectorUnpackLo(OutBarycentricVectors[0], OutBarycentricVectors[1]);
 				const VectorRegister4Float OutBarycentric2323 = VectorUnpackLo(OutBarycentricVectors[2], OutBarycentricVectors[3]);
