Chaos - fix non-manifold collision detection path to call collision detection every iteration

diff --git a/Engine/Source/Runtime/Experimental/Chaos/Private/Chaos/Collision/PBDCollisionConstraint.cpp b/Engine/Source/Runtime/Experimental/Chaos/Private/Chaos/Collision/PBDCollisionConstraint.cpp
--- a/Engine/Source/Runtime/Experimental/Chaos/Private/Chaos/Collision/PBDCollisionConstraint.cpp
+++ b/Engine/Source/Runtime/Experimental/Chaos/Private/Chaos/Collision/PBDCollisionConstraint.cpp
@@ -281,18 +281,11 @@ namespace Chaos
 
 		CullDistance = InCullDistance;
 
-		// Are we allowing manifolds? If manifolds are enabled, we will build a one-shot manifold
-		// if supported by the shape pair, otherwise an incremental manifold will be created and
-		// we call collision detection every iteration to add new points (this is expensive).
-		// NOTE: bUseIncrementalManifold will get set to false later if we add a one-shot manifold
-		Flags.bUseManifold = bInUseManifold;
-		Flags.bUseIncrementalManifold = bInUseManifold;
-
 		// Initialize the is-probe and is-probe-unmodified flags to the same value.
 		// Contact modification may change bIsProbe but we want to store the unmodified value
 		// so that it can be reset between frames.
-		Flags.bIsProbe = 
-		Flags.bIsProbeUnmodified = (Shape[0] && Shape[0]->GetIsProbe()) || (Shape[1] && Shape[1]->GetIsProbe());
+		Flags.bIsProbe = (Shape[0] && Shape[0]->GetIsProbe()) || (Shape[1] && Shape[1]->GetIsProbe());
+		Flags.bIsProbeUnmodified = Flags.bIsProbe;
 
 		const FReal Margin0 = GetImplicit0()->GetMargin();
 		const FReal Margin1 = GetImplicit1()->GetMargin();
@@ -300,6 +293,16 @@ namespace Chaos
 		const EImplicitObjectType ImplicitType1 = GetInnerType(GetImplicit1()->GetCollisionType());
 		InitMarginsAndTolerances(ImplicitType0, ImplicitType1, Margin0, Margin1);
 
+		// Are we allowing manifolds?
+		Flags.bUseManifold = bInUseManifold;
+		Flags.bUseIncrementalManifold = false;
+
+		// Only levelsets use incremental collision manifolds
+		if (bInUseManifold && ((ImplicitType0 == ImplicitObjectType::LevelSet) || (ImplicitType1 == ImplicitObjectType::LevelSet)))
+		{
+			Flags.bUseIncrementalManifold = true;
+		}
+
 		// Debug testing for solver stiffness
 		if (Chaos_Collision_Stiffness >= 0)
 		{
@@ -587,7 +590,7 @@ namespace Chaos
 			return;
 		}
 
-		if (Flags.bUseIncrementalManifold)
+		if (Flags.bUseManifold)
 		{
 			// See if the manifold point already exists
 			int32 ManifoldPointIndex = FindManifoldPoint(ContactPoint);
diff --git a/Engine/Source/Runtime/Experimental/Chaos/Private/Chaos/Collision/SolverCollisionContainer.cpp b/Engine/Source/Runtime/Experimental/Chaos/Private/Chaos/Collision/SolverCollisionContainer.cpp
--- a/Engine/Source/Runtime/Experimental/Chaos/Private/Chaos/Collision/SolverCollisionContainer.cpp
+++ b/Engine/Source/Runtime/Experimental/Chaos/Private/Chaos/Collision/SolverCollisionContainer.cpp
@@ -65,14 +65,16 @@ namespace Chaos
 	public:
 		FPBDCollisionSolverAdapter()
 			: Constraint(nullptr)
-			, bIsIncremental(false)
+			, bIsManifold(true)
+			, bIsIncrementalManifold(false)
 		{
 		}
 
 		FPBDCollisionSolver& GetSolver() { return Solver; }
 		FPBDCollisionConstraint* GetConstraint() { return Constraint; }
 
-		bool IsIncrementalManifold() const { return bIsIncremental; }
+		bool IsManifold() const { return bIsManifold; }
+		bool IsIncrementalManifold() const { return bIsIncrementalManifold; }
 
 		void PreGatherInput(
 			const FReal Dt,
@@ -192,6 +194,9 @@ namespace Chaos
 			Solver.SolverBody0().SetInvIScale(Constraint->GetInvInertiaScale0());
 			Solver.SolverBody1().SetInvMScale(Constraint->GetInvMassScale1());
 			Solver.SolverBody1().SetInvIScale(Constraint->GetInvInertiaScale1());
+
+			bIsManifold = Constraint->GetUseManifold();
+			bIsIncrementalManifold = bIsManifold && Constraint->GetUseIncrementalCollisionDetection();
 		}
 
 		void UpdateManifoldPoints(
@@ -202,7 +207,7 @@ namespace Chaos
 			const FConstraintSolverBody& Body1 = Solver.SolverBody1();
 
 			// We handle incremental manifolds by just collecting any new contacts
-			const int32 BeginPointIndex = Solver.NumManifoldPoints();
+			const int32 BeginPointIndex = bIsIncrementalManifold ? Solver.NumManifoldPoints() : 0;
 			const int32 EndPointIndex = Solver.SetNumManifoldPoints(Constraint->GetManifoldPoints().Num());
 
 			const FSolverReal RestitutionVelocityThreshold = FSolverReal(Constraint->GetRestitutionThreshold()) * Dt;
@@ -352,7 +357,8 @@ namespace Chaos
 	private:
 		FPBDCollisionSolver Solver;
 		FPBDCollisionConstraint* Constraint;
-		bool bIsIncremental;
+		bool bIsManifold;
+		bool bIsIncrementalManifold;
 	};
 
 
@@ -364,7 +370,7 @@ namespace Chaos
 
 	FPBDCollisionSolverContainer::FPBDCollisionSolverContainer()
 		: FConstraintSolverContainer()
-		, bRequiresIncrementalCollisionDetection(false)
+		, bPerIterationCollisionDetection(false)
 		, bDeferredCollisionDetection(false)
 	{
 	}
@@ -415,7 +421,7 @@ namespace Chaos
 			CollisionSolver.UpdateManifoldPoints(Dt);
 		}
 
-		bRequiresIncrementalCollisionDetection |= CollisionSolver.IsIncrementalManifold();
+		bPerIterationCollisionDetection |= (!CollisionSolver.IsManifold() || CollisionSolver.IsIncrementalManifold());
 	}
 
 	void FPBDCollisionSolverContainer::UpdatePositionShockPropagation(const FReal Dt, const int32 It, const int32 NumIts, const int32 BeginIndex, const int32 EndIndex, const FPBDCollisionSolverSettings& SolverSettings)
@@ -491,10 +497,9 @@ namespace Chaos
 
 		// We run collision detection here under two conditions (normally it is run after Integration and before the constraint solver phase):
 		// 1) When deferring collision detection until the solver phase for better joint-collision behaviour (RBAN). In this case, we only do this on the first iteration.
-		// 2) When using incremental manifolds, where we may add/replace manifold points every iteration.
+		// 2) When using no manifolds or incremental manifolds, where we may add/replace manifold points every iteration.
 		const bool bDeferredCollisions = bDeferredCollisionDetection && (It == 0);
-		const bool bIncrementalManifolds = bRequiresIncrementalCollisionDetection;
-		if (bIncrementalManifolds || bDeferredCollisions)
+		if (bDeferredCollisions || bPerIterationCollisionDetection)
 		{
 			UpdateCollisions(Dt, BeginIndex, EndIndex);
 		}
@@ -614,7 +619,7 @@ namespace Chaos
 		for (int32 SolverIndex = BeginIndex; SolverIndex < EndIndex; ++SolverIndex)
 		{
 			FPBDCollisionSolverAdapter& CollisionSolver = CollisionSolvers[SolverIndex];
-			if (CollisionSolver.IsIncrementalManifold() || bDeferredCollisionDetection)
+			if (!CollisionSolver.IsManifold() || CollisionSolver.IsIncrementalManifold() || bDeferredCollisionDetection)
 			{
 				Collisions::Update(*CollisionSolver.GetConstraint(), Dt);
 				CollisionSolver.UpdateManifoldPoints(Dt);
diff --git a/Engine/Source/Runtime/Experimental/Chaos/Public/Chaos/Collision/PBDCollisionConstraint.h b/Engine/Source/Runtime/Experimental/Chaos/Public/Chaos/Collision/PBDCollisionConstraint.h
--- a/Engine/Source/Runtime/Experimental/Chaos/Public/Chaos/Collision/PBDCollisionConstraint.h
+++ b/Engine/Source/Runtime/Experimental/Chaos/Public/Chaos/Collision/PBDCollisionConstraint.h
@@ -502,10 +502,6 @@ namespace Chaos
 				{
 					ClosestManifoldPointIndex = ManifoldPointIndex;
 				}
-
-				// If we can use a one-shot manifold with this shape pair, disable the incremental flag
-				// to prevent calling collision detection again
-				Flags.bUseIncrementalManifold = false;
 			}
 		}
 
@@ -519,9 +515,6 @@ namespace Chaos
 		{
 			ResetActiveManifoldContacts();
 
-			// Disable the incremental flag to prevent calling collision detection again
-			Flags.bUseIncrementalManifold = false;
-
 			FReal MinPhi = TNumericLimits<FReal>::Max();
 			const int32 NumContacts = FMath::Min(ContactPoints.Num(), MaxManifoldPoints);
 			for (int32 ContactIndex = 0; ContactIndex < NumContacts; ++ContactIndex)
diff --git a/Engine/Source/Runtime/Experimental/Chaos/Public/Chaos/Collision/SolverCollisionContainer.h b/Engine/Source/Runtime/Experimental/Chaos/Public/Chaos/Collision/SolverCollisionContainer.h
--- a/Engine/Source/Runtime/Experimental/Chaos/Public/Chaos/Collision/SolverCollisionContainer.h
+++ b/Engine/Source/Runtime/Experimental/Chaos/Public/Chaos/Collision/SolverCollisionContainer.h
@@ -92,7 +92,7 @@ namespace Chaos
 		void UpdateCollisions(const FReal InDt, const int32 BeginIndex, const int32 EndIndex);
 
 		TArray<FPBDCollisionSolverAdapter> CollisionSolvers;
-		bool bRequiresIncrementalCollisionDetection;
+		bool bPerIterationCollisionDetection;
 		bool bDeferredCollisionDetection;
 	};
 }
