Dataflow - Threading support for node evaluations

diff --git a/Engine/Source/Runtime/Experimental/Dataflow/Core/Private/Dataflow/DataflowInputOutput.cpp b/Engine/Source/Runtime/Experimental/Dataflow/Core/Private/Dataflow/DataflowInputOutput.cpp
--- a/Engine/Source/Runtime/Experimental/Dataflow/Core/Private/Dataflow/DataflowInputOutput.cpp
+++ b/Engine/Source/Runtime/Experimental/Dataflow/Core/Private/Dataflow/DataflowInputOutput.cpp
@@ -78,6 +78,7 @@ void FDataflowInput::Invalidate()
 FDataflowOutput::FDataflowOutput(const Dataflow::FOutputParameters& Param, FGuid InGuid)
 	: FDataflowConnection(Dataflow::FPin::EDirection::OUTPUT, Param.Type, Param.Name, Param.Owner, Param.Property, InGuid)
 {
+	OutputLock = MakeShared<FCriticalSection>();
 }
 
 const TArray<FDataflowInput*>& FDataflowOutput::GetConnections() const { return Connections; }
diff --git a/Engine/Source/Runtime/Experimental/Dataflow/Core/Public/Dataflow/DataflowInputOutput.h b/Engine/Source/Runtime/Experimental/Dataflow/Core/Public/Dataflow/DataflowInputOutput.h
--- a/Engine/Source/Runtime/Experimental/Dataflow/Core/Public/Dataflow/DataflowInputOutput.h
+++ b/Engine/Source/Runtime/Experimental/Dataflow/Core/Public/Dataflow/DataflowInputOutput.h
@@ -7,6 +7,8 @@
 #include "Dataflow/DataflowConnection.h"
 #include "Dataflow/DataflowNode.h"
 #include "Templates/Function.h"
+#include "Async/Async.h"
+#include "GenericPlatform/GenericPlatformCriticalSection.h"
 
 #include "DataflowInputOutput.generated.h"
 
@@ -53,7 +55,11 @@ public:
 	virtual TArray< FDataflowOutput* > GetConnectedOutputs();
 	virtual const TArray< const FDataflowOutput* > GetConnectedOutputs() const;
 
-	template<class T> const T& GetValue(Dataflow::FContext& Context, const T& Default) const;
+	template<class T>
+	const T& GetValue(Dataflow::FContext& Context, const T& Default) const;
+
+	template<class T>
+	TFuture<const T&> GetValueParallel(Dataflow::FContext& Context, const T& Default) const;
 
 	virtual void Invalidate() override;
 };
@@ -87,12 +93,18 @@ struct DATAFLOWCORE_API FDataflowOutput : public FDataflowConnection
 	mutable uint32 CacheKeyValue = UINT_MAX;
 	mutable Dataflow::TCacheValue Cache;
 
+
 	TArray< FDataflowInput* > Connections;
 
 	size_t PassthroughOffsetAddress = INDEX_NONE;
 
 public:
 	static FDataflowOutput NoOpOutput;
+	
+	mutable TSharedPtr<FCriticalSection> OutputLock; // @todo(dataflow) find alternatives
+	// Why is this made into a SharedPtr? FCriticalSection cannot be copied - making FCriticalSection a member of
+	// FDataflowOutput deletes the copy constructor of FDataflowOutput. But some other part of the code needs
+	// FDataflowOutput to be copyable. Hence wrapped it in a sharedPtr inspired by NiagaraEmitter.h:227.
 
 	FDataflowOutput(const Dataflow::FOutputParameters& Param = {}, FGuid InGuid = FGuid::NewGuid());
 
@@ -146,6 +158,10 @@ public:
 
 	template<class T>
 	bool Evaluate(Dataflow::FContext& Context) const;
+
+	template<class T>
+	TFuture<bool> EvaluateParallel(Dataflow::FContext& Context) const;
+
 	virtual void Invalidate() override;
 
 };
@@ -171,10 +187,18 @@ const T& FDataflowInput::GetValue(Dataflow::FContext& Context, const T& Default)
 	}
 	return Default;
 }
- 
+
+template<class T>
+TFuture<const T&> FDataflowInput::GetValueParallel(Dataflow::FContext& Context, const T& Default) const
+{
+	return Async(EAsyncExecution::TaskGraph, [&]() -> const T& { return this->GetValue<T>(Context, Default); });
+}
+
 template<class T>
 bool FDataflowOutput::Evaluate(Dataflow::FContext& Context) const
 {
+	OutputLock->Lock(); ON_SCOPE_EXIT { this->OutputLock->Unlock(); };
+
 	check(OwningNode);
  
 	if (OwningNode->bActive)
@@ -202,4 +226,10 @@ bool FDataflowOutput::Evaluate(Dataflow::FContext& Context) const
 	}
  
 	return false;
+}
+
+template<class T>
+TFuture<bool> FDataflowOutput::EvaluateParallel(Dataflow::FContext& Context) const
+{
+	return Async(EAsyncExecution::TaskGraph, [&]() -> bool { return this->Evaluate<T>(Context); });
 }
\ No newline at end of file
diff --git a/Engine/Source/Runtime/Experimental/Dataflow/Core/Public/Dataflow/DataflowNode.h b/Engine/Source/Runtime/Experimental/Dataflow/Core/Public/Dataflow/DataflowNode.h
--- a/Engine/Source/Runtime/Experimental/Dataflow/Core/Public/Dataflow/DataflowNode.h
+++ b/Engine/Source/Runtime/Experimental/Dataflow/Core/Public/Dataflow/DataflowNode.h
@@ -130,6 +130,13 @@ struct DATAFLOWCORE_API FDataflowNode
 		checkSlow(FindInput(Reference));
 		return FindInput(Reference)->template GetValue<T>(Context, *Reference);
 	}
+	
+	template<class T>
+	TFuture<const T&> GetValueParallel(Dataflow::FContext& Context, const T* Reference) const
+	{
+		checkSlow(FindInput(Reference));
+		return FindInput(Reference)->template GetValueParallel<T>(Context, *Reference);
+	}
 
 	/**
 	*   GetValue(...)
diff --git a/Engine/Source/Runtime/Experimental/Dataflow/Core/Public/Dataflow/DataflowNodeParameters.h b/Engine/Source/Runtime/Experimental/Dataflow/Core/Public/Dataflow/DataflowNodeParameters.h
--- a/Engine/Source/Runtime/Experimental/Dataflow/Core/Public/Dataflow/DataflowNodeParameters.h
+++ b/Engine/Source/Runtime/Experimental/Dataflow/Core/Public/Dataflow/DataflowNodeParameters.h
@@ -5,6 +5,7 @@
 #include "CoreMinimal.h"
 #include "UObject/UnrealType.h"
 #include "Templates/UniquePtr.h"
+#include "GenericPlatform/GenericPlatformCriticalSection.h"
 
 class  UDataflow;
 
@@ -53,6 +54,7 @@ namespace Dataflow
 	class DATAFLOWCORE_API FContext
 	{
 		TMap<int64, TUniquePtr<FContextCacheBase>> DataStore;
+		TSharedPtr<FCriticalSection> CacheLock;
 		
 		FContext(FContext&&) = default;
 		FContext& operator=(FContext&&) = default;
@@ -65,7 +67,9 @@ namespace Dataflow
 		FContext(float InTime, FString InType = FString(""))
 			: Timestamp(InTime)
 			, Type(StaticType().Append(InType))
-		{}
+		{
+			CacheLock = MakeShared<FCriticalSection>();
+		}
 
 		float Timestamp = 0.f;
 		FString Type;
@@ -89,6 +93,8 @@ namespace Dataflow
 		template<typename T>
 		void SetData(size_t Key, FProperty* Property, const T& Value)
 		{
+			CacheLock->Lock(); ON_SCOPE_EXIT { CacheLock->Unlock(); };
+			
 			int64 IntKey = (int64)Key;
 			TUniquePtr<FContextCache<T>> DataStoreEntry = MakeUnique<FContextCache<T>>(Property, Value, FPlatformTime::Cycles64());
 			DataStore.Emplace(IntKey, MoveTemp(DataStoreEntry));
@@ -97,6 +103,8 @@ namespace Dataflow
 		template<typename T>
 		void SetData(size_t Key, FProperty* Property, T&& Value)
 		{
+			CacheLock->Lock(); ON_SCOPE_EXIT { CacheLock->Unlock(); };
+			
 			int64 IntKey = (int64)Key;
 			TUniquePtr<FContextCache<T>> DataStoreEntry = MakeUnique<FContextCache<T>>(Property, Forward<T>(Value), FPlatformTime::Cycles64());
 			DataStore.Emplace(IntKey, MoveTemp(DataStoreEntry));
@@ -105,6 +113,8 @@ namespace Dataflow
 		template<class T>
 		const T& GetData(size_t Key, FProperty* Property, const T& Default = T())
 		{
+			CacheLock->Lock(); ON_SCOPE_EXIT { CacheLock->Unlock(); };
+			
 			if (TUniquePtr<FContextCacheBase>* Cache = DataStore.Find(Key))
 			{
 				return (*Cache)->GetTypedData<T>(Property);
@@ -114,6 +124,8 @@ namespace Dataflow
 
 		bool HasData(size_t Key, uint64 StoredAfter = 0)
 		{
+			CacheLock->Lock(); ON_SCOPE_EXIT { CacheLock->Unlock(); };
+			
 			int64 IntKey = (int64)Key;
 			return DataStore.Contains(IntKey) && DataStore[Key]->Timestamp >= StoredAfter;
 		}
