Chaos : Flesh :
Add PoissonSolve function into Poisson.

diff --git a/Engine/Source/Runtime/Experimental/Chaos/Public/Chaos/Math/Poisson.h b/Engine/Source/Runtime/Experimental/Chaos/Public/Chaos/Math/Poisson.h
--- a/Engine/Source/Runtime/Experimental/Chaos/Public/Chaos/Math/Poisson.h
+++ b/Engine/Source/Runtime/Experimental/Chaos/Public/Chaos/Math/Poisson.h
@@ -141,6 +141,62 @@ inline int32 MinFlatIndex(const TArray<int32>& ElemIdx, const TArray<int32>& Loc
 	return (ElemIdx[MinIdx] * 4) + LocalIdx[MinIdx];
 }
 
+template <class T, class TV_INT = FIntVector4, int d = 3>
+void PoissonSolve(const TArray<int32>& InConstrainedNodes, 
+	const TArray<T>& ConstrainedWeights, 
+	const TArray<TV_INT>& Mesh,
+	const TArray<TArray<int32>>& IncidentElements,
+	const TArray<TArray<int32>>& IncidentElementsLocalIndex,
+	const TArray<T>& De_inverse,
+	const TArray<T>& measure, 
+	const int32 MaxItCG,
+	const T CGTol,
+	TArray<T>& Weights)
+{
+	auto ProjectBCs = [&InConstrainedNodes] (TArray<T>& U)
+	{
+		for (int32 i = 0; i < InConstrainedNodes.Num(); i++)
+		{
+			U[InConstrainedNodes[i]] = (T)0.;
+		}
+	};
+
+	auto MultiplyLaplacian = [&](const TArray<T>& U, TArray<T>& LU)
+	{
+		Laplacian(
+			Mesh,
+			IncidentElements,
+			IncidentElementsLocalIndex,
+			De_inverse,
+			measure,
+			U,
+			LU);
+		ProjectBCs(LU);
+	};
+
+	Fill(Weights, T(0));
+
+	TArray<T> InitialGuess;
+	InitialGuess.Init((T)0., Weights.Num());
+	for (int32 i = 0; i < InConstrainedNodes.Num(); i++)
+	{
+		InitialGuess[InConstrainedNodes[i]] = ConstrainedWeights[i];
+	}
+	TArray<T> MinusResidual;
+	MinusResidual.Init((T)0., Weights.Num());
+	MultiplyLaplacian(InitialGuess, MinusResidual);
+
+	TArray<T> MinusDw;
+	MinusDw.Init((T)0., Weights.Num());
+	Chaos::LanczosCG(MultiplyLaplacian, MinusDw, MinusResidual, MaxItCG, CGTol, nullptr);
+
+	for (int32 i = 0; i < InitialGuess.Num(); i++)
+	{
+		InitialGuess[i] -= MinusDw[i];
+	}
+
+}
+
 template <class T, class TV_INT=FIntVector4, int d=3>
 void 
 Laplacian(
