ChaosFlesh : Deformable Solver Init
- Removed bSimulationInitialized
- Locked threaded solver advanced to TickEnd
- Defaulted threaded advanced to enabled.

diff --git a/Engine/Source/Runtime/Experimental/Chaos/Private/Chaos/Deformable/ChaosDeformableSolver.cpp b/Engine/Source/Runtime/Experimental/Chaos/Private/Chaos/Deformable/ChaosDeformableSolver.cpp
--- a/Engine/Source/Runtime/Experimental/Chaos/Private/Chaos/Deformable/ChaosDeformableSolver.cpp
+++ b/Engine/Source/Runtime/Experimental/Chaos/Private/Chaos/Deformable/ChaosDeformableSolver.cpp
@@ -263,7 +263,6 @@ namespace Chaos::Softs
 
 
 							Time = 0.f;
-							bSimulationInitialized = true;
 						}
 					}
 				}
@@ -380,6 +379,8 @@ namespace Chaos::Softs
 	}
 	void FDeformableSolver::TickSimulation(FSolverReal DeltaTime)
 	{
+		if (!Proxies.Num()) return;
+
 		if (!Property.FixTimeStep)
 		{
 			Evolution->AdvanceOneTimeStep(DeltaTime);
diff --git a/Engine/Source/Runtime/Experimental/Chaos/Public/Chaos/Deformable/ChaosDeformableSolver.h b/Engine/Source/Runtime/Experimental/Chaos/Public/Chaos/Deformable/ChaosDeformableSolver.h
--- a/Engine/Source/Runtime/Experimental/Chaos/Public/Chaos/Deformable/ChaosDeformableSolver.h
+++ b/Engine/Source/Runtime/Experimental/Chaos/Public/Chaos/Deformable/ChaosDeformableSolver.h
@@ -39,7 +39,6 @@ namespace Chaos::Softs
 
 			/* Simulation Advance */
 			bool Advance(FSolverReal DeltaTime);
-			bool IsInitialized() { return Solver.IsInitialized(); }
 			void Reset(const FDeformableSolverProperties&);
 			void TickSimulation(FSolverReal DeltaTime);
 			void UpdateOutputState(FThreadingProxy&);
@@ -88,7 +87,6 @@ namespace Chaos::Softs
 
 		/* Simulation Advance */
 		bool Advance(FSolverReal DeltaTime);
-		bool IsInitialized() { return bSimulationInitialized; }
 		void Reset(const FDeformableSolverProperties&);
 		void TickSimulation(FSolverReal DeltaTime);
 		void UpdateOutputState(FThreadingProxy&);
