Chaos : Flesh :
Update PoissonSolveto use just mesh. Add a test for PoissonSolve.

diff --git a/Engine/Source/Runtime/Experimental/Chaos/Public/Chaos/Math/Poisson.h b/Engine/Source/Runtime/Experimental/Chaos/Public/Chaos/Math/Poisson.h
--- a/Engine/Source/Runtime/Experimental/Chaos/Public/Chaos/Math/Poisson.h
+++ b/Engine/Source/Runtime/Experimental/Chaos/Public/Chaos/Math/Poisson.h
@@ -7,6 +7,7 @@
 #include "Chaos/Vector.h"
 #include "Math/Matrix.h"
 #include "Math/Vector.h"
+#include "Chaos/Utilities.h"
 
 namespace Chaos {
 
@@ -141,18 +142,22 @@ inline int32 MinFlatIndex(const TArray<int32>& ElemIdx, const TArray<int32>& Loc
 	return (ElemIdx[MinIdx] * 4) + LocalIdx[MinIdx];
 }
 
-template <class T, class TV_INT = FIntVector4, int d = 3>
+template <class T, class TV, class TV_INT = FIntVector4, int d = 3>
 void PoissonSolve(const TArray<int32>& InConstrainedNodes, 
 	const TArray<T>& ConstrainedWeights, 
 	const TArray<TV_INT>& Mesh,
-	const TArray<TArray<int32>>& IncidentElements,
-	const TArray<TArray<int32>>& IncidentElementsLocalIndex,
-	const TArray<T>& De_inverse,
-	const TArray<T>& measure, 
+	const TArray<TV>& X,
 	const int32 MaxItCG,
 	const T CGTol,
 	TArray<T>& Weights)
 {
+	TArray<T> De_inverse;
+	TArray<T> measure;
+	ComputeDeInverseAndElementMeasures(Mesh, X, De_inverse, measure);
+
+	TArray<TArray<int32>> IncidentElementsLocalIndex;
+	TArray<TArray<int32>> IncidentElements = Chaos::Utilities::ComputeIncidentElements<4>(Mesh, &IncidentElementsLocalIndex);
+	
 	auto ProjectBCs = [&InConstrainedNodes] (TArray<T>& U)
 	{
 		for (int32 i = 0; i < InConstrainedNodes.Num(); i++)
@@ -161,7 +166,7 @@ void PoissonSolve(const TArray<int32>& InConstrainedNodes,
 		}
 	};
 
-	auto MultiplyLaplacian = [&](const TArray<T>& U, TArray<T>& LU)
+	auto MultiplyLaplacian = [&](TArray<T>& LU, const TArray<T>& U)
 	{
 		Laplacian(
 			Mesh,
@@ -184,15 +189,15 @@ void PoissonSolve(const TArray<int32>& InConstrainedNodes,
 	}
 	TArray<T> MinusResidual;
 	MinusResidual.Init((T)0., Weights.Num());
-	MultiplyLaplacian(InitialGuess, MinusResidual);
+	MultiplyLaplacian(MinusResidual, InitialGuess);
 
 	TArray<T> MinusDw;
 	MinusDw.Init((T)0., Weights.Num());
-	Chaos::LanczosCG(MultiplyLaplacian, MinusDw, MinusResidual, MaxItCG, CGTol, nullptr);
+	Chaos::LanczosCG(MultiplyLaplacian, MinusDw, MinusResidual, MaxItCG, CGTol, false);
 
 	for (int32 i = 0; i < InitialGuess.Num(); i++)
 	{
-		InitialGuess[i] -= MinusDw[i];
+		Weights[i] = InitialGuess[i] - MinusDw[i];
 	}
 
 }
