ChaosFlesh : Added insights macros
- Moved MPM option to experimental folder on deformable solver component.

diff --git a/Engine/Source/Runtime/Experimental/Chaos/Private/Chaos/Deformable/ChaosDeformableSolver.cpp b/Engine/Source/Runtime/Experimental/Chaos/Private/Chaos/Deformable/ChaosDeformableSolver.cpp
--- a/Engine/Source/Runtime/Experimental/Chaos/Private/Chaos/Deformable/ChaosDeformableSolver.cpp
+++ b/Engine/Source/Runtime/Experimental/Chaos/Private/Chaos/Deformable/ChaosDeformableSolver.cpp
@@ -30,6 +30,7 @@
 #include "HAL/IConsoleManager.h"
 #include "Misc/Paths.h"
 
+DECLARE_CYCLE_STAT(TEXT("DeformableSolver.Advance"), STAT_DeformableSolver_Advance, STATGROUP_Chaos);
 
 DEFINE_LOG_CATEGORY_STATIC(LogChaosDeformableSolver, Log, All);
 namespace Chaos::Softs
@@ -49,6 +50,8 @@ namespace Chaos::Softs
 	}
 	void FDeformableSolver::Reset(const FDeformableSolverProperties& InProps)
 	{
+		TRACE_CPUPROFILER_EVENT_SCOPE(DeformableSolver_Reset);
+
 		Property = InProps;
 		MObjects = TArrayCollectionArray<const UObject*>();
 		FSolverParticles LocalParticlesDummy;
@@ -78,6 +81,9 @@ namespace Chaos::Softs
 	}
 	bool FDeformableSolver::Advance(FSolverReal DeltaTime)
 	{
+		SCOPE_CYCLE_COUNTER(STAT_DeformableSolver_Advance);
+		TRACE_CPUPROFILER_EVENT_SCOPE(DeformableSolver_Advance);
+
 		int32 NumIterations = FMath::Clamp<int32>(Property.NumSolverSubSteps, 0, INT_MAX);
 		if (NumIterations)
 		{
@@ -101,6 +107,8 @@ namespace Chaos::Softs
 	}
 	void FDeformableSolver::InitializeSimulationObjects()
 	{
+		TRACE_CPUPROFILER_EVENT_SCOPE(DeformableSolver_InitializeSimulationObjects);
+
 		for (TUniquePtr<FThreadingProxy>& Proxy : UninitializedProxys)
 		{
 			InitializeSimulationObject(*Proxy);
@@ -132,6 +140,8 @@ namespace Chaos::Softs
 	}
 	void FDeformableSolver::InitializeSimulationObject(FThreadingProxy& InProxy)
 	{
+		TRACE_CPUPROFILER_EVENT_SCOPE(DeformableSolver_InitializeSimulationObject);
+
 		if (FFleshThreadingProxy* Proxy = InProxy.As<FFleshThreadingProxy>())
 		{
 			if (const FManagedArrayCollection* Dynamic = &Proxy->GetDynamicCollection())
@@ -276,6 +286,7 @@ namespace Chaos::Softs
 	}
 	void FDeformableSolver::InitializeCollisionBodies()
 	{
+		TRACE_CPUPROFILER_EVENT_SCOPE(DeformableSolver_InitializeCollisionBodies);
 		if (Property.bUseFloor && Evolution->CollisionParticles().Size()==0)
 		{
 			Chaos::FVec3 Position(0.f);
@@ -295,6 +306,8 @@ namespace Chaos::Softs
 	{
 		auto MKineticUpdate = [this](FSolverParticles& MParticles, const FSolverReal Dt, const FSolverReal MTime, const int32 Index)
 		{
+			TRACE_CPUPROFILER_EVENT_SCOPE(DeformableSolver_MKineticUpdate);
+
 			if (0 <= Index && Index < this->MObjects.Num())
 			{
 				if (const UObject* Owner = this->MObjects[Index])
@@ -324,6 +337,8 @@ namespace Chaos::Softs
 	}
 	void FDeformableSolver::InitializeSelfCollisionVariables()
 	{
+		TRACE_CPUPROFILER_EVENT_SCOPE(DeformableSolver_InitializeSelfCollisionVariables);
+
 		int32 NumParticles = Evolution->Particles().Size();
 		SurfaceTriangleMesh->Init(*SurfaceElements);
 		TriangleMeshCollisions.Reset(new FPBDTriangleMeshCollisions(
@@ -353,6 +368,8 @@ namespace Chaos::Softs
 
 	void FDeformableSolver::InitializeGridBasedConstraintVariables()
 	{
+		TRACE_CPUPROFILER_EVENT_SCOPE(DeformableSolver_InitializeGridBasedConstraintVariables);
+
 		GridBasedCorotatedConstraint.Reset(new Chaos::Softs::FXPBDGridBasedCorotatedConstraints<FSolverReal, FSolverParticles>(
 			Evolution->Particles(), *AllElements, Property.GridDx, /*bRecordMetric = */false, (Chaos::Softs::FSolverReal).1, (Chaos::Softs::FSolverReal).01, (Chaos::Softs::FSolverReal).4, (Chaos::Softs::FSolverReal)1000.0));
 		Evolution->ResetConstraintRules();
@@ -379,6 +396,8 @@ namespace Chaos::Softs
 	}
 	void FDeformableSolver::TickSimulation(FSolverReal DeltaTime)
 	{
+		TRACE_CPUPROFILER_EVENT_SCOPE(DeformableSolver_TickSimulation);
+
 		if (!Proxies.Num()) return;
 
 		if (!Property.FixTimeStep)
@@ -420,6 +439,7 @@ namespace Chaos::Softs
 	void FDeformableSolver::PushPackage(int32 InFrame, FOutputDataMap&& InPackage)
 	{
 		FScopeLock Lock(&PackageMutex);
+		TRACE_CPUPROFILER_EVENT_SCOPE(DeformableSolver_PushPackage);
 		OutputPackages.Push(TUniquePtr< FOutputPackage >(new FOutputPackage(InFrame, MoveTemp(InPackage))));
 	}
 
@@ -431,6 +451,7 @@ namespace Chaos::Softs
 	TUniquePtr<FOutputPackage> FDeformableSolver::PullPackage()
 	{
 		FScopeLock Lock(&PackageMutex);
+		TRACE_CPUPROFILER_EVENT_SCOPE(DeformableSolver_PullPackage);
 		if (OutputPackages.Num())
 			return OutputPackages.Pop();
 		return TUniquePtr<FOutputPackage>(nullptr);
@@ -451,6 +472,7 @@ namespace Chaos::Softs
 	}
 	void FDeformableSolver::UpdateOutputState(FThreadingProxy& InProxy)
 	{
+		TRACE_CPUPROFILER_EVENT_SCOPE(DeformableSolver_UpdateOutputState);
 		if (FFleshThreadingProxy* Proxy = InProxy.As<FFleshThreadingProxy>())
 		{
 
@@ -476,6 +498,7 @@ namespace Chaos::Softs
 	}
 	void FDeformableSolver::WriteFrame(FThreadingProxy& InProxy, const FSolverReal DeltaTime)
 	{
+		TRACE_CPUPROFILER_EVENT_SCOPE(DeformableSolver_WriteFrame);
 		if (FFleshThreadingProxy* Proxy = InProxy.As<FFleshThreadingProxy>())
 		{
 			if (const FManagedArrayCollection* Rest = &Proxy->GetRestCollection())
