fix duplicate materials on new/reset of a geometry collection

diff --git a/Engine/Source/Runtime/Experimental/GeometryCollectionEngine/Private/GeometryCollection/GeometryCollectionEngineConversion.cpp b/Engine/Source/Runtime/Experimental/GeometryCollectionEngine/Private/GeometryCollection/GeometryCollectionEngineConversion.cpp
--- a/Engine/Source/Runtime/Experimental/GeometryCollectionEngine/Private/GeometryCollection/GeometryCollectionEngineConversion.cpp
+++ b/Engine/Source/Runtime/Experimental/GeometryCollectionEngine/Private/GeometryCollection/GeometryCollectionEngineConversion.cpp
@@ -384,7 +384,7 @@ FMeshDescription* FGeometryCollectionEngineConversion::GetMaxResMeshDescriptionW
 	return MeshDescription;
 }
 
-int32 FGeometryCollectionEngineConversion::AppendMaterials(const TArray<UMaterialInterface*>& Materials, UGeometryCollection* GeometryCollectionObject)
+int32 FGeometryCollectionEngineConversion::AppendMaterials(const TArray<UMaterialInterface*>& Materials, UGeometryCollection* GeometryCollectionObject, bool bAddInteriorCopy)
 {
 	// for each material, add a reference in our GeometryCollectionObject
 	const int32 MaterialStart = GeometryCollectionObject->Materials.Num();
@@ -403,7 +403,10 @@ int32 FGeometryCollectionEngineConversion::AppendMaterials(const TArray<UMateria
 
 		// We add the material twice, once for interior and again for exterior.
 		GeometryCollectionObject->Materials.Add(CurrMaterial);
-		GeometryCollectionObject->Materials.Add(CurrMaterial);
+		if (bAddInteriorCopy)
+		{
+			GeometryCollectionObject->Materials.Add(CurrMaterial);
+		}
 	}
 	return MaterialStart;
 }
@@ -421,7 +424,7 @@ bool FGeometryCollectionEngineConversion::AppendStaticMesh(const UStaticMesh* St
 
 	if (AppendStaticMesh(StaticMesh, StartMaterialIndex, StaticMeshTransform, GeometryCollection, bReindexMaterials))
 	{
-		AppendMaterials(Materials, GeometryCollectionObject);
+		AppendMaterials(Materials, GeometryCollectionObject, true);
 		return true;
 	}
 
@@ -701,7 +704,7 @@ void FGeometryCollectionEngineConversion::AppendGeometryCollection(const UGeomet
 	FGeometryCollection* GeometryCollection = GeometryCollectionPtr.Get();
 	check(GeometryCollection);
 
-	int32 MaterialStart = AppendMaterials(Materials, TargetGeometryCollectionObject);
+	int32 MaterialStart = AppendMaterials(Materials, TargetGeometryCollectionObject, false);
 
 	AppendGeometryCollection(SourceGeometryCollectionPtr.Get(), MaterialStart, GeometryCollectionTransform, GeometryCollection, bReindexMaterials);
 }
@@ -752,7 +755,7 @@ int32 FGeometryCollectionEngineConversion::AppendGeometryCollectionMaterials(con
 	TargetGeometryCollectionObject->Materials.Remove(BoneSelectedMaterial);
 	Materials.Remove(BoneSelectedMaterial);
 
-	return AppendMaterials(Materials, TargetGeometryCollectionObject);
+	return AppendMaterials(Materials, TargetGeometryCollectionObject, false);
 }
 
 
diff --git a/Engine/Source/Runtime/Experimental/GeometryCollectionEngine/Public/GeometryCollection/GeometryCollectionEngineConversion.h b/Engine/Source/Runtime/Experimental/GeometryCollectionEngine/Public/GeometryCollection/GeometryCollectionEngineConversion.h
--- a/Engine/Source/Runtime/Experimental/GeometryCollectionEngine/Public/GeometryCollection/GeometryCollectionEngineConversion.h
+++ b/Engine/Source/Runtime/Experimental/GeometryCollectionEngine/Public/GeometryCollection/GeometryCollectionEngineConversion.h
@@ -29,7 +29,7 @@ public:
 	 * @param Materials : Materials fetched from the StaticMeshComponent used to configure this geometry
 	 * @param GeometryCollection  : Collection to append the mesh into.
 	 */
-	static int32 AppendMaterials(const TArray<UMaterialInterface*>& Materials, UGeometryCollection* GeometryCollectionObject);
+	static int32 AppendMaterials(const TArray<UMaterialInterface*>& Materials, UGeometryCollection* GeometryCollectionObject, bool bAddInteriorCopy);
 
 	/**
 	 * Appends a MeshDescription to a GeometryCollection.
