Chaos - fix non-manifold collision update in collision solver

diff --git a/Engine/Source/Runtime/Experimental/Chaos/Private/Chaos/Collision/SolverCollisionContainer.cpp b/Engine/Source/Runtime/Experimental/Chaos/Private/Chaos/Collision/SolverCollisionContainer.cpp
--- a/Engine/Source/Runtime/Experimental/Chaos/Private/Chaos/Collision/SolverCollisionContainer.cpp
+++ b/Engine/Source/Runtime/Experimental/Chaos/Private/Chaos/Collision/SolverCollisionContainer.cpp
@@ -3,6 +3,7 @@
 
 #include "Chaos/Collision/CollisionApplyType.h"
 #include "Chaos/Collision/PBDCollisionConstraint.h"
+#include "Chaos/CollisionResolution.h"
 #include "Chaos/DebugDrawQueue.h"
 #include "Chaos/Evolution/SolverBodyContainer.h"
 #include "Chaos/Particle/ParticleUtilities.h"
@@ -621,7 +622,23 @@ namespace Chaos
 			FPBDCollisionSolverAdapter& CollisionSolver = CollisionSolvers[SolverIndex];
 			if (!CollisionSolver.IsManifold() || CollisionSolver.IsIncrementalManifold() || bDeferredCollisionDetection)
 			{
-				Collisions::Update(*CollisionSolver.GetConstraint(), Dt);
+				FPBDCollisionConstraint* Constraint = CollisionSolver.GetConstraint();
+
+				// Run collision detection at the current transforms including any correction from previous iterations
+				const FSolverBody& Body0 = CollisionSolver.GetSolver().SolverBody0().SolverBody();
+				const FSolverBody& Body1 = CollisionSolver.GetSolver().SolverBody1().SolverBody();
+				const FRigidTransform3 CorrectedActorWorldTransform0 = FRigidTransform3(Body0.CorrectedActorP(), Body0.CorrectedActorQ());
+				const FRigidTransform3 CorrectedActorWorldTransform1 = FRigidTransform3(Body1.CorrectedActorP(), Body1.CorrectedActorQ());
+				const FRigidTransform3 CorrectedShapeWorldTransform0 = Constraint->GetShapeRelativeTransform0() * CorrectedActorWorldTransform0;
+				const FRigidTransform3 CorrectedShapeWorldTransform1 = Constraint->GetShapeRelativeTransform1() * CorrectedActorWorldTransform1;
+
+				// NOTE: We deliberately have not updated the ShapwWorldTranforms on the constraint. If we did that, we would calculate 
+				// errors incorrectly in UpdateManifoldPoints, because the solver assumes nothing has been moved as we iterate (we accumulate 
+				// corrections that will be applied later.)
+				Constraint->ResetPhi(Constraint->GetCullDistance());
+				Collisions::UpdateConstraint(*Constraint, CorrectedShapeWorldTransform0, CorrectedShapeWorldTransform1, Dt);
+
+				// Update the manifold based on the new or updated contacts
 				CollisionSolver.UpdateManifoldPoints(Dt);
 			}
 		}
diff --git a/Engine/Source/Runtime/Experimental/Chaos/Public/Chaos/Evolution/SolverBody.h b/Engine/Source/Runtime/Experimental/Chaos/Public/Chaos/Evolution/SolverBody.h
--- a/Engine/Source/Runtime/Experimental/Chaos/Public/Chaos/Evolution/SolverBody.h
+++ b/Engine/Source/Runtime/Experimental/Chaos/Public/Chaos/Evolution/SolverBody.h
@@ -255,17 +255,27 @@ namespace Chaos
 			State.DQ = FSolverVec3(0);
 		}
 
+		/**
+		 * @brief Get the world-space Actor position 
+		*/
+		inline FVec3 ActorP() const { return P() - ActorQ().RotateVector(CoM()); }
+
+		/**
+		 * @brief Get the world-space Actor rotation 
+		*/
+		inline FRotation3 ActorQ() const { return Q() * RoM().Inverse(); }
+
 		/**
 		 * @brief Get the current world-space Actor position 
 		 * @note This is recalculated from the current CoM transform including the accumulated position and rotation corrections.
 		*/
-		inline FVec3 ActorP() const { return CorrectedP() - ActorQ().RotateVector(CoM()); }
+		inline FVec3 CorrectedActorP() const { return CorrectedP() - CorrectedActorQ().RotateVector(CoM()); }
 
 		/**
 		 * @brief Get the current world-space Actor rotation
 		 * @note This is recalculated from the current CoM transform including the accumulated position and rotation corrections.
 		*/
-		inline FRotation3 ActorQ() const { return CorrectedQ() * RoM().Inverse(); }
+		inline FRotation3 CorrectedActorQ() const { return CorrectedQ() * RoM().Inverse(); }
 
 		/**
 		 * @brief Contact graph level. This is used in shock propagation to determine which of two bodies should have its inverse mass scaled
@@ -515,6 +525,8 @@ namespace Chaos
 		inline const FRotation3& Q() const { return Body->Q(); }
 		inline const FVec3 ActorP() const { return Body->ActorP(); }
 		inline const FRotation3 ActorQ() const { return Body->ActorQ(); }
+		inline const FVec3 CorrectedActorP() const { return Body->CorrectedActorP(); }
+		inline const FRotation3 CorrectedActorQ() const { return Body->CorrectedActorQ(); }
 		inline const FVec3& V() const { return Body->V(); }
 		inline const FVec3& W() const { return Body->W(); }
 		inline int32 Level() const { return Body->Level(); }
