Chaos : Flesh :
Add APIC into MPM transfer ans associated test.

diff --git a/Engine/Source/Runtime/Experimental/Chaos/Public/Chaos/MPMTransfer.h b/Engine/Source/Runtime/Experimental/Chaos/Public/Chaos/MPMTransfer.h
--- a/Engine/Source/Runtime/Experimental/Chaos/Public/Chaos/MPMTransfer.h
+++ b/Engine/Source/Runtime/Experimental/Chaos/Public/Chaos/MPMTransfer.h
@@ -35,6 +35,9 @@ public:
 	TArray<TArray<TArray<int32>>> ElementGridNodeIncidentElements;
 	TArray<TArray<int32>> ElementGridNodesSet;
 
+	//APIC data;
+	TArray<T> AArray;
+	uint32 NumModes = 3;
 
 	TMPMTransfer(){}
 	TMPMTransfer(TMPMGrid<T>& _Grid) :Grid(_Grid) 
@@ -44,6 +47,13 @@ public:
 		NPerEle = NPerSec * Grid.NPerDir;
 		//CellColors.SetNum(NPerEle);
 	}
+	TMPMTransfer(TMPMGrid<T>& _Grid, const int32 NumParticles) :Grid(_Grid)
+	{
+		ensure(Grid.NPerDir == 2 || Grid.NPerDir == 3);
+		NPerSec = Grid.NPerDir * Grid.NPerDir;
+		NPerEle = NPerSec * Grid.NPerDir;
+		AArray.Init((T)0., 3 * 3 * NumParticles);
+	}
 
 	//template <typename ReorderFunctor, typename SplatFunctor>
 	//This one does an initial splat of momentum and mass to the grid. 
@@ -149,15 +159,11 @@ public:
 							T Nkk = Grid.Nijk(Weights[p][2], kkk);
 							T NProd = Nii * Njj * Nkk;
 							GridData[(NTransfer + 1) * GlobIndex] += NProd * InParticles.M(p);
-							//if (GlobIndex == 13472)
-							//{
-							//	std::cout << "";
-							//}
 							for (int32 alpha = 0; alpha < 3; alpha++)
 							{
 								GridData[(NTransfer + 1) * GlobIndex + alpha + 1] += NProd * InParticles.M(p) * InParticles.V(p)[alpha];
 							}
-
+							APICP2G(p, GlobIndex, NProd * InParticles.M(p), InParticles, GridData);
 						}
 					}
 				}
@@ -165,6 +171,20 @@ public:
 		}
 	}
 
+	void APICP2G(const int32 p, const int32 GlobIndex, const T mip, const TDynamicParticles<T, 3>& Particles, TArray<T>& GridData)
+	{
+		TVector<T, 3> xi_minus_xp = Grid.Node(GlobIndex) - Particles.X(p);
+		for (int32 l = 0; l < 3; ++l) 
+		{
+			T vl = T(0);
+			for (int32 m = 0; m < int32(NumModes); ++m)
+				vl += AArray[NumModes * 3 * p + NumModes * l + m] * xi_minus_xp[m];
+
+			GridData[l + (NTransfer + 1) * GlobIndex + 1] += mip * vl;
+		}
+	
+	}
+
 	void ComputeElementMetaData(const TArray<TVector<int32, 4>>& InMesh)
 	{
 		ElementGridNodes.SetNum(InMesh.Num());
@@ -278,6 +298,114 @@ public:
 		return Result;
 	}
 
+	void ComputeAArray(const TArray<T>& GridData, const TDynamicParticles<T, 3>& Particles)
+	{
+		AArray.Init((T)0., NumModes * 3 * Particles.Size());
+		for (int32 p = 0; p < Particles.Size(); p++)
+		{
+			TVector<int32, 3> BaseIndex = Indices[p];
+			for (int32 ii = 0; ii < int32(Grid.NPerDir); ii++)
+			{
+				T Nii = Grid.Nijk(Weights[p][0], ii);
+				T dNii = Grid.dNijk(Weights[p][0], ii, Grid.GetDx()[0]);
+				for (int32 jj = 0; jj < int32(Grid.NPerDir); jj++)
+				{
+					T Njj = Grid.Nijk(Weights[p][1], jj);
+					T dNjj = Grid.dNijk(Weights[p][1], jj, Grid.GetDx()[0]);
+					for (int32 kk = 0; kk < int32(Grid.NPerDir); kk++)
+					{
+						T Nkk = Grid.Nijk(Weights[p][2], kk);
+						T dNkk = Grid.dNijk(Weights[p][2], kk, Grid.GetDx()[0]);
+						TVector<T, 3> Ni = { Nii, Njj, Nkk };
+						TVector<T, 3> dNi = { dNii, dNjj, dNkk };
+						TVector<int32, 3> LocIndex = { ii, jj, kk };
+						TVector<int32, 3> GlobIndex = Grid.Loc2GlobIndex(BaseIndex, LocIndex);
+						int32 GlobIndexFlat = Grid.FlatIndex(GlobIndex);
+						if (GridData[(NTransfer + 1) * GlobIndexFlat] != T(0))
+						{
+							if (Grid.interp == TMPMGrid<T>::linear)
+							{
+								TVector<T, 3> dN((T)0.);
+								Grid.GradNi(Ni, dNi, dN);
+								for (int32 l = 0; l < 3; l++)
+								{
+									T GridQuantity = GridData[(NTransfer + 1) * GlobIndexFlat + 1 + l] / GridData[(NTransfer + 1) * GlobIndexFlat];
+									for (int32 iii = 0; iii < NumModes; ++iii)
+										AArray[NumModes * 3 * p + NumModes * l + iii] += dN[iii] * GridQuantity;
+								}
+							}
+							else
+							{
+								TVector<T, 3> rip = Grid.Node(GlobIndexFlat) - Particles.X(p);
+								T NiProdbInv = (T(4) * Ni[0] * Ni[1] * Ni[2]) / (Grid.GetDx()[0] * Grid.GetDx()[0]);
+								for (int32 l = 0; l < 3; l++)
+								{
+									T GridQuantity = GridData[(NTransfer + 1) * GlobIndexFlat + 1 + l] / GridData[(NTransfer + 1) * GlobIndexFlat];
+									for (int32 iii = 0; iii < NumModes; ++iii)
+										AArray[NumModes * 3 * p + NumModes * l + iii] += NiProdbInv * rip[iii] * GridQuantity;
+								}
+							}
+						}
+						
+					}
+				}
+			}
+		}
+	}
+
+	void GridPositionsToGridData(const TArray<TVector<T, 3>>& GridPositions, const T Dt, TArray<T>& GridData)
+	{
+		for (int32 i = 0; i < GridPositions.Num(); i++)
+		{
+			T Mi = GridData[(NTransfer + 1) * i];
+			if (Mi != T(0))
+			{
+				TVector<T, 3> NodePos = Grid.Node(i);
+				for (int32 alpha = 0; alpha < 3; alpha++)
+				{
+					GridData[(NTransfer + 1) * i + alpha + 1] = Mi * (GridPositions[i][alpha] - NodePos[alpha]) / Dt;
+				}
+			}
+			//else 
+			//{
+			//	for (int32 alpha = 0; alpha < 3; alpha++)
+			//	{
+			//	}
+			//}
+		}
+	
+	}
+
+	void G2P(const TArray<T>& GridData, TDynamicParticles<T, 3>& Particles)
+	{
+		PhysicsParallelFor(Particles.Size(), [&](const int32 p)
+			{
+				Particles.V(p) = TVec3<T>((T)0.);
+				for (int32 ii = 0; ii < int32(Grid.NPerDir); ii++)
+				{
+					T Nii = Grid.Nijk(Weights[p][0], ii);
+					for (int32 jj = 0; jj < int32(Grid.NPerDir); jj++)
+					{
+						T Njj = Grid.Nijk(Weights[p][1], jj);
+						for (int32 kk = 0; kk < int32(Grid.NPerDir); kk++)
+						{
+							T Nkk = Grid.Nijk(Weights[p][2], kk);
+							TVector<int32, 3> LocIndex = { ii, jj, kk };
+							TVector<int32, 3> GlobIndex = Grid.Loc2GlobIndex(Indices[p], LocIndex);
+							int32 GlobIndexFlat = Grid.FlatIndex(GlobIndex);
+							if (GridData[GlobIndexFlat * (NTransfer + 1)] > T(0))
+							{
+								for (int32 alpha = 0; alpha < 3; alpha++)
+								{
+									Particles.V(p)[alpha] += Nii * Njj * Nkk * GridData[GlobIndexFlat * (NTransfer + 1) + alpha + 1] / GridData[GlobIndexFlat * (NTransfer + 1)];
+								}
+							}
+						}
+					}
+				}
+			}, Particles.Size() < 50);
+	}
+
 
 	TArray<TVector<T, 3>> SparseP2G(const TVec4<TVector<T, 3>>& InGradient, const int32 ElementIndex)
 	{
diff --git a/Engine/Source/Runtime/Experimental/Chaos/Public/Chaos/UniformGrid.h b/Engine/Source/Runtime/Experimental/Chaos/Public/Chaos/UniformGrid.h
--- a/Engine/Source/Runtime/Experimental/Chaos/Public/Chaos/UniformGrid.h
+++ b/Engine/Source/Runtime/Experimental/Chaos/Public/Chaos/UniformGrid.h
@@ -335,6 +335,20 @@ public:
 			return ((T(1.5) * (w * w - w) - T(0.25)) * (T(ii) - 1) + (T(0.5) * w - T(0.25))) * (T(ii) - 1) - w * w + w + T(0.5);
 	}
 
+	inline void GradNi(const TVector<T, 3>& Ni, const TVector<T, 3>& dNi, TVector<T, 3>& result) const 
+	{
+		result[0] = dNi[0] * Ni[1] * Ni[2];
+		result[1] = Ni[0] * dNi[1] * Ni[2];
+		result[2] = Ni[0] * Ni[1] * dNi[2];
+	}
+
+	inline T dNijk(T w, int32 ii, T dx) const {
+		if (interp == linear)
+			return (T(2) * T(ii) - T(1)) / dx;
+		else
+			return ((w - 1) * (T(ii) - 1) * (T(ii) - 2) * T(0.5) + (2 * w - 1) * T(ii) * (T(ii) - 2) + w * T(ii) * (T(ii) - 1) * T(0.5)) / dx;
+	}
+
 	inline TVector<int32, 3> Loc2GlobIndex(const TVector<int32, 3>& IndexIn, const TVector<int32, 3>& LocalIndexIn) const {
 		TVector<int32, 3> Result;
 		if (interp == linear) {
@@ -381,6 +395,7 @@ public:
 	TVector<int32, 3> Lin2MultiIndex(const int32 IndexIn) const;
 	const TVector<int32, 3> GetCells() const { return MCells; }
 	const TVector<T, 3> GetMinCorner() const { return MMinCorner; } 
+	const TVector<T, 3>& GetDx() const { return MDx; }
 	void SetDx(const TVector<T, 3> DxIn) { MDx = DxIn; }
 	void UpdateGridFromPositions(const Chaos::TDynamicParticles<T, 3>& InParticles);
 	
diff --git a/Engine/Source/Runtime/Experimental/Chaos/Public/Chaos/XPBDGridBasedCorotatedConstraints.h b/Engine/Source/Runtime/Experimental/Chaos/Public/Chaos/XPBDGridBasedCorotatedConstraints.h
--- a/Engine/Source/Runtime/Experimental/Chaos/Public/Chaos/XPBDGridBasedCorotatedConstraints.h
+++ b/Engine/Source/Runtime/Experimental/Chaos/Public/Chaos/XPBDGridBasedCorotatedConstraints.h
@@ -40,7 +40,7 @@ namespace Chaos::Softs
 			const T& EMesh = (T)1000.0,
 			const T& NuMesh = (T).3
 		)
-			: Base(InParticles, InMesh, GridDx, EMesh, NuMesh), MaxDt(MaxDtIn), MinDt(MinDtIn), CFLCoeff(CFLCoeffIn), MPMGrid(GridDx), MPMTransfer(MPMGrid), PreviousColoring(nullptr)
+			: Base(InParticles, InMesh, GridDx, EMesh, NuMesh), MaxDt(MaxDtIn), MinDt(MinDtIn), CFLCoeff(CFLCoeffIn), MPMGrid(GridDx), MPMTransfer(MPMGrid, InParticles.Size()), PreviousColoring(nullptr)
 		{
 			//TMPMGrid<T> MPMGrid(GridN);
 			MPMGrid.UpdateGridFromPositions(InParticles);
@@ -138,16 +138,25 @@ namespace Chaos::Softs
 				TRACE_CPUPROFILER_EVENT_SCOPE(TEXT("STAT_ChaosXPBDGridBasedCorotatedComputeGridPositions"));
 				MPMTransfer.ComputeGridPositions(GridData, Dt, GridPositions);
 			}
+		}
+
+		void TimeStepPostprocessing(const ParticleType& InParticles, const T Dt)
+		{
+			{
+				TRACE_CPUPROFILER_EVENT_SCOPE(TEXT("STAT_ChaosXPBDGridBasedCorotatedGrid::FinalG2P"));
+				FinalG2P(InParticles);
+			}
+			
+			{
+				TRACE_CPUPROFILER_EVENT_SCOPE(TEXT("STAT_ChaosXPBDGridBasedCorotated::GridPositions2GridData"));
+				MPMTransfer.GridPositionsToGridData(GridPositions, Dt, GridData);
+			}
 
-			//if (!InitialGridDataWritten)
-			//{
-			//	WriteGridNodes();
-			//}
-			//MPMTransfer.Grid.UpdateGridFromPositions(InParticles);
-			//MPMTransfer.InitialP2G(InParticles, GridData);
-			//MPMTransfer.ComputeElementMetaData(MeshConstraints);
-			//ComputeGridBasedGraphSubColoringPointer(ElementsPerColor, MPMTransfer.Grid, MPMTransfer.Grid.Size(), PreviousColoring, MPMTransfer.ElementGridNodesSet, ElementsPerSubColors);
-			//MPMTransfer.ComputeGridPositions(GridData, Dt, GridPositions);
+			{
+				TRACE_CPUPROFILER_EVENT_SCOPE(TEXT("STAT_ChaosXPBDGridBasedCorotated::ComputeAArray"));
+				MPMTransfer.ComputeAArray(GridData, InParticles);
+			}
+		
 		}
 
 		T ComputeCFLDt(const ParticleType& InParticles)
@@ -445,12 +454,8 @@ namespace Chaos::Softs
 		//interpolate particle positions from grid positions
 		void FinalG2P(ParticleType& Particles)
 		{
-			TVec4<TVector<T, 3>> Result(TVector<T, 3>((T)0.));
-
 			PhysicsParallelFor(Particles.Size(), [&](const int32 p)
-				//for (int32 p = 0; p < int32(Particles.Size()); p++) 
 				{
-					//TVector<int32, 3> BaseIndex = MPMTransfer.Indices[p];
 					Particles.P(p) = TVec3<T>((T)0.);
 					for (int32 ii = 0; ii < int32(MPMGrid.NPerDir); ii++)
 					{
@@ -464,12 +469,12 @@ namespace Chaos::Softs
 								TVector<int32, 3> LocIndex = { ii, jj, kk };
 								TVector<int32, 3> GlobIndex = MPMGrid.Loc2GlobIndex(MPMTransfer.Indices[p], LocIndex);
 								int32 GlobIndexFlat = MPMGrid.FlatIndex(GlobIndex);
+								//G2PHelper(Nii * Njj * Nkk, GridPositions[GlobIndexFlat], p, Particles);
 								Particles.P(p) += Nii * Njj * Nkk * GridPositions[GlobIndexFlat];
 							}
 						}
 					}
 				}, Particles.Size() < 50);
-
 		}
 
 	protected:
