Fixed per-particle velocity damping not using the preprocessor-defined default value in non-shipping builds.

diff --git a/Engine/Source/Runtime/Experimental/Chaos/Private/Chaos/PerParticleDampVelocity.cpp b/Engine/Source/Runtime/Experimental/Chaos/Private/Chaos/PerParticleDampVelocity.cpp
--- a/Engine/Source/Runtime/Experimental/Chaos/Private/Chaos/PerParticleDampVelocity.cpp
+++ b/Engine/Source/Runtime/Experimental/Chaos/Private/Chaos/PerParticleDampVelocity.cpp
@@ -14,9 +14,8 @@
 
 #if INTEL_ISPC && !UE_BUILD_SHIPPING
 static_assert(sizeof(ispc::FVector3f) == sizeof(Chaos::Softs::FSolverVec3), "sizeof(ispc::FVector3f) != sizeof(Chaos::Softs::FSolverVec3)");
-
-bool bChaos_DampVelocity_ISPC_Enabled = true;
-FAutoConsoleVariableRef CVarChaosDampVelocityISPCEnabled(TEXT("p.Chaos.DampVelocity.ISPC"), bChaos_DampVelocity_ISPC_Enabled, TEXT("Whether to use ISPC optimizations in per particle damp velocity calculation"));
+bool bChaos_DampVelocity_ISPC_Enabled = CHAOS_DAMP_VELOCITY_ISPC_ENABLED_DEFAULT;
+static FAutoConsoleVariableRef CVarChaosDampVelocityISPCEnabled(TEXT("p.Chaos.DampVelocity.ISPC"), bChaos_DampVelocity_ISPC_Enabled, TEXT("Whether to use ISPC optimizations in per particle damp velocity calculation"));
 #endif
 
 namespace Chaos::Softs {
diff --git a/Engine/Source/Runtime/Experimental/Chaos/Public/Chaos/PerParticleDampVelocity.h b/Engine/Source/Runtime/Experimental/Chaos/Public/Chaos/PerParticleDampVelocity.h
--- a/Engine/Source/Runtime/Experimental/Chaos/Public/Chaos/PerParticleDampVelocity.h
+++ b/Engine/Source/Runtime/Experimental/Chaos/Public/Chaos/PerParticleDampVelocity.h
@@ -39,8 +39,10 @@ private:
 #endif
 
 // Support run-time toggling on supported platforms in non-shipping configurations
-#if !INTEL_ISPC || UE_BUILD_SHIPPING
-static constexpr bool bChaos_DampVelocity_ISPC_Enabled = INTEL_ISPC && CHAOS_DAMP_VELOCITY_ISPC_ENABLED_DEFAULT;
+#if INTEL_ISPC
+#if UE_BUILD_SHIPPING
+static constexpr bool bChaos_DampVelocity_ISPC_Enabled = CHAOS_DAMP_VELOCITY_ISPC_ENABLED_DEFAULT;
 #else
 extern CHAOS_API bool bChaos_DampVelocity_ISPC_Enabled;
-#endif
+#endif // UE_BUILD_SHIPPING
+#endif
\ No newline at end of file
