Chaos - another dangling collision pointer fix

This was caused by a sleeping island holding a pointer to an awake collision which then gets deleted.
There was an implicit assumption that particles are awake or get woken when a constraint is added to the graph. This is not true.

 The likely sequence of events was
- a collision between a sleeping body and a kinematic body is created and added to the graph
- the body and constraint are otherwise isolated and end up in their own island
- the island is set to sleeping
- when we sync the island sleep state to the constraints, we skip sleeping islands so the constraint sleep state is  never set
- the constraint is later destroyed but is still referenced in a sleeping island

diff --git a/Engine/Source/Runtime/Experimental/Chaos/Private/Chaos/Island/IslandManager.cpp b/Engine/Source/Runtime/Experimental/Chaos/Private/Chaos/Island/IslandManager.cpp
--- a/Engine/Source/Runtime/Experimental/Chaos/Private/Chaos/Island/IslandManager.cpp
+++ b/Engine/Source/Runtime/Experimental/Chaos/Private/Chaos/Island/IslandManager.cpp
@@ -191,7 +191,7 @@ inline void PopulateIslands(FPBDIslandManager::GraphType* IslandGraph)
 		if (IslandGraph->GraphIslands.IsValidIndex(IslandIndex))
 		{
 			FPBDIslandSolver* IslandSolver = IslandGraph->GraphIslands[IslandIndex].IslandItem;
-			if (!IslandSolver->IsSleeping())
+			if (!IslandSolver->IsSleeping() || IslandSolver->SleepingChanged())
 			{
 				IslandSolver->AddParticle(GraphNode.NodeItem);
 			}
@@ -230,7 +230,11 @@ inline void PopulateIslands(FPBDIslandManager::GraphType* IslandGraph)
 		if (IslandGraph->GraphIslands.IsValidIndex(GraphEdge.IslandIndex))
 		{
 			FPBDIslandSolver* IslandSolver = IslandGraph->GraphIslands[GraphEdge.IslandIndex].IslandItem;
-			if (!IslandSolver->IsSleeping())
+
+			// Note: we also populate the island if the sleeping state just changed. This is because we use this
+			// data to sync the constraint's sleeping state in SyncIslands.
+			// @todo(chaos): This is not great - we need a more robust way to sync the sleep state of islands/constraints
+			if (!IslandSolver->IsSleeping() || IslandSolver->SleepingChanged())
 			{
 #if CHAOS_CONSTRAINTHANDLE_DEBUG_ENABLED
 				ensure(GraphEdge.bValidEdge);
@@ -700,6 +704,9 @@ void FPBDIslandManager::SyncIslands(FPBDRigidsSOAs& Particles, const int32 NumCo
 	{
 		if (IslandGraph->GraphIslands.IsValidIndex(IslandIndex))
 		{
+			// NOTE: We do not check !IsSleeping() here because we need to sync the sleep state of islands
+			// that were just put to sleep this tick. See PopulateIslands.
+
 			// If the island item is not set on the graph we check if the solver island
 			// at the right index is already there. If not we create a new solver island.
 			if(!IslandSolvers.IsValidIndex(IslandIndex))
@@ -779,7 +786,10 @@ void FPBDIslandManager::BuildGroups(const int32 NumContainers)
 	int32 NumConstraints = 0;
 	for (TUniquePtr<FPBDIslandSolver>& IslandSolver : IslandSolvers)
 	{
-		NumConstraints += IslandSolver->NumConstraints();
+		if (!IslandSolver->IsSleeping())
+		{
+			NumConstraints += IslandSolver->NumConstraints();
+		}
 	}
 	
 	const int32 NumGroups = IslandGroups.Num();
@@ -797,26 +807,29 @@ void FPBDIslandManager::BuildGroups(const int32 NumContainers)
 	{
 		if( FPBDIslandSolver* IslandSolver = IslandSolvers[SortedIndex].Get())
 		{
-			checkf(GroupIndex < NumGroups, TEXT("Index %d is less than NumGroups %d"), GroupIndex, NumGroups);
-			IslandGroups[GroupIndex]->AddIsland(IslandSolver);
-			IslandGroups[GroupIndex]->NumParticles() += IslandSolver->NumParticles();
-			IslandGroups[GroupIndex]->NumConstraints() += IslandSolver->NumConstraints();
-
-			check(IslandSolver->NumContainerIds() == IslandGroups[GroupIndex]->NumContainerIds());
-			for(int32 ContainerIndex = 0; ContainerIndex < IslandSolver->NumContainerIds(); ++ContainerIndex)
+			if (!IslandSolver->IsSleeping())
 			{
-				IslandGroups[GroupIndex]->ConstraintCount(ContainerIndex) +=
-					IslandSolver->ConstraintCount(ContainerIndex);
-			}
+				checkf(GroupIndex < NumGroups, TEXT("Index %d is less than NumGroups %d"), GroupIndex, NumGroups);
+				IslandGroups[GroupIndex]->AddIsland(IslandSolver);
+				IslandGroups[GroupIndex]->NumParticles() += IslandSolver->NumParticles();
+				IslandGroups[GroupIndex]->NumConstraints() += IslandSolver->NumConstraints();
+
+				check(IslandSolver->NumContainerIds() == IslandGroups[GroupIndex]->NumContainerIds());
+				for(int32 ContainerIndex = 0; ContainerIndex < IslandSolver->NumContainerIds(); ++ContainerIndex)
+				{
+					IslandGroups[GroupIndex]->ConstraintCount(ContainerIndex) +=
+						IslandSolver->ConstraintCount(ContainerIndex);
+				}
 			
-			IslandSolver->SetGroupIndex(GroupIndex);
-			GroupOffset += IslandSolver->NumConstraints();
+				IslandSolver->SetGroupIndex(GroupIndex);
+				GroupOffset += IslandSolver->NumConstraints();
 
-			// @todo: In theory we should never need the second check	
-			if(GroupOffset > GroupSize && (GroupIndex + 1) < NumGroups)
-			{
-				GroupIndex++;
-				GroupOffset = 0;
+				// @todo: In theory we should never need the second check	
+				if(GroupOffset > GroupSize && (GroupIndex + 1) < NumGroups)
+				{
+					GroupIndex++;
+					GroupOffset = 0;
+				}
 			}
 		}
 	}
@@ -1115,11 +1128,30 @@ bool FPBDIslandManager::DebugCheckIslandConstraints() const
 			const bool bInAwakeIsland = (GraphEdge.IslandIndex != INDEX_NONE) && IslandGraph->GraphIslands.IsValidIndex(GraphEdge.IslandIndex) && !IslandGraph->GraphIslands[GraphEdge.IslandIndex].bIsSleeping;
 			const bool bIsDynamic0 = (ConstraintHolder.GetParticle0() != nullptr) ? FConstGenericParticleHandle(ConstraintHolder.GetParticle0())->IsDynamic() : false;
 			const bool bIsDynamic1 = (ConstraintHolder.GetParticle1() != nullptr) ? FConstGenericParticleHandle(ConstraintHolder.GetParticle1())->IsDynamic() : false;
+			
 			if (bInAwakeIsland)
 			{
+				// If the island is awake, both particle should be too
 				bSuccess = bSuccess && (bIsDynamic0 || bIsDynamic1);
 				ensure(bIsDynamic0 || bIsDynamic1);
+
+				// If the island is awake, constraint should be too
+				if (ConstraintHolder->SupportsSleeping())
+				{
+					const bool bIsAsleep = ConstraintHolder->IsSleeping();
+					ensure(!bIsAsleep);
+				}
 			}
+			else
+			{
+				// If the island is asleep, make sure the constraint is too
+				if (ConstraintHolder->SupportsSleeping())
+				{
+					const bool bIsAsleep = ConstraintHolder->IsSleeping();
+					ensure(bIsAsleep);
+				}
+			}
+
 			if (!bIsDynamic0 && !bIsDynamic1)
 			{
 				bSuccess = bSuccess && !bInAwakeIsland;
@@ -1214,8 +1246,15 @@ void FPBDIslandManager::DebugCheckConstraintIslands(const FConstraintHandle* Con
 				ensure(ConstraintHandle->IsSleeping());
 			}
 
-			// If we are sleeping, we should not appear in any solver islands (we should only be in the graph)
-			ensure(EdgeSolverIslands.Num() == 0);
+			// If we are sleeping, we usually don't appear in the solver islands (because we don't need to solve sleepers)
+			// but we do create solver islands for 1 tick after the island goes to sleep so that we can sync the sleep state
+			// (see PopulateIslands and SyncSleepState). 
+			if (EdgeSolverIslands.Num() > 0)
+			{
+				ensure(EdgeSolverIslands.Num() == 1);
+				const int32 EdgeIsland = IslandIndexing[EdgeSolverIslands[0]];
+				ensure(EdgeIslandIndex == EdgeIsland);
+			}
 		}
 		else
 		{
