Iris
- Minor review fixes

diff --git a/Engine/Source/Runtime/Experimental/Iris/Core/Private/Iris/ReplicationSystem/Filtering/ReplicationFiltering.h b/Engine/Source/Runtime/Experimental/Iris/Core/Private/Iris/ReplicationSystem/Filtering/ReplicationFiltering.h
--- a/Engine/Source/Runtime/Experimental/Iris/Core/Private/Iris/ReplicationSystem/Filtering/ReplicationFiltering.h
+++ b/Engine/Source/Runtime/Experimental/Iris/Core/Private/Iris/ReplicationSystem/Filtering/ReplicationFiltering.h
@@ -86,7 +86,7 @@ public:
 	void NotifyAddedDependentObject(uint32 ObjectIndex);
 	void NotifyRemovedDependentObject(uint32 ObjectIndex);
 	
-	// SubOjectFilter status
+	// SubObjectFilter status
 	void AddSubObjectFilter(FNetObjectGroupHandle GroupHandle);
 	void RemoveSubObjectFilter(FNetObjectGroupHandle GroupHandle);
 	bool IsSubObjectFilterGroup(FNetObjectGroupHandle GroupHandle) const { return GroupHandle && SubObjectFilterGroups.GetBit(GroupHandle); }
@@ -154,7 +154,7 @@ private:
 
 	bool HasDynamicFilters() const;
 
-	// Helper to update and reset group filter effects if objects are removed from a filter or after a filter status change, returns true of the group filter was changed
+	// Helper to update and reset group filter effects if objects are removed from a filter or after a filter status change, returns true if the group filter was changed
 	bool UpdateGroupFilterEffectsForObject(uint32 ObjectIndex, uint32 ConnectionId, const FNetBitArray& ScopableObjects);
 	bool HasOwnerFilter(uint32 ObjectIndex) const;
 	bool HasConnectionFilter(uint32 ObjectIndex) const;
@@ -243,7 +243,6 @@ private:
 	FNetBitArray DynamicFilterEnabledObjects;
 	FNetBitArray ObjectsRequiringDynamicFilterUpdate;
 
-	//
 	uint32 bHasNewConnection : 1;
 	uint32 bHasRemovedConnection : 1;
 	uint32 bHasDirtyConnectionFilter: 1;
diff --git a/Engine/Source/Runtime/Experimental/Iris/Core/Private/Iris/ReplicationSystem/NetHandleManager.cpp b/Engine/Source/Runtime/Experimental/Iris/Core/Private/Iris/ReplicationSystem/NetHandleManager.cpp
--- a/Engine/Source/Runtime/Experimental/Iris/Core/Private/Iris/ReplicationSystem/NetHandleManager.cpp
+++ b/Engine/Source/Runtime/Experimental/Iris/Core/Private/Iris/ReplicationSystem/NetHandleManager.cpp
@@ -72,7 +72,7 @@ uint32 FNetHandleManager::InternalCreateNetObject(const FNetHandle NetHandle, co
 	// verify that the handle is free
 	if (HandleMap.Contains(NetHandle))
 	{
-		ensureAlwaysMsgf(false, TEXT("NetHandleManager::InternalCreateNetHandleHandle %s already exists"), *NetHandle.ToString());
+		ensureAlwaysMsgf(false, TEXT("NetHandleManager::InternalCreateNetObject %s already exists"), *NetHandle.ToString());
 		return InvalidInternalIndex;
 	}
 
@@ -195,7 +195,7 @@ FNetHandle FNetHandleManager::CreateNetObject(FNetHandle WantedHandle, const FRe
 // Create NetHandle not owned by us
 FNetHandle FNetHandleManager::CreateNetObjectFromRemote(FNetHandle WantedHandle, const FReplicationProtocol* ReplicationProtocol)
 {
-	if (!ensureAlwaysMsgf(WantedHandle.IsValid() && !WantedHandle.IsCompleteHandle(), TEXT("FNetHandleManager::CreateNetHandleFromRemote Expected WantedHandle %s to be valid and incomplete"), *WantedHandle.ToString()))
+	if (!ensureAlwaysMsgf(WantedHandle.IsValid() && !WantedHandle.IsCompleteHandle(), TEXT("FNetHandleManager::CreateNetObjectFromRemote Expected WantedHandle %s to be valid and incomplete"), *WantedHandle.ToString()))
 	{
 		return FNetHandle();
 	}
diff --git a/Engine/Source/Runtime/Experimental/Iris/Core/Private/Iris/ReplicationSystem/NetHandleManager.h b/Engine/Source/Runtime/Experimental/Iris/Core/Private/Iris/ReplicationSystem/NetHandleManager.h
--- a/Engine/Source/Runtime/Experimental/Iris/Core/Private/Iris/ReplicationSystem/NetHandleManager.h
+++ b/Engine/Source/Runtime/Experimental/Iris/Core/Private/Iris/ReplicationSystem/NetHandleManager.h
@@ -162,10 +162,10 @@ public:
 	// Verify a handle against internal handle, A handle is valid if it matches internal storage
 	inline bool IsValidNetHandle(FNetHandle Handle) const;
 
-	// Returns true of the handle is for local replicated object
+	// Returns true if the handle is for local replicated object
 	inline bool IsLocalNetHandle(FNetHandle Handle) const;
 
-	// Returns true of the handle is for a remote replicated object
+	// Returns true if the handle is for a remote replicated object
 	inline bool IsRemoteNetHandle(FNetHandle Handle) const;
 
 	// Extract an Full handle from an incomplete one, a incomplete handle is a handle with only an index
@@ -377,19 +377,17 @@ bool FNetHandleManager::IsValidNetHandle(FNetHandle Handle) const
 	return HandleMap.Contains(Handle);
 }
 
-// Returns true of the handle is a local handle
 bool FNetHandleManager::IsLocalNetHandle(FNetHandle Handle) const
 {
 	if (const uint32 InternalIndex = GetInternalIndex(Handle))
 	{
-		// For the time being only replicated objects owned by this peer has a state buffer`1
+		// For the time being only replicated objects owned by this peer has a state buffer
 		return ReplicatedObjectStateBuffers[InternalIndex] != nullptr;
 	}
 
 	return false;
 }
 
-// Returns true of the handle is a remote handle
 bool FNetHandleManager::IsRemoteNetHandle(FNetHandle Handle) const
 {
 	if (const uint32 InternalIndex = GetInternalIndex(Handle))
diff --git a/Engine/Source/Runtime/Experimental/Iris/Core/Private/Iris/ReplicationSystem/ReplicationOperationsInternal.h b/Engine/Source/Runtime/Experimental/Iris/Core/Private/Iris/ReplicationSystem/ReplicationOperationsInternal.h
--- a/Engine/Source/Runtime/Experimental/Iris/Core/Private/Iris/ReplicationSystem/ReplicationOperationsInternal.h
+++ b/Engine/Source/Runtime/Experimental/Iris/Core/Private/Iris/ReplicationSystem/ReplicationOperationsInternal.h
@@ -50,7 +50,7 @@ struct FReplicationStateOperationsInternal
 	/** Free the dynamic state from a state for a single replication state */
 	static IRISCORE_API void FreeDynamicState(uint8* ObjectStateBuffer, const FReplicationStateDescriptor* Descriptor);
 
-	/** Clone a qauntized state, Note: DstInternalBuffer is expected to be uninitialized */
+	/** Clone a quantized state, Note: DstInternalBuffer is expected to be uninitialized */
 	static IRISCORE_API void CloneQuantizedState(FNetSerializationContext& Context, uint8* RESTRICT DstInternalBuffer, const uint8* RESTRICT SrcInternalBuffer, const FReplicationStateDescriptor* Descriptor);
 	
 	/** Collect references from a state that does not have any changemask information */
diff --git a/Engine/Source/Runtime/Experimental/Iris/Core/Private/Iris/Serialization/ObjectNetSerializer.cpp b/Engine/Source/Runtime/Experimental/Iris/Core/Private/Iris/Serialization/ObjectNetSerializer.cpp
--- a/Engine/Source/Runtime/Experimental/Iris/Core/Private/Iris/Serialization/ObjectNetSerializer.cpp
+++ b/Engine/Source/Runtime/Experimental/Iris/Core/Private/Iris/Serialization/ObjectNetSerializer.cpp
@@ -1,18 +1,13 @@
 // Copyright Epic Games, Inc. All Rights Reserved.
 
 #include "Iris/Serialization/ObjectNetSerializer.h"
-
 #include "Iris/Core/NetObjectReference.h"
 #include "Iris/Core/IrisLog.h"
-
 #include "Net/Core/Trace/NetTrace.h"
-
 #include "Iris/Serialization/InternalNetSerializer.h"
 #include "Iris/Serialization/NetBitStreamReader.h"
 #include "Iris/Serialization/NetBitStreamWriter.h"
-#include "Iris/Serialization/ObjectNetSerializer.h"
 #include "Iris/Serialization/InternalNetSerializationContext.h"
-
 #include "Iris/ReplicationSystem/NetHandleManager.h"
 #include "Iris/ReplicationSystem/ObjectReferenceCache.h"
 #include "Iris/Serialization/NetBitStreamUtil.h"
diff --git a/Engine/Source/Runtime/Experimental/Iris/Core/Public/Iris/ReplicationState/ReplicationStateFwd.h b/Engine/Source/Runtime/Experimental/Iris/Core/Public/Iris/ReplicationState/ReplicationStateFwd.h
--- a/Engine/Source/Runtime/Experimental/Iris/Core/Public/Iris/ReplicationState/ReplicationStateFwd.h
+++ b/Engine/Source/Runtime/Experimental/Iris/Core/Public/Iris/ReplicationState/ReplicationStateFwd.h
@@ -4,9 +4,11 @@
 
 #include "CoreTypes.h"
 
-namespace UE {
-namespace Net {
-namespace Private {
+namespace UE::Net
+{
+
+namespace Private
+{
 	struct FReplicationStateHeaderAccessor;
 }
 
@@ -21,7 +23,7 @@ struct FReplicationStateHeader
 	bool IsBound() const { return ReplicationIndex != 0U; }
 
 private:
-	friend struct Private::FReplicationStateHeaderAccessor;
+	friend Private::FReplicationStateHeaderAccessor;
 
 	// All replication states that are bound by a instance protocol as assigned an index used for dirty state tracking
 	uint32 ReplicationIndex : 24;
@@ -35,7 +37,8 @@ private:
 
 static_assert(sizeof(FReplicationStateHeader) == sizeof(uint32) && alignof(FReplicationStateHeader) == alignof(uint32), "FReplicationStateHeader must currently have the same size and alignment as uint32");
 
-namespace Private {
+namespace Private
+{
 
 /** 
  * Internal access, should only be used by internal code
@@ -52,4 +55,6 @@ struct FReplicationStateHeaderAccessor
 	static void SetReplicationIndex(FReplicationStateHeader& Header, uint32 Index, uint32 ReplicationSystemId) { Header.ReplicationIndex = Index; Header.ReplicationSystemId = ReplicationSystemId; }
 };
 
-}}}
+}
+
+}
diff --git a/Engine/Source/Runtime/Experimental/Iris/Core/Public/Iris/ReplicationSystem/Filtering/NetObjectFilter.h b/Engine/Source/Runtime/Experimental/Iris/Core/Public/Iris/ReplicationSystem/Filtering/NetObjectFilter.h
--- a/Engine/Source/Runtime/Experimental/Iris/Core/Public/Iris/ReplicationSystem/Filtering/NetObjectFilter.h
+++ b/Engine/Source/Runtime/Experimental/Iris/Core/Public/Iris/ReplicationSystem/Filtering/NetObjectFilter.h
@@ -78,7 +78,7 @@ struct FNetObjectFilteringParams
 	/** ID of the connection that the filtering applies to. */
 	uint32 ConnectionId;
 
-	/** The view associated with the connection and its sub-sconnections that objects are filtered for. */
+	/** The view associated with the connection and its sub-connections that objects are filtered for. */
 	UE::Net::FReplicationView View;
 };
 
@@ -182,10 +182,10 @@ public:
 	/** A new object want to use this filter. Opportunity to cache some information for it. The info struct passed has been zeroed. Must be overriden. */
 	IRISCORE_API virtual bool AddObject(uint32 ObjectIndex, FNetObjectFilterAddObjectParams&) PURE_VIRTUAL(AddObject, return false;)
 
-	/** An object do no longer want to use this filter. */
+	/** An object no longer wants to use this filter. */
 	IRISCORE_API virtual void RemoveObject(uint32 ObjectIndex, const FNetObjectFilteringInfo&) PURE_VIRTUAL(RemoveObject,)
 
-	/** A set of objects ususing this filter has been updated. An opportunity for the filter to update cached data. */
+	/** A set of objects using this filter has been updated. An opportunity for the filter to update cached data. */
 	IRISCORE_API virtual void UpdateObjects(FNetObjectFilterUpdateParams&) PURE_VIRTUAL(UpdateObjects,)
 
 	/**
diff --git a/Engine/Source/Runtime/Experimental/Iris/Core/Public/Iris/ReplicationSystem/Filtering/NetObjectFilterDefinitions.h b/Engine/Source/Runtime/Experimental/Iris/Core/Public/Iris/ReplicationSystem/Filtering/NetObjectFilterDefinitions.h
--- a/Engine/Source/Runtime/Experimental/Iris/Core/Public/Iris/ReplicationSystem/Filtering/NetObjectFilterDefinitions.h
+++ b/Engine/Source/Runtime/Experimental/Iris/Core/Public/Iris/ReplicationSystem/Filtering/NetObjectFilterDefinitions.h
@@ -9,7 +9,7 @@
 USTRUCT()
 struct FNetObjectFilterDefinition
 {
-	GENERATED_USTRUCT_BODY()
+	GENERATED_BODY()
 
 	/** Filter identifier. Used to uniquely identify a filter. */
 	UPROPERTY()
@@ -38,7 +38,7 @@ class UNetObjectFilterDefinitions final : public UObject
 	GENERATED_BODY()
 
 public:
-	/** Returns the filter definitions exactly as configured. May contain invalid defintions, such as ClassNames not inheriting from UNetObjectFilter. */
+	/** Returns the filter definitions exactly as configured. May contain invalid definitions, such as ClassNames not inheriting from UNetObjectFilter. */
 	IRISCORE_API TConstArrayView<FNetObjectFilterDefinition> GetFilterDefinitions() const;
 
 private:
diff --git a/Engine/Source/Runtime/Experimental/Iris/Core/Public/Iris/ReplicationSystem/NetTokenStore.h b/Engine/Source/Runtime/Experimental/Iris/Core/Public/Iris/ReplicationSystem/NetTokenStore.h
--- a/Engine/Source/Runtime/Experimental/Iris/Core/Public/Iris/ReplicationSystem/NetTokenStore.h
+++ b/Engine/Source/Runtime/Experimental/Iris/Core/Public/Iris/ReplicationSystem/NetTokenStore.h
@@ -49,8 +49,6 @@ class FNetTokenStore
 {
 	UE_NONCOPYABLE(FNetTokenStore);
 public:
-	using FNetTokenStoreState = UE::Net::FNetTokenStoreState;
-
 	FNetTokenStore();
 	~FNetTokenStore();
 
diff --git a/Engine/Source/Runtime/Experimental/Iris/Core/Public/Iris/ReplicationSystem/ReplicationSystem.h b/Engine/Source/Runtime/Experimental/Iris/Core/Public/Iris/ReplicationSystem/ReplicationSystem.h
--- a/Engine/Source/Runtime/Experimental/Iris/Core/Public/Iris/ReplicationSystem/ReplicationSystem.h
+++ b/Engine/Source/Runtime/Experimental/Iris/Core/Public/Iris/ReplicationSystem/ReplicationSystem.h
@@ -340,7 +340,7 @@ public:
 	 * @param FilterName The name of the filter. Names of valid filters are configured in UNetObjectFilterDefinitions.
 	 * @return A valid handle if a filter with the given name has been successfully created in this system or InvalidNetObjectFilterHandle if not.
 	 * @see SetFilter
-	 * @ee UNetObjectFilterDefinitions
+	 * @see UNetObjectFilterDefinitions
 	 */
 	IRISCORE_API UE::Net::FNetObjectFilterHandle GetFilterHandle(const FName FilterName) const;
 
@@ -348,7 +348,7 @@ public:
 	 * Gets the filter with a given name.
 	 * @param FilterName The name of the filter. Names of valid filters are configured in UNetObjectFilterDefinitions.
 	 * @return A pointer to the filter if it exists, nullptr if not.
-	 * @ee UNetObjectFilterDefinitions
+	 * @see UNetObjectFilterDefinitions
 	 */
 	IRISCORE_API UNetObjectFilter* GetFilter(const FName FilterName) const;
 
diff --git a/Engine/Source/Runtime/Experimental/Iris/Core/Public/Iris/Serialization/NetSerializer.h b/Engine/Source/Runtime/Experimental/Iris/Core/Public/Iris/Serialization/NetSerializer.h
--- a/Engine/Source/Runtime/Experimental/Iris/Core/Public/Iris/Serialization/NetSerializer.h
+++ b/Engine/Source/Runtime/Experimental/Iris/Core/Public/Iris/Serialization/NetSerializer.h
@@ -68,7 +68,7 @@ struct FExampleNetSerializer
 
 	/** Required. Serialize is responsible for writing the quantized data to a bit stream provided by the serialization context. */
 	static void Serialize(FNetSerializationContext&, const FNetSerializeArgs&);
-	/** Required. Serialize is responsible for reading the quantized data from a bit stream provided by the serialization context. */
+	/** Required. Deserialize is responsible for reading the quantized data from a bit stream provided by the serialization context. */
 	static void Deserialize(FNetSerializationContext&, const FNetDeserializeArgs&);
 
 	/**
