Chaos : Flesh :
Fix sigma uninitialization error in 3X3 SVD.

diff --git a/Engine/Source/Runtime/Experimental/ChaosCore/Public/Chaos/ImplicitQRSVD.h b/Engine/Source/Runtime/Experimental/ChaosCore/Public/Chaos/ImplicitQRSVD.h
--- a/Engine/Source/Runtime/Experimental/ChaosCore/Public/Chaos/ImplicitQRSVD.h
+++ b/Engine/Source/Runtime/Experimental/ChaosCore/Public/Chaos/ImplicitQRSVD.h
@@ -595,7 +595,11 @@ namespace Chaos {
 	  */
 	template <class T>
 	inline int SingularValueDecomposition(const PMatrix<T, 3, 3>& A, PMatrix<T, 3, 3>& U, TVector<T, 3>& sigma, PMatrix<T, 3, 3>& V,
-		T tol = 1024 * TMathUtilConstants<T>::Epsilon) {
+		T tol = 1024 * TMathUtilConstants<T>::Epsilon) 
+	{
+		sigma[0] = T(0.);
+		sigma[1] = T(0.);
+		sigma[2] = T(0.);
 
 		PMatrix<T, 3, 3> B = A;
 		U = PMatrix<T, 3, 3>(1, 1, 1);
