prevent crash if nanite data is null on geometry collection object

diff --git a/Engine/Source/Runtime/Experimental/GeometryCollectionEngine/Private/GeometryCollection/GeometryCollectionObject.cpp b/Engine/Source/Runtime/Experimental/GeometryCollectionEngine/Private/GeometryCollection/GeometryCollectionObject.cpp
--- a/Engine/Source/Runtime/Experimental/GeometryCollectionEngine/Private/GeometryCollection/GeometryCollectionObject.cpp
+++ b/Engine/Source/Runtime/Experimental/GeometryCollectionEngine/Private/GeometryCollection/GeometryCollectionObject.cpp
@@ -1263,7 +1263,10 @@ void UGeometryCollection::EnsureDataIsCooked(bool bInitResources, bool bIsTransa
 			// If there is no geometry in the collection, we leave Nanite data alone.
 			if (GeometryCollection->NumElements(FGeometryCollection::GeometryGroup) > 0)
 			{
-				NaniteData->InitResources(this);
+				if (NaniteData)
+				{
+					NaniteData->InitResources(this);
+				}
 			}
 		}
 		LastBuiltGuid = StateGuid;
