Dataflow Graph:
- moved DataflowSettings from DataflowEditor to Dataflow/Core

diff --git a/Engine/Plugins/Experimental/Dataflow/Source/DataflowEditor/Private/Dataflow/DataflowSettings.cpp b/Engine/Plugins/Experimental/Dataflow/Source/DataflowEditor/Private/Dataflow/DataflowSettings.cpp
deleted file mode 100644
diff --git a/Engine/Plugins/Experimental/Dataflow/Source/DataflowEditor/Public/Dataflow/DataflowSettings.h b/Engine/Plugins/Experimental/Dataflow/Source/DataflowEditor/Public/Dataflow/DataflowSettings.h
deleted file mode 100644
diff --git a/Engine/Source/Runtime/Experimental/Dataflow/Core/Private/Dataflow/DataflowSettings.cpp b/Engine/Source/Runtime/Experimental/Dataflow/Core/Private/Dataflow/DataflowSettings.cpp
new file mode 100644
--- /dev/null
+++ b/Engine/Source/Runtime/Experimental/Dataflow/Core/Private/Dataflow/DataflowSettings.cpp
@@ -0,0 +1,43 @@
+// Copyright Epic Games, Inc. All Rights Reserved.
+
+#include "Dataflow/DataflowSettings.h"
+
+UDataflowSettings::UDataflowSettings(const FObjectInitializer& ObjectInitlaizer)
+	: Super(ObjectInitlaizer)
+{
+	ManagedArrayCollectionPinTypeColor = FLinearColor(0.353393f, 0.454175f, 1.0f, 1.0f);
+	ArrayPinTypeColor = FLinearColor(1.0f, 0.172585f, 0.0f, 1.0f);
+	BoxPinTypeColor = FLinearColor(0.013575f, 0.770000f, 0.429609f, 1.0f);
+
+	GeometryCollectionCategoryNodeTitleColor = FLinearColor(0.55f, 0.0f, 1.f);
+	GeometryCollectionCategoryNodeBodyTintColor = FLinearColor(0.55f, 0.0f, 1.f);
+}
+
+FName UDataflowSettings::GetCategoryName() const
+{
+	return TEXT("Plugins");
+}
+
+FText UDataflowSettings::GetSectionText() const
+{
+	return NSLOCTEXT("DataflowPlugin", "DataflowSettingsSection", "Dataflow");
+}
+
+void UDataflowSettings::PostEditChangeProperty(FPropertyChangedEvent& PropertyChangedEvent)
+{
+	if (PropertyChangedEvent.Property != nullptr)
+	{
+		SettingsChangedDelegate.Broadcast(PropertyChangedEvent.Property->GetFName(), this);
+	}
+}
+
+UDataflowSettings::FOnDataflowSettingsChanged& UDataflowSettings::OnSettingsChanged()
+{
+	return SettingsChangedDelegate;
+}
+
+UDataflowSettings::FOnDataflowSettingsChanged UDataflowSettings::SettingsChangedDelegate;
+
+
+
+
diff --git a/Engine/Source/Runtime/Experimental/Dataflow/Core/Public/Dataflow/DataflowSettings.h b/Engine/Source/Runtime/Experimental/Dataflow/Core/Public/Dataflow/DataflowSettings.h
new file mode 100644
--- /dev/null
+++ b/Engine/Source/Runtime/Experimental/Dataflow/Core/Public/Dataflow/DataflowSettings.h
@@ -0,0 +1,55 @@
+// Copyright Epic Games, Inc. All Rights Reserved.
+
+#pragma once
+
+#include "Engine/DeveloperSettings.h"
+#include "InputCoreTypes.h"
+
+#include "DataflowSettings.generated.h"
+
+UCLASS(config = Dataflow, defaultconfig, meta = (DisplayName = "Dataflow"))
+class DATAFLOWEDITOR_API UDataflowSettings : public UDeveloperSettings
+{
+	GENERATED_UCLASS_BODY()
+
+	/** TArray<> pin type color. The other pin colors are defined in the general editor settings. */
+	UPROPERTY(config, EditAnywhere, Category = PinColors)
+	FLinearColor ArrayPinTypeColor;
+
+	/** FManagedArrayCollection pin type color. The other pin colors are defined in the general editor settings. */
+	UPROPERTY(config, EditAnywhere, Category = PinColors)
+	FLinearColor ManagedArrayCollectionPinTypeColor;
+
+	/** FBox pin type color. The other pin colors are defined in the general editor settings. */
+	UPROPERTY(config, EditAnywhere, Category = PinColors)
+	FLinearColor BoxPinTypeColor;
+
+	/** GeometryCollection category NodeTitle color. */
+	UPROPERTY(config, EditAnywhere, Category = NodeColors)
+	FLinearColor GeometryCollectionCategoryNodeTitleColor;
+
+	/** GeometryCollection category NodeBodyTint color. */
+	UPROPERTY(config, EditAnywhere, Category = NodeColors)
+	FLinearColor GeometryCollectionCategoryNodeBodyTintColor;
+
+	// Begin UDeveloperSettings Interface
+	virtual FName GetCategoryName() const override;
+#if WITH_EDITOR
+	virtual FText GetSectionText() const override;
+#endif
+	// END UDeveloperSettings Interface
+
+#if WITH_EDITOR
+	virtual void PostEditChangeProperty(FPropertyChangedEvent& PropertyChangedEvent) override;
+
+public:
+	DECLARE_MULTICAST_DELEGATE_TwoParams(FOnDataflowSettingsChanged, const FName&, const UDataflowSettings*);
+
+	/** Gets a multicast delegate which is called whenever one of the parameters in this settings object changes. */
+	static FOnDataflowSettingsChanged& OnSettingsChanged();
+
+protected:
+	static FOnDataflowSettingsChanged SettingsChangedDelegate;
+
+#endif
+};
