Iris
- Minor review fixes

diff --git a/Engine/Source/Runtime/Experimental/Iris/Core/Private/Iris/ReplicationSystem/DeltaCompression/DeltaCompressionBaselineManager.cpp b/Engine/Source/Runtime/Experimental/Iris/Core/Private/Iris/ReplicationSystem/DeltaCompression/DeltaCompressionBaselineManager.cpp
--- a/Engine/Source/Runtime/Experimental/Iris/Core/Private/Iris/ReplicationSystem/DeltaCompression/DeltaCompressionBaselineManager.cpp
+++ b/Engine/Source/Runtime/Experimental/Iris/Core/Private/Iris/ReplicationSystem/DeltaCompression/DeltaCompressionBaselineManager.cpp
@@ -55,6 +55,7 @@ FDeltaCompressionBaselineManager::FDeltaCompressionBaselineManager()
 {
 	// If alignment of storage type for FPerObjectInfo is too small then it's totally unsafe to cast such pointers to FPerObjectInfo.
 	static_assert(InvalidObjectInfoIndex == 0, "InvalidObjectInfoIndex needs to be zero in order for current code to function properly.");
+	static_assert(InvalidInternalIndex == FNetHandleManager::InvalidInternalIndex, "InvalidInternalIndex differs from FNetHandleManager::InvalidInternalIndex");
 }
 
 FDeltaCompressionBaselineManager::~FDeltaCompressionBaselineManager()
diff --git a/Engine/Source/Runtime/Experimental/Iris/Core/Private/Iris/ReplicationSystem/DeltaCompression/DeltaCompressionBaselineManager.h b/Engine/Source/Runtime/Experimental/Iris/Core/Private/Iris/ReplicationSystem/DeltaCompression/DeltaCompressionBaselineManager.h
--- a/Engine/Source/Runtime/Experimental/Iris/Core/Private/Iris/ReplicationSystem/DeltaCompression/DeltaCompressionBaselineManager.h
+++ b/Engine/Source/Runtime/Experimental/Iris/Core/Private/Iris/ReplicationSystem/DeltaCompression/DeltaCompressionBaselineManager.h
@@ -103,6 +103,9 @@ private:
 		InvalidObjectInfoIndex = 0,
 		InvalidBaselineStateInfoIndex = 0,
 
+		// Must match NetHandleManager::InvalidInternalIndex
+		InvalidInternalIndex = 0,
+
 		ObjectInfoGrowCount = 128,
 	};
 
@@ -264,6 +267,9 @@ inline const FDeltaCompressionBaselineManager::FPerObjectInfo* FDeltaCompression
 
 inline FDeltaCompressionBaselineManager::FPerObjectInfo* FDeltaCompressionBaselineManager::GetPerObjectInfoForObject(uint32 ObjectIndex)
 {
+	checkSlow(ObjectIndex != InvalidInternalIndex);
+	checkSlow(ObjectIndex < (uint32)ObjectIndexToObjectInfoIndex.Num());
+
 	const ObjectInfoIndexType InfoIndex = ObjectIndexToObjectInfoIndex[ObjectIndex];
 	if (InfoIndex != InvalidObjectInfoIndex)
 	{
@@ -275,6 +281,9 @@ inline FDeltaCompressionBaselineManager::FPerObjectInfo* FDeltaCompressionBaseli
 
 inline const FDeltaCompressionBaselineManager::FPerObjectInfo* FDeltaCompressionBaselineManager::GetPerObjectInfoForObject(uint32 ObjectIndex) const
 {
+	checkSlow(ObjectIndex != InvalidInternalIndex);
+	checkSlow(ObjectIndex < (uint32)ObjectIndexToObjectInfoIndex.Num());
+
 	const ObjectInfoIndexType InfoIndex = ObjectIndexToObjectInfoIndex[ObjectIndex];
 	if (InfoIndex != InvalidObjectInfoIndex)
 	{
diff --git a/Engine/Source/Runtime/Experimental/Iris/Core/Private/Iris/ReplicationSystem/Filtering/ReplicationFiltering.cpp b/Engine/Source/Runtime/Experimental/Iris/Core/Private/Iris/ReplicationSystem/Filtering/ReplicationFiltering.cpp
--- a/Engine/Source/Runtime/Experimental/Iris/Core/Private/Iris/ReplicationSystem/Filtering/ReplicationFiltering.cpp
+++ b/Engine/Source/Runtime/Experimental/Iris/Core/Private/Iris/ReplicationSystem/Filtering/ReplicationFiltering.cpp
@@ -521,7 +521,7 @@ void FReplicationFiltering::InitNewConnections()
 		// Everything enabled by default in groups, i.e. not filtered out.
 		ConnectionInfo.GroupFilteredOutObjects.Init(MaxObjectCount);
 
-		// The combined result of scoped objects, connection filteirng and group filtering.
+		// The combined result of scoped objects, connection filtering and group filtering.
 		ConnectionInfo.ObjectsInScopeBeforeDynamicFiltering.Init(MaxObjectCount);
 
 		// The final result of all filtering.
@@ -1368,7 +1368,7 @@ void FReplicationFiltering::SetSubObjectFilterStatus(FNetObjectGroupHandle Group
 
 	if (IsReservedNetObjectGroupHandle(GroupHandle))
 	{
-		UE_LOG(LogIrisFiltering, Warning, TEXT("FReplicationFiltering::SetSubObjectFilterStatus - Trying to set filter for reserved GroupHandle: %u which is not allowed."), GroupHandle);
+		UE_LOG(LogIrisFiltering, Warning, TEXT("FReplicationFiltering::SetSubObjectFilterStatus - Trying to set filter for reserved GroupIndex: %u which is not allowed."), GroupHandle);
 		return;
 	}
 
@@ -1396,7 +1396,7 @@ void FReplicationFiltering::SetSubObjectFilterStatus(FNetObjectGroupHandle Group
 
 	if (IsReservedNetObjectGroupHandle(GroupHandle))
 	{
-		UE_LOG(LogIrisFiltering, Warning, TEXT("FReplicationFiltering::SetSubObjectFilterStatus - Trying to set filter for reserved GroupHandle: %u which is not allowed."), GroupHandle);
+		UE_LOG(LogIrisFiltering, Warning, TEXT("FReplicationFiltering::SetSubObjectFilterStatus - Trying to set filter for reserved GroupIndex: %u which is not allowed."), GroupHandle);
 		return;
 	}
 
@@ -1428,7 +1428,7 @@ void FReplicationFiltering::SetSubObjectFilterStatus(FNetObjectGroupHandle Group
 {
 	if (IsReservedNetObjectGroupHandle(GroupHandle))
 	{
-		UE_LOG(LogIrisFiltering, Warning, TEXT("FReplicationFiltering::SetSubObjectFilterStatus - Trying to set filter for reserved GroupHandle: %u which is not allowed."), GroupHandle);
+		UE_LOG(LogIrisFiltering, Warning, TEXT("FReplicationFiltering::SetSubObjectFilterStatus - Trying to set filter for reserved GroupIndex: %u which is not allowed."), GroupHandle);
 		return;
 	}
 
@@ -1501,7 +1501,7 @@ void FReplicationFiltering::SetGroupFilterStatus(FNetObjectGroupHandle GroupHand
 
 	if (IsReservedNetObjectGroupHandle(GroupHandle))
 	{
-		UE_LOG(LogIrisFiltering, Warning, TEXT("FReplicationFiltering::SetGroupFilterStatus - Trying to set filter for reserer NetObjectGroupHandle: %u which is not allowed."), GroupHandle);
+		UE_LOG(LogIrisFiltering, Warning, TEXT("FReplicationFiltering::SetGroupFilterStatus - Trying to set filter for reserved GroupIndex: %u which is not allowed."), GroupHandle);
 		return;
 	}
 
@@ -1524,7 +1524,7 @@ void FReplicationFiltering::SetGroupFilterStatus(FNetObjectGroupHandle GroupHand
 		
 	if (IsReservedNetObjectGroupHandle(GroupHandle))
 	{
-		UE_LOG(LogIrisFiltering, Warning, TEXT("FReplicationFiltering::SetGroupFilterStatus - Trying to set filter for reserer NetObjectGroupHandle: %u which is not allowed."), GroupHandle);		
+		UE_LOG(LogIrisFiltering, Warning, TEXT("FReplicationFiltering::SetGroupFilterStatus - Trying to set filter for reserved GroupIndex: %u which is not allowed."), GroupHandle);		
 		return;
 	}
 
@@ -1558,7 +1558,7 @@ void FReplicationFiltering::SetGroupFilterStatus(FNetObjectGroupHandle GroupHand
 
 	if (IsReservedNetObjectGroupHandle(GroupHandle))
 	{
-		UE_LOG(LogIrisFiltering, Warning, TEXT("FReplicationFiltering::SetGroupFilterStatus - Trying to set filter for reserer NetObjectGroupHandle: %u which is not allowed."), GroupHandle);
+		UE_LOG(LogIrisFiltering, Warning, TEXT("FReplicationFiltering::SetGroupFilterStatus - Trying to set filter for reserved GroupIndex: %u which is not allowed."), GroupHandle);
 		return;
 	}
 
