Chaos Cloth - Make parallelfor lambda captures explicit.

diff --git a/Engine/Source/Runtime/Experimental/Chaos/Private/Chaos/PBDAxialSpringConstraints.cpp b/Engine/Source/Runtime/Experimental/Chaos/Private/Chaos/PBDAxialSpringConstraints.cpp
--- a/Engine/Source/Runtime/Experimental/Chaos/Private/Chaos/PBDAxialSpringConstraints.cpp
+++ b/Engine/Source/Runtime/Experimental/Chaos/Private/Chaos/PBDAxialSpringConstraints.cpp
@@ -123,7 +123,7 @@ void FPBDAxialSpringConstraints::Apply(FSolverParticles& Particles, const FSolve
 				{
 					const int32 ColorStart = ConstraintsPerColorStartIndex[ConstraintColorIndex];
 					const int32 ColorSize = ConstraintsPerColorStartIndex[ConstraintColorIndex + 1] - ColorStart;
-					PhysicsParallelFor(ColorSize, [&](const int32 Index)
+					PhysicsParallelFor(ColorSize, [this, &Particles, Dt, ColorStart, ExpStiffnessValue](const int32 Index)
 					{
 						const int32 ConstraintIndex = ColorStart + Index;
 						ApplyHelper(Particles, Dt, ConstraintIndex, ExpStiffnessValue);
@@ -157,7 +157,7 @@ void FPBDAxialSpringConstraints::Apply(FSolverParticles& Particles, const FSolve
 				{
 					const int32 ColorStart = ConstraintsPerColorStartIndex[ConstraintColorIndex];
 					const int32 ColorSize = ConstraintsPerColorStartIndex[ConstraintColorIndex + 1] - ColorStart;
-					PhysicsParallelFor(ColorSize, [&](const int32 Index)
+					PhysicsParallelFor(ColorSize, [this, &Particles, Dt, ColorStart](const int32 Index)
 					{
 						const int32 ConstraintIndex = ColorStart + Index;
 						const FSolverReal ExpStiffnessValue = Stiffness[ConstraintIndex];
diff --git a/Engine/Source/Runtime/Experimental/Chaos/Private/Chaos/PBDSpringConstraints.cpp b/Engine/Source/Runtime/Experimental/Chaos/Private/Chaos/PBDSpringConstraints.cpp
--- a/Engine/Source/Runtime/Experimental/Chaos/Private/Chaos/PBDSpringConstraints.cpp
+++ b/Engine/Source/Runtime/Experimental/Chaos/Private/Chaos/PBDSpringConstraints.cpp
@@ -120,7 +120,7 @@ void FPBDSpringConstraints::Apply(FSolverParticles& Particles, const FSolverReal
 				{
 					const int32 ColorStart = ConstraintsPerColorStartIndex[ConstraintColorIndex];
 					const int32 ColorSize = ConstraintsPerColorStartIndex[ConstraintColorIndex + 1] - ColorStart;
-					PhysicsParallelFor(ColorSize, [&](const int32 Index)
+					PhysicsParallelFor(ColorSize, [this, &Particles, Dt, ColorStart, ExpStiffnessValue](const int32 Index)
 					{
 						const int32 ConstraintIndex = ColorStart + Index;
 						ApplyHelper(Particles, Dt, ConstraintIndex, ExpStiffnessValue);
@@ -153,7 +153,7 @@ void FPBDSpringConstraints::Apply(FSolverParticles& Particles, const FSolverReal
 				{
 					const int32 ColorStart = ConstraintsPerColorStartIndex[ConstraintColorIndex];
 					const int32 ColorSize = ConstraintsPerColorStartIndex[ConstraintColorIndex + 1] - ColorStart;
-					PhysicsParallelFor(ColorSize, [&](const int32 Index)
+					PhysicsParallelFor(ColorSize, [this, &Particles, Dt, ColorStart](const int32 Index)
 					{
 						const int32 ConstraintIndex = ColorStart + Index;
 						const FSolverReal ExpStiffnessValue = Stiffness[ConstraintIndex];
diff --git a/Engine/Source/Runtime/Experimental/Chaos/Public/Chaos/PBDSphericalConstraint.h b/Engine/Source/Runtime/Experimental/Chaos/Public/Chaos/PBDSphericalConstraint.h
--- a/Engine/Source/Runtime/Experimental/Chaos/Public/Chaos/PBDSphericalConstraint.h
+++ b/Engine/Source/Runtime/Experimental/Chaos/Public/Chaos/PBDSphericalConstraint.h
@@ -67,7 +67,7 @@ private:
 	{
 		const int32 ParticleCount = SphereRadii.Num();
 
-		PhysicsParallelFor(ParticleCount, [&](int32 Index)  // TODO: profile need for parallel loop based on particle count
+		PhysicsParallelFor(ParticleCount, [this, &Particles, Dt](int32 Index)  // TODO: profile need for parallel loop based on particle count
 		{
 			const int32 ParticleIndex = ParticleOffset + Index;
 
@@ -185,7 +185,7 @@ private:
 	{
 		const int32 ParticleCount = SphereRadii.Num();
 
-		PhysicsParallelFor(ParticleCount, [&](int32 Index)  // TODO: profile need for parallel loop based on particle count
+		PhysicsParallelFor(ParticleCount, [this, &Particles, Dt](int32 Index)  // TODO: profile need for parallel loop based on particle count
 		{
 			const int32 ParticleIndex = ParticleOffset + Index;
 
@@ -222,7 +222,7 @@ private:
 	{
 		const int32 ParticleCount = SphereRadii.Num();
 
-		PhysicsParallelFor(ParticleCount, [&](int32 Index)  // TODO: profile need for parallel loop based on particle count
+		PhysicsParallelFor(ParticleCount, [this, &Particles, Dt](int32 Index)  // TODO: profile need for parallel loop based on particle count
 		{
 			const int32 ParticleIndex = ParticleOffset + Index;
 
