Chaos: Fixed static analysis issue in CalculateVolumeAndCenterOfMassImpl. Previously we ensure that we have triangular elements then proceed to rely on exactly three elements. Upgraded this to a check as if the count is greater than three we write out of bounds and if it's less we end up with garbage in the CoM and Volume.

diff --git a/Engine/Source/Runtime/Experimental/Chaos/Private/Chaos/MassProperties.cpp b/Engine/Source/Runtime/Experimental/Chaos/Private/Chaos/MassProperties.cpp
--- a/Engine/Source/Runtime/Experimental/Chaos/Private/Chaos/MassProperties.cpp
+++ b/Engine/Source/Runtime/Experimental/Chaos/Private/Chaos/MassProperties.cpp
@@ -146,7 +146,7 @@ namespace Chaos
 		for (const auto& Element : Surfaces)
 		{
 			// For now we only support triangular elements
-			ensure(Element.Num() == 3);
+			check(Element.Num() == 3);
 
 			PMatrix<T,3,3> DeltaMatrix;
 			TVec PerElementSize;
@@ -157,12 +157,15 @@ namespace Chaos
 				DeltaMatrix.M[1][i] = DeltaVector[1];
 				DeltaMatrix.M[2][i] = DeltaVector[2];
 			}
+			
 			PerElementSize[0] = DeltaMatrix.M[0][0] + DeltaMatrix.M[0][1] + DeltaMatrix.M[0][2];
 			PerElementSize[1] = DeltaMatrix.M[1][0] + DeltaMatrix.M[1][1] + DeltaMatrix.M[1][2];
 			PerElementSize[2] = DeltaMatrix.M[2][0] + DeltaMatrix.M[2][1] + DeltaMatrix.M[2][2];
+
 			T Det = DeltaMatrix.M[0][0] * (DeltaMatrix.M[1][1] * DeltaMatrix.M[2][2] - DeltaMatrix.M[1][2] * DeltaMatrix.M[2][1]) -
 				DeltaMatrix.M[0][1] * (DeltaMatrix.M[1][0] * DeltaMatrix.M[2][2] - DeltaMatrix.M[1][2] * DeltaMatrix.M[2][0]) +
 				DeltaMatrix.M[0][2] * (DeltaMatrix.M[1][0] * DeltaMatrix.M[2][1] - DeltaMatrix.M[1][1] * DeltaMatrix.M[2][0]);
+			
 			Volume += Det;
 			VolumeTimesSum += Det * PerElementSize;
 		}
