Chaos : Flesh :
Simplify computation for XPBD corotated constraint.

diff --git a/Engine/Source/Runtime/Experimental/Chaos/Public/Chaos/XPBDCorotatedConstraints.h b/Engine/Source/Runtime/Experimental/Chaos/Public/Chaos/XPBDCorotatedConstraints.h
--- a/Engine/Source/Runtime/Experimental/Chaos/Public/Chaos/XPBDCorotatedConstraints.h
+++ b/Engine/Source/Runtime/Experimental/Chaos/Public/Chaos/XPBDCorotatedConstraints.h
@@ -204,43 +204,19 @@ namespace Chaos::Softs
 
 		TVec4<TVector<T, 3>> GetPolarGradient(const PMatrix<T, 3, 3>& Fe, const PMatrix<T, 3, 3>& Re, const PMatrix<T, 3, 3>& DmInvT, const T C1) const
 		{
-			TVector<T, 81> dRdF((T)0.);
-			Chaos::dRdFCorotated(Fe, dRdF);
+			//TVector<T, 81> dRdF((T)0.);
+			//Chaos::dRdFCorotated(Fe, dRdF);
 			TVec4<TVector<T, 3>> dC1(TVector<T, 3>((T)0.));
 			//dC1 = dC1dF * dFdX
+			PMatrix<T, 3, 3> A = DmInvT * (Fe - Re);
 			for (int alpha = 0; alpha < 3; alpha++) {
 				for (int l = 0; l < 3; l++) {
-
-					dC1[0][alpha] += (DmInvT * Re).GetAt(alpha, l) - (DmInvT * Fe).GetAt(alpha, l);
-
-				}
-			}
-			for (int ie = 0; ie < 3; ie++) {
-				for (int alpha = 0; alpha < 3; alpha++) {
-					dC1[ie + 1][alpha] = (DmInvT * Fe).GetAt(alpha, ie) - (DmInvT * Re).GetAt(alpha, ie);
+					dC1[0][alpha] -= A.GetAt(alpha, l);
 				}
 			}
-			//it's really ie-1 here
 			for (int ie = 0; ie < 3; ie++) {
 				for (int alpha = 0; alpha < 3; alpha++) {
-					for (int m = 0; m < 3; m++) {
-						for (int n = 0; n < 3; n++) {
-							for (int j = 0; j < 3; j++) {
-								dC1[ie + 1][alpha] -= (Fe.GetAt(m, n) - Re.GetAt(m, n)) * dRdF[9 * (alpha * 3 + j) + m * 3 + n] * DmInvT.GetAt(j, ie);
-							}
-						}
-					}
-				}
-			}
-			for (int alpha = 0; alpha < 3; alpha++) {
-				for (int m = 0; m < 3; m++) {
-					for (int n = 0; n < 3; n++) {
-						for (int l = 0; l < 3; l++) {
-							for (int j = 0; j < 3; j++) {
-								dC1[0][alpha] += (Fe.GetAt(m, n) - Re.GetAt(m, n)) * dRdF[9 * (alpha * 3 + j) + m * 3 + n] * DmInvT.GetAt(j, l);
-							}
-						}
-					}
+					dC1[ie + 1][alpha] = A.GetAt(alpha, ie);
 				}
 			}
 
@@ -433,57 +409,6 @@ namespace Chaos::Softs
 
 			PMatrix<T, 3, 3> DmInvT = ElementDmInv(ElementIndex).GetTransposed();
 
-			////TODO: deifnitely test the initialization
-			//TVec4<TVector<T, 3>> dC1(TVector<T, 3>((T)0.));
-			////dC1 = dC1dF * dFdX
-			//for (int alpha = 0; alpha < 3; alpha++) {
-			//	for (int l = 0; l < 3; l++) {
-
-			//		dC1[0][alpha] += (DmInvT * Re).GetAt(alpha, l) - (DmInvT * Fe).GetAt(alpha, l);
-
-			//	}
-			//}
-			//for (int ie = 0; ie < 3; ie++) {
-			//	for (int alpha = 0; alpha < 3; alpha++) {
-			//		dC1[ie+1][alpha] = (DmInvT * Fe).GetAt(alpha, ie) - (DmInvT * Re).GetAt(alpha, ie);
-			//	}
-			//}
-			////it's really ie-1 here
-			//for (int ie = 0; ie < 3; ie++) {
-			//	for (int alpha = 0; alpha < 3; alpha++) {
-			//		for (int m = 0; m < 3; m++) {
-			//			for (int n = 0; n < 3; n++) {
-			//				for (int j = 0; j < 3; j++) {
-			//					dC1[ie + 1][alpha] -= (Fe.GetAt(m, n) - Re.GetAt(m, n)) * dRdF[9*(alpha*3+j)+ m * 3 + n] * DmInvT.GetAt(j, ie);
-			//				}
-			//			}
-			//		}
-			//	}
-			//}
-			//for (int alpha = 0; alpha < 3; alpha++) {
-			//	for (int m = 0; m < 3; m++) {
-			//		for (int n = 0; n < 3; n++) {
-			//			for (int l = 0; l < 3; l++) {
-			//				for (int j = 0; j < 3; j++) {
-			//					dC1[0][alpha] += (Fe.GetAt(m, n) - Re.GetAt(m, n)) * dRdF[9*(alpha*3+j) + m * 3 + n] * DmInvT.GetAt(j, l);
-			//				}
-			//			}
-			//		}
-			//	}
-			//}
-
-			//
-			//if (C1 != 0)
-			//{
-			//	for (int i = 0; i < 4; i++)
-			//	{
-			//		for (int j = 0; j < 3; j++)
-			//		{
-			//			dC1[i][j] /= C1;
-			//		}
-			//	}
-			//}
-
 			TVec4<TVector<T, 3>> dC1 = GetPolarGradient(Fe, Re, DmInvT, C1);
 
 			T AlphaTilde = (T)1. / (Dt * Dt * Mu * Measure[ElementIndex]);
