Modular Vehicles - Generate vehicle Pawn rather than plain a Actor, better for setting up player controller

diff --git a/Engine/Source/Runtime/Experimental/ChaosVehicles/ChaosVehiclesCore/Private/SimModule/SimModuleTree.cpp b/Engine/Source/Runtime/Experimental/ChaosVehicles/ChaosVehiclesCore/Private/SimModule/SimModuleTree.cpp
--- a/Engine/Source/Runtime/Experimental/ChaosVehicles/ChaosVehiclesCore/Private/SimModule/SimModuleTree.cpp
+++ b/Engine/Source/Runtime/Experimental/ChaosVehicles/ChaosVehiclesCore/Private/SimModule/SimModuleTree.cpp
@@ -210,33 +210,35 @@ void FSimModuleTree::UpdateModuleVelocites(FGeometryCollectionPhysicsProxy* Phys
 	{
 		if (ISimulationModuleBase* Module = SimulationModuleTree[i].SimModule)
 		{
-			Chaos::FPBDRigidClusteredParticleHandle* ParentParticle = Clusters[Module->GetTransformIndex()];
-			const FTransform BodyTransform(ParentParticle->R(), ParentParticle->X());
-
-			if (Module->IsBehaviourType(eSimModuleTypeFlags::Velocity))
+			if (Chaos::FPBDRigidClusteredParticleHandle* ParentParticle = Clusters[Module->GetTransformIndex()])
 			{
-				Chaos::FPBDRigidClusteredParticleHandle* Particle = nullptr;
-				if (Module->IsClustered())
-				{
-					Particle = Clusters[Module->GetTransformIndex()];
-				}
-				else
-				{
-					Particle = Particles[Module->GetTransformIndex()];				
-				}
+				const FTransform BodyTransform(ParentParticle->R(), ParentParticle->X());
 
-				if (Particle)
+				if (Module->IsBehaviourType(eSimModuleTypeFlags::Velocity))
 				{
-					FVector WorldLocation = BodyTransform.TransformPosition(Module->GetParentRelativeTransform().GetLocation());
-					const Chaos::FVec3 Arm = WorldLocation - Particle->X();
-
-					FVector WorldVelocity = Particle->V() - Chaos::FVec3::CrossProduct(Arm, Particle->W());
-					FVector LocalVelocity = BodyTransform.InverseTransformVector(WorldVelocity);
-					LocalVelocity = Module->GetClusteredTransform().InverseTransformVector(LocalVelocity);
-
-					Module->SetLocalVelocity(LocalVelocity);
-				}
+					Chaos::FPBDRigidClusteredParticleHandle* Particle = nullptr;
+					if (Module->IsClustered())
+					{
+						Particle = Clusters[Module->GetTransformIndex()];
+					}
+					else
+					{
+						Particle = Particles[Module->GetTransformIndex()];				
+					}
+
+					if (Particle)
+					{
+						FVector WorldLocation = BodyTransform.TransformPosition(Module->GetParentRelativeTransform().GetLocation());
+						const Chaos::FVec3 Arm = WorldLocation - Particle->X();
+
+						FVector WorldVelocity = Particle->V() - Chaos::FVec3::CrossProduct(Arm, Particle->W());
+						FVector LocalVelocity = BodyTransform.InverseTransformVector(WorldVelocity);
+						LocalVelocity = Module->GetClusteredTransform().InverseTransformVector(LocalVelocity);
+
+						Module->SetLocalVelocity(LocalVelocity);
+					}
 			
+				}
 			}
 		}
 	}
