Add a Vector2 attribute to an InterchangeNode.

diff --git a/Engine/Source/Runtime/Experimental/Interchange/Core/Private/Nodes/InterchangeBaseNode.cpp b/Engine/Source/Runtime/Experimental/Interchange/Core/Private/Nodes/InterchangeBaseNode.cpp
--- a/Engine/Source/Runtime/Experimental/Interchange/Core/Private/Nodes/InterchangeBaseNode.cpp
+++ b/Engine/Source/Runtime/Experimental/Interchange/Core/Private/Nodes/InterchangeBaseNode.cpp
@@ -268,6 +268,16 @@ bool UInterchangeBaseNode::GetLinearColorAttribute(const FName& NodeAttributeKey
 	INTERCHANGE_BASE_NODE_GET_ATTRIBUTE(FLinearColor);
 }
 
+bool UInterchangeBaseNode::AddVector2Attribute(const FName& NodeAttributeKey, const FVector2f& Value)
+{
+	INTERCHANGE_BASE_NODE_ADD_ATTRIBUTE(FVector2f);
+}
+
+bool UInterchangeBaseNode::GetVector2Attribute(const FName& NodeAttributeKey, FVector2f& OutValue) const
+{
+	INTERCHANGE_BASE_NODE_GET_ATTRIBUTE(FVector2f);
+}
+
 FString UInterchangeBaseNode::GetUniqueID() const
 {
 	ensure(bIsInitialized);
diff --git a/Engine/Source/Runtime/Experimental/Interchange/Core/Public/Nodes/InterchangeBaseNode.h b/Engine/Source/Runtime/Experimental/Interchange/Core/Public/Nodes/InterchangeBaseNode.h
--- a/Engine/Source/Runtime/Experimental/Interchange/Core/Public/Nodes/InterchangeBaseNode.h
+++ b/Engine/Source/Runtime/Experimental/Interchange/Core/Public/Nodes/InterchangeBaseNode.h
@@ -434,6 +434,18 @@ public:
 	UFUNCTION(BlueprintCallable, Category = "Interchange | Node")
 	bool GetLinearColorAttribute(const FName& NodeAttributeKey, FLinearColor& OutValue) const;
 
+	/**
+	 * Add a Vector2 attribute to this node. Return false if the attribute do not exist or if we cannot add it
+	 */
+	UFUNCTION(BlueprintCallable, Category = "Interchange | Node")
+	bool AddVector2Attribute(const FName& NodeAttributeKey, const FVector2f& Value);
+
+	/**
+	 * Get a Vector2 attribute from this node. Return false if the attribute do not exist
+	 */
+	UFUNCTION(BlueprintCallable, Category = "Interchange | Node")
+	bool GetVector2Attribute(const FName& NodeAttributeKey, FVector2f& OutValue) const;
+
 	template<typename AttributeType>
 	AttributeType GetAttributeChecked(const FName& NodeAttributeKey) const
 	{
