Chaos : Flesh :
Added timestep postprocessing for constraints into pbd evolution.

diff --git a/Engine/Source/Runtime/Experimental/Chaos/Private/Chaos/Deformable/ChaosDeformableSolver.cpp b/Engine/Source/Runtime/Experimental/Chaos/Private/Chaos/Deformable/ChaosDeformableSolver.cpp
--- a/Engine/Source/Runtime/Experimental/Chaos/Private/Chaos/Deformable/ChaosDeformableSolver.cpp
+++ b/Engine/Source/Runtime/Experimental/Chaos/Private/Chaos/Deformable/ChaosDeformableSolver.cpp
@@ -385,6 +385,12 @@ namespace Chaos::Softs
 		{
 			this->GridBasedCorotatedConstraint->ApplyInParallel(InParticles, Dt);
 		};
+		int32 PostprocessingIndex1 = Evolution->AddConstraintPostprocessingsRange(1, true);
+		Evolution->ConstraintPostprocessings()[PostprocessingIndex1] =
+			[this](FSolverParticles& InParticles, const FSolverReal Dt)
+		{
+			this->GridBasedCorotatedConstraint->TimeStepPostprocessing(InParticles, Dt);
+		};
 	
 	
 	}
diff --git a/Engine/Source/Runtime/Experimental/Chaos/Private/Chaos/PBDEvolution.cpp b/Engine/Source/Runtime/Experimental/Chaos/Private/Chaos/PBDEvolution.cpp
--- a/Engine/Source/Runtime/Experimental/Chaos/Private/Chaos/PBDEvolution.cpp
+++ b/Engine/Source/Runtime/Experimental/Chaos/Private/Chaos/PBDEvolution.cpp
@@ -24,6 +24,7 @@ DECLARE_CYCLE_STAT(TEXT("Chaos PBD Iteration Loop"), STAT_ChaosPBDIterationLoop,
 DECLARE_CYCLE_STAT(TEXT("Chaos PBD Post Iteration Updates"), STAT_ChaosPBDPostIterationUpdates, STATGROUP_Chaos);
 DECLARE_CYCLE_STAT(TEXT("Chaos PBD Constraint Rules"), STAT_ChaosPBDConstraintRule, STATGROUP_Chaos);
 DECLARE_CYCLE_STAT(TEXT("Chaos PBD Post Collision Constraint Rules"), STAT_ChaosPBDPostCollisionConstraintRule, STATGROUP_Chaos);
+DECLARE_CYCLE_STAT(TEXT("Chaos PBD Constraint Postprocessings"), STAT_ChaosPBDConstraintPostprocessings, STATGROUP_Chaos);
 DECLARE_CYCLE_STAT(TEXT("Chaos PBD Self Collision"), STAT_ChaosPBDSelfCollisionRule, STATGROUP_Chaos);
 DECLARE_CYCLE_STAT(TEXT("Chaos PBD Collision Rule"), STAT_ChaosPBDCollisionRule, STATGROUP_Chaos);
 DECLARE_CYCLE_STAT(TEXT("Chaos PBD Collider Friction"), STAT_ChaosPBDCollisionRuleFriction, STATGROUP_Chaos);
@@ -87,6 +88,7 @@ FPBDEvolution::FPBDEvolution(
 	, MConstraintInitsActiveView(MConstraintInits)
 	, MConstraintRulesActiveView(MConstraintRules)
 	, MPostCollisionConstraintRulesActiveView(MPostCollisionConstraintRules)
+	, MConstraintPostprocessingsActiveView(MConstraintPostprocessings)
 	, MNumIterations(NumIterations)
 	, MGravity(FSolverVec3((FSolverReal)0., (FSolverReal)0., (FSolverReal)-980.665))
 	, MCollisionThickness(CollisionThickness)
@@ -208,6 +210,15 @@ int32 FPBDEvolution::AddPostCollisionConstraintRuleRange(int32 NumConstraints, b
 	return MPostCollisionConstraintRulesActiveView.AddRange(NumConstraints, bActivate);
 }
 
+int32 FPBDEvolution::AddConstraintPostprocessingsRange(int32 NumConstraints, bool bActivate)
+{
+	// Add new constraint rule functions
+	MConstraintPostprocessings.AddDefaulted(NumConstraints);
+
+	// Add range
+	return MConstraintPostprocessingsActiveView.AddRange(NumConstraints, bActivate);
+}
+
 template<bool bForceRule, bool bVelocityField, bool bDampVelocityRule>
 void FPBDEvolution::PreIterationUpdate(
 	const FSolverReal Dt,
@@ -501,6 +512,21 @@ void FPBDEvolution::AdvanceOneTimeStep(const FSolverReal Dt, const bool bSmoothD
 				}, bUseSingleThreadedRange);
 		}
 
+		{
+			TRACE_CPUPROFILER_EVENT_SCOPE(ChaosPBDConstraintPostprocessings);
+			SCOPE_CYCLE_COUNTER(STAT_ChaosPBDConstraintPostprocessings);
+			MConstraintPostprocessingsActiveView.RangeFor(
+				[this, Dt](TArray<TFunction<void(FSolverParticles&, const FSolverReal)>>& ConstraintPostprocessings, int32 Offset, int32 Range)
+				{
+					TRACE_CPUPROFILER_EVENT_SCOPE(ChaosPBDConstraintPostprocessings);
+					SCOPE_CYCLE_COUNTER(STAT_ChaosPBDConstraintPostprocessings);
+					for (int32 ConstraintIndex = Offset; ConstraintIndex < Range; ++ConstraintIndex)
+					{
+						ConstraintPostprocessings[ConstraintIndex](MParticles, Dt); // P +/-= ...
+					}
+				}, bUseSingleThreadedRange);
+		}
+
 		{
 			SCOPE_CYCLE_COUNTER(STAT_ChaosPBDPostIterationUpdates);
 
diff --git a/Engine/Source/Runtime/Experimental/Chaos/Public/Chaos/MPMTransfer.h b/Engine/Source/Runtime/Experimental/Chaos/Public/Chaos/MPMTransfer.h
--- a/Engine/Source/Runtime/Experimental/Chaos/Public/Chaos/MPMTransfer.h
+++ b/Engine/Source/Runtime/Experimental/Chaos/Public/Chaos/MPMTransfer.h
@@ -301,56 +301,57 @@ public:
 	void ComputeAArray(const TArray<T>& GridData, const TDynamicParticles<T, 3>& Particles)
 	{
 		AArray.Init((T)0., NumModes * 3 * Particles.Size());
-		for (int32 p = 0; p < Particles.Size(); p++)
-		{
-			TVector<int32, 3> BaseIndex = Indices[p];
-			for (int32 ii = 0; ii < int32(Grid.NPerDir); ii++)
+		PhysicsParallelFor(int32(Particles.Size()), [&](const int32 p)
+			//for (int32 p = 0; p < int32(Particles.Size()); p++)
 			{
-				T Nii = Grid.Nijk(Weights[p][0], ii);
-				T dNii = Grid.dNijk(Weights[p][0], ii, Grid.GetDx()[0]);
-				for (int32 jj = 0; jj < int32(Grid.NPerDir); jj++)
+				TVector<int32, 3> BaseIndex = Indices[p];
+				for (int32 ii = 0; ii < int32(Grid.NPerDir); ii++)
 				{
-					T Njj = Grid.Nijk(Weights[p][1], jj);
-					T dNjj = Grid.dNijk(Weights[p][1], jj, Grid.GetDx()[0]);
-					for (int32 kk = 0; kk < int32(Grid.NPerDir); kk++)
+					T Nii = Grid.Nijk(Weights[p][0], ii);
+					T dNii = Grid.dNijk(Weights[p][0], ii, Grid.GetDx()[0]);
+					for (int32 jj = 0; jj < int32(Grid.NPerDir); jj++)
 					{
-						T Nkk = Grid.Nijk(Weights[p][2], kk);
-						T dNkk = Grid.dNijk(Weights[p][2], kk, Grid.GetDx()[0]);
-						TVector<T, 3> Ni = { Nii, Njj, Nkk };
-						TVector<T, 3> dNi = { dNii, dNjj, dNkk };
-						TVector<int32, 3> LocIndex = { ii, jj, kk };
-						TVector<int32, 3> GlobIndex = Grid.Loc2GlobIndex(BaseIndex, LocIndex);
-						int32 GlobIndexFlat = Grid.FlatIndex(GlobIndex);
-						if (GridData[(NTransfer + 1) * GlobIndexFlat] != T(0))
+						T Njj = Grid.Nijk(Weights[p][1], jj);
+						T dNjj = Grid.dNijk(Weights[p][1], jj, Grid.GetDx()[0]);
+						for (int32 kk = 0; kk < int32(Grid.NPerDir); kk++)
 						{
-							if (Grid.interp == TMPMGrid<T>::linear)
+							T Nkk = Grid.Nijk(Weights[p][2], kk);
+							T dNkk = Grid.dNijk(Weights[p][2], kk, Grid.GetDx()[0]);
+							TVector<T, 3> Ni = { Nii, Njj, Nkk };
+							TVector<T, 3> dNi = { dNii, dNjj, dNkk };
+							TVector<int32, 3> LocIndex = { ii, jj, kk };
+							TVector<int32, 3> GlobIndex = Grid.Loc2GlobIndex(BaseIndex, LocIndex);
+							int32 GlobIndexFlat = Grid.FlatIndex(GlobIndex);
+							if (GridData[(NTransfer + 1) * GlobIndexFlat] != T(0))
 							{
-								TVector<T, 3> dN((T)0.);
-								Grid.GradNi(Ni, dNi, dN);
-								for (int32 l = 0; l < 3; l++)
+								if (Grid.interp == TMPMGrid<T>::linear)
 								{
-									T GridQuantity = GridData[(NTransfer + 1) * GlobIndexFlat + 1 + l] / GridData[(NTransfer + 1) * GlobIndexFlat];
-									for (int32 iii = 0; iii < NumModes; ++iii)
-										AArray[NumModes * 3 * p + NumModes * l + iii] += dN[iii] * GridQuantity;
+									TVector<T, 3> dN((T)0.);
+									Grid.GradNi(Ni, dNi, dN);
+									for (int32 l = 0; l < 3; l++)
+									{
+										T GridQuantity = GridData[(NTransfer + 1) * GlobIndexFlat + 1 + l] / GridData[(NTransfer + 1) * GlobIndexFlat];
+										for (int32 iii = 0; iii < int32(NumModes); ++iii)
+											AArray[NumModes * 3 * p + NumModes * l + iii] += dN[iii] * GridQuantity;
+									}
 								}
-							}
-							else
-							{
-								TVector<T, 3> rip = Grid.Node(GlobIndexFlat) - Particles.X(p);
-								T NiProdbInv = (T(4) * Ni[0] * Ni[1] * Ni[2]) / (Grid.GetDx()[0] * Grid.GetDx()[0]);
-								for (int32 l = 0; l < 3; l++)
+								else
 								{
-									T GridQuantity = GridData[(NTransfer + 1) * GlobIndexFlat + 1 + l] / GridData[(NTransfer + 1) * GlobIndexFlat];
-									for (int32 iii = 0; iii < NumModes; ++iii)
-										AArray[NumModes * 3 * p + NumModes * l + iii] += NiProdbInv * rip[iii] * GridQuantity;
+									TVector<T, 3> rip = Grid.Node(GlobIndexFlat) - Particles.X(p);
+									T NiProdbInv = (T(4) * Ni[0] * Ni[1] * Ni[2]) / (Grid.GetDx()[0] * Grid.GetDx()[0]);
+									for (int32 l = 0; l < 3; l++)
+									{
+										T GridQuantity = GridData[(NTransfer + 1) * GlobIndexFlat + 1 + l] / GridData[(NTransfer + 1) * GlobIndexFlat];
+										for (int32 iii = 0; iii < int32(NumModes); ++iii)
+											AArray[NumModes * 3 * p + NumModes * l + iii] += NiProdbInv * rip[iii] * GridQuantity;
+									}
 								}
 							}
+
 						}
-						
 					}
 				}
-			}
-		}
+			}, Particles.Size() < 50);
 	}
 
 	void GridPositionsToGridData(const TArray<TVector<T, 3>>& GridPositions, const T Dt, TArray<T>& GridData)
diff --git a/Engine/Source/Runtime/Experimental/Chaos/Public/Chaos/PBDEvolution.h b/Engine/Source/Runtime/Experimental/Chaos/Public/Chaos/PBDEvolution.h
--- a/Engine/Source/Runtime/Experimental/Chaos/Public/Chaos/PBDEvolution.h
+++ b/Engine/Source/Runtime/Experimental/Chaos/Public/Chaos/PBDEvolution.h
@@ -75,25 +75,30 @@ class CHAOS_API FPBDEvolution : public TArrayCollection
 		MConstraintInits.Reset(); 
 		MConstraintRules.Reset(); 
 		MPostCollisionConstraintRules.Reset();
+		MConstraintPostprocessings.Reset();
 		MConstraintInitsActiveView.Reset(); 
 		MConstraintRulesActiveView.Reset();  
 		MPostCollisionConstraintRulesActiveView.Reset();
+		MConstraintPostprocessingsActiveView.Reset();
 	}
 
 	// Add constraints. Return the index of the first added constraint.
 	int32 AddConstraintInitRange(int32 NumConstraints, bool bActivate);
 	int32 AddConstraintRuleRange(int32 NumConstraints, bool bActivate);
 	int32 AddPostCollisionConstraintRuleRange(int32 NumConstraints, bool bActivate);
+	int32 AddConstraintPostprocessingsRange(int32 NumConstraints, bool bActivate);
 	
 	// Return the number of particles of the block starting at Offset
 	int32 GetConstraintInitRangeSize(int32 Offset) const { return MConstraintInitsActiveView.GetRangeSize(Offset); }
 	int32 GetConstraintRuleRangeSize(int32 Offset) const { return MConstraintRulesActiveView.GetRangeSize(Offset); }
 	int32 GetPostCollisionConstraintRuleRangeSize(int32 Offset) const { return MPostCollisionConstraintRulesActiveView.GetRangeSize(Offset); }
+	int32 GetConstraintPostprocessingsRangeSize(int32 Offset) const { return MConstraintPostprocessingsActiveView.GetRangeSize(Offset); }
 
 	// Set a block of constraints active or inactive, using the index of the first added particle to identify the block.
 	void ActivateConstraintInitRange(int32 Offset, bool bActivate) { MConstraintInitsActiveView.ActivateRange(Offset, bActivate); }
 	void ActivateConstraintRuleRange(int32 Offset, bool bActivate) { MConstraintRulesActiveView.ActivateRange(Offset, bActivate); }
 	void ActivatePostCollisionConstraintRuleRange(int32 Offset, bool bActivate) { MPostCollisionConstraintRulesActiveView.ActivateRange(Offset, bActivate); }
+	void ActivateConstraintPostprocessingsRange(int32 Offset, bool bActivate) { MConstraintPostprocessingsActiveView.ActivateRange(Offset, bActivate); }
 
 	// Constraint accessors
 	const TArray<TFunction<void(FSolverParticles&, const FSolverReal)>>& ConstraintInits() const { return MConstraintInits; }
@@ -102,6 +107,8 @@ class CHAOS_API FPBDEvolution : public TArrayCollection
 	TArray<TFunction<void(FSolverParticles&, const FSolverReal)>>& ConstraintRules() { return MConstraintRules; }
 	const TArray<TFunction<void(FSolverParticles&, const FSolverReal)>>& PostCollisionConstraintRules() const { return MPostCollisionConstraintRules; }
 	TArray<TFunction<void(FSolverParticles&, const FSolverReal)>>& PostCollisionConstraintRules() { return MPostCollisionConstraintRules; }
+	const TArray<TFunction<void(FSolverParticles&, const FSolverReal)>>& ConstraintPostprocessings() const { return MConstraintPostprocessings; }
+	TArray<TFunction<void(FSolverParticles&, const FSolverReal)>>& ConstraintPostprocessings() { return MConstraintPostprocessings; }
 	
 	void SetKinematicUpdateFunction(TFunction<void(FSolverParticles&, const FSolverReal, const FSolverReal, const int32)> KinematicUpdate) { MKinematicUpdate = KinematicUpdate; }
 	void SetCollisionKinematicUpdateFunction(TFunction<void(FSolverRigidParticles&, const FSolverReal, const FSolverReal, const int32)> KinematicUpdate) { MCollisionKinematicUpdate = KinematicUpdate; }
@@ -189,6 +196,8 @@ private:
 	TPBDActiveView<TArray<TFunction<void(FSolverParticles&, const FSolverReal)>>> MConstraintRulesActiveView;
 	TArray<TFunction<void(FSolverParticles&, const FSolverReal)>> MPostCollisionConstraintRules;
 	TPBDActiveView<TArray<TFunction<void(FSolverParticles&, const FSolverReal)>>> MPostCollisionConstraintRulesActiveView;
+	TArray<TFunction<void(FSolverParticles&, const FSolverReal)>> MConstraintPostprocessings;
+	TPBDActiveView<TArray<TFunction<void(FSolverParticles&, const FSolverReal)>>> MConstraintPostprocessingsActiveView;
 
 	TFunction<void(FSolverParticles&, const FSolverReal, const FSolverReal, const int32)> MKinematicUpdate;
 	TFunction<void(FSolverRigidParticles&, const FSolverReal, const FSolverReal, const int32)> MCollisionKinematicUpdate;
diff --git a/Engine/Source/Runtime/Experimental/Chaos/Public/Chaos/XPBDGridBasedCorotatedConstraints.h b/Engine/Source/Runtime/Experimental/Chaos/Public/Chaos/XPBDGridBasedCorotatedConstraints.h
--- a/Engine/Source/Runtime/Experimental/Chaos/Public/Chaos/XPBDGridBasedCorotatedConstraints.h
+++ b/Engine/Source/Runtime/Experimental/Chaos/Public/Chaos/XPBDGridBasedCorotatedConstraints.h
@@ -140,7 +140,7 @@ namespace Chaos::Softs
 			}
 		}
 
-		void TimeStepPostprocessing(const ParticleType& InParticles, const T Dt)
+		void TimeStepPostprocessing(ParticleType& InParticles, const T Dt)
 		{
 			{
 				TRACE_CPUPROFILER_EVENT_SCOPE(TEXT("STAT_ChaosXPBDGridBasedCorotatedGrid::FinalG2P"));
@@ -310,7 +310,7 @@ namespace Chaos::Softs
 				}
 			}
 			//final G2P to compute particle positions:
-			FinalG2P(Particles);
+			//FinalG2P(Particles);
 		}
 
 	private:
