Chaos Physics: A dynamic tree can now be used as a dirty element acceleration structure for static trees

diff --git a/Engine/Source/Programs/HeadlessChaos/Private/HeadlessChaosApp.cpp b/Engine/Source/Programs/HeadlessChaos/Private/HeadlessChaosApp.cpp
--- a/Engine/Source/Programs/HeadlessChaos/Private/HeadlessChaosApp.cpp
+++ b/Engine/Source/Programs/HeadlessChaos/Private/HeadlessChaosApp.cpp
@@ -262,6 +262,7 @@ TEST(BP, BroadphaseTests) {
 	ChaosTest::GridBPTest2();
 	ChaosTest::AABBTreeTest();
 	ChaosTest::AABBTreeTestDynamic();
+	ChaosTest::AABBTreeDirtyTreeTest();
 	ChaosTest::AABBTreeDirtyGridTest();
 	ChaosTest::AABBTreeTimesliceTest();
 	ChaosTest::DoForSweepIntersectCellsImpTest();
diff --git a/Engine/Source/Programs/HeadlessChaos/Private/HeadlessChaosTestBroadphase.cpp b/Engine/Source/Programs/HeadlessChaos/Private/HeadlessChaosTestBroadphase.cpp
--- a/Engine/Source/Programs/HeadlessChaos/Private/HeadlessChaosTestBroadphase.cpp
+++ b/Engine/Source/Programs/HeadlessChaos/Private/HeadlessChaosTestBroadphase.cpp
@@ -570,6 +570,28 @@ namespace ChaosTest
 		
 	}
 
+	void AABBTreeDirtyTreeTest()
+	{
+		using TreeType = TAABBTree<int32, TAABBTreeLeafArray<int32>, true>;
+
+		// Do the standard tests
+		{
+			TUniquePtr<TBox<FReal, 3>> Box;
+			auto Boxes = BuildBoxes(Box);
+
+			TArray<TSOAView<FGeometryParticles>> EmptyArray;
+			TreeType Spatial{ MakeParticleView(MoveTemp(EmptyArray)),5, 5, 10000.0f, 1000, false, true};
+
+			int32 Idx;
+			for (Idx = 0; Idx < (int32)Boxes->Size(); ++Idx)
+			{
+				Spatial.UpdateElement(Idx, Boxes->WorldSpaceInflatedBounds(Idx), true);
+			}		
+
+			SpatialTestHelper(Spatial, Boxes.Get(), Box);
+		}
+	}
+
 	void AABBTreeDirtyGridTest()
 	{
 		using TreeType = TAABBTree<int32, TBoundingVolume<int32>>;
diff --git a/Engine/Source/Programs/HeadlessChaos/Public/HeadlessChaosTestBroadphase.h b/Engine/Source/Programs/HeadlessChaos/Public/HeadlessChaosTestBroadphase.h
--- a/Engine/Source/Programs/HeadlessChaos/Public/HeadlessChaosTestBroadphase.h
+++ b/Engine/Source/Programs/HeadlessChaos/Public/HeadlessChaosTestBroadphase.h
@@ -15,6 +15,8 @@ namespace ChaosTest {
 
 	void AABBTreeTestDynamic();
 
+	void AABBTreeDirtyTreeTest();
+
 	void AABBTreeDirtyGridTest();
 
 	void AABBTreeTimesliceTest();
diff --git a/Engine/Source/Runtime/Experimental/Chaos/Private/Chaos/PBDRigidsEvolution.cpp b/Engine/Source/Runtime/Experimental/Chaos/Private/Chaos/PBDRigidsEvolution.cpp
--- a/Engine/Source/Runtime/Experimental/Chaos/Private/Chaos/PBDRigidsEvolution.cpp
+++ b/Engine/Source/Runtime/Experimental/Chaos/Private/Chaos/PBDRigidsEvolution.cpp
@@ -37,6 +37,9 @@ namespace Chaos
 	CHAOS_API int32 AccelerationStructureUseDynamicTree = 1;
 	FAutoConsoleVariableRef CVarAccelerationStructureUseDynamicTree(TEXT("p.Chaos.AccelerationStructureUseDynamicTree"), AccelerationStructureUseDynamicTree, TEXT("Use a dynamic BVH tree structure for dynamic objects"));
 
+	CHAOS_API int32 AccelerationStructureUseDirtyTreeInsteadOfGrid = 1;
+	FAutoConsoleVariableRef CVarAccelerationStructureUseDirtyTreeInsteadOfGrid(TEXT("p.Chaos.AccelerationStructureUseDirtyTreeInsteadOfGrid"), AccelerationStructureUseDirtyTreeInsteadOfGrid, TEXT("Use a dynamic tree structure for dirty elements instead of a 2D grid"));
+
 	/** Console variable to enable the caching of the overlapping leaves if the dynamic tree is enable */
 	CHAOS_API int32 GAccelerationStructureCacheOverlappingLeaves = 1;
 	FAutoConsoleVariableRef CVarAccelerationStructureCacheOverlappingLeaves(TEXT("p.Chaos.AccelerationStructureCacheOverlappingLeaves"), GAccelerationStructureCacheOverlappingLeaves, TEXT("Set to 1: Cache the overlapping leaves for faster overlap query, any other value will disable the feature"));
@@ -175,7 +178,7 @@ namespace Chaos
 					{
 						return MakeUnique<AABBDynamicTreeType>(Particles, ConfigSettings.MaxChildrenInLeaf, ConfigSettings.MaxTreeDepth, ConfigSettings.MaxPayloadSize, ForceFullBuild ? 0 : ConfigSettings.IterationsPerTimeSlice, true);
 					}
-					return MakeUnique<AABBTreeType>(Particles, ConfigSettings.MaxChildrenInLeaf, ConfigSettings.MaxTreeDepth, ConfigSettings.MaxPayloadSize, ForceFullBuild ? 0 : ConfigSettings.IterationsPerTimeSlice);
+					return MakeUnique<AABBTreeType>(Particles, ConfigSettings.MaxChildrenInLeaf, ConfigSettings.MaxTreeDepth, ConfigSettings.MaxPayloadSize, ForceFullBuild ? 0 : ConfigSettings.IterationsPerTimeSlice, false, AccelerationStructureUseDirtyTreeInsteadOfGrid == 1);
 				}
 				else if (ConfigSettings.BroadphaseType == 4 || ConfigSettings.BroadphaseType == 2)
 				{
diff --git a/Engine/Source/Runtime/Experimental/Chaos/Public/Chaos/AABBTree.h b/Engine/Source/Runtime/Experimental/Chaos/Public/Chaos/AABBTree.h
--- a/Engine/Source/Runtime/Experimental/Chaos/Public/Chaos/AABBTree.h
+++ b/Engine/Source/Runtime/Experimental/Chaos/Public/Chaos/AABBTree.h
@@ -701,6 +701,11 @@ public:
 		FirstFreeLeafNode = INDEX_NONE;
 
 		this->SetAsyncTimeSlicingComplete(true);
+
+		if (DirtyElementTree != nullptr)
+		{
+			DirtyElementTree->Reset();
+		}
 	}
 
 	virtual void ProgressAsyncTimeSlicing(bool ForceBuildCompletion) override
@@ -729,7 +734,7 @@ public:
 	}
 
 	template <typename TParticles>
-	TAABBTree(const TParticles& Particles, int32 InMaxChildrenInLeaf = DefaultMaxChildrenInLeaf, int32 InMaxTreeDepth = DefaultMaxTreeDepth, T InMaxPayloadBounds = DefaultMaxPayloadBounds, int32 InMaxNumToProcess = DefaultMaxNumToProcess, bool bInDynamicTree = false)
+	TAABBTree(const TParticles& Particles, int32 InMaxChildrenInLeaf = DefaultMaxChildrenInLeaf, int32 InMaxTreeDepth = DefaultMaxTreeDepth, T InMaxPayloadBounds = DefaultMaxPayloadBounds, int32 InMaxNumToProcess = DefaultMaxNumToProcess, bool bInDynamicTree = false, bool bInUseDirtyTree = false)
 		: ISpatialAcceleration<TPayloadType, T, 3>(StaticType)
 		, bDynamicTree(bInDynamicTree)
 		, MaxChildrenInLeaf(InMaxChildrenInLeaf)
@@ -737,8 +742,12 @@ public:
 		, MaxPayloadBounds(InMaxPayloadBounds)
 		, MaxNumToProcess(InMaxNumToProcess)
 		, bShouldRebuild(true)
-
 	{
+		if (bInUseDirtyTree)
+		{
+			DirtyElementTree = TUniquePtr<TAABBTree<TPayloadType, TLeafType, bMutable, T>>(new TAABBTree<TPayloadType, TLeafType, bMutable, T>());
+			DirtyElementTree->SetTreeToDynamic();
+		}		
 		GenerateTree(Particles);
 	}
 
@@ -861,7 +870,7 @@ public:
 
 	FORCEINLINE_DEBUGGABLE bool DirtyElementGridEnabled() const
 	{
-		return DirtyElementGridCellSize > 0.0f &&
+		return DirtyElementTree == nullptr && DirtyElementGridCellSize > 0.0f &&
 			DirtyElementMaxGridCellQueryCount > 0 &&
 			DirtyElementMaxPhysicalSizeInCells > 0 &&
 			DirtyElementMaxCellCapacity > 0;
@@ -1102,7 +1111,14 @@ public:
 		return AllocatedNodeIdx;
 	}
 
-	int32 AllocateLeafNodeAndLeaf(const TPayloadType& Payload, const TAABB<T, 3>& NewBounds)
+
+	struct NodeAndLeafIndices
+	{
+		int32 NodeIdx;
+		int32 LeafIdx;
+	};
+
+	NodeAndLeafIndices AllocateLeafNodeAndLeaf(const TPayloadType& Payload, const TAABB<T, 3>& NewBounds)
 	{
 		int32 AllocatedNodeIdx = FirstFreeLeafNode;
 		int32 LeafIndex;
@@ -1134,12 +1150,8 @@ public:
 		Nodes[AllocatedNodeIdx].ChildrenBounds[0] = ExpandedBounds;
 
 		Nodes[AllocatedNodeIdx].ParentNode = INDEX_NONE;
-		FAABBTreePayloadInfo* PayloadInfo = PayloadToInfo.Find(Payload);
-		check(PayloadInfo);
-		PayloadInfo->LeafIdx = LeafIndex;
-		PayloadInfo->NodeIdx = AllocatedNodeIdx;
-
-		return AllocatedNodeIdx;
+		
+		return NodeAndLeafIndices{ AllocatedNodeIdx , LeafIndex };
 	}
 
 	void DeAllocateInternalNode(int32 NodeIdx)
@@ -1339,7 +1351,8 @@ public:
 		}
 	}
 
-	void InsertLeaf(const TPayloadType& Payload, const TAABB<T, 3>& NewBounds)
+	// Returns the inserted node and leaf
+	NodeAndLeafIndices InsertLeaf(const TPayloadType& Payload, const TAABB<T, 3>& NewBounds)
 	{
 
 		// Slow Debug Code
@@ -1351,9 +1364,9 @@ public:
 		// New tree?
 		if(RootNode == INDEX_NONE)
 		{
-			int32 NewLeafNode = AllocateLeafNodeAndLeaf(Payload, NewBounds);
-			RootNode = NewLeafNode;
-			return;
+			NodeAndLeafIndices NewIndices = AllocateLeafNodeAndLeaf(Payload, NewBounds);
+			RootNode = NewIndices.NodeIdx;
+			return NewIndices;
 		}
 
 		// Find the best sibling
@@ -1362,38 +1375,35 @@ public:
 
 		if (bAddToLeaf)
 		{
-			int32 LeafIdx = Nodes[BestSibling].ChildrenNodes[0];
+			const int32 LeafIdx = Nodes[BestSibling].ChildrenNodes[0];
 			Leaves[LeafIdx].AddElement(TPayloadBoundsElement<TPayloadType, T>{Payload, NewBounds});
 			Leaves[LeafIdx].RecomputeBounds();
 			TAABB<T, 3> ExpandedBounds = Leaves[LeafIdx].GetBounds();
 			ExpandedBounds.Thicken(FAABBTreeCVars::DynamicTreeBoundingBoxPadding);
 			Nodes[BestSibling].ChildrenBounds[0] = ExpandedBounds;
 			UpdateAncestorBounds(BestSibling, true);
-			FAABBTreePayloadInfo* PayloadInfo = PayloadToInfo.Find(Payload);
-			PayloadInfo->LeafIdx = LeafIdx;
-			PayloadInfo->NodeIdx = BestSibling;
-			return;
+			return NodeAndLeafIndices{ BestSibling , LeafIdx };
 		}
 
-		int32 NewLeafNode = AllocateLeafNodeAndLeaf(Payload, NewBounds);
+		const NodeAndLeafIndices NewLeafIndices = AllocateLeafNodeAndLeaf(Payload, NewBounds);
 
 		// New internal parent node
-		int32 OldParent = Nodes[BestSibling].ParentNode;
-		int32 NewParent = AllocateInternalNode();
+		const int32 OldParent = Nodes[BestSibling].ParentNode;
+		const int32 NewParent = AllocateInternalNode();
 		FNode& NewParentNode = Nodes[NewParent];
 		NewParentNode.ParentNode = OldParent;
 		NewParentNode.ChildrenNodes[0] = BestSibling;
-		NewParentNode.ChildrenNodes[1] = NewLeafNode;
+		NewParentNode.ChildrenNodes[1] = NewLeafIndices.NodeIdx;
 		NewParentNode.ChildrenBounds[0] = Nodes[BestSibling].ChildrenBounds[0];
 		if (!Nodes[BestSibling].bLeaf)
 		{
 			NewParentNode.ChildrenBounds[0].GrowToInclude(Nodes[BestSibling].ChildrenBounds[1]);
 		}
-		NewParentNode.ChildrenBounds[1] = Nodes[NewLeafNode].ChildrenBounds[0];
+		NewParentNode.ChildrenBounds[1] = Nodes[NewLeafIndices.NodeIdx].ChildrenBounds[0];
 
 		if (OldParent != INDEX_NONE)
 		{
-			int32 ChildIdx = WhichChildAmI(BestSibling);
+			const int32 ChildIdx = WhichChildAmI(BestSibling);
 			Nodes[OldParent].ChildrenNodes[ChildIdx] = NewParent;
 		}
 		else
@@ -1402,9 +1412,11 @@ public:
 		}
 
 		Nodes[BestSibling].ParentNode = NewParent;
-		Nodes[NewLeafNode].ParentNode = NewParent;
+		Nodes[NewLeafIndices.NodeIdx].ParentNode = NewParent;
 
 		UpdateAncestorBounds(NewParent, true);
+
+		return NewLeafIndices;
 	}	
 
 	void  UpdateAncestorBounds(int32 NodeIdx, bool bDoRotation = false)
@@ -1501,7 +1513,11 @@ public:
 				}
 				else if (PayloadInfo->DirtyPayloadIdx != INDEX_NONE)
 				{
-					if (DirtyElementGridEnabled())
+					if (DirtyElementTree != nullptr)
+					{
+						DirtyElementTree->RemoveLeafNode(PayloadInfo->DirtyPayloadIdx, Payload);
+					}
+					else if (DirtyElementGridEnabled())
 					{
 						DeleteDirtyParticleEverywhere(PayloadInfo->DirtyPayloadIdx, PayloadInfo->DirtyGridOverflowIdx);
 					}
@@ -1622,28 +1638,59 @@ public:
 				{
 					if (bDynamicTree)
 					{
-						InsertLeaf(Payload, NewBounds);
+						NodeAndLeafIndices Indices = InsertLeaf(Payload, NewBounds);
+						PayloadInfo->NodeIdx = Indices.NodeIdx;
+						PayloadInfo->LeafIdx = Indices.LeafIdx;
 					}
 					else
 					{
-						PayloadInfo->DirtyPayloadIdx = DirtyElements.Add(FElement{ Payload, NewBounds });
-						if (DirtyElementGridEnabled())
+						if (DirtyElementTree)
 						{
-							//CSV_SCOPED_TIMING_STAT(ChaosPhysicsTimers, AABBAddElement)
-							PayloadInfo->DirtyGridOverflowIdx = AddDirtyElementToGrid(NewBounds, PayloadInfo->DirtyPayloadIdx);
+							// We save the node index for the dirty tree here:
+							PayloadInfo->DirtyPayloadIdx = DirtyElementTree->InsertLeaf(Payload, NewBounds).NodeIdx;
 						}
+						else
+						{
+							PayloadInfo->DirtyPayloadIdx = DirtyElements.Add(FElement{ Payload, NewBounds });
+
+							if (DirtyElementGridEnabled())
+							{
+								PayloadInfo->DirtyGridOverflowIdx = AddDirtyElementToGrid(NewBounds, PayloadInfo->DirtyPayloadIdx);
+							}
+						}						
 					}
 				}
 				else
 				{
 					const int32 DirtyElementIndex = PayloadInfo->DirtyPayloadIdx;
-					if (DirtyElementGridEnabled())
+					if (DirtyElementTree)
 					{
-						//CSV_SCOPED_TIMING_STAT(ChaosPhysicsTimers, AABBUpElement)
-						PayloadInfo->DirtyGridOverflowIdx = UpdateDirtyElementInGrid(NewBounds, DirtyElementIndex, PayloadInfo->DirtyGridOverflowIdx);
+						// The leaf node bounds can be larger than the actual leave bound
+						const TAABB<T, 3>& LeafNodeBounds = DirtyElementTree->Nodes[DirtyElementIndex].ChildrenBounds[0];
+						if (LeafNodeBounds.Contains(NewBounds.Min()) && LeafNodeBounds.Contains(NewBounds.Max()))
+						{
+							// We still need to update the constituent bounds
+							const int32 DirtyLeafIndex = DirtyElementTree->Nodes[DirtyElementIndex].ChildrenNodes[0]; // A Leaf node's first child is the leaf index
+							DirtyElementTree->Leaves[DirtyLeafIndex].UpdateElement(Payload, NewBounds, bHasBounds);
+							DirtyElementTree->Leaves[DirtyLeafIndex].RecomputeBounds();
+						}
+						else
+						{
+							// Reinsert the leaf
+							DirtyElementTree->RemoveLeafNode(DirtyElementIndex, Payload);
+							PayloadInfo->DirtyPayloadIdx = DirtyElementTree->InsertLeaf(Payload, NewBounds).NodeIdx;
+						}
 					}
-					DirtyElements[DirtyElementIndex].Bounds = NewBounds;
-					UpdateElementHelper(DirtyElements[DirtyElementIndex].Payload, Payload);
+					else
+					{
+						if (DirtyElementGridEnabled())
+						{
+							//CSV_SCOPED_TIMING_STAT(ChaosPhysicsTimers, AABBUpElement)
+							PayloadInfo->DirtyGridOverflowIdx = UpdateDirtyElementInGrid(NewBounds, DirtyElementIndex, PayloadInfo->DirtyGridOverflowIdx);
+						}
+						DirtyElements[DirtyElementIndex].Bounds = NewBounds;
+						UpdateElementHelper(DirtyElements[DirtyElementIndex].Payload, Payload);
+					}				
 				}
 
 				// Handle something that previously did not have bounds that may be in global elements.
@@ -1675,18 +1722,25 @@ public:
 				// Handle something that previously had bounds that may be in dirty elements.
 				if (PayloadInfo->DirtyPayloadIdx != INDEX_NONE)
 				{
-					if (DirtyElementGridEnabled())
+					if (DirtyElementTree)
 					{
-						DeleteDirtyParticleEverywhere(PayloadInfo->DirtyPayloadIdx, PayloadInfo->DirtyGridOverflowIdx);
+						DirtyElementTree->RemoveLeafNode(PayloadInfo->DirtyPayloadIdx, Payload);
 					}
 					else
 					{
-						if (PayloadInfo->DirtyPayloadIdx + 1 < DirtyElements.Num())
+						if (DirtyElementGridEnabled())
 						{
-							auto LastDirtyPayload = DirtyElements.Last().Payload;
-							PayloadToInfo.FindChecked(LastDirtyPayload).DirtyPayloadIdx = PayloadInfo->DirtyPayloadIdx;
+							DeleteDirtyParticleEverywhere(PayloadInfo->DirtyPayloadIdx, PayloadInfo->DirtyGridOverflowIdx);
+						}
+						else
+						{
+							if (PayloadInfo->DirtyPayloadIdx + 1 < DirtyElements.Num())
+							{
+								auto LastDirtyPayload = DirtyElements.Last().Payload;
+								PayloadToInfo.FindChecked(LastDirtyPayload).DirtyPayloadIdx = PayloadInfo->DirtyPayloadIdx;
+							}
+							DirtyElements.RemoveAtSwap(PayloadInfo->DirtyPayloadIdx);
 						}
-						DirtyElements.RemoveAtSwap(PayloadInfo->DirtyPayloadIdx);
 					}
 
 					PayloadInfo->DirtyPayloadIdx = INDEX_NONE;
@@ -1695,7 +1749,7 @@ public:
 			}
 		}
 
-		if(!bDynamicTree && DirtyElements.Num() > MaxDirtyElements)
+		if(!DirtyElementTree && !bDynamicTree && DirtyElements.Num() > MaxDirtyElements)
 		{
 			UE_LOG(LogChaos, Verbose, TEXT("Bounding volume exceeded maximum dirty elements (%d dirty of max %d) and is forcing a tree rebuild."), DirtyElements.Num(), MaxDirtyElements);
 			ReoptimizeTree();
@@ -1758,7 +1812,7 @@ public:
 	virtual void ClearShouldRebuild() override { bShouldRebuild = false; }
 
 	virtual bool IsTreeDynamic() const override { return bDynamicTree; }
-	bool SetTreeToDynamic() { bDynamicTree = true; } // Tree cannot be changed back to static for now
+	void SetTreeToDynamic() { bDynamicTree = true; } // Tree cannot be changed back to static for now
 
 	virtual void PrepareCopyTimeSliced(const  ISpatialAcceleration<TPayloadType, T, 3>& InFrom) override
 	{
@@ -1805,6 +1859,12 @@ public:
 		OverlappingCounts.Reserve(From.OverlappingCounts.Num());
 
 		this->SetAsyncTimeSlicingComplete(false);
+
+		if (DirtyElementTree)
+		{
+			check(From.DirtyElementTree);
+			DirtyElementTree->PrepareCopyTimeSliced(*(From.DirtyElementTree));
+		}
 	}
 	
 	virtual void ProgressCopyTimeSliced(const  ISpatialAcceleration<TPayloadType, T, 3>& InFrom, int MaximumBytesToCopy) override
@@ -1862,6 +1922,16 @@ public:
 			return;
 		}
 
+		if (DirtyElementTree)
+		{
+			check(From.DirtyElementTree);
+			DirtyElementTree->ProgressCopyTimeSliced(*(From.DirtyElementTree), MaximumBytesToCopy);
+			if (!DirtyElementTree->IsAsyncTimeSlicingComplete())
+			{
+				return;
+			}
+		}
+
 		this->SetAsyncTimeSlicingComplete(true);
 	}
 
@@ -2226,6 +2296,7 @@ private:
 
 	void ReoptimizeTree()
 	{
+		check(!DirtyElementTree && !bDynamicTree);
 		TRACE_CPUPROFILER_EVENT_SCOPE(TAABBTree::ReoptimizeTree);
 		TArray<FElement> AllElements;
 
@@ -2411,7 +2482,14 @@ private:
 			}
 		}
 
-		if (bMutable)
+		if (DirtyElementTree)
+		{
+			if (!DirtyElementTree->template QueryImp<Query, TQueryFastData, SQVisitor>(Start, CurData, QueryHalfExtents, QueryBounds, Visitor, Dir, InvDir, bParallel))
+			{
+				return false;
+			}
+		}
+		else if (bMutable)
 		{	// Returns true if we should continue
 			auto IntersectAndVisit = [&](const FElement& Elem) -> bool
 			{
@@ -2670,6 +2748,11 @@ private:
 		CellHashToFlatArray.Reset(); 
 		FlattenedCellArrayOfDirtyIndices.Reset();
 		DirtyElementsGridOverflow.Reset();
+		if (DirtyElementTree)
+		{
+			DirtyElementTree->Reset();
+		}
+		
 		TreeStats.Reset();
 		TreeExpensiveStats.Reset();
 		PayloadToInfo.Reset();
@@ -2678,7 +2761,6 @@ private:
 
 		if (bDynamicTree)
 		{
-			// Todo for now, don't ever return async task
 			int32 Idx = 0;
 			for (auto& Particle : Particles)
 			{
@@ -3177,6 +3259,7 @@ private:
 		, CellHashToFlatArray(Other.CellHashToFlatArray)
 		, FlattenedCellArrayOfDirtyIndices(Other.FlattenedCellArrayOfDirtyIndices)
 		, DirtyElementsGridOverflow(Other.DirtyElementsGridOverflow)
+		, DirtyElementTree(nullptr)
 		, DirtyElementGridCellSize(Other.DirtyElementGridCellSize)
 		, DirtyElementGridCellSizeInv(Other.DirtyElementGridCellSizeInv)
 		, DirtyElementMaxGridCellQueryCount(Other.DirtyElementMaxGridCellQueryCount)
@@ -3199,6 +3282,10 @@ private:
 	{
 		ensure(bDynamicTree == Other.IsTreeDynamic());
 		PriorityQ.Reserve(32);
+		if (Other.DirtyElementTree)
+		{
+			DirtyElementTree = TUniquePtr<TAABBTree<TPayloadType, TLeafType, bMutable, T>>(new TAABBTree<TPayloadType, TLeafType, bMutable, T>(*(Other.DirtyElementTree)));
+		}
 	}
 
 	virtual ISpatialAcceleration<TPayloadType, T, 3>& operator=(const ISpatialAcceleration<TPayloadType, T, 3>& Other) override
@@ -3247,6 +3334,11 @@ private:
 			MaxNumToProcess = Rhs.MaxNumToProcess;
 			NumProcessedThisSlice = Rhs.NumProcessedThisSlice;
 			bShouldRebuild = Rhs.bShouldRebuild;
+			if (Rhs.DirtyElementTree)
+			{
+				check(DirtyElementTree); // We should have allocated this already
+				*DirtyElementTree = *Rhs.DirtyElementTree;
+			}
 		}
 
 		return *this;
@@ -3267,6 +3359,11 @@ private:
 	TMap<int32, DirtyGridHashEntry> CellHashToFlatArray; // Index, size into flat grid structure (FlattenedCellArrayOfDirtyIndices)
 	TArray<int32> FlattenedCellArrayOfDirtyIndices; // Array of indices into dirty Elements array, indices are always sorted within a 2D cell
 	TArray<int32> DirtyElementsGridOverflow; // Array of indices of DirtyElements that is not in the grid for some reason
+
+	// Members for using a dynamic tree as a dirty element acceleration structure
+	TUniquePtr<TAABBTree<TPayloadType, TLeafType, bMutable, T>> DirtyElementTree;
+
+
 	// Copy of CVARS
 	T DirtyElementGridCellSize;
 	T DirtyElementGridCellSizeInv;
