Add undoable support for constraints

diff --git a/Engine/Source/Runtime/Experimental/Animation/Constraints/Public/ConstraintsManager.h b/Engine/Source/Runtime/Experimental/Animation/Constraints/Public/ConstraintsManager.h
--- a/Engine/Source/Runtime/Experimental/Animation/Constraints/Public/ConstraintsManager.h
+++ b/Engine/Source/Runtime/Experimental/Animation/Constraints/Public/ConstraintsManager.h
@@ -204,7 +204,7 @@ public:
 		// unique name (we may want to use another approach here to manage uniqueness)
 		const FName Name = MakeUniqueObjectName(Manager, TConstraint::StaticClass(), InBaseName);
 
-		TConstraint* NewConstraint = NewObject<TConstraint>(Manager, Name);		
+		TConstraint* NewConstraint = NewObject<TConstraint>(Manager, Name, RF_Transactional);		
 		return NewConstraint;
 	}
 
