HeightField Raycast optimizatrion
Three level of walking raycast :
- Squared low level heightfield walking 6x6 resolution 36 cells checked
- Fast walking only on cells interescting the ray 2 cell length check (min 2, max 4)
- Slow walking one cell at the time
If the fast or the slow walking don't find intersection they go back to the low resolution height field walking
Additional memory for low resolution height field before this CL was 8% now it is 2x2.8 = 5.6% so 2.4% saved

diff --git a/Engine/Source/Runtime/Core/Public/UObject/UE5MainStreamObjectVersion.h b/Engine/Source/Runtime/Core/Public/UObject/UE5MainStreamObjectVersion.h
--- a/Engine/Source/Runtime/Core/Public/UObject/UE5MainStreamObjectVersion.h
+++ b/Engine/Source/Runtime/Core/Public/UObject/UE5MainStreamObjectVersion.h
@@ -265,6 +265,9 @@ struct CORE_API FUE5MainStreamObjectVersion
 		// Add low resolution data in Height Field
 		AddLowResolutionHeightField,
 
+		// Low resolution data in Height Field will store one height for (6x6) 36 cells
+		DecreaseLowResolutionHeightField,
+
 		// -----<new versions can be added above this line>-------------------------------------------------
 		VersionPlusOne,
 		LatestVersion = VersionPlusOne - 1
diff --git a/Engine/Source/Runtime/Experimental/Chaos/Private/Chaos/HeightField.cpp b/Engine/Source/Runtime/Experimental/Chaos/Private/Chaos/HeightField.cpp
--- a/Engine/Source/Runtime/Experimental/Chaos/Private/Chaos/HeightField.cpp
+++ b/Engine/Source/Runtime/Experimental/Chaos/Private/Chaos/HeightField.cpp
@@ -52,7 +52,6 @@ namespace Chaos
 			VectorRegister4Float DirSimd = MakeVectorRegisterFloatFromDouble(MakeVectorRegister(Dir.X, Dir.Y, Dir.Z, 0.0));
 			Parallel = VectorCompareLT(VectorAbs(DirSimd), GlobalVectorConstants::SmallNumber);
 			InvDirSimd = VectorBitwiseNotAnd(Parallel, VectorDivide(VectorOne(), DirSimd));
-			ThicknessSimd = VectorSet1(static_cast<FRealSingle>(InThickness));
 		}
 
 		enum class ERaycastType
@@ -61,17 +60,11 @@ namespace Chaos
 			Sweep
 		};
 
-		FORCEINLINE_DEBUGGABLE bool VisitRaycast(int32 Payload, FReal& CurrentLength)
+		FORCEINLINE_DEBUGGABLE bool VisitRaycast(int32 Payload, FReal CurrentLength)
 		{
-			const int32 SubX = Payload % (GeomData->NumCols - 1);
 			const int32 SubY = Payload / (GeomData->NumCols - 1);
-
 			const int32 FullIndex = Payload + SubY;
 
-			const FReal Radius = Thickness + UE_SMALL_NUMBER;
-			const FReal Radius2 = Radius * Radius;
-			bool bIntersection = false;
-
 			// return if the triangle was hit or not 
 			auto TestTriangle = [&](int32 FaceIndex, const FVec3& A, const FVec3& B, const FVec3& C) -> bool
 			{
@@ -95,7 +88,6 @@ namespace Chaos
 							OutNormal = Normal;
 							OutTime = Time;
 							OutFaceIndex = FaceIndex;
-							CurrentLength = Time;
 							return true;
 						}
 					}
@@ -105,18 +97,11 @@ namespace Chaos
 			};
 
 			FVec3 Points[4];
-			FAABB3 CellBounds;
-			GeomData->GetPointsAndBoundsScaled(FullIndex, Points, CellBounds);
-			FAABBVectorized Bounds(CellBounds);
-			Bounds.Thicken(ThicknessSimd);
-
+			GeomData->GetPointsScaled(FullIndex, Points);
 			bool bHit = false;
-			if (Bounds.RaycastFast(StartPointSimd, InvDirSimd, Parallel, VectorSet1(static_cast<FRealSingle>(CurrentLength))))
-			{
 				// Test both triangles that are in this cell, as we could hit both in any order
 				bHit |= TestTriangle(Payload * 2, Points[0], Points[1], Points[3]);
 				bHit |= TestTriangle(Payload * 2 + 1, Points[0], Points[3], Points[2]);
-			}
 			const bool bShouldContinueVisiting = !bHit;
 			return bShouldContinueVisiting;
 		}
@@ -283,7 +268,6 @@ namespace Chaos
 		VectorRegister4Float StartPointSimd;
 		VectorRegister4Float InvDirSimd;
 		VectorRegister4Float Parallel;
-		VectorRegister4Float ThicknessSimd;
 		
 		const typename FHeightField::FDataType* GeomData;
 	private:
@@ -1044,8 +1028,250 @@ namespace Chaos
 		return false;
 	}
 
+	// Raycast is navigating slowly one cell at the time, on the high resolution height field
+	// This navigation is the most precise and will check if we need to compute the raycast at the triangle level
+	FORCEINLINE bool FHeightField::WalkSlow(TVec2<int32>& CellIdx, FHeightfieldRaycastVisitor& Visitor, FReal CurrentLength, const VectorRegister4Float& CurrentLengthSimd,
+		const FVec2& ScaledMin, FReal ZMidPoint, const FVec3& Dir, const FVec3& InvDir, bool bParallel[3], const FVec2& ScaledDx2D, FVec3& NextStart, 
+		const FVec3& ScaleSign, const FVec3& ScaledDx, int32 IndexLowResX, int32 IndexLowResY) const
+	{
+		const int32 LowResInc = Visitor.GeomData->LowResInc;
+		while (true)
+		{
+			if (FlatGrid.IsValid(CellIdx))
+			{
+				PHYSICS_CSV_CUSTOM_VERY_EXPENSIVE(PhysicsCounters, NumRayHeightfieldCellVisited, 1, ECsvCustomStatOp::Accumulate);
+				// Test for the cell bounding box is done in the visitor at the same time as fetching the points for the triangles
+				// this avoid fetching the points twice ( here and in the visitor )
+
+				const int32 Payload = CellIdx[1] * (GeomData.NumCols - 1) + CellIdx[0];
+				const int32 SubY = (Payload) / (GeomData.NumCols - 1);
+				const int32 FullIndex = Payload + SubY;
+				FAABBVectorized Bounds;
+				Visitor.GeomData->GetBoundsScaledSimd(FullIndex, Bounds);
+				if (Bounds.RaycastFast(Visitor.StartPointSimd, Visitor.InvDirSimd, Visitor.Parallel, CurrentLengthSimd))
+				{
+					const bool bContinue = Visitor.VisitRaycast(CellIdx[1] * (GeomData.NumCols - 1) + CellIdx[0], CurrentLength);
+					if (!bContinue)
+					{
+						return false;
+					}
+				}
+			}
+
+			//find next cell
+			//We want to know which plane we used to cross into next cell
+			const FVec2 ScaledCellCenter2D = ScaledMin + FVec2(static_cast<FReal>(CellIdx[0]) + 0.5f, static_cast<FReal>(CellIdx[1]) + 0.5f) * ScaledDx2D;
+			const FVec3 ScaledCellCenter(ScaledCellCenter2D[0], ScaledCellCenter2D[1], ZMidPoint);
+
+			FReal Times[3];
+			FReal BestTime = CurrentLength;
+			bool bTerminate = true;
+			for (int Axis = 0; Axis < 3; ++Axis)
+			{
+				if (!bParallel[Axis])
+				{
+					const FReal CrossPoint = (Dir[Axis] * ScaleSign[Axis]) > 0 ? ScaledCellCenter[Axis] + ScaledDx[Axis] / 2 : ScaledCellCenter[Axis] - ScaledDx[Axis] / 2;
+					const FReal Distance = CrossPoint - NextStart[Axis];	//note: CellCenter already has /2, we probably want to use the corner instead
+					const FReal Time = Distance * InvDir[Axis];
+					Times[Axis] = Time;
+					if (Time < BestTime)
+					{
+						bTerminate = false;	//found at least one plane to pass through
+						BestTime = Time;
+					}
+				}
+				else
+				{
+					Times[Axis] = TNumericLimits<FReal>::Max();
+				}
+			}
+
+			if (bTerminate)
+			{
+				return false;
+			}
+
+			const TVec2<int32> PrevIdx = CellIdx;
+
+			for (int Axis = 0; Axis < 2; ++Axis)
+			{
+				CellIdx[Axis] += (Times[Axis] <= BestTime) ? ((Dir[Axis] * ScaleSign[Axis]) > 0 ? 1 : -1) : 0;
+				if (CellIdx[Axis] < 0 || CellIdx[Axis] >= FlatGrid.Counts()[Axis])
+				{
+					return false;
+				}
+			}
+
+			if (PrevIdx == CellIdx)
+			{
+				//crossed on z plane which means no longer in heightfield bounds
+				return false;
+			}
+
+			NextStart = NextStart + Dir * BestTime;
+
+			if (IndexLowResX != CellIdx[0] / LowResInc || IndexLowResY != CellIdx[1] / LowResInc)
+			{
+				return true;
+			}
+
+		}
+		return true;
+	}
+
+	// Raycast is navigating fast on the high resolution height field
+	// This navigation is more precise than the low resolution walk. It would have more cache miss because it grabs several cells not necessarily contiguous in memory.
+	// It navigates only on the cell that need to be checked.
+	FORCEINLINE bool FHeightField::WalkFast(TVec2<int32>& CellIdx, FHeightfieldRaycastVisitor& Visitor, FReal CurrentLength, const VectorRegister4Float& CurrentLengthSimd,
+		const FVec2& ScaledMin, FReal ZMidPoint, const FVec3& Dir, const FVec3& InvDir, bool bParallel[3], const FVec2& ScaledDx2D, FVec3& NextStart,
+		const FVec3& ScaleSign, const FVec3& ScaledDx, const FVec2& Scale2D, const FVec3& DirScaled) const 
+	{
+		const FReal BoundsMinZ = CachedBounds.Min().Z;
+		const FReal BoundsMaxZ = CachedBounds.Max().Z;
+
+		const int32 LowResInc = Visitor.GeomData->LowResInc;
+		const int32 IndexLowResX = CellIdx[0] / LowResInc;
+		const int32 IndexLowResY = CellIdx[1] / LowResInc;
+
+		const TVec2<int32> NextStartInt = TVec2<int32>(static_cast<int32>(NextStart[0] / Scale2D[0]), static_cast<int32>(NextStart[1] / Scale2D[1]));
+		CellIdx = FlatGrid.Cell(NextStartInt);
+
+		const FVec3 NextStartOri = NextStart;
+		const FReal Length2D = FMath::Sqrt(DirScaled[0] * DirScaled[0] + DirScaled[1] * DirScaled[1]);
+		constexpr FReal FastInc2D = 2.0;  // Length in cell dimension
+		const FReal FastIncScaled = FastInc2D / Length2D;
+
+		const FVec3 InspectionStepVector = DirScaled * FastIncScaled * ScaledDx;
+		FReal NextStartZ = FMath::Clamp(NextStart.Z, BoundsMinZ, BoundsMaxZ); // Numerical errors can violate this condition, so force it here.
+		const FReal FastIncScaled3D = InspectionStepVector.Length();
+		const FReal CellsToInspectZ = FastIncScaled3D * Dir.Z;
+		FReal DistanceProcessed = 0.0;
+
+		while (true)
+		{
+			if (!FlatGrid.IsValid(CellIdx))
+			{
+				return false;
+			}
+
+			const FVec3 NewNextStart = (InspectionStepVector + NextStart);
+			const TVec2<int32> NewNextStartInt = TVec2<int32>(static_cast<int32>(NewNextStart[0] / Scale2D[0]), static_cast<int32>(NewNextStart[1] / Scale2D[1]));
+			const TVec2<int32> NextStartInt2 = TVec2<int32>(static_cast<int32>(NextStart[0] / Scale2D[0]), static_cast<int32>(NextStart[1] / Scale2D[1]));
+
+			const TVec2<int32> DiffInt = NewNextStartInt - NextStartInt2;
+			const TVec2<int32> AddedCellIdx = TVec2<int32>(static_cast<int32>(DiffInt[0]), static_cast<int32>(DiffInt[1]));
+
+			FAABBVectorized Bounds;
+			Visitor.GeomData->GetBoundsScaled(CellIdx, AddedCellIdx, Bounds);
+			if (Bounds.RaycastFast(Visitor.StartPointSimd, Visitor.InvDirSimd, Visitor.Parallel, CurrentLengthSimd))
+			{
+				NextStart = NextStartOri + Dir * DistanceProcessed;
+				bool bContinue = WalkSlow(CellIdx, Visitor, CurrentLength, CurrentLengthSimd, ScaledMin, ZMidPoint, Dir, InvDir, bParallel, ScaledDx2D, NextStart, ScaleSign, ScaledDx, IndexLowResX, IndexLowResY);
+				if (bContinue)
+				{
+					return true;
+				}
+				return false;
+			}
+
+			NextStart = NewNextStart;
+			CellIdx = FlatGrid.Cell(NewNextStartInt);
+			if (DistanceProcessed > CurrentLength || NextStartZ < BoundsMinZ || NextStartZ > BoundsMaxZ)
+			{
+				return false;
+			}
+			DistanceProcessed += FastIncScaled3D;
+			NextStartZ += CellsToInspectZ;
+			if (IndexLowResX != CellIdx[0] / LowResInc || IndexLowResY != CellIdx[1] / LowResInc)
+			{
+				// Walk on low resolution height field
+				return true;
+			}
+		} 
+		return true;
+	}
+
+	// Raycast is navigating a low resolution height field
+	// This navigation is less precise can have lot of false positive intersection but navigate fast with few cache miss.
+	// It navigates on a squared cell number.
+	FORCEINLINE bool FHeightField::WalkOnLowRes(TVec2<int32>& CellIdx, FHeightfieldRaycastVisitor& Visitor, FReal CurrentLength, const VectorRegister4Float& CurrentLengthSimd,
+		const FVec2& ScaledMin, FReal ZMidPoint, const FVec3& Dir, const FVec3& InvDir, bool bParallel[3], const FVec2& ScaledDx2D, FVec3& NextStart,
+		const FVec3& ScaleSign, const FVec3& ScaledDx, const FVec2& Scale2D, const FVec3& DirScaled) const
+	{
+		const int32 LowResInc = Visitor.GeomData->LowResInc;
+		while (true)
+		{
+			if (!FlatGrid.IsValid(CellIdx))
+			{
+				return false;
+			}
+
+			FAABBVectorized Bounds;
+			Visitor.GeomData->GetLowResBoundsScaled(CellIdx, Bounds);
+			if (Bounds.RaycastFast(Visitor.StartPointSimd, Visitor.InvDirSimd, Visitor.Parallel, CurrentLengthSimd))
+			{
+				bool bContinue = WalkFast(CellIdx, Visitor, CurrentLength, CurrentLengthSimd, ScaledMin, ZMidPoint, Dir, InvDir, bParallel, ScaledDx2D, NextStart,
+					ScaleSign, ScaledDx, Scale2D, DirScaled);
+				if (bContinue)
+				{
+					continue; 
+				}
+				return false;
+			}
+
+			const FVec2 ScaledCellCenter2D = ScaledMin + FVec2(static_cast<FReal>(CellIdx[0] / LowResInc) + 0.5f, static_cast<FReal>(CellIdx[1] / LowResInc) + 0.5f) * (ScaledDx2D * (FReal)LowResInc);
+			const FVec3 ScaledCellCenter(ScaledCellCenter2D[0], ScaledCellCenter2D[1], ZMidPoint);
+
+			FReal Times[3];
+			FReal BestTime = CurrentLength;
+			bool bTerminate = true;
+			for (int Axis = 0; Axis < 3; ++Axis)
+			{
+				if (!bParallel[Axis])
+				{
+					const FReal CrossPoint = (Dir[Axis] * ScaleSign[Axis]) > 0 ? ScaledCellCenter[Axis] + ScaledDx[Axis] * (FReal)LowResInc / 2 : ScaledCellCenter[Axis] - ScaledDx[Axis] * (FReal)LowResInc / 2;
+					const FReal Distance = CrossPoint - NextStart[Axis];	//note: CellCenter already has /2, we probably want to use the corner instead
+					const FReal Time = Distance * InvDir[Axis];
+					Times[Axis] = Time;
+					if (Time < BestTime)
+					{
+						bTerminate = false;	//found at least one plane to pass through
+						BestTime = Time;
+					}
+				}
+				else
+				{
+					Times[Axis] = TNumericLimits<FReal>::Max();
+				}
+			}
 
-	bool FHeightField::GridCast(const FVec3& StartPoint, const FVec3& Dir, const FReal Length, FHeightfieldRaycastVisitor& Visitor) const
+			if (bTerminate)
+			{
+				return false;
+			}
+
+			const TVec2<int32> PrevIdx = CellIdx;
+
+			for (int Axis = 0; Axis < 2; ++Axis)
+			{
+				CellIdx[Axis] += (Times[Axis] <= BestTime) ? ((Dir[Axis] * ScaleSign[Axis]) > 0 ? LowResInc : -LowResInc) : 0;
+				if (CellIdx[Axis] < 0 || CellIdx[Axis] >= FlatGrid.Counts()[Axis])
+				{
+					return false;
+				}
+			}
+
+			if (PrevIdx == CellIdx)
+			{
+				//crossed on z plane which means no longer in heightfield bounds
+				return false;
+			}
+			NextStart = NextStart + Dir * BestTime;
+		}
+		return false;
+	}
+	
+	FORCEINLINE bool FHeightField::GridCast(const FVec3& StartPoint, const FVec3& Dir, const FReal Length, FHeightfieldRaycastVisitor& Visitor) const
 	{
 		//Is this check needed?
 		if(Length < 1e-4)
@@ -1076,7 +1302,7 @@ namespace Chaos
 			CurrentLength = RayExitTime + 1e-2; // to account for precision errors 
 			FVec3 NextStart = StartPoint + (Dir * RayEntryTime);
 
-			const FVec2 Scale2D(GeomData.Scale[0],GeomData.Scale[1]);
+			const FVec2 Scale2D(GeomData.Scale[0], GeomData.Scale[1]);
 			TVec2<int32> CellIdx = FlatGrid.Cell(TVec2<int32>(static_cast<int32>(NextStart[0] / Scale2D[0]), static_cast<int32>(NextStart[1] / Scale2D[1])));
 			const FReal ZDx = CachedBounds.Extents()[2];
 			const FReal ZMidPoint = CachedBounds.Min()[2] + ZDx * 0.5f;
@@ -1088,147 +1314,23 @@ namespace Chaos
 			const FVec3 DirScaled = Dir / ScaledDx;
 			const FReal SumPlaneAxis = FMath::Abs(DirScaled[0]) + FMath::Abs(DirScaled[1]);
 			const bool bCanWalk = SumPlaneAxis > UE_SMALL_NUMBER;
+			const VectorRegister4Float CurrentLengthSimd = VectorSet1(static_cast<FRealSingle>(CurrentLength));
 			if (bCanWalk)
 			{
-				const VectorRegister4Float CurrentLengthSimd = VectorSet1(static_cast<FRealSingle>(CurrentLength));
-				const int32 LowResInc = Visitor.GeomData->LowResInc;
-
-				while (true)
-				{
-					if (!FlatGrid.IsValid(CellIdx))
-					{
-						return false;
-					}
-
-					FAABBVectorized Bounds;
-					Visitor.GeomData->GetLowResBoundsScaled(CellIdx, Bounds);
-					Bounds.Thicken(Visitor.ThicknessSimd);
-					if (Bounds.RaycastFast(Visitor.StartPointSimd, Visitor.InvDirSimd, Visitor.Parallel, CurrentLengthSimd))
-					{
-						TVec2<int32> NextStartInt = TVec2<int32>(static_cast<int32>(NextStart[0] / Scale2D[0]), static_cast<int32>(NextStart[1] / Scale2D[1]));
-						CellIdx = FlatGrid.Cell(NextStartInt);
-						break;
-					}
-
-					const FVec2 ScaledCellCenter2D = ScaledMin + FVec2(static_cast<FReal>(CellIdx[0] / LowResInc) + 0.5f, static_cast<FReal>(CellIdx[1] / LowResInc) + 0.5f) * (ScaledDx2D*(FReal)LowResInc);
-					const FVec3 ScaledCellCenter(ScaledCellCenter2D[0], ScaledCellCenter2D[1], ZMidPoint);
-
-					FReal Times[3];
-					FReal BestTime = CurrentLength;
-					bool bTerminate = true;
-					for (int Axis = 0; Axis < 3; ++Axis)
-					{
-						if (!bParallel[Axis])
-						{
-							const FReal CrossPoint = (Dir[Axis] * ScaleSign[Axis]) > 0 ? ScaledCellCenter[Axis] + ScaledDx[Axis] * (FReal)LowResInc / 2 : ScaledCellCenter[Axis] - ScaledDx[Axis] * (FReal)LowResInc / 2;
-							const FReal Distance = CrossPoint - NextStart[Axis];	//note: CellCenter already has /2, we probably want to use the corner instead
-							const FReal Time = Distance * InvDir[Axis];
-							Times[Axis] = Time;
-							if (Time < BestTime)
-							{
-								bTerminate = false;	//found at least one plane to pass through
-								BestTime = Time;
-							}
-						}
-						else
-						{
-							Times[Axis] = TNumericLimits<FReal>::Max();
-						}
-					}
-
-					if (bTerminate)
-					{
-						return false;
-					}
-
-					const TVec2<int32> PrevIdx = CellIdx;
-
-					for (int Axis = 0; Axis < 2; ++Axis)
-					{
-						CellIdx[Axis] += (Times[Axis] <= BestTime) ? ((Dir[Axis] * ScaleSign[Axis]) > 0 ? LowResInc : -LowResInc) : 0;
-						if (CellIdx[Axis] < 0 || CellIdx[Axis] >= FlatGrid.Counts()[Axis])
-						{
-							return false;
-						}
-					}
-
-					if (PrevIdx == CellIdx)
-					{
-						//crossed on z plane which means no longer in heightfield bounds
-						return false;
-					}
-					NextStart = NextStart + Dir * BestTime;
-				}
+				WalkOnLowRes(CellIdx, Visitor, CurrentLength, CurrentLengthSimd, ScaledMin, ZMidPoint, Dir, InvDir, bParallel, ScaledDx2D, NextStart, ScaleSign, ScaledDx, Scale2D, DirScaled);
 			}
-
-			//START
-			do
+			else 
 			{
-				if (FlatGrid.IsValid(CellIdx))
-				{
-					PHYSICS_CSV_CUSTOM_VERY_EXPENSIVE(PhysicsCounters, NumRayHeightfieldCellVisited, 1, ECsvCustomStatOp::Accumulate);
-					// Test for the cell bounding box is done in the visitor at the same time as fetching the points for the triangles
-					// this avoid fetching the points twice ( here and in the visitor ) 
-					bool bContinue = Visitor.VisitRaycast(CellIdx[1] * (GeomData.NumCols - 1) + CellIdx[0], CurrentLength);
-					if (!bContinue)
-					{
-						return false;
-					}
-				}
-
-				//find next cell
-				//We want to know which plane we used to cross into next cell
-				const FVec2 ScaledCellCenter2D = ScaledMin + FVec2(static_cast<FReal>(CellIdx[0]) + 0.5f, static_cast<FReal>(CellIdx[1]) + 0.5f) * ScaledDx2D;
-				const FVec3 ScaledCellCenter(ScaledCellCenter2D[0], ScaledCellCenter2D[1], ZMidPoint);
-
-				FReal Times[3];
-				FReal BestTime = CurrentLength;
-				bool bTerminate = true;
-				for(int Axis = 0; Axis < 3; ++Axis)
-				{
-					if(!bParallel[Axis])
-					{
-						const FReal CrossPoint = (Dir[Axis] * ScaleSign[Axis]) > 0 ? ScaledCellCenter[Axis] + ScaledDx[Axis] / 2 : ScaledCellCenter[Axis] - ScaledDx[Axis] / 2;
-						const FReal Distance = CrossPoint - NextStart[Axis];	//note: CellCenter already has /2, we probably want to use the corner instead
-						const FReal Time = Distance * InvDir[Axis];
-						Times[Axis] = Time;
-						if(Time < BestTime)
-						{
-							bTerminate = false;	//found at least one plane to pass through
-							BestTime = Time;
-						}
-					} else
-					{
-						Times[Axis] = TNumericLimits<FReal>::Max();
-					}
-				}
-
-				if(bTerminate)
-				{
-					return false;
-				}
-
-				const TVec2<int32> PrevIdx = CellIdx;
-
-				for(int Axis = 0; Axis < 2; ++Axis)
-				{
-					CellIdx[Axis] += (Times[Axis] <= BestTime) ? ((Dir[Axis] * ScaleSign[Axis]) > 0 ? 1 : -1) : 0;
-					if(CellIdx[Axis] < 0 || CellIdx[Axis] >= FlatGrid.Counts()[Axis])
-					{
-						return false;
-					}
-				}
-
-				if(PrevIdx == CellIdx)
+				const int32 LowResInc = Visitor.GeomData->LowResInc;
+				const int32 IndexLowResX = CellIdx[0] / LowResInc;
+				const int32 IndexLowResY = CellIdx[1] / LowResInc;
+				const bool bContinue = WalkSlow(CellIdx, Visitor, CurrentLength, CurrentLengthSimd, ScaledMin, ZMidPoint, Dir, InvDir, bParallel, ScaledDx2D, NextStart, ScaleSign, ScaledDx, IndexLowResX, IndexLowResY);
+				if (!bContinue)
 				{
-					//crossed on z plane which means no longer in heightfield bounds
 					return false;
 				}
-
-				NextStart = NextStart + Dir * BestTime;
-			} while(true);
+			}
 		}
-
 		return false;
 	}
 
diff --git a/Engine/Source/Runtime/Experimental/Chaos/Public/Chaos/HeightField.h b/Engine/Source/Runtime/Experimental/Chaos/Public/Chaos/HeightField.h
--- a/Engine/Source/Runtime/Experimental/Chaos/Public/Chaos/HeightField.h
+++ b/Engine/Source/Runtime/Experimental/Chaos/Public/Chaos/HeightField.h
@@ -270,7 +270,7 @@ namespace Chaos
 			// Range of the chosen type (unsigned so Min is always 0)
 			static constexpr int32 StorageRange = TNumericLimits<StorageType>::Max();
 
-			static constexpr int32 LowResInc = 5;
+			static constexpr int32 LowResInc = 6;
 
 			// Heights in the chosen format. final placement of the vertex will be at
 			// MinValue + Heights[Index] * HeightPerUnit
@@ -406,6 +406,23 @@ namespace Chaos
 				OutBounds = FAABBVectorized(MakeVectorRegisterFloat((FRealSingle)X, (FRealSingle)Y, MinZ, 0.0f),
 											MakeVectorRegisterFloat((FRealSingle)X + 1, (FRealSingle)Y + 1, MaxZ, 0.0f));
 			}
+
+			FORCEINLINE void GetBoundsSimd(int32 Index, FAABBVectorized& OutBounds) const
+			{
+				const FRealSingle H0 = static_cast<FRealSingle>(MinValue + Heights[Index] * HeightPerUnit);
+				const FRealSingle H1 = static_cast<FRealSingle>(MinValue + Heights[Index + 1] * HeightPerUnit);
+				const FRealSingle H2 = static_cast<FRealSingle>(MinValue + Heights[Index + NumCols] * HeightPerUnit);
+				const FRealSingle H3 = static_cast<FRealSingle>(MinValue + Heights[Index + NumCols + 1] * HeightPerUnit);
+
+				const int32 X = Index % (NumCols);
+				const int32 Y = Index / (NumCols);
+				
+				const FRealSingle MinZ = FMath::Min<FRealSingle>(H0, FMath::Min<FRealSingle>(H1, FMath::Min<FRealSingle>(H2, H3)));
+				const FRealSingle MaxZ = FMath::Max<FRealSingle>(H0, FMath::Max<FRealSingle>(H1, FMath::Max<FRealSingle>(H2, H3)));
+
+				OutBounds = FAABBVectorized(MakeVectorRegisterFloat((FRealSingle)X, (FRealSingle)Y, MinZ, 0.0f),
+					MakeVectorRegisterFloat((FRealSingle)X + 1, (FRealSingle)Y + 1, MaxZ, 0.0f));
+			}
 			
 			FORCEINLINE void GetBounds(TVec2<int32> CellIdx, TVec2<int32> Area, FAABBVectorized& OutBounds) const
 			{
@@ -451,13 +468,10 @@ namespace Chaos
 
 			FORCEINLINE void GetLowResBounds(TVec2<int32> CellIdx, FAABBVectorized& OutBounds) const
 			{
-				int32 FirstIndexX = CellIdx[0];
-				int32 FirstIndexY = CellIdx[1];
-
-				const int32 FirstIndexLowResX = FirstIndexX / LowResInc;
-				const int32 FirstIndexLowResY = FirstIndexY / LowResInc;
+				const int32 IndexLowResX = CellIdx[0] / LowResInc;
+				const int32 IndexLowResY = CellIdx[1] / LowResInc;
 
-				const int32 Index = FirstIndexLowResY * (NumColsLowRes - 1) + FirstIndexLowResX + FirstIndexLowResY;
+				const int32 Index = IndexLowResY * (NumColsLowRes - 1) + IndexLowResX + IndexLowResY;
 				check(Index < LowResolutionHeights.Num());
 
 				FRealSingle MinHeight = LowResolutionHeights[Index].Min;
@@ -468,8 +482,8 @@ namespace Chaos
 				MinHeight = MinValueSingle + MinHeight * HeightPerUnitSingle;
 				MaxHeight = MinValueSingle + MaxHeight * HeightPerUnitSingle;
 
-				OutBounds = FAABBVectorized(MakeVectorRegisterFloat(static_cast<FRealSingle>(FirstIndexLowResX*LowResInc), static_cast<FRealSingle>(FirstIndexLowResY*LowResInc), MinHeight, 0.0f),
-					MakeVectorRegisterFloat(static_cast<FRealSingle>((FirstIndexLowResX+1) * LowResInc), static_cast<FRealSingle>((FirstIndexLowResY + 1) * LowResInc), MaxHeight, 0.0f));
+				OutBounds = FAABBVectorized(MakeVectorRegisterFloat(static_cast<FRealSingle>(IndexLowResX *LowResInc), static_cast<FRealSingle>(IndexLowResY *LowResInc), MinHeight, 0.0f),
+					MakeVectorRegisterFloat(static_cast<FRealSingle>((IndexLowResX +1) * LowResInc), static_cast<FRealSingle>((IndexLowResY + 1) * LowResInc), MaxHeight, 0.0f));
 			}
 
 			FORCEINLINE void GetPointsScaled(int32 Index, FVec3 OutPts[4]) const
@@ -511,6 +525,18 @@ namespace Chaos
 				OutBounds = FAABBVectorized(Min, Max);
 			}
 
+			FORCEINLINE void GetBoundsScaledSimd(int32 Index, FAABBVectorized& OutBounds) const
+			{
+				GetBoundsSimd(Index, OutBounds);
+
+				VectorRegister4Float P0 = VectorMultiply(OutBounds.GetMin(), ScaleSimd);
+				VectorRegister4Float P1 = VectorMultiply(OutBounds.GetMax(), ScaleSimd);
+				VectorRegister4Float Min = VectorMin(P0, P1);
+				VectorRegister4Float Max = VectorMax(P0, P1);
+
+				OutBounds = FAABBVectorized(Min, Max);
+			}
+
 			FORCEINLINE void GetBoundsScaled(TVec2<int32> CellIdx, TVec2<int32> Area, FAABBVectorized& OutBounds) const
 			{
 				GetBounds(CellIdx, Area, OutBounds);
@@ -523,6 +549,18 @@ namespace Chaos
 				OutBounds = FAABBVectorized(Min, Max);
 			}
 
+			FORCEINLINE void GetBoundsScaled(TVec2<int32> CellIdx, FAABBVectorized& OutBounds) const
+			{
+				GetBounds(CellIdx, OutBounds);
+
+				VectorRegister4Float P0 = VectorMultiply(OutBounds.GetMin(), ScaleSimd);
+				VectorRegister4Float P1 = VectorMultiply(OutBounds.GetMax(), ScaleSimd);
+				VectorRegister4Float Min = VectorMin(P0, P1);
+				VectorRegister4Float Max = VectorMax(P0, P1);
+
+				OutBounds = FAABBVectorized(Min, Max);
+			}
+
 			FORCEINLINE void GetLowResBoundsScaled(TVec2<int32> CellIdx, FAABBVectorized& OutBounds) const
 			{
 				GetLowResBounds(CellIdx, OutBounds);
@@ -601,7 +639,12 @@ namespace Chaos
 				}
 
 				Ar.UsingCustomVersion(FUE5MainStreamObjectVersion::GUID);
-				if (Ar.CustomVer(FUE5MainStreamObjectVersion::GUID) >= FUE5MainStreamObjectVersion::AddLowResolutionHeightField)
+				if (Ar.CustomVer(FUE5MainStreamObjectVersion::GUID) == FUE5MainStreamObjectVersion::AddLowResolutionHeightField)
+				{
+					Ar << LowResolutionHeights;
+					Ar << NumColsLowRes;
+				}
+				else if (Ar.CustomVer(FUE5MainStreamObjectVersion::GUID) >= FUE5MainStreamObjectVersion::DecreaseLowResolutionHeightField)
 				{
 					Ar << LowResolutionHeights;
 					Ar << NumColsLowRes;
@@ -881,6 +924,19 @@ namespace Chaos
 		// Query functions - sweep, ray, overlap
 		template<typename SQVisitor>
 		bool GridSweep(const FVec3& StartPoint, const FVec3& Dir, const FReal Length, const FVec3 InHalfExtents, SQVisitor& Visitor) const;
+		
+		bool WalkSlow(TVec2<int32>& CellIdx, FHeightfieldRaycastVisitor& Visitor, FReal CurrentLength, const VectorRegister4Float& CurrentLengthSimd,
+			const FVec2& ScaledMin, FReal ZMidPoint, const FVec3& Dir, const FVec3& InvDir, bool bParallel[3], const FVec2& ScaledDx2D, FVec3& NextStart,
+			const FVec3& ScaleSign, const FVec3& ScaledDx, int32 IndexLowResX, int32 IndexLowResY) const;
+		
+		bool WalkFast(TVec2<int32>& CellIdx, FHeightfieldRaycastVisitor& Visitor, FReal CurrentLength, const VectorRegister4Float& CurrentLengthSimd,
+			const FVec2& ScaledMin, FReal ZMidPoint, const FVec3& Dir, const FVec3& InvDir, bool bParallel[3], const FVec2& ScaledDx2D, FVec3& NextStart,
+			const FVec3& ScaleSign, const FVec3& ScaledDx, const FVec2& Scale2D, const FVec3& DirScaled) const;
+
+		bool WalkOnLowRes(TVec2<int32>& CellIdx, FHeightfieldRaycastVisitor& Visitor, FReal CurrentLength, const VectorRegister4Float& CurrentLengthSimd,
+			const FVec2& ScaledMin, FReal ZMidPoint, const FVec3& Dir, const FVec3& InvDir, bool bParallel[3], const FVec2& ScaledDx2D, FVec3& NextStart,
+			const FVec3& ScaleSign, const FVec3& ScaledDx, const FVec2& Scale2D, const FVec3& DirScaled) const;
+
 		bool GridCast(const FVec3& StartPoint, const FVec3& Dir, const FReal Length, FHeightfieldRaycastVisitor& Visitor) const;
 		bool GetGridIntersections(FBounds2D InFlatBounds, TArray<TVec2<int32>>& OutInterssctions) const;
 		bool GetGridIntersectionsBatch(FBounds2D InFlatBounds, TArray<TVec2<int32>>& OutIntersections, const FAABBVectorized& Bounds) const;
diff --git a/Engine/Source/Runtime/Experimental/ChaosCore/Public/Chaos/Core.h b/Engine/Source/Runtime/Experimental/ChaosCore/Public/Chaos/Core.h
--- a/Engine/Source/Runtime/Experimental/ChaosCore/Public/Chaos/Core.h
+++ b/Engine/Source/Runtime/Experimental/ChaosCore/Public/Chaos/Core.h
@@ -51,7 +51,7 @@ namespace Chaos
 	// NOTE: if you get a merge conflict on the GUID, you must replace it with a new GUID - do not accept the source or target
 	// or you will likely get DDC version conflicts resulting in crashes during load.
 	// Core version string for Chaos data. Any DDC builder dependent on Chaos for serialization should depend on this version
-	static const TCHAR* const ChaosVersionGUID = TEXT("08C803C2-E2A7-4A85-AAA6-3353713AEE1D");
+	static const TCHAR* const ChaosVersionGUID = TEXT("8D7C3D39-DEF4-4775-85D1-80187F4A3D48");
 
 	inline FString GetChaosVersionStringInner()
 	{
