Interchange: Add the meshes build options for SM and SK. The default are the same has the legacy importer except for the remove degenerate which is set to false to keep the triangles connections by default.
#rnx

diff --git a/Engine/Plugins/Experimental/Interchange/Editor/Source/Pipelines/Private/InterchangeEditorPipelineDetails.cpp b/Engine/Plugins/Experimental/Interchange/Editor/Source/Pipelines/Private/InterchangeEditorPipelineDetails.cpp
--- a/Engine/Plugins/Experimental/Interchange/Editor/Source/Pipelines/Private/InterchangeEditorPipelineDetails.cpp
+++ b/Engine/Plugins/Experimental/Interchange/Editor/Source/Pipelines/Private/InterchangeEditorPipelineDetails.cpp
@@ -137,11 +137,12 @@ void FInterchangePipelineBaseDetailsCustomization::CustomizeDetails(IDetailLayou
 	CachedDetailBuilder->GetCategoryNames(AllCategoryNames);
 	TMap<FName, TArray<FName>> PropertiesPerCategorys;
 	InternalGetPipelineProperties(InterchangePipeline.Get(), AllCategoryNames, PropertiesPerCategorys);
-
-	TMap<FName, TMap<FName, TArray<FInternalPropertyData>>> SubCategoriesPropertiesPerMainCategory;
-
+	
 	for (const TPair<FName, TArray<FName>>& CategoryAndProperties : PropertiesPerCategorys)
 	{
+		//Category meta value Subgroup data
+		TMap<FName, IDetailGroup*> SubCategoryGroups;
+
 		const FName CategoryName = CategoryAndProperties.Key;
 		IDetailCategoryBuilder& Category = CachedDetailBuilder->EditCategory(CategoryName);
 		
@@ -169,9 +170,32 @@ void FInterchangePipelineBaseDetailsCustomization::CustomizeDetails(IDetailLayou
 			}
 			FName PropertyPath = FName(PropertyPtr->GetPathName());
 			CachedDetailBuilder->HideProperty(PropertyHandle);
+
+			const FName SubCategoryData = FName(PropertyHandle->GetMetaData(TEXT("SubCategory")));
+			IDetailGroup* GroupPtr = nullptr;
+			auto GetGroupPtr = [&GroupPtr, &SubCategoryData, &SubCategoryGroups, &Category]()
+			{
+				if (SubCategoryData != NAME_None)
+				{
+					if (!SubCategoryGroups.Contains(SubCategoryData))
+					{
+						SubCategoryGroups.Add(SubCategoryData, &(Category.AddGroup(SubCategoryData, FText::FromName(SubCategoryData))));
+					}
+					GroupPtr = SubCategoryGroups.FindChecked(SubCategoryData);
+				}
+			};
+
+			//This is the import/re-import dialog mode
 			if (!bAllowPropertyStatesEdition)
 			{
-				const FName SubCategoryData = FName(PropertyHandle->GetMetaData(TEXT("SubCategory")));
+				if (bIsReimportContext && PropertyHandle->GetBoolMetaData(FName("AdjustPipelineAndRefreshDetailOnChange")))
+				{
+					PropertyHandle->SetOnPropertyValueChanged(FSimpleDelegate::CreateLambda([this]() {
+						InterchangePipeline->AdjustSettingsFromCache();
+						RefreshCustomDetail();
+					}));
+				}
+
 				bool IsLocked = false;
 				if (const FInterchangePipelinePropertyStates* PropertyStates = InterchangePipeline->GetPropertyStates(PropertyPath))
 				{
@@ -181,28 +205,20 @@ void FInterchangePipelineBaseDetailsCustomization::CustomizeDetails(IDetailLayou
 					}
 					IsLocked = PropertyStates->IsPropertyLocked();
 				}
-				if (SubCategoryData != NAME_None)
+				GetGroupPtr();
+				IDetailPropertyRow& PropertyRow = GroupPtr ? GroupPtr->AddPropertyRow(PropertyHandle) : Category.AddProperty(PropertyHandle);
+				//When we use the pipeline in interchange
+				if (IsLocked)
 				{
-					TMap<FName, TArray<FInternalPropertyData>>& SubCategoriesProperties = SubCategoriesPropertiesPerMainCategory.FindOrAdd(CategoryName);
-					TArray<FInternalPropertyData>& SubCategoryProperties = SubCategoriesProperties.FindOrAdd(SubCategoryData);
-					FInternalPropertyData& PropertyData = SubCategoryProperties.AddDefaulted_GetRef();
-					PropertyData.PropertyHandle = PropertyHandle;
-					PropertyData.bReadOnly = IsLocked;
-				}
-				else
-				{
-					IDetailPropertyRow& PropertyRow = Category.AddProperty(PropertyHandle);
-					//When we use the pipeline in interchange
-					if (IsLocked)
-					{
-						LockPropertyHandleRow(PropertyHandle, PropertyRow);
-					}
+					LockPropertyHandleRow(PropertyHandle, PropertyRow);
 				}
 			}
+			//This is the asset editor mode, we allow users to set default value and locks
 			else
 			{
 				constexpr bool bShowChildren = true;
-				IDetailPropertyRow& PropertyRow = Category.AddProperty(PropertyHandle);
+				GetGroupPtr();
+				IDetailPropertyRow& PropertyRow = GroupPtr ? GroupPtr->AddPropertyRow(PropertyHandle) : Category.AddProperty(PropertyHandle);
 				TSharedPtr<SWidget> NameWidget;
 				TSharedPtr<SWidget> ValueWidget;
 				FDetailWidgetRow Row;
@@ -383,8 +399,6 @@ void FInterchangePipelineBaseDetailsCustomization::CustomizeDetails(IDetailLayou
 			}
 		}
 	}
-
-	AddSubCategory(DetailBuilder, SubCategoriesPropertiesPerMainCategory);
 }
 
 FInterchangeBaseNodeDetailsCustomization::FInterchangeBaseNodeDetailsCustomization()
diff --git a/Engine/Plugins/Experimental/Interchange/Runtime/Source/FactoryNodes/Private/InterchangeMeshFactoryNode.cpp b/Engine/Plugins/Experimental/Interchange/Runtime/Source/FactoryNodes/Private/InterchangeMeshFactoryNode.cpp
--- a/Engine/Plugins/Experimental/Interchange/Runtime/Source/FactoryNodes/Private/InterchangeMeshFactoryNode.cpp
+++ b/Engine/Plugins/Experimental/Interchange/Runtime/Source/FactoryNodes/Private/InterchangeMeshFactoryNode.cpp
@@ -1,6 +1,62 @@
 // Copyright Epic Games, Inc. All Rights Reserved.
 
 #include "InterchangeMeshFactoryNode.h"
+#include "Engine/StaticMesh.h"
+#include "Engine/SkeletalMesh.h"
+
+//The static mesh source model which contain the build settings are editor only
+#if WITH_EDITOR
+
+#define IMPLEMENT_MESH_BUILD_VALUE_TO_ASSET(AttributeName, AttributeType, PropertyName)	\
+AttributeType ValueData;															\
+if (GetAttribute<AttributeType>(Macro_Custom##AttributeName##Key.Key, ValueData))	\
+{																					\
+	if (USkeletalMesh* SkeletalMesh = Cast<USkeletalMesh>(Asset))					\
+	{																				\
+		if (FSkeletalMeshLODInfo* LodInfo = SkeletalMesh->GetLODInfo(0))			\
+		{																			\
+			LodInfo->BuildSettings.PropertyName = ValueData;						\
+			return true;															\
+		}																			\
+	}																				\
+	else if (UStaticMesh* StaticMesh = Cast<UStaticMesh>(Asset))					\
+	{																				\
+		if (StaticMesh->GetNumSourceModels() > 0)									\
+		{																			\
+			StaticMesh->GetSourceModel(0).BuildSettings.PropertyName = ValueData;	\
+			return true;															\
+		}																			\
+	}																				\
+}																					\
+return false;
+
+#define IMPLEMENT_MESH_BUILD_ASSET_TO_VALUE(AttributeName, PropertyName)	\
+if (UStaticMesh* StaticMesh = Cast<UStaticMesh>(Asset))								\
+{																					\
+	if (StaticMesh->GetNumSourceModels() > 0)										\
+	{																				\
+		return SetAttribute(Macro_Custom##AttributeName##Key.Key, StaticMesh->GetSourceModel(0).BuildSettings.PropertyName);	\
+	}																				\
+}																					\
+else if (USkeletalMesh* SkeletalMesh = Cast<USkeletalMesh>(Asset))					\
+{																					\
+	if (FSkeletalMeshLODInfo* LodInfo = SkeletalMesh->GetLODInfo(0))				\
+	{																				\
+		return SetAttribute(Macro_Custom##AttributeName##Key.Key, LodInfo->BuildSettings.PropertyName);	\
+	}																				\
+}																					\
+return false;
+
+#else //WITH_EDITOR
+
+#define IMPLEMENT_MESH_BUILD_VALUE_TO_ASSET(AttributeName, AttributeType, PropertyName)	\
+return false;
+
+#define IMPLEMENT_MESH_BUILD_ASSET_TO_VALUE(AttributeName, PropertyName)	\
+return false;
+#endif // else WITH_EDITOR
+
+
 
 namespace UE::Interchange
 {
@@ -15,6 +71,7 @@ namespace UE::Interchange
 		static FAttributeKey SlotMaterialDependency_BaseKey(TEXT("__SlotMaterialDependency__"));
 		return SlotMaterialDependency_BaseKey;
 	}
+
 } // namespace UE::Interchange
 
 
@@ -116,3 +173,162 @@ bool UInterchangeMeshFactoryNode::RemoveSlotMaterialDependencyUid(const FString&
 {
 	return SlotMaterialDependencies.RemoveKey(SlotName);
 }
+
+bool UInterchangeMeshFactoryNode::GetCustomRecomputeNormals(bool& AttributeValue) const
+{
+	IMPLEMENT_NODE_ATTRIBUTE_GETTER(RecomputeNormals, bool)
+}
+
+bool UInterchangeMeshFactoryNode::SetCustomRecomputeNormals(const bool& AttributeValue, bool bAddApplyDelegate)
+{
+	IMPLEMENT_NODE_ATTRIBUTE_SETTER_WITH_CUSTOM_DELEGATE(UInterchangeMeshFactoryNode, RecomputeNormals, bool);
+}
+
+bool UInterchangeMeshFactoryNode::ApplyCustomRecomputeNormalsToAsset(UObject* Asset) const
+{
+	IMPLEMENT_MESH_BUILD_VALUE_TO_ASSET(RecomputeNormals, bool, bRecomputeNormals);
+}
+
+bool UInterchangeMeshFactoryNode::FillCustomRecomputeNormalsFromAsset(UObject* Asset)
+{
+	IMPLEMENT_MESH_BUILD_ASSET_TO_VALUE(RecomputeNormals, bRecomputeNormals);
+}
+
+bool UInterchangeMeshFactoryNode::GetCustomRecomputeTangents(bool& AttributeValue) const
+{
+	IMPLEMENT_NODE_ATTRIBUTE_GETTER(RecomputeTangents, bool)
+}
+
+bool UInterchangeMeshFactoryNode::SetCustomRecomputeTangents(const bool& AttributeValue, bool bAddApplyDelegate)
+{
+	IMPLEMENT_NODE_ATTRIBUTE_SETTER_WITH_CUSTOM_DELEGATE(UInterchangeMeshFactoryNode, RecomputeTangents, bool);
+}
+
+bool UInterchangeMeshFactoryNode::ApplyCustomRecomputeTangentsToAsset(UObject* Asset) const
+{
+	IMPLEMENT_MESH_BUILD_VALUE_TO_ASSET(RecomputeTangents, bool, bRecomputeTangents);
+}
+
+bool UInterchangeMeshFactoryNode::FillCustomRecomputeTangentsFromAsset(UObject* Asset)
+{
+	IMPLEMENT_MESH_BUILD_ASSET_TO_VALUE(RecomputeTangents, bRecomputeTangents);
+}
+
+bool UInterchangeMeshFactoryNode::GetCustomUseMikkTSpace(bool& AttributeValue) const
+{
+	IMPLEMENT_NODE_ATTRIBUTE_GETTER(UseMikkTSpace, bool)
+}
+
+bool UInterchangeMeshFactoryNode::SetCustomUseMikkTSpace(const bool& AttributeValue, bool bAddApplyDelegate)
+{
+	IMPLEMENT_NODE_ATTRIBUTE_SETTER_WITH_CUSTOM_DELEGATE(UInterchangeMeshFactoryNode, UseMikkTSpace, bool);
+}
+
+bool UInterchangeMeshFactoryNode::ApplyCustomUseMikkTSpaceToAsset(UObject* Asset) const
+{
+	IMPLEMENT_MESH_BUILD_VALUE_TO_ASSET(UseMikkTSpace, bool, bUseMikkTSpace);
+}
+
+bool UInterchangeMeshFactoryNode::FillCustomUseMikkTSpaceFromAsset(UObject* Asset)
+{
+	IMPLEMENT_MESH_BUILD_ASSET_TO_VALUE(UseMikkTSpace, bUseMikkTSpace);
+}
+
+bool UInterchangeMeshFactoryNode::GetCustomComputeWeightedNormals(bool& AttributeValue) const
+{
+	IMPLEMENT_NODE_ATTRIBUTE_GETTER(ComputeWeightedNormals, bool)
+}
+
+bool UInterchangeMeshFactoryNode::SetCustomComputeWeightedNormals(const bool& AttributeValue, bool bAddApplyDelegate)
+{
+	IMPLEMENT_NODE_ATTRIBUTE_SETTER_WITH_CUSTOM_DELEGATE(UInterchangeMeshFactoryNode, ComputeWeightedNormals, bool);
+}
+
+bool UInterchangeMeshFactoryNode::ApplyCustomComputeWeightedNormalsToAsset(UObject* Asset) const
+{
+	IMPLEMENT_MESH_BUILD_VALUE_TO_ASSET(ComputeWeightedNormals, bool, bComputeWeightedNormals);
+}
+
+bool UInterchangeMeshFactoryNode::FillCustomComputeWeightedNormalsFromAsset(UObject* Asset)
+{
+	IMPLEMENT_MESH_BUILD_ASSET_TO_VALUE(ComputeWeightedNormals, bComputeWeightedNormals);
+}
+
+bool UInterchangeMeshFactoryNode::GetCustomUseHighPrecisionTangentBasis(bool& AttributeValue) const
+{
+	IMPLEMENT_NODE_ATTRIBUTE_GETTER(UseHighPrecisionTangentBasis, bool)
+}
+
+bool UInterchangeMeshFactoryNode::SetCustomUseHighPrecisionTangentBasis(const bool& AttributeValue, bool bAddApplyDelegate)
+{
+	IMPLEMENT_NODE_ATTRIBUTE_SETTER_WITH_CUSTOM_DELEGATE(UInterchangeMeshFactoryNode, UseHighPrecisionTangentBasis, bool);
+}
+
+bool UInterchangeMeshFactoryNode::ApplyCustomUseHighPrecisionTangentBasisToAsset(UObject* Asset) const
+{
+	IMPLEMENT_MESH_BUILD_VALUE_TO_ASSET(UseHighPrecisionTangentBasis, bool, bUseHighPrecisionTangentBasis);
+}
+
+bool UInterchangeMeshFactoryNode::FillCustomUseHighPrecisionTangentBasisFromAsset(UObject* Asset)
+{
+	IMPLEMENT_MESH_BUILD_ASSET_TO_VALUE(UseHighPrecisionTangentBasis, bUseHighPrecisionTangentBasis);
+}
+
+bool UInterchangeMeshFactoryNode::GetCustomUseFullPrecisionUVs(bool& AttributeValue) const
+{
+	IMPLEMENT_NODE_ATTRIBUTE_GETTER(UseFullPrecisionUVs, bool)
+}
+
+bool UInterchangeMeshFactoryNode::SetCustomUseFullPrecisionUVs(const bool& AttributeValue, bool bAddApplyDelegate)
+{
+	IMPLEMENT_NODE_ATTRIBUTE_SETTER_WITH_CUSTOM_DELEGATE(UInterchangeMeshFactoryNode, UseFullPrecisionUVs, bool);
+}
+
+bool UInterchangeMeshFactoryNode::ApplyCustomUseFullPrecisionUVsToAsset(UObject* Asset) const
+{
+	IMPLEMENT_MESH_BUILD_VALUE_TO_ASSET(UseFullPrecisionUVs, bool, bUseFullPrecisionUVs);
+}
+
+bool UInterchangeMeshFactoryNode::FillCustomUseFullPrecisionUVsFromAsset(UObject* Asset)
+{
+	IMPLEMENT_MESH_BUILD_ASSET_TO_VALUE(UseFullPrecisionUVs, bUseFullPrecisionUVs);
+}
+
+bool UInterchangeMeshFactoryNode::GetCustomUseBackwardsCompatibleF16TruncUVs(bool& AttributeValue) const
+{
+	IMPLEMENT_NODE_ATTRIBUTE_GETTER(UseBackwardsCompatibleF16TruncUVs, bool)
+}
+
+bool UInterchangeMeshFactoryNode::SetCustomUseBackwardsCompatibleF16TruncUVs(const bool& AttributeValue, bool bAddApplyDelegate)
+{
+	IMPLEMENT_NODE_ATTRIBUTE_SETTER_WITH_CUSTOM_DELEGATE(UInterchangeMeshFactoryNode, UseBackwardsCompatibleF16TruncUVs, bool);
+}
+
+bool UInterchangeMeshFactoryNode::ApplyCustomUseBackwardsCompatibleF16TruncUVsToAsset(UObject* Asset) const
+{
+	IMPLEMENT_MESH_BUILD_VALUE_TO_ASSET(UseBackwardsCompatibleF16TruncUVs, bool, bUseBackwardsCompatibleF16TruncUVs);
+}
+bool UInterchangeMeshFactoryNode::FillCustomUseBackwardsCompatibleF16TruncUVsFromAsset(UObject* Asset)
+{
+	IMPLEMENT_MESH_BUILD_ASSET_TO_VALUE(UseBackwardsCompatibleF16TruncUVs, bUseBackwardsCompatibleF16TruncUVs);
+}
+
+bool UInterchangeMeshFactoryNode::GetCustomRemoveDegenerates(bool& AttributeValue) const
+{
+	IMPLEMENT_NODE_ATTRIBUTE_GETTER(RemoveDegenerates, bool)
+}
+
+bool UInterchangeMeshFactoryNode::SetCustomRemoveDegenerates(const bool& AttributeValue, bool bAddApplyDelegate)
+{
+	IMPLEMENT_NODE_ATTRIBUTE_SETTER_WITH_CUSTOM_DELEGATE(UInterchangeMeshFactoryNode, RemoveDegenerates, bool);
+}
+
+bool UInterchangeMeshFactoryNode::ApplyCustomRemoveDegeneratesToAsset(UObject* Asset) const
+{
+	IMPLEMENT_MESH_BUILD_VALUE_TO_ASSET(RemoveDegenerates, bool, bRemoveDegenerates);
+}
+
+bool UInterchangeMeshFactoryNode::FillCustomRemoveDegeneratesFromAsset(UObject* Asset)
+{
+	IMPLEMENT_MESH_BUILD_ASSET_TO_VALUE(RemoveDegenerates, bRemoveDegenerates);
+}
\ No newline at end of file
diff --git a/Engine/Plugins/Experimental/Interchange/Runtime/Source/FactoryNodes/Private/InterchangeSkeletalMeshFactoryNode.cpp b/Engine/Plugins/Experimental/Interchange/Runtime/Source/FactoryNodes/Private/InterchangeSkeletalMeshFactoryNode.cpp
--- a/Engine/Plugins/Experimental/Interchange/Runtime/Source/FactoryNodes/Private/InterchangeSkeletalMeshFactoryNode.cpp
+++ b/Engine/Plugins/Experimental/Interchange/Runtime/Source/FactoryNodes/Private/InterchangeSkeletalMeshFactoryNode.cpp
@@ -2,15 +2,48 @@
 
 #include "InterchangeSkeletalMeshFactoryNode.h"
 
-#if WITH_ENGINE
 #include "Engine/SkeletalMesh.h"
-#endif
+
+#if WITH_EDITOR
+
+#define IMPLEMENT_SKELETAL_BUILD_VALUE_TO_ASSET(AttributeName, AttributeType, PropertyName)	\
+AttributeType ValueData;															\
+if (GetAttribute<AttributeType>(Macro_Custom##AttributeName##Key.Key, ValueData))	\
+{																					\
+	if (USkeletalMesh* SkeletalMesh = Cast<USkeletalMesh>(Asset))					\
+	{																				\
+		if (FSkeletalMeshLODInfo* LodInfo = SkeletalMesh->GetLODInfo(0))			\
+		{																			\
+			LodInfo->BuildSettings.PropertyName = ValueData;						\
+			return true;															\
+		}																			\
+	}																				\
+}																					\
+return false;
+
+#define IMPLEMENT_SKELETALMESH_BUILD_ASSET_TO_VALUE(AttributeName, PropertyName)					\
+if (USkeletalMesh* SkeletalMesh = Cast<USkeletalMesh>(Asset))						\
+{																					\
+	if (FSkeletalMeshLODInfo* LodInfo = SkeletalMesh->GetLODInfo(0))				\
+	{																				\
+		return SetAttribute(Macro_Custom##AttributeName##Key.Key, LodInfo->BuildSettings.PropertyName);	\
+	}																				\
+}																					\
+return false;
+
+#else //WITH_EDITOR
+
+#define IMPLEMENT_SKELETAL_BUILD_VALUE_TO_ASSET(AttributeName, AttributeType, PropertyName)	\
+return false;
+
+#define IMPLEMENT_SKELETALMESH_BUILD_ASSET_TO_VALUE(AttributeName, PropertyName)					\
+return false;
+
+#endif //else WITH_EDIOTR
 
 UInterchangeSkeletalMeshFactoryNode::UInterchangeSkeletalMeshFactoryNode()
 {
-#if WITH_ENGINE
 	AssetClass = nullptr;
-#endif
 }
 
 void UInterchangeSkeletalMeshFactoryNode::InitializeSkeletalMeshNode(const FString& UniqueID, const FString& DisplayLabel, const FString& InAssetClass)
@@ -32,11 +65,7 @@ FString UInterchangeSkeletalMeshFactoryNode::GetTypeName() const
 UClass* UInterchangeSkeletalMeshFactoryNode::GetObjectClass() const
 {
 	ensure(bIsNodeClassInitialized);
-#if WITH_ENGINE
 	return AssetClass.Get() != nullptr ? AssetClass.Get() : USkeletalMesh::StaticClass();
-#else
-	return nullptr;
-#endif
 }
 
 bool UInterchangeSkeletalMeshFactoryNode::GetCustomSkeletonSoftObjectPath(FSoftObjectPath& AttributeValue) const
@@ -79,6 +108,74 @@ bool UInterchangeSkeletalMeshFactoryNode::SetCustomPhysicAssetSoftObjectPath(con
 	IMPLEMENT_NODE_ATTRIBUTE_SETTER_NODELEGATE(PhysicAssetSoftObjectPath, FSoftObjectPath)
 }
 
+bool UInterchangeSkeletalMeshFactoryNode::GetCustomThresholdPosition(float& AttributeValue) const
+{
+	IMPLEMENT_NODE_ATTRIBUTE_GETTER(ThresholdPosition, float)
+}
+bool UInterchangeSkeletalMeshFactoryNode::SetCustomThresholdPosition(const float& AttributeValue, bool bAddApplyDelegate)
+{
+	IMPLEMENT_NODE_ATTRIBUTE_SETTER_WITH_CUSTOM_DELEGATE(UInterchangeSkeletalMeshFactoryNode, ThresholdPosition, float);
+}
+bool UInterchangeSkeletalMeshFactoryNode::ApplyCustomThresholdPositionToAsset(UObject* Asset) const
+{
+	IMPLEMENT_SKELETAL_BUILD_VALUE_TO_ASSET(ThresholdPosition, float, ThresholdPosition);
+}
+bool UInterchangeSkeletalMeshFactoryNode::FillCustomThresholdPositionFromAsset(UObject* Asset)
+{
+	IMPLEMENT_SKELETALMESH_BUILD_ASSET_TO_VALUE(ThresholdPosition, ThresholdPosition);
+}
+
+bool UInterchangeSkeletalMeshFactoryNode::GetCustomThresholdTangentNormal(float& AttributeValue) const
+{
+	IMPLEMENT_NODE_ATTRIBUTE_GETTER(ThresholdTangentNormal, float)
+}
+bool UInterchangeSkeletalMeshFactoryNode::SetCustomThresholdTangentNormal(const float& AttributeValue, bool bAddApplyDelegate)
+{
+	IMPLEMENT_NODE_ATTRIBUTE_SETTER_WITH_CUSTOM_DELEGATE(UInterchangeSkeletalMeshFactoryNode, ThresholdTangentNormal, float);
+}
+bool UInterchangeSkeletalMeshFactoryNode::ApplyCustomThresholdTangentNormalToAsset(UObject* Asset) const
+{
+	IMPLEMENT_SKELETAL_BUILD_VALUE_TO_ASSET(ThresholdTangentNormal, float, ThresholdTangentNormal);
+}
+bool UInterchangeSkeletalMeshFactoryNode::FillCustomThresholdTangentNormalFromAsset(UObject* Asset)
+{
+	IMPLEMENT_SKELETALMESH_BUILD_ASSET_TO_VALUE(ThresholdTangentNormal, ThresholdTangentNormal);
+}
+
+bool UInterchangeSkeletalMeshFactoryNode::GetCustomThresholdUV(float& AttributeValue) const
+{
+	IMPLEMENT_NODE_ATTRIBUTE_GETTER(ThresholdUV, float)
+}
+bool UInterchangeSkeletalMeshFactoryNode::SetCustomThresholdUV(const float& AttributeValue, bool bAddApplyDelegate)
+{
+	IMPLEMENT_NODE_ATTRIBUTE_SETTER_WITH_CUSTOM_DELEGATE(UInterchangeSkeletalMeshFactoryNode, ThresholdUV, float);
+}
+bool UInterchangeSkeletalMeshFactoryNode::ApplyCustomThresholdUVToAsset(UObject* Asset) const
+{
+	IMPLEMENT_SKELETAL_BUILD_VALUE_TO_ASSET(ThresholdUV, float, ThresholdUV);
+}
+bool UInterchangeSkeletalMeshFactoryNode::FillCustomThresholdUVFromAsset(UObject* Asset)
+{
+	IMPLEMENT_SKELETALMESH_BUILD_ASSET_TO_VALUE(ThresholdUV, ThresholdUV);
+}
+
+bool UInterchangeSkeletalMeshFactoryNode::GetCustomMorphThresholdPosition(float& AttributeValue) const
+{
+	IMPLEMENT_NODE_ATTRIBUTE_GETTER(MorphThresholdPosition, float)
+}
+bool UInterchangeSkeletalMeshFactoryNode::SetCustomMorphThresholdPosition(const float& AttributeValue, bool bAddApplyDelegate)
+{
+	IMPLEMENT_NODE_ATTRIBUTE_SETTER_WITH_CUSTOM_DELEGATE(UInterchangeSkeletalMeshFactoryNode, MorphThresholdPosition, float);
+}
+bool UInterchangeSkeletalMeshFactoryNode::ApplyCustomMorphThresholdPositionToAsset(UObject* Asset) const
+{
+	IMPLEMENT_SKELETAL_BUILD_VALUE_TO_ASSET(MorphThresholdPosition, float, MorphThresholdPosition);
+}
+bool UInterchangeSkeletalMeshFactoryNode::FillCustomMorphThresholdPositionFromAsset(UObject* Asset)
+{
+	IMPLEMENT_SKELETALMESH_BUILD_ASSET_TO_VALUE(MorphThresholdPosition, MorphThresholdPosition);
+}
+
 bool UInterchangeSkeletalMeshFactoryNode::GetCustomImportContentType(EInterchangeSkeletalMeshContentType& AttributeValue) const
 {
 	FString OperationName = GetTypeName() + TEXT(".GetImportContentType");
diff --git a/Engine/Plugins/Experimental/Interchange/Runtime/Source/FactoryNodes/Private/InterchangeStaticMeshFactoryNode.cpp b/Engine/Plugins/Experimental/Interchange/Runtime/Source/FactoryNodes/Private/InterchangeStaticMeshFactoryNode.cpp
--- a/Engine/Plugins/Experimental/Interchange/Runtime/Source/FactoryNodes/Private/InterchangeStaticMeshFactoryNode.cpp
+++ b/Engine/Plugins/Experimental/Interchange/Runtime/Source/FactoryNodes/Private/InterchangeStaticMeshFactoryNode.cpp
@@ -2,9 +2,44 @@
 
 #include "InterchangeStaticMeshFactoryNode.h"
 
-#if WITH_ENGINE
 #include "Engine/StaticMesh.h"
-#endif
+
+#if WITH_EDITOR
+
+#define IMPLEMENT_STATICMESH_BUILD_VALUE_TO_ASSET(AttributeName, AttributeType, PropertyName)	\
+AttributeType ValueData;															\
+if (GetAttribute<AttributeType>(Macro_Custom##AttributeName##Key.Key, ValueData))	\
+{																					\
+	if (UStaticMesh* StaticMesh = Cast<UStaticMesh>(Asset))							\
+	{																				\
+		if (StaticMesh->GetNumSourceModels() > 0)									\
+		{																			\
+			StaticMesh->GetSourceModel(0).BuildSettings.PropertyName = ValueData;	\
+			return true;															\
+		}																			\
+	}																				\
+}																					\
+return false;
+
+#define IMPLEMENT_STATICMESH_BUILD_ASSET_TO_VALUE(AttributeName, PropertyName)	\
+if (UStaticMesh* StaticMesh = Cast<UStaticMesh>(Asset))								\
+{																					\
+	if (StaticMesh->GetNumSourceModels() > 0)										\
+	{																				\
+		return SetAttribute(Macro_Custom##AttributeName##Key.Key, StaticMesh->GetSourceModel(0).BuildSettings.PropertyName);	\
+	}																				\
+}																					\
+return false;
+
+#else //WITH_EDITOR
+
+#define IMPLEMENT_STATICMESH_BUILD_VALUE_TO_ASSET(AttributeName, AttributeType, PropertyName)	\
+return false;
+
+#define IMPLEMENT_STATICMESH_BUILD_ASSET_TO_VALUE(AttributeName, PropertyName)	\
+return false;
+
+#endif // else WITH_EDITOR
 
 namespace UE
 {
@@ -60,7 +95,7 @@ bool UInterchangeStaticMeshFactoryNode::GetCustomBuildNanite(bool& AttributeValu
 
 bool UInterchangeStaticMeshFactoryNode::SetCustomBuildNanite(const bool& AttributeValue, bool bAddApplyDelegate /*= true*/)
 {
-	IMPLEMENT_NODE_ATTRIBUTE_SETTER_WITH_CUSTOM_DELEGATE(UInterchangeStaticMeshFactoryNode, BuildNanite, bool, UStaticMesh);
+	IMPLEMENT_NODE_ATTRIBUTE_SETTER_WITH_CUSTOM_DELEGATE(UInterchangeStaticMeshFactoryNode, BuildNanite, bool);
 }
 
 int32 UInterchangeStaticMeshFactoryNode::GetSocketUidCount() const
@@ -118,3 +153,215 @@ bool UInterchangeStaticMeshFactoryNode::SetNodeClassFromClassAttribute()
 	}
 	return bIsNodeClassInitialized;
 }
+
+bool UInterchangeStaticMeshFactoryNode::GetCustomBuildReversedIndexBuffer(bool& AttributeValue) const
+{
+	IMPLEMENT_NODE_ATTRIBUTE_GETTER(BuildReversedIndexBuffer, bool)
+}
+bool UInterchangeStaticMeshFactoryNode::SetCustomBuildReversedIndexBuffer(const bool& AttributeValue, bool bAddApplyDelegate)
+{
+	IMPLEMENT_NODE_ATTRIBUTE_SETTER_WITH_CUSTOM_DELEGATE(UInterchangeStaticMeshFactoryNode, BuildReversedIndexBuffer, bool);
+}
+bool UInterchangeStaticMeshFactoryNode::ApplyCustomBuildReversedIndexBufferToAsset(UObject* Asset) const
+{
+	IMPLEMENT_STATICMESH_BUILD_VALUE_TO_ASSET(BuildReversedIndexBuffer, bool, bBuildReversedIndexBuffer);
+}
+bool UInterchangeStaticMeshFactoryNode::FillCustomBuildReversedIndexBufferFromAsset(UObject* Asset)
+{
+	IMPLEMENT_STATICMESH_BUILD_ASSET_TO_VALUE(BuildReversedIndexBuffer, bBuildReversedIndexBuffer);
+}
+
+bool UInterchangeStaticMeshFactoryNode::GetCustomGenerateLightmapUVs(bool& AttributeValue) const
+{
+	IMPLEMENT_NODE_ATTRIBUTE_GETTER(GenerateLightmapUVs, bool)
+}
+bool UInterchangeStaticMeshFactoryNode::SetCustomGenerateLightmapUVs(const bool& AttributeValue, bool bAddApplyDelegate)
+{
+	IMPLEMENT_NODE_ATTRIBUTE_SETTER_WITH_CUSTOM_DELEGATE(UInterchangeStaticMeshFactoryNode, GenerateLightmapUVs, bool);
+}
+bool UInterchangeStaticMeshFactoryNode::ApplyCustomGenerateLightmapUVsToAsset(UObject* Asset) const
+{
+	IMPLEMENT_STATICMESH_BUILD_VALUE_TO_ASSET(GenerateLightmapUVs, bool, bGenerateLightmapUVs);
+}
+bool UInterchangeStaticMeshFactoryNode::FillCustomGenerateLightmapUVsFromAsset(UObject* Asset)
+{
+	IMPLEMENT_STATICMESH_BUILD_ASSET_TO_VALUE(GenerateLightmapUVs, bGenerateLightmapUVs);
+}
+
+bool UInterchangeStaticMeshFactoryNode::GetCustomGenerateDistanceFieldAsIfTwoSided(bool& AttributeValue) const
+{
+	IMPLEMENT_NODE_ATTRIBUTE_GETTER(GenerateDistanceFieldAsIfTwoSided, bool)
+}
+bool UInterchangeStaticMeshFactoryNode::SetCustomGenerateDistanceFieldAsIfTwoSided(const bool& AttributeValue, bool bAddApplyDelegate)
+{
+	IMPLEMENT_NODE_ATTRIBUTE_SETTER_WITH_CUSTOM_DELEGATE(UInterchangeStaticMeshFactoryNode, GenerateDistanceFieldAsIfTwoSided, bool);
+}
+bool UInterchangeStaticMeshFactoryNode::ApplyCustomGenerateDistanceFieldAsIfTwoSidedToAsset(UObject* Asset) const
+{
+	IMPLEMENT_STATICMESH_BUILD_VALUE_TO_ASSET(GenerateDistanceFieldAsIfTwoSided, bool, bGenerateDistanceFieldAsIfTwoSided);
+
+}
+bool UInterchangeStaticMeshFactoryNode::FillCustomGenerateDistanceFieldAsIfTwoSidedFromAsset(UObject* Asset)
+{
+	IMPLEMENT_STATICMESH_BUILD_ASSET_TO_VALUE(GenerateDistanceFieldAsIfTwoSided, bGenerateDistanceFieldAsIfTwoSided);
+}
+
+bool UInterchangeStaticMeshFactoryNode::GetCustomSupportFaceRemap(bool& AttributeValue) const
+{
+	IMPLEMENT_NODE_ATTRIBUTE_GETTER(SupportFaceRemap, bool)
+}
+bool UInterchangeStaticMeshFactoryNode::SetCustomSupportFaceRemap(const bool& AttributeValue, bool bAddApplyDelegate)
+{
+	IMPLEMENT_NODE_ATTRIBUTE_SETTER_WITH_CUSTOM_DELEGATE(UInterchangeStaticMeshFactoryNode, SupportFaceRemap, bool);
+}
+bool UInterchangeStaticMeshFactoryNode::ApplyCustomSupportFaceRemapToAsset(UObject* Asset) const
+{
+	IMPLEMENT_STATICMESH_BUILD_VALUE_TO_ASSET(SupportFaceRemap, bool, bSupportFaceRemap);
+}
+bool UInterchangeStaticMeshFactoryNode::FillCustomSupportFaceRemapFromAsset(UObject* Asset)
+{
+	IMPLEMENT_STATICMESH_BUILD_ASSET_TO_VALUE(SupportFaceRemap, bSupportFaceRemap);
+}
+
+bool UInterchangeStaticMeshFactoryNode::GetCustomMinLightmapResolution(int32& AttributeValue) const
+{
+	IMPLEMENT_NODE_ATTRIBUTE_GETTER(MinLightmapResolution, int32)
+}
+bool UInterchangeStaticMeshFactoryNode::SetCustomMinLightmapResolution(const int32& AttributeValue, bool bAddApplyDelegate)
+{
+	IMPLEMENT_NODE_ATTRIBUTE_SETTER_WITH_CUSTOM_DELEGATE(UInterchangeStaticMeshFactoryNode, MinLightmapResolution, int32);
+}
+bool UInterchangeStaticMeshFactoryNode::ApplyCustomMinLightmapResolutionToAsset(UObject* Asset) const
+{
+	IMPLEMENT_STATICMESH_BUILD_VALUE_TO_ASSET(MinLightmapResolution, int32, MinLightmapResolution);
+}
+bool UInterchangeStaticMeshFactoryNode::FillCustomMinLightmapResolutionFromAsset(UObject* Asset)
+{
+	IMPLEMENT_STATICMESH_BUILD_ASSET_TO_VALUE(MinLightmapResolution, MinLightmapResolution);
+}
+
+bool UInterchangeStaticMeshFactoryNode::GetCustomSrcLightmapIndex(int32& AttributeValue) const
+{
+	IMPLEMENT_NODE_ATTRIBUTE_GETTER(SrcLightmapIndex, int32)
+}
+bool UInterchangeStaticMeshFactoryNode::SetCustomSrcLightmapIndex(const int32& AttributeValue, bool bAddApplyDelegate)
+{
+	IMPLEMENT_NODE_ATTRIBUTE_SETTER_WITH_CUSTOM_DELEGATE(UInterchangeStaticMeshFactoryNode, SrcLightmapIndex, int32);
+}
+bool UInterchangeStaticMeshFactoryNode::ApplyCustomSrcLightmapIndexToAsset(UObject* Asset) const
+{
+	IMPLEMENT_STATICMESH_BUILD_VALUE_TO_ASSET(SrcLightmapIndex, int32, SrcLightmapIndex);
+}
+bool UInterchangeStaticMeshFactoryNode::FillCustomSrcLightmapIndexFromAsset(UObject* Asset)
+{
+	IMPLEMENT_STATICMESH_BUILD_ASSET_TO_VALUE(SrcLightmapIndex, SrcLightmapIndex);
+}
+
+bool UInterchangeStaticMeshFactoryNode::GetCustomDstLightmapIndex(int32& AttributeValue) const
+{
+	IMPLEMENT_NODE_ATTRIBUTE_GETTER(DstLightmapIndex, int32)
+}
+bool UInterchangeStaticMeshFactoryNode::SetCustomDstLightmapIndex(const int32& AttributeValue, bool bAddApplyDelegate)
+{
+	IMPLEMENT_NODE_ATTRIBUTE_SETTER_WITH_CUSTOM_DELEGATE(UInterchangeStaticMeshFactoryNode, DstLightmapIndex, int32);
+}
+bool UInterchangeStaticMeshFactoryNode::ApplyCustomDstLightmapIndexToAsset(UObject* Asset) const
+{
+	IMPLEMENT_STATICMESH_BUILD_VALUE_TO_ASSET(DstLightmapIndex, int32, DstLightmapIndex);
+}
+bool UInterchangeStaticMeshFactoryNode::FillCustomDstLightmapIndexFromAsset(UObject* Asset)
+{
+	IMPLEMENT_STATICMESH_BUILD_ASSET_TO_VALUE(DstLightmapIndex, DstLightmapIndex);
+}
+
+bool UInterchangeStaticMeshFactoryNode::GetCustomBuildScale3D(FVector& AttributeValue) const
+{
+	IMPLEMENT_NODE_ATTRIBUTE_GETTER(BuildScale3D, FVector)
+}
+bool UInterchangeStaticMeshFactoryNode::SetCustomBuildScale3D(const FVector& AttributeValue, bool bAddApplyDelegate)
+{
+	IMPLEMENT_NODE_ATTRIBUTE_SETTER_WITH_CUSTOM_DELEGATE(UInterchangeStaticMeshFactoryNode, BuildScale3D, FVector);
+}
+bool UInterchangeStaticMeshFactoryNode::ApplyCustomBuildScale3DToAsset(UObject* Asset) const
+{
+	IMPLEMENT_STATICMESH_BUILD_VALUE_TO_ASSET(BuildScale3D, FVector, BuildScale3D);
+}
+bool UInterchangeStaticMeshFactoryNode::FillCustomBuildScale3DFromAsset(UObject* Asset)
+{
+	IMPLEMENT_STATICMESH_BUILD_ASSET_TO_VALUE(BuildScale3D, BuildScale3D);
+}
+
+bool UInterchangeStaticMeshFactoryNode::GetCustomDistanceFieldResolutionScale(float& AttributeValue) const
+{
+	IMPLEMENT_NODE_ATTRIBUTE_GETTER(DistanceFieldResolutionScale, float)
+}
+bool UInterchangeStaticMeshFactoryNode::SetCustomDistanceFieldResolutionScale(const float& AttributeValue, bool bAddApplyDelegate)
+{
+	IMPLEMENT_NODE_ATTRIBUTE_SETTER_WITH_CUSTOM_DELEGATE(UInterchangeStaticMeshFactoryNode, DistanceFieldResolutionScale, float);
+}
+bool UInterchangeStaticMeshFactoryNode::ApplyCustomDistanceFieldResolutionScaleToAsset(UObject* Asset) const
+{
+	IMPLEMENT_STATICMESH_BUILD_VALUE_TO_ASSET(DistanceFieldResolutionScale, float, DistanceFieldResolutionScale);
+}
+bool UInterchangeStaticMeshFactoryNode::FillCustomDistanceFieldResolutionScaleFromAsset(UObject* Asset)
+{
+	IMPLEMENT_STATICMESH_BUILD_ASSET_TO_VALUE(DistanceFieldResolutionScale, DistanceFieldResolutionScale);
+}
+
+bool UInterchangeStaticMeshFactoryNode::GetCustomDistanceFieldReplacementMesh(FSoftObjectPath& AttributeValue) const
+{
+	IMPLEMENT_NODE_ATTRIBUTE_GETTER(DistanceFieldReplacementMesh, FSoftObjectPath)
+}
+bool UInterchangeStaticMeshFactoryNode::SetCustomDistanceFieldReplacementMesh(const FSoftObjectPath& AttributeValue, bool bAddApplyDelegate)
+{
+	IMPLEMENT_NODE_ATTRIBUTE_SETTER_WITH_CUSTOM_DELEGATE(UInterchangeStaticMeshFactoryNode, DistanceFieldReplacementMesh, FSoftObjectPath);
+}
+bool UInterchangeStaticMeshFactoryNode::ApplyCustomDistanceFieldReplacementMeshToAsset(UObject* Asset) const
+{
+#if WITH_EDITOR
+	FSoftObjectPath ValueData;
+	if (GetAttribute<FSoftObjectPath>(Macro_CustomDistanceFieldReplacementMeshKey.Key, ValueData))
+	{
+		if (UStaticMesh* StaticMesh = Cast<UStaticMesh>(Asset))
+		{
+			if (StaticMesh->GetNumSourceModels() > 0)
+			{
+				StaticMesh->GetSourceModel(0).BuildSettings.DistanceFieldReplacementMesh = Cast<UStaticMesh>(ValueData.TryLoad());
+				return true;
+			}
+		}
+	}
+#endif //WITH_EDITOR
+	return false;
+}
+bool UInterchangeStaticMeshFactoryNode::FillCustomDistanceFieldReplacementMeshFromAsset(UObject* Asset)
+{
+#if WITH_EDITOR
+	if (UStaticMesh* StaticMesh = Cast<UStaticMesh>(Asset))
+	{
+		if (StaticMesh->GetNumSourceModels() > 0)
+		{
+			FSoftObjectPath SoftObjectPath(StaticMesh->GetSourceModel(0).BuildSettings.DistanceFieldReplacementMesh);
+			return SetAttribute(Macro_CustomDistanceFieldReplacementMeshKey.Key, SoftObjectPath);
+		}
+	}
+#endif //WITH_EDITOR
+	return false;
+}
+
+bool UInterchangeStaticMeshFactoryNode::GetCustomMaxLumenMeshCards(int32& AttributeValue) const
+{
+	IMPLEMENT_NODE_ATTRIBUTE_GETTER(MaxLumenMeshCards, int32)
+}
+bool UInterchangeStaticMeshFactoryNode::SetCustomMaxLumenMeshCards(const int32& AttributeValue, bool bAddApplyDelegate)
+{
+	IMPLEMENT_NODE_ATTRIBUTE_SETTER_WITH_CUSTOM_DELEGATE(UInterchangeStaticMeshFactoryNode, MaxLumenMeshCards, int32);
+}
+bool UInterchangeStaticMeshFactoryNode::ApplyCustomMaxLumenMeshCardsToAsset(UObject* Asset) const
+{
+	IMPLEMENT_STATICMESH_BUILD_VALUE_TO_ASSET(MaxLumenMeshCards, int32, MaxLumenMeshCards);
+}
+bool UInterchangeStaticMeshFactoryNode::FillCustomMaxLumenMeshCardsFromAsset(UObject* Asset)
+{
+	IMPLEMENT_STATICMESH_BUILD_ASSET_TO_VALUE(MaxLumenMeshCards, MaxLumenMeshCards);
+}
\ No newline at end of file
diff --git a/Engine/Plugins/Experimental/Interchange/Runtime/Source/FactoryNodes/Public/InterchangeCineCameraFactoryNode.h b/Engine/Plugins/Experimental/Interchange/Runtime/Source/FactoryNodes/Public/InterchangeCineCameraFactoryNode.h
--- a/Engine/Plugins/Experimental/Interchange/Runtime/Source/FactoryNodes/Public/InterchangeCineCameraFactoryNode.h
+++ b/Engine/Plugins/Experimental/Interchange/Runtime/Source/FactoryNodes/Public/InterchangeCineCameraFactoryNode.h
@@ -26,7 +26,7 @@ public:
 	UFUNCTION(BlueprintCallable, Category = "Interchange | Node | CameraFactory")
 	bool SetCustomFocalLength(const float& AttributeValue, bool bAddApplyDelegate = true)
 	{
-		IMPLEMENT_NODE_ATTRIBUTE_SETTER_WITH_CUSTOM_DELEGATE(UInterchangeCineCameraFactoryNode, FocalLength, float, UCineCameraComponent);
+		IMPLEMENT_NODE_ATTRIBUTE_SETTER_WITH_CUSTOM_DELEGATE_WITH_CUSTOM_CLASS(UInterchangeCineCameraFactoryNode, FocalLength, float, UCineCameraComponent);
 	}
 
 	UFUNCTION(BlueprintCallable, Category = "Interchange | Node | CameraFactory")
@@ -38,7 +38,7 @@ public:
 	UFUNCTION(BlueprintCallable, Category = "Interchange | Node | CameraFactory")
 	bool SetCustomSensorWidth(const float& AttributeValue, bool bAddApplyDelegate = true)
 	{
-		IMPLEMENT_NODE_ATTRIBUTE_SETTER_WITH_CUSTOM_DELEGATE(UInterchangeCineCameraFactoryNode, SensorWidth, float, UCineCameraComponent);
+		IMPLEMENT_NODE_ATTRIBUTE_SETTER_WITH_CUSTOM_DELEGATE_WITH_CUSTOM_CLASS(UInterchangeCineCameraFactoryNode, SensorWidth, float, UCineCameraComponent);
 	}
 
 	UFUNCTION(BlueprintCallable, Category = "Interchange | Node | CameraFactory")
@@ -50,7 +50,7 @@ public:
 	UFUNCTION(BlueprintCallable, Category = "Interchange | Node | CameraFactory")
 	bool SetCustomSensorHeight(const float& AttributeValue, bool bAddApplyDelegate = true)
 	{
-		IMPLEMENT_NODE_ATTRIBUTE_SETTER_WITH_CUSTOM_DELEGATE(UInterchangeCineCameraFactoryNode, SensorHeight, float, UCineCameraComponent);
+		IMPLEMENT_NODE_ATTRIBUTE_SETTER_WITH_CUSTOM_DELEGATE_WITH_CUSTOM_CLASS(UInterchangeCineCameraFactoryNode, SensorHeight, float, UCineCameraComponent);
 	}
 
 private:
diff --git a/Engine/Plugins/Experimental/Interchange/Runtime/Source/FactoryNodes/Public/InterchangeMeshFactoryNode.h b/Engine/Plugins/Experimental/Interchange/Runtime/Source/FactoryNodes/Public/InterchangeMeshFactoryNode.h
--- a/Engine/Plugins/Experimental/Interchange/Runtime/Source/FactoryNodes/Public/InterchangeMeshFactoryNode.h
+++ b/Engine/Plugins/Experimental/Interchange/Runtime/Source/FactoryNodes/Public/InterchangeMeshFactoryNode.h
@@ -59,27 +59,27 @@ public:
 	UFUNCTION(BlueprintCallable, Category = "Interchange | FactoryNode | Mesh")
 	bool RemoveLodDataUniqueId(const FString& LodDataUniqueId);
 
-	/** Query weather the static mesh factory should replace the vertex color. Return false if the attribute was not set.*/
+	/** Query whether the static mesh factory should replace the vertex color. Return false if the attribute was not set.*/
 	UFUNCTION(BlueprintCallable, Category = "Interchange | FactoryNode | Mesh")
 	bool GetCustomVertexColorReplace(bool& AttributeValue) const;
 
-	/** Set weather the static mesh factory should replace the vertex color. Return false if the attribute cannot be set.*/
+	/** Set whether the static mesh factory should replace the vertex color. Return false if the attribute cannot be set.*/
 	UFUNCTION(BlueprintCallable, Category = "Interchange | FactoryNode | Mesh")
 	bool SetCustomVertexColorReplace(const bool& AttributeValue);
 
-	/** Query weather the static mesh factory should ignore the vertex color. Return false if the attribute was not set.*/
+	/** Query whether the static mesh factory should ignore the vertex color. Return false if the attribute was not set.*/
 	UFUNCTION(BlueprintCallable, Category = "Interchange | FactoryNode | Mesh")
 	bool GetCustomVertexColorIgnore(bool& AttributeValue) const;
 
-	/** Set weather the static mesh factory should ignore the vertex color. Return false if the attribute cannot be set.*/
+	/** Set whether the static mesh factory should ignore the vertex color. Return false if the attribute cannot be set.*/
 	UFUNCTION(BlueprintCallable, Category = "Interchange | FactoryNode | Mesh")
 	bool SetCustomVertexColorIgnore(const bool& AttributeValue);
 
-	/** Query weather the static mesh factory should override the vertex color. Return false if the attribute was not set.*/
+	/** Query whether the static mesh factory should override the vertex color. Return false if the attribute was not set.*/
 	UFUNCTION(BlueprintCallable, Category = "Interchange | FactoryNode | Mesh")
 	bool GetCustomVertexColorOverride(FColor& AttributeValue) const;
 
-	/** Set weather the static mesh factory should override the vertex color. Return false if the attribute cannot be set.*/
+	/** Set whether the static mesh factory should override the vertex color. Return false if the attribute cannot be set.*/
 	UFUNCTION(BlueprintCallable, Category = "Interchange | FactoryNode | Mesh")
 	bool SetCustomVertexColorOverride(const FColor& AttributeValue);
 
@@ -99,11 +99,83 @@ public:
 	UFUNCTION(BlueprintCallable, Category = "Interchange | FactoryNode | Mesh")
 	bool RemoveSlotMaterialDependencyUid(const FString& SlotName);
 
+	/** Query whether normals in the imported mesh are ignored and recomputed. When normals are recomputed the tangents are also recomputed. */
+	UFUNCTION(BlueprintCallable, Category = "Interchange | FactoryNode | Mesh")
+	bool GetCustomRecomputeNormals(bool& AttributeValue) const;
+
+	/** Set whether normals in the imported mesh are ignored and recomputed. When normals are recomputed the tangents are also recomputed. */
+	UFUNCTION(BlueprintCallable, Category = "Interchange | FactoryNode | Mesh")
+	bool SetCustomRecomputeNormals(const bool& AttributeValue, bool bAddApplyDelegate = true);
+	
+	/** Query whether tangents in the imported mesh are ignored and recomputed. */
+	UFUNCTION(BlueprintCallable, Category = "Interchange | FactoryNode | Mesh")
+	bool GetCustomRecomputeTangents(bool& AttributeValue) const;
+
+	/** Set whether tangents in the imported mesh are ignored and recomputed. */
+	UFUNCTION(BlueprintCallable, Category = "Interchange | FactoryNode | Mesh")
+	bool SetCustomRecomputeTangents(const bool& AttributeValue, bool bAddApplyDelegate = true);
+
+	/** Query whether tangents recompute will use mikkt space if tangents are recomputed. */
+	UFUNCTION(BlueprintCallable, Category = "Interchange | FactoryNode | Mesh")
+	bool GetCustomUseMikkTSpace(bool& AttributeValue) const;
+
+	/** Set whether tangents recompute will use mikkt space if tangents are recomputed. */
+	UFUNCTION(BlueprintCallable, Category = "Interchange | FactoryNode | Mesh")
+	bool SetCustomUseMikkTSpace(const bool& AttributeValue, bool bAddApplyDelegate = true);
+
+	/** Query whether normals recompute will use the surface area and the corner angle of the triangle as a ratio when computing the normals. */
+	UFUNCTION(BlueprintCallable, Category = "Interchange | FactoryNode | Mesh")
+	bool GetCustomComputeWeightedNormals(bool& AttributeValue) const;
+
+	/** Set whether normals recompute will use the surface area and the corner angle of the triangle as a ratio when computing the normals. */
+	UFUNCTION(BlueprintCallable, Category = "Interchange | FactoryNode | Mesh")
+	bool SetCustomComputeWeightedNormals(const bool& AttributeValue, bool bAddApplyDelegate = true);
+
+	/** Query whether Tangents will be stored at 16 bit vs 8 bit precision. */
+	UFUNCTION(BlueprintCallable, Category = "Interchange | FactoryNode | Mesh")
+	bool GetCustomUseHighPrecisionTangentBasis(bool& AttributeValue) const;
+
+	/** Set whether Tangents will be stored at 16 bit vs 8 bit precision. */
+	UFUNCTION(BlueprintCallable, Category = "Interchange | FactoryNode | Mesh")
+	bool SetCustomUseHighPrecisionTangentBasis(const bool& AttributeValue, bool bAddApplyDelegate = true);
+
+	/** Query whether UVs will be stored at full floating point precision. */
+	UFUNCTION(BlueprintCallable, Category = "Interchange | FactoryNode | Mesh")
+	bool GetCustomUseFullPrecisionUVs(bool& AttributeValue) const;
+
+	/** Set whether UVs will be stored at full floating point precision. */
+	UFUNCTION(BlueprintCallable, Category = "Interchange | FactoryNode | Mesh")
+	bool SetCustomUseFullPrecisionUVs(const bool& AttributeValue, bool bAddApplyDelegate = true);
+
+	/** Query whether UVs will be stored at full floating point precision. */
+	UFUNCTION(BlueprintCallable, Category = "Interchange | FactoryNode | Mesh")
+	bool GetCustomUseBackwardsCompatibleF16TruncUVs(bool& AttributeValue) const;
+
+	/** Set whether UVs will be stored at full floating point precision. */
+	UFUNCTION(BlueprintCallable, Category = "Interchange | FactoryNode | Mesh")
+	bool SetCustomUseBackwardsCompatibleF16TruncUVs(const bool& AttributeValue, bool bAddApplyDelegate = true);
+
+	/** Query whether degenerate triangles will be removed. */
+	UFUNCTION(BlueprintCallable, Category = "Interchange | FactoryNode | Mesh")
+	bool GetCustomRemoveDegenerates(bool& AttributeValue) const;
+
+	/** Set whether degenerate triangles will be removed. */
+	UFUNCTION(BlueprintCallable, Category = "Interchange | FactoryNode | Mesh")
+	bool SetCustomRemoveDegenerates(const bool& AttributeValue, bool bAddApplyDelegate = true);
+
 private:
 
 	const UE::Interchange::FAttributeKey Macro_CustomVertexColorReplaceKey = UE::Interchange::FAttributeKey(TEXT("VertexColorReplace"));
 	const UE::Interchange::FAttributeKey Macro_CustomVertexColorIgnoreKey = UE::Interchange::FAttributeKey(TEXT("VertexColorIgnore"));
 	const UE::Interchange::FAttributeKey Macro_CustomVertexColorOverrideKey = UE::Interchange::FAttributeKey(TEXT("VertexColorOverride"));
+	const UE::Interchange::FAttributeKey Macro_CustomRecomputeNormalsKey = UE::Interchange::FAttributeKey(TEXT("RecomputeNormals"));
+	const UE::Interchange::FAttributeKey Macro_CustomRecomputeTangentsKey = UE::Interchange::FAttributeKey(TEXT("RecomputeTangents"));
+	const UE::Interchange::FAttributeKey Macro_CustomUseMikkTSpaceKey = UE::Interchange::FAttributeKey(TEXT("UseMikkTSpace"));
+	const UE::Interchange::FAttributeKey Macro_CustomComputeWeightedNormalsKey = UE::Interchange::FAttributeKey(TEXT("ComputeWeightedNormals"));
+	const UE::Interchange::FAttributeKey Macro_CustomUseHighPrecisionTangentBasisKey = UE::Interchange::FAttributeKey(TEXT("UseHighPrecisionTangentBasis"));
+	const UE::Interchange::FAttributeKey Macro_CustomUseFullPrecisionUVsKey = UE::Interchange::FAttributeKey(TEXT("UseFullPrecisionUVs"));
+	const UE::Interchange::FAttributeKey Macro_CustomUseBackwardsCompatibleF16TruncUVsKey = UE::Interchange::FAttributeKey(TEXT("UseBackwardsCompatibleF16TruncUVs"));
+	const UE::Interchange::FAttributeKey Macro_CustomRemoveDegeneratesKey = UE::Interchange::FAttributeKey(TEXT("RemoveDegenerates"));
 
 	UE::Interchange::TArrayAttributeHelper<FString> LodDependencies;
 	UE::Interchange::TMapAttributeHelper<FString, FString> SlotMaterialDependencies;
@@ -114,5 +186,22 @@ protected:
 
 	const UE::Interchange::FAttributeKey ClassNameAttributeKey = UE::Interchange::FBaseNodeStaticData::ClassTypeAttributeKey();
 
+	bool ApplyCustomRecomputeNormalsToAsset(UObject * Asset) const;
+	bool FillCustomRecomputeNormalsFromAsset(UObject * Asset);
+	bool ApplyCustomRecomputeTangentsToAsset(UObject * Asset) const;
+	bool FillCustomRecomputeTangentsFromAsset(UObject * Asset);
+	bool ApplyCustomUseMikkTSpaceToAsset(UObject * Asset) const;
+	bool FillCustomUseMikkTSpaceFromAsset(UObject * Asset);
+	bool ApplyCustomComputeWeightedNormalsToAsset(UObject * Asset) const;
+	bool FillCustomComputeWeightedNormalsFromAsset(UObject * Asset);
+	bool ApplyCustomUseHighPrecisionTangentBasisToAsset(UObject * Asset) const;
+	bool FillCustomUseHighPrecisionTangentBasisFromAsset(UObject * Asset);
+	bool ApplyCustomUseFullPrecisionUVsToAsset(UObject * Asset) const;
+	bool FillCustomUseFullPrecisionUVsFromAsset(UObject * Asset);
+	bool ApplyCustomUseBackwardsCompatibleF16TruncUVsToAsset(UObject * Asset) const;
+	bool FillCustomUseBackwardsCompatibleF16TruncUVsFromAsset(UObject * Asset);
+	bool ApplyCustomRemoveDegeneratesToAsset(UObject * Asset) const;
+	bool FillCustomRemoveDegeneratesFromAsset(UObject * Asset);
+
 	bool bIsNodeClassInitialized = false;
 };
diff --git a/Engine/Plugins/Experimental/Interchange/Runtime/Source/FactoryNodes/Public/InterchangeSkeletalMeshFactoryNode.h b/Engine/Plugins/Experimental/Interchange/Runtime/Source/FactoryNodes/Public/InterchangeSkeletalMeshFactoryNode.h
--- a/Engine/Plugins/Experimental/Interchange/Runtime/Source/FactoryNodes/Public/InterchangeSkeletalMeshFactoryNode.h
+++ b/Engine/Plugins/Experimental/Interchange/Runtime/Source/FactoryNodes/Public/InterchangeSkeletalMeshFactoryNode.h
@@ -88,6 +88,38 @@ public:
 	UFUNCTION(BlueprintCallable, Category = "Interchange | Node | SkeletalMesh")
 	bool SetCustomImportContentType(const EInterchangeSkeletalMeshContentType& AttributeValue);
 
+	/** Query the skeletal mesh threshold use to decide if two vertex position are equal. */
+	UFUNCTION(BlueprintCallable, Category = "Interchange | Node | SkeletalMesh")
+	bool GetCustomThresholdPosition(float& AttributeValue) const;
+
+	/** Set the skeletal mesh threshold use to decide if two vertex position are equal. */
+	UFUNCTION(BlueprintCallable, Category = "Interchange | Node | SkeletalMesh")
+	bool SetCustomThresholdPosition(const float& AttributeValue, bool bAddApplyDelegate = true);
+
+	/** Query the skeletal mesh threshold use to decide if two normal, tangents or bi-normals are equal. */
+	UFUNCTION(BlueprintCallable, Category = "Interchange | Node | SkeletalMesh")
+	bool GetCustomThresholdTangentNormal(float& AttributeValue) const;
+
+	/** Set the skeletal mesh threshold use to decide if two normal, tangents or bi-normals are equal. */
+	UFUNCTION(BlueprintCallable, Category = "Interchange | Node | SkeletalMesh")
+	bool SetCustomThresholdTangentNormal(const float& AttributeValue, bool bAddApplyDelegate = true);
+
+	/** Query the skeletal mesh threshold use to decide if two UVs are equal. */
+	UFUNCTION(BlueprintCallable, Category = "Interchange | Node | SkeletalMesh")
+	bool GetCustomThresholdUV(float& AttributeValue) const;
+
+	/** Set the skeletal mesh threshold use to decide if two UVs are equal. */
+	UFUNCTION(BlueprintCallable, Category = "Interchange | Node | SkeletalMesh")
+	bool SetCustomThresholdUV(const float& AttributeValue, bool bAddApplyDelegate = true);
+
+	/** Query the skeletal mesh threshold to compare vertex position equality when computing morph target deltas. */
+	UFUNCTION(BlueprintCallable, Category = "Interchange | Node | SkeletalMesh")
+	bool GetCustomMorphThresholdPosition(float& AttributeValue) const;
+
+	/** Set the skeletal mesh threshold to compare vertex position equality when computing morph target deltas. */
+	UFUNCTION(BlueprintCallable, Category = "Interchange | Node | SkeletalMesh")
+	bool SetCustomMorphThresholdPosition(const float& AttributeValue, bool bAddApplyDelegate = true);
+
 	/**
 	 * The skeletal mesh thumbnail can have an overlay if the last re-import was geometry only. This thumbnail overlay feature use the metadata to find out if the last import was geometry only.
 	 */
@@ -103,9 +135,19 @@ private:
 	const UE::Interchange::FAttributeKey Macro_CustomCreatePhysicsAssetKey = UE::Interchange::FAttributeKey(TEXT("CreatePhysicsAsset"));
 	const UE::Interchange::FAttributeKey Macro_CustomPhysicAssetSoftObjectPathKey = UE::Interchange::FAttributeKey(TEXT("PhysicAssetSoftObjectPath"));
 	const UE::Interchange::FAttributeKey Macro_CustomImportContentTypeKey = UE::Interchange::FAttributeKey(TEXT("ImportContentType"));
-
+	const UE::Interchange::FAttributeKey Macro_CustomThresholdPositionKey = UE::Interchange::FAttributeKey(TEXT("ThresholdPosition"));
+	const UE::Interchange::FAttributeKey Macro_CustomThresholdTangentNormalKey = UE::Interchange::FAttributeKey(TEXT("ThresholdTangentNormal"));
+	const UE::Interchange::FAttributeKey Macro_CustomThresholdUVKey = UE::Interchange::FAttributeKey(TEXT("ThresholdUV"));
+	const UE::Interchange::FAttributeKey Macro_CustomMorphThresholdPositionKey = UE::Interchange::FAttributeKey(TEXT("MorphThresholdPosition"));
+
+	bool ApplyCustomThresholdPositionToAsset(UObject* Asset) const;
+	bool FillCustomThresholdPositionFromAsset(UObject* Asset);
+	bool ApplyCustomThresholdTangentNormalToAsset(UObject* Asset) const;
+	bool FillCustomThresholdTangentNormalFromAsset(UObject* Asset);
+	bool ApplyCustomThresholdUVToAsset(UObject* Asset) const;
+	bool FillCustomThresholdUVFromAsset(UObject* Asset);
+	bool ApplyCustomMorphThresholdPositionToAsset(UObject* Asset) const;
+	bool FillCustomMorphThresholdPositionFromAsset(UObject* Asset);
 protected:
-#if WITH_ENGINE
 	TSubclassOf<USkeletalMesh> AssetClass = nullptr;
-#endif
 };
diff --git a/Engine/Plugins/Experimental/Interchange/Runtime/Source/FactoryNodes/Public/InterchangeStaticMeshFactoryNode.h b/Engine/Plugins/Experimental/Interchange/Runtime/Source/FactoryNodes/Public/InterchangeStaticMeshFactoryNode.h
--- a/Engine/Plugins/Experimental/Interchange/Runtime/Source/FactoryNodes/Public/InterchangeStaticMeshFactoryNode.h
+++ b/Engine/Plugins/Experimental/Interchange/Runtime/Source/FactoryNodes/Public/InterchangeStaticMeshFactoryNode.h
@@ -76,11 +76,131 @@ public:
 	UFUNCTION(BlueprintCallable, Category = "Interchange | Node | StaticMesh")
 	bool RemoveSocketUd(const FString& SocketUid);
 
-private:
+	/** Get weather the static mesh should build build a reverse index buffer. */
+	UFUNCTION(BlueprintCallable, Category = "Interchange | Node | StaticMesh")
+	bool GetCustomBuildReversedIndexBuffer(bool& AttributeValue) const;
+
+	/** Set weather the static mesh should build build a reverse index buffer. */
+	UFUNCTION(BlueprintCallable, Category = "Interchange | Node | StaticMesh")
+	bool SetCustomBuildReversedIndexBuffer(const bool& AttributeValue, bool bAddApplyDelegate = true);
+
+	/** Get weather the static mesh should build build a reverse index buffer. */
+	UFUNCTION(BlueprintCallable, Category = "Interchange | Node | StaticMesh")
+	bool GetCustomGenerateLightmapUVs(bool& AttributeValue) const;
+
+	/** Set weather the static mesh should build build a reverse index buffer. */
+	UFUNCTION(BlueprintCallable, Category = "Interchange | Node | StaticMesh")
+	bool SetCustomGenerateLightmapUVs(const bool& AttributeValue, bool bAddApplyDelegate = true);
+
+	/**
+	 * Get Whether to generate the distance field treating every triangle hit as a front face.  
+	 * When enabled prevents the distance field from being discarded due to the mesh being open, but also lowers Distance Field AO quality.
+	 */
+	UFUNCTION(BlueprintCallable, Category = "Interchange | Node | StaticMesh")
+	bool GetCustomGenerateDistanceFieldAsIfTwoSided(bool& AttributeValue) const;
+
+	/**
+	 * Set Whether to generate the distance field treating every triangle hit as a front face.
+	 * When enabled prevents the distance field from being discarded due to the mesh being open, but also lowers Distance Field AO quality.
+	 */
+	UFUNCTION(BlueprintCallable, Category = "Interchange | Node | StaticMesh")
+	bool SetCustomGenerateDistanceFieldAsIfTwoSided(const bool& AttributeValue, bool bAddApplyDelegate = true);
+
+	/** Get weather the static mesh ... (TODO add documentation) */
+	UFUNCTION(BlueprintCallable, Category = "Interchange | Node | StaticMesh")
+	bool GetCustomSupportFaceRemap(bool& AttributeValue) const;
+
+	/** Set weather the static mesh ... (TODO add documentation) */
+	UFUNCTION(BlueprintCallable, Category = "Interchange | Node | StaticMesh")
+	bool SetCustomSupportFaceRemap(const bool& AttributeValue, bool bAddApplyDelegate = true);
+
+	/** Get weather the static mesh ... (TODO add documentation) */
+	UFUNCTION(BlueprintCallable, Category = "Interchange | Node | StaticMesh")
+	bool GetCustomMinLightmapResolution(int32& AttributeValue) const;
+
+	/** Set weather the static mesh ... (TODO add documentation) */
+	UFUNCTION(BlueprintCallable, Category = "Interchange | Node | StaticMesh")
+	bool SetCustomMinLightmapResolution(const int32& AttributeValue, bool bAddApplyDelegate = true);
+
+	/** Get weather the static mesh ... (TODO add documentation) */
+	UFUNCTION(BlueprintCallable, Category = "Interchange | Node | StaticMesh")
+	bool GetCustomSrcLightmapIndex(int32& AttributeValue) const;
+
+	/** Set weather the static mesh ... (TODO add documentation) */
+	UFUNCTION(BlueprintCallable, Category = "Interchange | Node | StaticMesh")
+	bool SetCustomSrcLightmapIndex(const int32& AttributeValue, bool bAddApplyDelegate = true);
+
+	/** Get weather the static mesh ... (TODO add documentation) */
+	UFUNCTION(BlueprintCallable, Category = "Interchange | Node | StaticMesh")
+	bool GetCustomDstLightmapIndex(int32& AttributeValue) const;
+
+	/** Set weather the static mesh ... (TODO add documentation) */
+	UFUNCTION(BlueprintCallable, Category = "Interchange | Node | StaticMesh")
+	bool SetCustomDstLightmapIndex(const int32& AttributeValue, bool bAddApplyDelegate = true);
+
+	/** Get the static mesh local scale applied when building the static mesh. */
+	UFUNCTION(BlueprintCallable, Category = "Interchange | Node | StaticMesh")
+	bool GetCustomBuildScale3D(FVector& AttributeValue) const;
+
+	/** Set the static mesh local scale applied when building the static mesh. */
+	UFUNCTION(BlueprintCallable, Category = "Interchange | Node | StaticMesh")
+	bool SetCustomBuildScale3D(const FVector& AttributeValue, bool bAddApplyDelegate = true);
+
+	/**
+	 * Get the Scale to apply to the mesh when allocating the distance field volume texture.
+	 * The default scale is 1, which is assuming that the mesh will be placed unscaled in the world.
+	 */
+	UFUNCTION(BlueprintCallable, Category = "Interchange | Node | StaticMesh")
+	bool GetCustomDistanceFieldResolutionScale(float& AttributeValue) const;
+
+	/**
+	 * Set the Scale to apply to the mesh when allocating the distance field volume texture.
+	 * The default scale is 1, which is assuming that the mesh will be placed unscaled in the world.
+	 */
+	UFUNCTION(BlueprintCallable, Category = "Interchange | Node | StaticMesh")
+	bool SetCustomDistanceFieldResolutionScale(const float& AttributeValue, bool bAddApplyDelegate = true);
 
+	/** Get the static mesh local scale applied when building the static mesh. */
+	UFUNCTION(BlueprintCallable, Category = "Interchange | Node | StaticMesh")
+	bool GetCustomDistanceFieldReplacementMesh(FSoftObjectPath& AttributeValue) const;
+
+	/** Set the static mesh local scale applied when building the static mesh. */
+	UFUNCTION(BlueprintCallable, Category = "Interchange | Node | StaticMesh")
+	bool SetCustomDistanceFieldReplacementMesh(const FSoftObjectPath& AttributeValue, bool bAddApplyDelegate = true);
+
+	/**
+	 * Max Lumen mesh cards to generate for this mesh.
+	 * More cards means that surface will have better coverage, but will result in increased runtime overhead.
+	 * Set to 0 in order to disable mesh card generation for this mesh.
+	 * Default is 12.
+	 */
+	UFUNCTION(BlueprintCallable, Category = "Interchange | Node | StaticMesh")
+	bool GetCustomMaxLumenMeshCards(int32& AttributeValue) const;
+
+	/**
+	 * Max Lumen mesh cards to generate for this mesh.
+	 * More cards means that surface will have better coverage, but will result in increased runtime overhead.
+	 * Set to 0 in order to disable mesh card generation for this mesh.
+	 * Default is 12.
+	 */
+	UFUNCTION(BlueprintCallable, Category = "Interchange | Node | StaticMesh")
+	bool SetCustomMaxLumenMeshCards(const int32& AttributeValue, bool bAddApplyDelegate = true);
+
+private:
 	virtual void FillAssetClassFromAttribute() override;
 	virtual bool SetNodeClassFromClassAttribute() override;
 
+	const UE::Interchange::FAttributeKey Macro_CustomBuildReversedIndexBufferKey = UE::Interchange::FAttributeKey(TEXT("BuildReversedIndexBuffer"));
+	const UE::Interchange::FAttributeKey Macro_CustomGenerateLightmapUVsKey = UE::Interchange::FAttributeKey(TEXT("GenerateLightmapUVs"));
+	const UE::Interchange::FAttributeKey Macro_CustomGenerateDistanceFieldAsIfTwoSidedKey = UE::Interchange::FAttributeKey(TEXT("GenerateDistanceFieldAsIfTwoSided"));
+	const UE::Interchange::FAttributeKey Macro_CustomSupportFaceRemapKey = UE::Interchange::FAttributeKey(TEXT("SupportFaceRemap"));
+	const UE::Interchange::FAttributeKey Macro_CustomMinLightmapResolutionKey = UE::Interchange::FAttributeKey(TEXT("MinLightmapResolution"));
+	const UE::Interchange::FAttributeKey Macro_CustomSrcLightmapIndexKey = UE::Interchange::FAttributeKey(TEXT("SrcLightmapIndex"));
+	const UE::Interchange::FAttributeKey Macro_CustomDstLightmapIndexKey = UE::Interchange::FAttributeKey(TEXT("DstLightmapIndex"));
+	const UE::Interchange::FAttributeKey Macro_CustomBuildScale3DKey = UE::Interchange::FAttributeKey(TEXT("BuildScale3D"));
+	const UE::Interchange::FAttributeKey Macro_CustomDistanceFieldResolutionScaleKey = UE::Interchange::FAttributeKey(TEXT("DistanceFieldResolutionScale"));
+	const UE::Interchange::FAttributeKey Macro_CustomDistanceFieldReplacementMeshKey = UE::Interchange::FAttributeKey(TEXT("DistanceFieldReplacementMesh"));
+	const UE::Interchange::FAttributeKey Macro_CustomMaxLumenMeshCardsKey = UE::Interchange::FAttributeKey(TEXT("MaxLumenMeshCards"));
 	const UE::Interchange::FAttributeKey Macro_CustomBuildNaniteKey = UE::Interchange::FAttributeKey(TEXT("BuildNanite"));
 
 	UE::Interchange::TArrayAttributeHelper<FString> SocketUids;
@@ -89,6 +209,29 @@ protected:
 	
 	IMPLEMENT_NODE_ATTRIBUTE_DELEGATE_BY_PROPERTYNAME(BuildNanite, bool, UStaticMesh, TEXT("NaniteSettings.bEnabled"));
 
+	bool ApplyCustomBuildReversedIndexBufferToAsset(UObject* Asset) const;
+	bool FillCustomBuildReversedIndexBufferFromAsset(UObject* Asset);
+	bool ApplyCustomGenerateLightmapUVsToAsset(UObject* Asset) const;
+	bool FillCustomGenerateLightmapUVsFromAsset(UObject* Asset);
+	bool ApplyCustomGenerateDistanceFieldAsIfTwoSidedToAsset(UObject* Asset) const;
+	bool FillCustomGenerateDistanceFieldAsIfTwoSidedFromAsset(UObject* Asset);
+	bool ApplyCustomSupportFaceRemapToAsset(UObject* Asset) const;
+	bool FillCustomSupportFaceRemapFromAsset(UObject* Asset);
+	bool ApplyCustomMinLightmapResolutionToAsset(UObject* Asset) const;
+	bool FillCustomMinLightmapResolutionFromAsset(UObject* Asset);
+	bool ApplyCustomSrcLightmapIndexToAsset(UObject* Asset) const;
+	bool FillCustomSrcLightmapIndexFromAsset(UObject* Asset);
+	bool ApplyCustomDstLightmapIndexToAsset(UObject* Asset) const;
+	bool FillCustomDstLightmapIndexFromAsset(UObject* Asset);
+	bool ApplyCustomBuildScale3DToAsset(UObject* Asset) const;
+	bool FillCustomBuildScale3DFromAsset(UObject* Asset);
+	bool ApplyCustomDistanceFieldResolutionScaleToAsset(UObject* Asset) const;
+	bool FillCustomDistanceFieldResolutionScaleFromAsset(UObject* Asset);
+	bool ApplyCustomDistanceFieldReplacementMeshToAsset(UObject* Asset) const;
+	bool FillCustomDistanceFieldReplacementMeshFromAsset(UObject* Asset);
+	bool ApplyCustomMaxLumenMeshCardsToAsset(UObject* Asset) const;
+	bool FillCustomMaxLumenMeshCardsFromAsset(UObject* Asset);
+
 #if WITH_ENGINE
 	TSubclassOf<UStaticMesh> AssetClass = nullptr;
 #endif
diff --git a/Engine/Plugins/Experimental/Interchange/Runtime/Source/FactoryNodes/Public/InterchangeTextureFactoryNode.h b/Engine/Plugins/Experimental/Interchange/Runtime/Source/FactoryNodes/Public/InterchangeTextureFactoryNode.h
--- a/Engine/Plugins/Experimental/Interchange/Runtime/Source/FactoryNodes/Public/InterchangeTextureFactoryNode.h
+++ b/Engine/Plugins/Experimental/Interchange/Runtime/Source/FactoryNodes/Public/InterchangeTextureFactoryNode.h
@@ -602,7 +602,7 @@ public:
 	UFUNCTION(BlueprintCallable, Category = "Interchange | Node | Texture")
 	bool SetCustomDownscale(const float& AttributeValue, bool bAddApplyDelegate = true)
 	{
-		IMPLEMENT_NODE_ATTRIBUTE_SETTER_WITH_CUSTOM_DELEGATE(UInterchangeTextureFactoryNode, Downscale, float, UTexture)
+		IMPLEMENT_NODE_ATTRIBUTE_SETTER_WITH_CUSTOM_DELEGATE_WITH_CUSTOM_CLASS(UInterchangeTextureFactoryNode, Downscale, float, UTexture)
 	}
 
 	/** Return false if the Attribute was not set previously.*/
diff --git a/Engine/Plugins/Experimental/Interchange/Runtime/Source/Pipelines/Private/InterchangeGenericAssetsPipeline.cpp b/Engine/Plugins/Experimental/Interchange/Runtime/Source/Pipelines/Private/InterchangeGenericAssetsPipeline.cpp
--- a/Engine/Plugins/Experimental/Interchange/Runtime/Source/Pipelines/Private/InterchangeGenericAssetsPipeline.cpp
+++ b/Engine/Plugins/Experimental/Interchange/Runtime/Source/Pipelines/Private/InterchangeGenericAssetsPipeline.cpp
@@ -233,12 +233,9 @@ void UInterchangeGenericAssetsPipeline::ExecutePreImportPipeline(UInterchangeBas
 
 	ImplementUseSourceNameForAssetOption();
 	//Make sure all factory nodes have the specified strategy
-	BaseNodeContainer->IterateNodes([ReimportStrategyClosure = ReimportStrategy](const FString& NodeUid, UInterchangeBaseNode* Node)
+	BaseNodeContainer->IterateNodesOfType<UInterchangeFactoryBaseNode>([ReimportStrategyClosure = ReimportStrategy](const FString& NodeUid, UInterchangeFactoryBaseNode* FactoryNode)
 		{
-			if (Node->GetNodeContainerType() == EInterchangeNodeContainerType::FactoryData)
-			{
-				Node->SetReimportStrategyFlags(ReimportStrategyClosure);
-			}
+			FactoryNode->SetReimportStrategyFlags(ReimportStrategyClosure);
 		});
 }
 
diff --git a/Engine/Plugins/Experimental/Interchange/Runtime/Source/Pipelines/Private/InterchangeGenericMeshPipeline.cpp b/Engine/Plugins/Experimental/Interchange/Runtime/Source/Pipelines/Private/InterchangeGenericMeshPipeline.cpp
--- a/Engine/Plugins/Experimental/Interchange/Runtime/Source/Pipelines/Private/InterchangeGenericMeshPipeline.cpp
+++ b/Engine/Plugins/Experimental/Interchange/Runtime/Source/Pipelines/Private/InterchangeGenericMeshPipeline.cpp
@@ -37,43 +37,59 @@ void UInterchangeGenericMeshPipeline::AdjustSettingsForContext(EInterchangePipel
 		bCreatePhysicsAsset = false;
 		PhysicsAsset = nullptr;
 	}
+	const FString CommonMeshesCategory = TEXT("Common Meshes");
+	const FString StaticMeshesCategory = TEXT("Static Meshes");
+	const FString SkeletalMeshesCategory = TEXT("Skeletal Meshes");
+	const FString CommonSkeletalMeshesAndAnimationCategory = TEXT("Common Skeletal Meshes and Animations");
 
 	TArray<FString> HideCategories;
+	TArray<FString> HideSubCategories;
 	if (ImportType == EInterchangePipelineContext::AssetReimport)
 	{
+		HideSubCategories.Add(TEXT("Build"));
 		if (USkeletalMesh* SkeletalMesh = Cast<USkeletalMesh>(ReimportAsset))
 		{
-
 			//Set the skeleton to the current asset skeleton
 			CommonSkeletalMeshesAndAnimationsProperties->Skeleton = SkeletalMesh->GetSkeleton();
 			bImportStaticMeshes = false;
-			HideCategories.Add(TEXT("Static Meshes"));
+			HideCategories.Add(StaticMeshesCategory);
 		}
 		else if (UStaticMesh* StaticMesh = Cast<UStaticMesh>(ReimportAsset))
 		{
-			HideCategories.Add(TEXT("Skeletal Meshes"));
-			HideCategories.Add(TEXT("Common Skeletal Meshes and Animations"));
+			HideCategories.Add(SkeletalMeshesCategory);
+			HideCategories.Add(CommonSkeletalMeshesAndAnimationCategory);
 		}
 		else if (UAnimSequence* AnimSequence = Cast<UAnimSequence>(ReimportAsset))
 		{
-			HideCategories.Add(TEXT("Static Meshes"));
-			HideCategories.Add(TEXT("Skeletal Meshes"));
-			HideCategories.Add(TEXT("Common Meshes"));
+			HideCategories.Add(StaticMeshesCategory);
+			HideCategories.Add(SkeletalMeshesCategory);
+			HideCategories.Add(CommonMeshesCategory);
 		}
 		else if (ReimportAsset)
 		{
-			HideCategories.Add(TEXT("Static Meshes"));
-			HideCategories.Add(TEXT("Skeletal Meshes"));
-			HideCategories.Add(TEXT("Common Meshes"));
-			HideCategories.Add(TEXT("Common Skeletal Meshes and Animations"));
+			HideCategories.Add(StaticMeshesCategory);
+			HideCategories.Add(SkeletalMeshesCategory);
+			HideCategories.Add(CommonMeshesCategory);
+			HideCategories.Add(CommonSkeletalMeshesAndAnimationCategory);
 		}
 	}
 
 	if (UInterchangePipelineBase* OuterMostPipeline = GetMostPipelineOuter())
 	{
+		constexpr bool bDoTransientSubPipeline = true;
+		if (UInterchangeGenericAssetsPipeline* ParentPipeline = Cast<UInterchangeGenericAssetsPipeline>(OuterMostPipeline))
+		{
+			if (ParentPipeline->ReimportStrategy == EReimportStrategyFlags::ApplyNoProperties)
+			{
+				for (const FString& HideSubCategoryName : HideSubCategories)
+				{
+					HidePropertiesOfSubCategory(OuterMostPipeline, this, HideSubCategoryName, bDoTransientSubPipeline);
+				}
+			}
+		}
+
 		for (const FString& HideCategoryName : HideCategories)
 		{
-			constexpr bool bDoTransientSubPipeline = true;
 			HidePropertiesOfCategory(OuterMostPipeline, this, HideCategoryName, bDoTransientSubPipeline);
 		}
 	}
@@ -170,3 +186,46 @@ void UInterchangeGenericMeshPipeline::SetReimportSourceIndex(UClass* ReimportObj
 	}
 }
 
+#if WITH_EDITOR
+bool UInterchangeGenericMeshPipeline::DoClassesIncludeAllEditableStructProperties(const TArray<const UClass*>& Classes, const UStruct* Struct)
+{
+	const FName CategoryKey("Category");
+	for (const FProperty* Property = Struct->PropertyLink; Property; Property = Property->PropertyLinkNext)
+	{
+		//skip (transient, deprecated, const) property
+		if (Property->HasAnyPropertyFlags(CPF_Transient | CPF_Deprecated | CPF_EditConst))
+		{
+			continue;
+		}
+		//skip property that is not editable
+		if (!Property->HasAnyPropertyFlags(CPF_Edit))
+		{
+			continue;
+		}
+		const FObjectProperty* SubObject = CastField<FObjectProperty>(Property);
+		if (SubObject)
+		{
+			continue;
+		}
+		else if (const FString* PropertyCategoryString = Property->FindMetaData(CategoryKey))
+		{
+			FName PropertyName = Property->GetFName();
+			bool bFindProperty = false;
+			for (const UClass* Class : Classes)
+			{
+				if (Class->FindPropertyByName(PropertyName) != nullptr)
+				{
+					bFindProperty = true;
+					break;
+				}
+			}
+			//Ensure to notify
+			if (!bFindProperty)
+			{
+				return false;
+			}
+		}
+	}
+	return true;
+}
+#endif
\ No newline at end of file
diff --git a/Engine/Plugins/Experimental/Interchange/Runtime/Source/Pipelines/Private/InterchangeGenericSkeletalMeshPipeline.cpp b/Engine/Plugins/Experimental/Interchange/Runtime/Source/Pipelines/Private/InterchangeGenericSkeletalMeshPipeline.cpp
--- a/Engine/Plugins/Experimental/Interchange/Runtime/Source/Pipelines/Private/InterchangeGenericSkeletalMeshPipeline.cpp
+++ b/Engine/Plugins/Experimental/Interchange/Runtime/Source/Pipelines/Private/InterchangeGenericSkeletalMeshPipeline.cpp
@@ -5,6 +5,7 @@
 #include "AssetRegistry/AssetRegistryModule.h"
 #include "Async/Async.h"
 #include "CoreMinimal.h"
+#include "Engine/EngineTypes.h"
 #include "InterchangeGenericAssetsPipeline.h"
 #include "InterchangeMaterialFactoryNode.h"
 #include "InterchangeMeshNode.h"
@@ -83,6 +84,20 @@ void UInterchangeGenericMeshPipeline::ExecutePreImportPipelineSkeletalMesh()
 		//Nothing to import
 		return;
 	}
+
+#if WITH_EDITOR
+	//Make sure the generic pipeline we cover all skeletalmesh build settings by asserting when we import
+	{
+		TArray<const UClass*> Classes;
+		Classes.Add(UInterchangeGenericCommonMeshesProperties::StaticClass());
+		Classes.Add(UInterchangeGenericMeshPipeline::StaticClass());
+		if (!ensure(DoClassesIncludeAllEditableStructProperties(Classes, FSkeletalMeshBuildSettings::StaticStruct())))
+		{
+			UE_LOG(LogInterchangePipeline, Log, TEXT("UInterchangeGenericMeshPipeline: The generic pipeline does not cover all skeletal mesh build options."));
+		}
+	}
+#endif
+
 	const bool bConvertStaticMeshToSkeletalMesh = (CommonMeshesProperties->ForceAllMeshAsType == EInterchangeForceMeshType::IFMT_SkeletalMesh);
 	TMap<FString, TArray<FString>> SkeletalMeshFactoryDependencyOrderPerSkeletonRootNodeUid;
 
@@ -478,6 +493,22 @@ UInterchangeSkeletalMeshFactoryNode* UInterchangeGenericMeshPipeline::CreateSkel
 		SkeletonFactoryNode->SetEnabled(false);
 		SkeletalMeshFactoryNode->SetEnabled(false);
 	}
+
+	//Common meshes build options
+	SkeletalMeshFactoryNode->SetCustomRecomputeNormals(CommonMeshesProperties->bRecomputeNormals);
+	SkeletalMeshFactoryNode->SetCustomRecomputeTangents(CommonMeshesProperties->bRecomputeTangents);
+	SkeletalMeshFactoryNode->SetCustomUseMikkTSpace(CommonMeshesProperties->bUseMikkTSpace);
+	SkeletalMeshFactoryNode->SetCustomComputeWeightedNormals(CommonMeshesProperties->bComputeWeightedNormals);
+	SkeletalMeshFactoryNode->SetCustomUseHighPrecisionTangentBasis(CommonMeshesProperties->bUseHighPrecisionTangentBasis);
+	SkeletalMeshFactoryNode->SetCustomUseFullPrecisionUVs(CommonMeshesProperties->bUseFullPrecisionUVs);
+	SkeletalMeshFactoryNode->SetCustomUseBackwardsCompatibleF16TruncUVs(CommonMeshesProperties->bUseBackwardsCompatibleF16TruncUVs);
+	SkeletalMeshFactoryNode->SetCustomRemoveDegenerates(CommonMeshesProperties->bRemoveDegenerates);
+	//Skeletal meshes build options
+	SkeletalMeshFactoryNode->SetCustomThresholdPosition(ThresholdPosition);
+	SkeletalMeshFactoryNode->SetCustomThresholdTangentNormal(ThresholdTangentNormal);
+	SkeletalMeshFactoryNode->SetCustomThresholdUV(ThresholdUV);
+	SkeletalMeshFactoryNode->SetCustomMorphThresholdPosition(MorphThresholdPosition);
+
 	return SkeletalMeshFactoryNode;
 }
 
diff --git a/Engine/Plugins/Experimental/Interchange/Runtime/Source/Pipelines/Private/InterchangeGenericStaticMeshPipeline.cpp b/Engine/Plugins/Experimental/Interchange/Runtime/Source/Pipelines/Private/InterchangeGenericStaticMeshPipeline.cpp
--- a/Engine/Plugins/Experimental/Interchange/Runtime/Source/Pipelines/Private/InterchangeGenericStaticMeshPipeline.cpp
+++ b/Engine/Plugins/Experimental/Interchange/Runtime/Source/Pipelines/Private/InterchangeGenericStaticMeshPipeline.cpp
@@ -138,6 +138,19 @@ void UInterchangeGenericMeshPipeline::ExecutePreImportPipelineStaticMesh()
 {
 	check(!CommonMeshesProperties.IsNull());
 
+#if WITH_EDITOR
+	//Make sure the generic pipeline will cover all staticmesh build settings when we import
+	{
+		TArray<const UClass*> Classes;
+		Classes.Add(UInterchangeGenericCommonMeshesProperties::StaticClass());
+		Classes.Add(UInterchangeGenericMeshPipeline::StaticClass());
+		if (!ensure(DoClassesIncludeAllEditableStructProperties(Classes, FMeshBuildSettings::StaticStruct())))
+		{
+			UE_LOG(LogInterchangePipeline, Log, TEXT("UInterchangeGenericMeshPipeline: The generic pipeline does not cover all static mesh build options."));
+		}
+	}
+#endif
+
 	if (bImportStaticMeshes && (CommonMeshesProperties->ForceAllMeshAsType == EInterchangeForceMeshType::IFMT_None || CommonMeshesProperties->ForceAllMeshAsType == EInterchangeForceMeshType::IFMT_StaticMesh))
 	{
 		const bool bConvertSkeletalMeshToStaticMesh = (CommonMeshesProperties->ForceAllMeshAsType == EInterchangeForceMeshType::IFMT_StaticMesh);
@@ -416,10 +429,28 @@ UInterchangeStaticMeshFactoryNode* UInterchangeGenericMeshPipeline::CreateStatic
 		break;
 	}
 
-	if (bBuildNanite)
-	{
-		StaticMeshFactoryNode->SetCustomBuildNanite(bBuildNanite);
-	}
+	//Common meshes build options
+	StaticMeshFactoryNode->SetCustomRecomputeNormals(CommonMeshesProperties->bRecomputeNormals);
+	StaticMeshFactoryNode->SetCustomRecomputeTangents(CommonMeshesProperties->bRecomputeTangents);
+	StaticMeshFactoryNode->SetCustomUseMikkTSpace(CommonMeshesProperties->bUseMikkTSpace);
+	StaticMeshFactoryNode->SetCustomComputeWeightedNormals(CommonMeshesProperties->bComputeWeightedNormals);
+	StaticMeshFactoryNode->SetCustomUseHighPrecisionTangentBasis(CommonMeshesProperties->bUseHighPrecisionTangentBasis);
+	StaticMeshFactoryNode->SetCustomUseFullPrecisionUVs(CommonMeshesProperties->bUseFullPrecisionUVs);
+	StaticMeshFactoryNode->SetCustomUseBackwardsCompatibleF16TruncUVs(CommonMeshesProperties->bUseBackwardsCompatibleF16TruncUVs);
+	StaticMeshFactoryNode->SetCustomRemoveDegenerates(CommonMeshesProperties->bRemoveDegenerates);
+	//Static meshes build options
+	StaticMeshFactoryNode->SetCustomBuildReversedIndexBuffer(bBuildReversedIndexBuffer);
+	StaticMeshFactoryNode->SetCustomGenerateLightmapUVs(bGenerateLightmapUVs);
+	StaticMeshFactoryNode->SetCustomGenerateDistanceFieldAsIfTwoSided(bGenerateDistanceFieldAsIfTwoSided);
+	StaticMeshFactoryNode->SetCustomSupportFaceRemap(bSupportFaceRemap);
+	StaticMeshFactoryNode->SetCustomMinLightmapResolution(MinLightmapResolution);
+	StaticMeshFactoryNode->SetCustomSrcLightmapIndex(SrcLightmapIndex);
+	StaticMeshFactoryNode->SetCustomDstLightmapIndex(DstLightmapIndex);
+	StaticMeshFactoryNode->SetCustomBuildScale3D(BuildScale3D);
+	StaticMeshFactoryNode->SetCustomDistanceFieldResolutionScale(DistanceFieldResolutionScale);
+	StaticMeshFactoryNode->SetCustomDistanceFieldReplacementMesh(DistanceFieldReplacementMesh);
+	StaticMeshFactoryNode->SetCustomMaxLumenMeshCards(MaxLumenMeshCards);
+	StaticMeshFactoryNode->SetCustomBuildNanite(bBuildNanite);
 
 	return StaticMeshFactoryNode;
 }
diff --git a/Engine/Plugins/Experimental/Interchange/Runtime/Source/Pipelines/Public/InterchangeGenericAssetsPipeline.h b/Engine/Plugins/Experimental/Interchange/Runtime/Source/Pipelines/Public/InterchangeGenericAssetsPipeline.h
--- a/Engine/Plugins/Experimental/Interchange/Runtime/Source/Pipelines/Public/InterchangeGenericAssetsPipeline.h
+++ b/Engine/Plugins/Experimental/Interchange/Runtime/Source/Pipelines/Public/InterchangeGenericAssetsPipeline.h
@@ -33,7 +33,7 @@ public:
 	
 	//////	COMMON_CATEGORY Properties //////
 	/* Allow user to choose the re-import strategy. */
-	UPROPERTY(EditAnywhere, BlueprintReadWrite, Category = "Common")
+	UPROPERTY(EditAnywhere, BlueprintReadWrite, Category = "Common", meta = (AdjustPipelineAndRefreshDetailOnChange = "True"))
 	EReimportStrategyFlags ReimportStrategy = EReimportStrategyFlags::ApplyNoProperties;
 
 	/** If enable and there is only one asset and one source data, we will name the asset like the source data name. */
diff --git a/Engine/Plugins/Experimental/Interchange/Runtime/Source/Pipelines/Public/InterchangeGenericAssetsPipelineSharedSettings.h b/Engine/Plugins/Experimental/Interchange/Runtime/Source/Pipelines/Public/InterchangeGenericAssetsPipelineSharedSettings.h
--- a/Engine/Plugins/Experimental/Interchange/Runtime/Source/Pipelines/Public/InterchangeGenericAssetsPipelineSharedSettings.h
+++ b/Engine/Plugins/Experimental/Interchange/Runtime/Source/Pipelines/Public/InterchangeGenericAssetsPipelineSharedSettings.h
@@ -70,6 +70,38 @@ public:
 	/** Specify override color in the case that VertexColorImportOption is set to Override */
 	UPROPERTY(EditAnywhere, BlueprintReadWrite, Category = "Common Meshes")
 	FColor VertexOverrideColor;
+
+	/** If true, normals in the imported mesh are ignored and recomputed. */
+	UPROPERTY(EditAnywhere, BlueprintReadWrite, Category = "Common Meshes", meta = (SubCategory = "Build"))
+	bool bRecomputeNormals = true;
+
+	/** If true, tangents in the imported mesh are ignored and recomputed. */
+	UPROPERTY(EditAnywhere, BlueprintReadWrite, Category = "Common Meshes", meta = (SubCategory = "Build"))
+	bool bRecomputeTangents = true;
+
+	/** If true, recompute tangents will use mikkt space. */
+	UPROPERTY(EditAnywhere, BlueprintReadWrite, Category = "Common Meshes", meta = (SubCategory = "Build"))
+	bool bUseMikkTSpace = true;
+
+	/** If true, we will use the surface area and the corner angle of the triangle as a ratio when computing the normals. */
+	UPROPERTY(EditAnywhere, BlueprintReadWrite, Category = "Common Meshes", meta = (SubCategory = "Build"))
+	bool bComputeWeightedNormals = false;
+
+	/** If true, Tangents will be stored at 16 bit vs 8 bit precision. */
+	UPROPERTY(EditAnywhere, BlueprintReadWrite, Category = "Common Meshes", meta = (SubCategory = "Build"))
+	bool bUseHighPrecisionTangentBasis = false;
+
+	/** If true, UVs will be stored at full floating point precision. */
+	UPROPERTY(EditAnywhere, BlueprintReadWrite, Category = "Common Meshes", meta = (SubCategory = "Build"))
+	bool bUseFullPrecisionUVs = false;
+
+	/** If true, UVs will use backwards-compatible F16 conversion with truncation for legacy meshes. */
+	UPROPERTY(EditAnywhere, BlueprintReadWrite, Category = "Common Meshes", meta = (SubCategory = "Build"))
+	bool bUseBackwardsCompatibleF16TruncUVs = false;
+
+	/** If true, degenerate triangles will be removed. */
+	UPROPERTY(EditAnywhere, BlueprintReadWrite, Category = "Common Meshes", meta = (SubCategory = "Build"))
+	bool bRemoveDegenerates = false;
 };
 
 UCLASS(BlueprintType, hidedropdown, Experimental)
diff --git a/Engine/Plugins/Experimental/Interchange/Runtime/Source/Pipelines/Public/InterchangeGenericMeshPipeline.h b/Engine/Plugins/Experimental/Interchange/Runtime/Source/Pipelines/Public/InterchangeGenericMeshPipeline.h
--- a/Engine/Plugins/Experimental/Interchange/Runtime/Source/Pipelines/Public/InterchangeGenericMeshPipeline.h
+++ b/Engine/Plugins/Experimental/Interchange/Runtime/Source/Pipelines/Public/InterchangeGenericMeshPipeline.h
@@ -70,7 +70,57 @@ public:
 	/** If enabled this option will allow you to use Nanite rendering at runtime. Can only be used with simple opaque materials. */
 	UPROPERTY(EditAnywhere, BlueprintReadWrite, Category = "Static Meshes", meta = (SubCategory = "Build"))
 	bool bBuildNanite = false;
+
+	/** If enabled this option will make sure the staticmesh build a reverse index buffer. */
+	UPROPERTY(EditAnywhere, BlueprintReadWrite, Category = "Static Meshes", meta = (SubCategory = "Build"))
+	bool bBuildReversedIndexBuffer = true;
+	
+	/** If enabled this option will generate lightmap for this staticmesh. */
+	UPROPERTY(EditAnywhere, BlueprintReadWrite, Category = "Static Meshes", meta = (SubCategory = "Build"))
+	bool bGenerateLightmapUVs = true;
+	
+	/** 
+	 * Whether to generate the distance field treating every triangle hit as a front face.  
+	 * When enabled prevents the distance field from being discarded due to the mesh being open, but also lowers Distance Field AO quality.
+	 */
+	UPROPERTY(EditAnywhere, BlueprintReadWrite, Category = "Static Meshes", meta=(SubCategory = "Build", DisplayName="Two-Sided Distance Field Generation"))
+	bool bGenerateDistanceFieldAsIfTwoSided = false;
+	
+	UPROPERTY(EditAnywhere, BlueprintReadWrite, Category = "Static Meshes", meta = (SubCategory = "Build", DisplayName="Enable Physical Material Mask"))
+	bool bSupportFaceRemap = false;
+	
+	UPROPERTY(EditAnywhere, BlueprintReadWrite, Category = "Static Meshes", meta = (SubCategory = "Build"))
+	int32 MinLightmapResolution = 64;
+	
+	UPROPERTY(EditAnywhere, BlueprintReadWrite, Category = "Static Meshes", meta = (SubCategory = "Build", DisplayName="Source Lightmap Index"))
+	int32 SrcLightmapIndex = 0;
+	
+	UPROPERTY(EditAnywhere, BlueprintReadWrite, Category = "Static Meshes", meta = (SubCategory = "Build", DisplayName="Destination Lightmap Index"))
+	int32 DstLightmapIndex = 1;
+	
+	/** The local scale applied when building the mesh */
+	UPROPERTY(EditAnywhere, BlueprintReadWrite, Category = "Static Meshes", meta = (SubCategory = "Build", DisplayName="Build Scale"))
+	FVector BuildScale3D = FVector(1.0);
+	
+	/** 
+	 * Scale to apply to the mesh when allocating the distance field volume texture.
+	 * The default scale is 1, which is assuming that the mesh will be placed unscaled in the world.
+	 */
+	UPROPERTY(EditAnywhere, BlueprintReadWrite, Category = "Static Meshes", meta = (SubCategory = "Build"))
+	float DistanceFieldResolutionScale = 1.0f;
 	
+	UPROPERTY(EditAnywhere, BlueprintReadWrite, Category = "Static Meshes", meta = (SubCategory = "Build"))
+	TObjectPtr<class UStaticMesh> DistanceFieldReplacementMesh = nullptr;
+	
+	/** 
+	 * Max Lumen mesh cards to generate for this mesh.
+	 * More cards means that surface will have better coverage, but will result in increased runtime overhead.
+	 * Set to 0 in order to disable mesh card generation for this mesh.
+	 * Default is 12.
+	 */
+	UPROPERTY(EditAnywhere, BlueprintReadWrite, Category = "Static Meshes", meta = (SubCategory = "Build"))
+	int32 MaxLumenMeshCards = 12;
+
 	//////	SKELETAL_MESHES_CATEGORY Properties //////
 
 	/** If enable, import the animation asset find in the sources. */
@@ -105,6 +155,22 @@ public:
 	UPROPERTY(EditAnywhere, BlueprintReadWrite, Category = "Skeletal Meshes", meta = (editcondition = "!bCreatePhysicsAsset"))
 	TObjectPtr<UPhysicsAsset> PhysicsAsset;
 
+	/** Threshold use to decide if two vertex position are equal. */
+	UPROPERTY(EditAnywhere, BlueprintReadWrite, Category = "Skeletal Meshes", meta = (SubCategory = "Build"))
+	float ThresholdPosition = 0.00002f;
+	
+	/** Threshold use to decide if two normal, tangents or bi-normals are equal. */
+	UPROPERTY(EditAnywhere, BlueprintReadWrite, Category = "Skeletal Meshes", meta = (SubCategory = "Build"))
+	float ThresholdTangentNormal = 0.00002f;
+	
+	/** Threshold use to decide if two UVs are equal. */
+	UPROPERTY(EditAnywhere, BlueprintReadWrite, Category = "Skeletal Meshes", meta = (SubCategory = "Build"))
+	float ThresholdUV = 0.0009765625f;
+	
+	/** Threshold to compare vertex position equality when computing morph target deltas. */
+	UPROPERTY(EditAnywhere, BlueprintReadWrite, Category = "Skeletal Meshes", meta = (SubCategory = "Build"))
+	float MorphThresholdPosition = 0.015f;
+
 	virtual void AdjustSettingsForContext(EInterchangePipelineContext ImportType, TObjectPtr<UObject> ReimportAsset) override;
 
 	virtual void PreDialogCleanup(const FName PipelineStackName) override;
@@ -121,6 +187,18 @@ protected:
 
 	virtual void SetReimportSourceIndex(UClass* ReimportObjectClass, const int32 SourceFileIndex) override;
 
+#if WITH_EDITOR
+	/**
+	 * This function return true if all UPROPERTYs of the @Struct exist in the provided @Classes.
+	 * @Struct UPROPERTY tested must be: not transient, editable
+	 * 
+	 * @param Classes - The array of UClass that should contains the Struct properties
+	 * @param Struct - The struct that has the referenced properties
+	 * 
+	 */
+	static bool DoClassesIncludeAllEditableStructProperties(const TArray<const UClass*>& Classes, const UStruct* Struct);
+#endif
+
 private:
 
 	/* Meshes utilities, to parse the translated graph and extract the meshes informations. */
diff --git a/Engine/Source/Runtime/Experimental/Interchange/Core/Private/InterchangePipelineBase.cpp b/Engine/Source/Runtime/Experimental/Interchange/Core/Private/InterchangePipelineBase.cpp
--- a/Engine/Source/Runtime/Experimental/Interchange/Core/Private/InterchangePipelineBase.cpp
+++ b/Engine/Source/Runtime/Experimental/Interchange/Core/Private/InterchangePipelineBase.cpp
@@ -43,6 +43,10 @@ void UInterchangePipelineBase::SaveSettings(const FName PipelineStackName)
 
 void UInterchangePipelineBase::AdjustSettingsForContext(EInterchangePipelineContext ReimportType, TObjectPtr<UObject> ReimportAsset)
 {
+	CachePipelineContext = ReimportType;
+	CacheReimportObject = ReimportAsset;
+	CachePropertiesStates = PropertiesStates;
+
 	bAllowPropertyStatesEdition = (ReimportType == EInterchangePipelineContext::None);
 	bIsReimportContext = false;
 	switch (ReimportType)
@@ -58,6 +62,12 @@ void UInterchangePipelineBase::AdjustSettingsForContext(EInterchangePipelineCont
 	}
 }
 
+void UInterchangePipelineBase::AdjustSettingsFromCache()
+{
+	PropertiesStates = CachePropertiesStates;
+	AdjustSettingsForContext(CachePipelineContext, CacheReimportObject.Get());
+}
+
 const FInterchangePipelinePropertyStates* UInterchangePipelineBase::GetPropertyStates(const FName PropertyPath) const
 {
 	return PropertiesStates.Find(PropertyPath);
@@ -285,11 +295,10 @@ UInterchangePipelineBase* UInterchangePipelineBase::GetMostPipelineOuter() const
 	return Top;
 }
 
-void UInterchangePipelineBase::HidePropertiesOfCategory(UInterchangePipelineBase* OuterMostPipeline, UInterchangePipelineBase* Pipeline, const FString& HideCategoryName, bool bDoTransientSubPipeline /*= false*/)
+void UInterchangePipelineBase::InternalToggleVisibilityPropertiesOfMetaDataValue(UInterchangePipelineBase* OuterMostPipeline, UInterchangePipelineBase* Pipeline, bool bDoTransientSubPipeline, const FString& MetaDataKey, const FString& MetaDataValue, const bool VisibilityState)
 {
 #if WITH_EDITOR
 	UClass* PipelineClass = Pipeline->GetClass();
-	const FName CategoryKey("Category");
 	for (FProperty* Property = PipelineClass->PropertyLink; Property; Property = Property->PropertyLinkNext)
 	{
 		FObjectProperty* SubObject = CastField<FObjectProperty>(Property);
@@ -303,19 +312,31 @@ void UInterchangePipelineBase::HidePropertiesOfCategory(UInterchangePipelineBase
 
 		const FName PropertyName = Property->GetFName();
 		const FName PropertyPath = FName(Property->GetPathName());
-		
+
 		if (SubPipeline)
 		{
-			HidePropertiesOfCategory(OuterMostPipeline, SubPipeline, HideCategoryName);
+			InternalToggleVisibilityPropertiesOfMetaDataValue(OuterMostPipeline, SubPipeline, bDoTransientSubPipeline, MetaDataKey, MetaDataValue, VisibilityState);
 		}
 		else
 		{
-			FString CategoryName = Property->GetMetaData(CategoryKey);
-			if (CategoryName.Equals(HideCategoryName))
+			FString CategoryName = Property->GetMetaData(*MetaDataKey);
+			if (CategoryName.Equals(MetaDataValue))
 			{
-				OuterMostPipeline->FindOrAddPropertyStates(PropertyPath).ReimportStates.bVisible = false;
+				OuterMostPipeline->FindOrAddPropertyStates(PropertyPath).ReimportStates.bVisible = VisibilityState;
 			}
 		}
 	}
 #endif
 }
+
+void UInterchangePipelineBase::HidePropertiesOfCategory(UInterchangePipelineBase* OuterMostPipeline, UInterchangePipelineBase* Pipeline, const FString& HideCategoryName, bool bDoTransientSubPipeline /*= false*/)
+{
+	constexpr bool bVisibilityState = false;
+	InternalToggleVisibilityPropertiesOfMetaDataValue(OuterMostPipeline, Pipeline, bDoTransientSubPipeline, TEXT("Category"), HideCategoryName, bVisibilityState);
+}
+
+void UInterchangePipelineBase::HidePropertiesOfSubCategory(UInterchangePipelineBase* OuterMostPipeline, UInterchangePipelineBase* Pipeline, const FString& HideSubCategoryName, bool bDoTransientSubPipeline /*= false*/)
+{
+	constexpr bool bVisibilityState = false;
+	InternalToggleVisibilityPropertiesOfMetaDataValue(OuterMostPipeline, Pipeline, bDoTransientSubPipeline, TEXT("SubCategory"), HideSubCategoryName, bVisibilityState);
+}
diff --git a/Engine/Source/Runtime/Experimental/Interchange/Core/Private/Nodes/InterchangeBaseNode.cpp b/Engine/Source/Runtime/Experimental/Interchange/Core/Private/Nodes/InterchangeBaseNode.cpp
--- a/Engine/Source/Runtime/Experimental/Interchange/Core/Private/Nodes/InterchangeBaseNode.cpp
+++ b/Engine/Source/Runtime/Experimental/Interchange/Core/Private/Nodes/InterchangeBaseNode.cpp
@@ -60,12 +60,6 @@ namespace UE
 			return AttributeKey;
 		}
 
-		const FAttributeKey& FBaseNodeStaticData::ReimportStrategyFlagsKey()
-		{
-			static FAttributeKey AttributeKey(TEXT("__Reimport_Strategy_Flags_Key__"));
-			return AttributeKey;
-		}
-
 	} //ns Interchange
 } //ns UE
 
@@ -75,7 +69,6 @@ UInterchangeBaseNode::UInterchangeBaseNode()
 	TargetNodes.Initialize(Attributes, UE::Interchange::FBaseNodeStaticData::TargetAssetIDsKey().ToString());
 	RegisterAttribute<bool>(UE::Interchange::FBaseNodeStaticData::IsEnabledKey(), true);
 	RegisterAttribute<uint8>(UE::Interchange::FBaseNodeStaticData::NodeContainerTypeKey(), static_cast<uint8>(EInterchangeNodeContainerType::None));
-	RegisterAttribute<uint8>(UE::Interchange::FBaseNodeStaticData::ReimportStrategyFlagsKey(), static_cast<uint8>(EReimportStrategyFlags::ApplyNoProperties));
 }
 
 void UInterchangeBaseNode::InitializeNode(const FString& UniqueID, const FString& DisplayLabel, const EInterchangeNodeContainerType NodeContainerType)
@@ -144,10 +137,6 @@ FString UInterchangeBaseNode::GetKeyDisplayName(const UE::Interchange::FAttribut
 	{
 		KeyDisplayName = TEXT("Node Container Type");
 	}
-	else if (NodeAttributeKey == UE::Interchange::FBaseNodeStaticData::ReimportStrategyFlagsKey())
-	{
-		KeyDisplayName = TEXT("Re-Import Strategy");
-	}
 	else if (NodeAttributeKey == UE::Interchange::FBaseNodeStaticData::TargetAssetIDsKey())
 	{
 		KeyDisplayName = TEXT("Target Asset Count");
@@ -306,25 +295,6 @@ bool UInterchangeBaseNode::SetDisplayLabel(const FString& DisplayLabel)
 	return false;
 }
 
-EReimportStrategyFlags UInterchangeBaseNode::GetReimportStrategyFlags() const
-{
-	checkSlow(Attributes->ContainAttribute(UE::Interchange::FBaseNodeStaticData::ReimportStrategyFlagsKey()));
-	uint8 ReimportStrategyFlags = static_cast<uint8>(EReimportStrategyFlags::ApplyNoProperties);
-	Attributes->GetAttributeHandle<uint8>(UE::Interchange::FBaseNodeStaticData::ReimportStrategyFlagsKey()).Get(ReimportStrategyFlags);
-	return static_cast<EReimportStrategyFlags>(ReimportStrategyFlags);
-}
-
-bool UInterchangeBaseNode::SetReimportStrategyFlags(const EReimportStrategyFlags& ReimportStrategyFlags)
-{
-	UE::Interchange::EAttributeStorageResult Result = Attributes->RegisterAttribute(UE::Interchange::FBaseNodeStaticData::ReimportStrategyFlagsKey(), static_cast<uint8>(ReimportStrategyFlags));
-	if (IsAttributeStorageResultSuccess(Result))
-	{
-		UE::Interchange::FAttributeStorage::TAttributeHandle<uint8> Handle = Attributes->GetAttributeHandle<uint8>(UE::Interchange::FBaseNodeStaticData::ReimportStrategyFlagsKey());
-		return Handle.IsValid();
-	}
-	return false;
-}
-
 FString UInterchangeBaseNode::GetParentUid() const
 {
 	if (!Attributes->ContainAttribute(UE::Interchange::FBaseNodeStaticData::ParentIDKey()))
@@ -456,42 +426,6 @@ FString UInterchangeBaseNode::InvalidNodeUid()
 	return FString();
 }
 
-void UInterchangeBaseNode::ApplyAllCustomAttributeToObject(UObject* Object) const
-{
-	UClass* ObjectClass = Object->GetClass();
-	for (const TPair<UClass*, TArray<UE::Interchange::FApplyAttributeToAsset>>& ClassDelegatePair : ApplyCustomAttributeDelegates)
-	{
-		if(ObjectClass->IsChildOf(ClassDelegatePair.Key))
-		{
-			for (const UE::Interchange::FApplyAttributeToAsset& Delegate : ClassDelegatePair.Value)
-			{
-				if(Delegate.IsBound())
-				{
-					Delegate.Execute(Object);
-				}
-			}
-		}
-	}
-}
-
-void UInterchangeBaseNode::FillAllCustomAttributeFromObject(UObject* Object) const
-{
-	UClass* ObjectClass = Object->GetClass();
-	for (const TPair<UClass*, TArray<UE::Interchange::FFillAttributeToAsset>>& ClassDelegatePair : FillCustomAttributeDelegates)
-	{
-		if (ObjectClass->IsChildOf(ClassDelegatePair.Key))
-		{
-			for (const UE::Interchange::FFillAttributeToAsset& Delegate : ClassDelegatePair.Value)
-			{
-				if (Delegate.IsBound())
-				{
-					Delegate.Execute(Object);
-				}
-			}
-		}
-	}
-}
-
 void UInterchangeBaseNode::Serialize(FArchive& Ar)
 {
 	UE::Interchange::FAttributeStorage& RefAttributes = *(Attributes.Get());
diff --git a/Engine/Source/Runtime/Experimental/Interchange/Core/Private/Nodes/InterchangeFactoryBaseNode.cpp b/Engine/Source/Runtime/Experimental/Interchange/Core/Private/Nodes/InterchangeFactoryBaseNode.cpp
--- a/Engine/Source/Runtime/Experimental/Interchange/Core/Private/Nodes/InterchangeFactoryBaseNode.cpp
+++ b/Engine/Source/Runtime/Experimental/Interchange/Core/Private/Nodes/InterchangeFactoryBaseNode.cpp
@@ -18,6 +18,12 @@ namespace UE
 			return BaseNodeFactoryDependencies_BaseKey;
 		}
 
+		const FAttributeKey& FFactoryBaseNodeStaticData::ReimportStrategyFlagsKey()
+		{
+			static FAttributeKey AttributeKey(TEXT("__Reimport_Strategy_Flags_Key__"));
+			return AttributeKey;
+		}
+
 	} //ns Interchange
 } //ns UE
 
@@ -25,6 +31,8 @@ namespace UE
 UInterchangeFactoryBaseNode::UInterchangeFactoryBaseNode()
 {
 	FactoryDependencies.Initialize(Attributes, UE::Interchange::FFactoryBaseNodeStaticData::FactoryDependenciesBaseKey());
+
+	RegisterAttribute<uint8>(UE::Interchange::FFactoryBaseNodeStaticData::ReimportStrategyFlagsKey(), static_cast<uint8>(EReimportStrategyFlags::ApplyNoProperties));
 }
 
 
@@ -50,6 +58,10 @@ FString UInterchangeFactoryBaseNode::GetKeyDisplayName(const UE::Interchange::FA
 			KeyDisplayName += OriginalKeyName.RightChop(IndexPosition);
 		}
 	}
+	else if (NodeAttributeKey == UE::Interchange::FFactoryBaseNodeStaticData::ReimportStrategyFlagsKey())
+	{
+		KeyDisplayName = TEXT("Re-Import Strategy");
+	}
 	else
 	{
 		KeyDisplayName = Super::GetKeyDisplayName(NodeAttributeKey);
@@ -77,6 +89,25 @@ UClass* UInterchangeFactoryBaseNode::GetObjectClass() const
 	return nullptr;
 }
 
+EReimportStrategyFlags UInterchangeFactoryBaseNode::GetReimportStrategyFlags() const
+{
+	checkSlow(Attributes->ContainAttribute(UE::Interchange::FFactoryBaseNodeStaticData::ReimportStrategyFlagsKey()));
+	uint8 ReimportStrategyFlags = static_cast<uint8>(EReimportStrategyFlags::ApplyNoProperties);
+	Attributes->GetAttributeHandle<uint8>(UE::Interchange::FFactoryBaseNodeStaticData::ReimportStrategyFlagsKey()).Get(ReimportStrategyFlags);
+	return static_cast<EReimportStrategyFlags>(ReimportStrategyFlags);
+}
+
+bool UInterchangeFactoryBaseNode::SetReimportStrategyFlags(const EReimportStrategyFlags& ReimportStrategyFlags)
+{
+	UE::Interchange::EAttributeStorageResult Result = Attributes->RegisterAttribute(UE::Interchange::FFactoryBaseNodeStaticData::ReimportStrategyFlagsKey(), static_cast<uint8>(ReimportStrategyFlags));
+	if (IsAttributeStorageResultSuccess(Result))
+	{
+		UE::Interchange::FAttributeStorage::TAttributeHandle<uint8> Handle = Attributes->GetAttributeHandle<uint8>(UE::Interchange::FFactoryBaseNodeStaticData::ReimportStrategyFlagsKey());
+		return Handle.IsValid();
+	}
+	return false;
+}
+
 bool UInterchangeFactoryBaseNode::GetCustomSubPath(FString& AttributeValue) const
 {
 	IMPLEMENT_NODE_ATTRIBUTE_GETTER(SubPath, FString)
@@ -116,3 +147,39 @@ FString UInterchangeFactoryBaseNode::BuildFactoryNodeUid(const FString& Translat
 {
 	return TEXT("Factory_") + TranslatedNodeUid;
 }
+
+void UInterchangeFactoryBaseNode::ApplyAllCustomAttributeToObject(UObject* Object) const
+{
+	UClass* ObjectClass = Object->GetClass();
+	for (const TPair<UClass*, TArray<UE::Interchange::FApplyAttributeToAsset>>& ClassDelegatePair : ApplyCustomAttributeDelegates)
+	{
+		if (ObjectClass->IsChildOf(ClassDelegatePair.Key))
+		{
+			for (const UE::Interchange::FApplyAttributeToAsset& Delegate : ClassDelegatePair.Value)
+			{
+				if (Delegate.IsBound())
+				{
+					Delegate.Execute(Object);
+				}
+			}
+		}
+	}
+}
+
+void UInterchangeFactoryBaseNode::FillAllCustomAttributeFromObject(UObject* Object) const
+{
+	UClass* ObjectClass = Object->GetClass();
+	for (const TPair<UClass*, TArray<UE::Interchange::FFillAttributeToAsset>>& ClassDelegatePair : FillCustomAttributeDelegates)
+	{
+		if (ObjectClass->IsChildOf(ClassDelegatePair.Key))
+		{
+			for (const UE::Interchange::FFillAttributeToAsset& Delegate : ClassDelegatePair.Value)
+			{
+				if (Delegate.IsBound())
+				{
+					Delegate.Execute(Object);
+				}
+			}
+		}
+	}
+}
\ No newline at end of file
diff --git a/Engine/Source/Runtime/Experimental/Interchange/Core/Private/Types/AttributeStorage.cpp b/Engine/Source/Runtime/Experimental/Interchange/Core/Private/Types/AttributeStorage.cpp
--- a/Engine/Source/Runtime/Experimental/Interchange/Core/Private/Types/AttributeStorage.cpp
+++ b/Engine/Source/Runtime/Experimental/Interchange/Core/Private/Types/AttributeStorage.cpp
@@ -395,8 +395,14 @@ namespace UE
 			const uint64 NumberOfChar = AttributeAllocationInfo->Size / sizeof(TCHAR);
 			check(NumberOfChar > 0);
 
+			//Since the null character is always there, a empty String is only one TCHAR
+			if (NumberOfChar <= 1)
+			{
+				OutValue.Empty();
+				return;
+			}
 
-			//Empty and pre allocate the FString buffer (TArray<TCHAR>)
+			//Empty and pre allocate the FString buffer (TArray<TCHAR>) and add the null terminator
 			OutValue.Reset(NumberOfChar - 1);
 			//SLOW! We have to add the characters 1 x 1, there is no AddZeroed or similar stuff. We do not have a string to use the operator= which do what we need here
 			//We go up to NumberOfChar -1 because there is always a null terminating character add at the end of the FString
diff --git a/Engine/Source/Runtime/Experimental/Interchange/Core/Public/InterchangePipelineBase.h b/Engine/Source/Runtime/Experimental/Interchange/Core/Public/InterchangePipelineBase.h
--- a/Engine/Source/Runtime/Experimental/Interchange/Core/Public/InterchangePipelineBase.h
+++ b/Engine/Source/Runtime/Experimental/Interchange/Core/Public/InterchangePipelineBase.h
@@ -216,6 +216,7 @@ public:
 	 * @Param ReimportAsset - This is an optional parameter which is set when re-importing an asset.
 	 */
 	virtual void AdjustSettingsForContext(EInterchangePipelineContext ReimportType, TObjectPtr<UObject> ReimportAsset);
+	virtual void AdjustSettingsFromCache();
 
 	/**
 	 * This function is used to add the given message object directly into the results for this operation.
@@ -314,8 +315,10 @@ protected:
 
 	UInterchangePipelineBase* GetMostPipelineOuter() const;
 
+	static void InternalToggleVisibilityPropertiesOfMetaDataValue(UInterchangePipelineBase* OuterMostPipeline, UInterchangePipelineBase* Pipeline, bool bDoTransientSubPipeline, const FString& MetaDataKey, const FString& MetaDataValue, const bool bVisibilityState);
 	static void HidePropertiesOfCategory(UInterchangePipelineBase* OuterMostPipeline, UInterchangePipelineBase* Pipeline, const FString& HideCategoryName, bool bDoTransientSubPipeline = false);
-
+	static void HidePropertiesOfSubCategory(UInterchangePipelineBase* OuterMostPipeline, UInterchangePipelineBase* Pipeline, const FString& HideSubCategoryName, bool bDoTransientSubPipeline = false);
+	
 	/**
 	 * If true, the property editor for this pipeline instance will allow properties states edition.
 	 * If false, the property editor for this pipeline instance will apply the properties states.
@@ -342,4 +345,8 @@ protected:
 	*/
 	UPROPERTY()
 	TMap<FName, FInterchangePipelinePropertyStates> PropertiesStates;
+
+	mutable TMap<FName, FInterchangePipelinePropertyStates> CachePropertiesStates;
+	mutable EInterchangePipelineContext CachePipelineContext = EInterchangePipelineContext::None;
+	mutable TWeakObjectPtr<UObject> CacheReimportObject = nullptr;
 };
diff --git a/Engine/Source/Runtime/Experimental/Interchange/Core/Public/Nodes/InterchangeBaseNode.h b/Engine/Source/Runtime/Experimental/Interchange/Core/Public/Nodes/InterchangeBaseNode.h
--- a/Engine/Source/Runtime/Experimental/Interchange/Core/Public/Nodes/InterchangeBaseNode.h
+++ b/Engine/Source/Runtime/Experimental/Interchange/Core/Public/Nodes/InterchangeBaseNode.h
@@ -114,26 +114,9 @@ namespace InterchangePrivateNodeBase
  * @param AssetType - This is the asset you want to apply the storage value"
  * @param EnumType - Optional, specify it only if the AssetType member is an enum so we can type cast it in the apply function (we use uint8 to store the enum value)"
  */
-#define IMPLEMENT_NODE_ATTRIBUTE_KEY(AttributeName)																\
+#define IMPLEMENT_NODE_ATTRIBUTE_KEY(AttributeName)																				\
 const UE::Interchange::FAttributeKey Macro_Custom##AttributeName##Key = UE::Interchange::FAttributeKey(TEXT(#AttributeName));
 
-#if WITH_ENGINE
-#define IMPLEMENT_NODE_ATTRIBUTE_DELEGATE_BY_PROPERTYNAME(AttributeName, AttributeType, ObjectType, PropertyName)		\
-bool ApplyCustom##AttributeName##ToAsset(UObject* Asset) const														\
-{																													\
-	return ApplyAttributeToObject<AttributeType>(Macro_Custom##AttributeName##Key.Key, Asset, PropertyName);	\
-}																													\
-																													\
-bool FillCustom##AttributeName##FromAsset(UObject* Asset)															\
-{																													\
-	return FillAttributeFromObject<AttributeType>(Macro_Custom##AttributeName##Key.Key, Asset, PropertyName);\
-}
-
-#else //#if WITH_ENGINE
-
-#define IMPLEMENT_NODE_ATTRIBUTE_DELEGATE_BY_PROPERTYNAME(AttributeName, AttributeType, AssetType, PropertyName)
-
-#endif //#if WITH_ENGINE
 
 #define IMPLEMENT_NODE_ATTRIBUTE_GETTER(AttributeName, AttributeType)																					\
 	FString OperationName = GetTypeName() + TEXT(".Get" #AttributeName);																				\
@@ -143,68 +126,25 @@ bool FillCustom##AttributeName##FromAsset(UObject* Asset)															\
 	FString OperationName = GetTypeName() + TEXT(".Set" #AttributeName);																				\
 	return InterchangePrivateNodeBase::SetCustomAttribute<AttributeType>(*Attributes, Macro_Custom##AttributeName##Key, OperationName, AttributeValue);
 
-#if WITH_ENGINE
-
-#define IMPLEMENT_NODE_ATTRIBUTE_SETTER(NodeClassName, AttributeName, AttributeType, AssetType)														\
-	FString OperationName = GetTypeName() + TEXT(".Set" #AttributeName);																			\
-	if(InterchangePrivateNodeBase::SetCustomAttribute<AttributeType>(*Attributes, Macro_Custom##AttributeName##Key, OperationName, AttributeValue))	\
-	{																																				\
-		if(bAddApplyDelegate)																														\
-		{																																			\
-			FName AttributeNameImpl = TEXT(""#AttributeName);																						\
-			AddApplyAndFillDelegates<AttributeType>(Macro_Custom##AttributeName##Key.Key, AssetType::StaticClass(), AttributeNameImpl);				\
-		}																																			\
-		return true;																																\
-	}																																				\
-	return false;
-
-#define IMPLEMENT_NODE_ATTRIBUTE_SETTER_WITH_CUSTOM_DELEGATE(NodeClassName, AttributeName, AttributeType, AssetType)								\
-	FString OperationName = GetTypeName() + TEXT(".Set" #AttributeName);																			\
-	if(InterchangePrivateNodeBase::SetCustomAttribute<AttributeType>(*Attributes, Macro_Custom##AttributeName##Key, OperationName, AttributeValue))	\
-	{																																				\
-		if(bAddApplyDelegate)																														\
-		{																																			\
-			TArray<UE::Interchange::FApplyAttributeToAsset>& Delegates = ApplyCustomAttributeDelegates.FindOrAdd(AssetType::StaticClass());			\
-			Delegates.Add(UE::Interchange::FApplyAttributeToAsset::CreateUObject(this, &NodeClassName::ApplyCustom##AttributeName##ToAsset));		\
-			TArray<UE::Interchange::FFillAttributeToAsset>& FillDelegates = FillCustomAttributeDelegates.FindOrAdd(AssetType::StaticClass());		\
-			FillDelegates.Add(UE::Interchange::FFillAttributeToAsset::CreateUObject(this, &NodeClassName::FillCustom##AttributeName##FromAsset));	\
-		}																																			\
-		return true;																																\
-	}																																				\
-	return false;
-
-#else //#if WITH_ENGINE
-
-#define IMPLEMENT_NODE_ATTRIBUTE_SETTER(NodeClassName, AttributeName, AttributeType, AssetType)															\
-	FString OperationName = GetTypeName() + TEXT(".Set" #AttributeName);																				\
-	return InterchangePrivateNodeBase::SetCustomAttribute<AttributeType>(*Attributes, Macro_Custom##AttributeName##Key, OperationName, AttributeValue);
-
-#define IMPLEMENT_NODE_ATTRIBUTE_SETTER_WITH_CUSTOM_DELEGATE(NodeClassName, AttributeName, AttributeType, AssetType)								\
-	IMPLEMENT_NODE_ATTRIBUTE_SETTER(NodeClassName, AttributeName, AttributeType, AssetType)
 
-#endif //#if WITH_ENGINE
+#define INTERCHANGE_BASE_NODE_ADD_ATTRIBUTE(ValueType)																												\
+	UE::Interchange::FAttributeStorage::TAttributeHandle<ValueType> Handle = RegisterAttribute<ValueType>(UE::Interchange::FAttributeKey(NodeAttributeKey), Value);	\
+	return Handle.IsValid();
 
-#define INTERCHANGE_BASE_NODE_ADD_ATTRIBUTE(ValueType)																									\
-UE::Interchange::FAttributeStorage::TAttributeHandle<ValueType> Handle = RegisterAttribute<ValueType>(UE::Interchange::FAttributeKey(NodeAttributeKey), Value);	\
-return Handle.IsValid();
-
-#define INTERCHANGE_BASE_NODE_GET_ATTRIBUTE(ValueType)																						\
-if(HasAttribute(UE::Interchange::FAttributeKey(NodeAttributeKey)))																			\
-{																																			\
-	UE::Interchange::FAttributeStorage::TAttributeHandle<ValueType> Handle = GetAttributeHandle<ValueType>(UE::Interchange::FAttributeKey(NodeAttributeKey));	\
-	if (Handle.IsValid())																													\
-	{																																		\
-		return (Handle.Get(OutValue) == UE::Interchange::EAttributeStorageResult::Operation_Success);										\
-	}																																		\
-}																																			\
-return false;
+#define INTERCHANGE_BASE_NODE_GET_ATTRIBUTE(ValueType)																							\
+	if(HasAttribute(UE::Interchange::FAttributeKey(NodeAttributeKey)))																			\
+	{																																			\
+		UE::Interchange::FAttributeStorage::TAttributeHandle<ValueType> Handle = GetAttributeHandle<ValueType>(UE::Interchange::FAttributeKey(NodeAttributeKey));	\
+		if (Handle.IsValid())																													\
+		{																																		\
+			return (Handle.Get(OutValue) == UE::Interchange::EAttributeStorageResult::Operation_Success);										\
+		}																																		\
+	}																																			\
+	return false;
 
 //Interchange namespace
 namespace UE::Interchange
 {
-	DECLARE_DELEGATE_RetVal_OneParam(bool, FApplyAttributeToAsset, UObject*);
-	DECLARE_DELEGATE_RetVal_OneParam(bool, FFillAttributeToAsset, UObject*);
-
 	/**
 		* Helper struct use to declare static const data we use in the UInterchangeBaseNode
 		* Node that derive from UInterchangeBaseNode can also add a struct that derive from this one to add there static data
@@ -220,7 +160,6 @@ namespace UE::Interchange
 		static const FAttributeKey& ClassTypeAttributeKey();
 		static const FAttributeKey& AssetNameKey();
 		static const FAttributeKey& NodeContainerTypeKey();
-		static const FAttributeKey& ReimportStrategyFlagsKey();
 	};
 
 } //ns UE::Interchange
@@ -234,17 +173,6 @@ enum class EInterchangeNodeContainerType : uint8
 	FactoryData
 };
 
-UENUM()
-enum class EReimportStrategyFlags : uint8
-{
-	//Do not apply any property when re-importing, simply change the source data
-	ApplyNoProperties,
-	//Always apply all pipeline specified properties
-	ApplyPipelineProperties,
-	//Always apply all pipeline properties, but leave the properties modified in editor since the last import
-	ApplyEditorChangedProperties
-};
-
 /**
  * This struct is used to store and retrieve key value attributes. The attributes are store in a generic FAttributeStorage which serialize the value in a TArray64<uint8>
  * See UE::Interchange::EAttributeTypes to know the supported template types
@@ -468,35 +396,6 @@ public:
 		INTERCHANGE_BASE_NODE_ADD_ATTRIBUTE(AttributeType);
 	}
 
-	/**
-	 * Adds the delegates that will read and write the attribute value to a UObject.
-	 * @param NodeAttributeKey	The key of the attribute for which we are adding the delegates.
-	 * @param ObjectClass		The class of the object that we will apply and fill on.
-	 * @param PropertyName		The name of the property of the ObjectClass that we will read and write to.
-	 */
-	template<typename AttributeType>
-	inline void AddApplyAndFillDelegates(const FName& NodeAttributeKey, UClass* ObjectClass, const FName PropertyName);
-
-	/**
-	 * Writes an attribute value to a UObject.
-	 * @param NodeAttributeKey	The key for the attribute to write.
-	 * @param Object			The object to write to.
-	 * @param PropertyName		The name of the property to write to on the Object.
-	 * @result					True if we succeeded.
-	 */
-	template<typename AttributeType>
-	inline bool ApplyAttributeToObject(const FName& NodeAttributeKey, UObject* Object, const FName PropertyName) const;
-
-	/**
-	 * Reads an attribute value from a UObject.
-	 * @param NodeAttributeKey	The key for the attribute to update.
-	 * @param Object			The object to read from.
-	 * @param PropertyName		The name of the property to read from on the Object.
-	 * @result					True if we succeeded.
-	 */
-	template<typename AttributeType>
-	inline bool FillAttributeFromObject(const FName& NodeAttributeKey, UObject* Object, const FName PropertyName);
-
 	/**
 	 * Return the unique id pass in the constructor.
 	 */
@@ -515,20 +414,6 @@ public:
 	UFUNCTION(BlueprintCallable, Category = "Interchange | Node")
 	bool SetDisplayLabel(const FString& DisplayName);
 
-	/**
-	 * Return the reimport strategy flags.
-	 *
-	 */
-	UFUNCTION(BlueprintCallable, Category = "Interchange | Node")
-	EReimportStrategyFlags GetReimportStrategyFlags() const;
-
-	/**
-	 * Change the reimport strategy flags.
-	 *
-	 */
-	UFUNCTION(BlueprintCallable, Category = "Interchange | Node")
-	bool SetReimportStrategyFlags(const EReimportStrategyFlags& ReimportStrategyFlags);
-
 	/**
 	 * Return the parent unique id. In case the attribute does not exist it will return InvalidNodeUid()
 	 */
@@ -612,15 +497,6 @@ public:
 	/** Return the invalid unique ID */
 	static FString InvalidNodeUid();
 
-	/**
-	 * Each Attribute that was set and have a delegate set for the specified UObject->UClass will
-	 * get the delegate execute so it apply the attribute to the UObject property.
-	 * See the macros IMPLEMENT_NODE_ATTRIBUTE_SETTER at the top of the file to know how delegates are setup for property.
-	 */
-	void ApplyAllCustomAttributeToObject(UObject* Object) const;
-
-	void FillAllCustomAttributeFromObject(UObject* Object) const;
-
 	/**
 	 * Serialize the node, by default only the attribute storage is serialize for a node
 	 */
@@ -641,11 +517,6 @@ protected:
 	/** The storage use to store the Key value attribute for this node. */
 	TSharedPtr<UE::Interchange::FAttributeStorage, ESPMode::ThreadSafe> Attributes;
 
-	/* This array hold the delegate to apply the attribute that has to be set on an UObject */
-	TMap<UClass*, TArray<UE::Interchange::FApplyAttributeToAsset>> ApplyCustomAttributeDelegates;
-
-	TMap<UClass*, TArray<UE::Interchange::FFillAttributeToAsset>> FillCustomAttributeDelegates;
-
 	bool bIsInitialized = false;
 
 	/**
@@ -653,247 +524,3 @@ protected:
 	 */
 	mutable UE::Interchange::TArrayAttributeHelper<FString> TargetNodes;
 };
-
-template<typename AttributeType>
-inline void UInterchangeBaseNode::AddApplyAndFillDelegates(const FName& NodeAttributeKey, UClass* ObjectClass, const FName PropertyName)
-{
-	TArray<UE::Interchange::FApplyAttributeToAsset>& ApplyDelegates = ApplyCustomAttributeDelegates.FindOrAdd(ObjectClass);
-	ApplyDelegates.Add(
-		UE::Interchange::FApplyAttributeToAsset::CreateLambda(
-		[NodeAttributeKey, PropertyName, this](UObject* Asset)
-		{
-			return ApplyAttributeToObject<AttributeType>(NodeAttributeKey, Asset, PropertyName);
-		})
-	);
-
-	TArray<UE::Interchange::FFillAttributeToAsset>& FillDelegates = FillCustomAttributeDelegates.FindOrAdd(ObjectClass);
-	FillDelegates.Add(
-		UE::Interchange::FApplyAttributeToAsset::CreateLambda(
-		[NodeAttributeKey, PropertyName, this](UObject* Asset)
-		{
-			return FillAttributeFromObject<AttributeType>(NodeAttributeKey, Asset, PropertyName);
-		})
-	);
-}
-
-template<typename AttributeType>
-inline bool UInterchangeBaseNode::ApplyAttributeToObject(const FName& NodeAttributeKey, UObject* Object, const FName PropertyName) const
-{
-	if (!Object)
-	{
-		return false;
-	}
-
-	AttributeType ValueData;
-	if (GetAttribute<AttributeType>(NodeAttributeKey, ValueData))
-	{
-		TVariant<UObject*, uint8*> Container;
-		Container.Set<UObject*>(Object);
-		if(FProperty* Property = InterchangePrivateNodeBase::FindPropertyByPathChecked(Container, Object->GetClass(), PropertyName.ToString()))
-		{
-			AttributeType* PropertyValueAddress;
-			if (Container.IsType<UObject*>())
-			{
-				PropertyValueAddress = Property->ContainerPtrToValuePtr<AttributeType>(Container.Get<UObject*>());
-			}
-			else
-			{
-				PropertyValueAddress = Property->ContainerPtrToValuePtr<AttributeType>(Container.Get<uint8*>());
-			}
-
-			*PropertyValueAddress = ValueData;
-		}
-		return true;
-	}
-	return false;
-}
-
-/**
- * Specialized version of ApplyAttributeToObject for strings.
- * If the target property is a FObjectPropertyBase, treat the string as an object path.
- */
-template<>
-inline bool UInterchangeBaseNode::ApplyAttributeToObject<FString>(const FName& NodeAttributeKey, UObject* Object, const FName PropertyName) const
-{
-	if (!Object)
-	{
-		return false;
-	}
-
-	FString ValueData;
-	if (GetAttribute<FString>(NodeAttributeKey, ValueData))
-	{
-		TVariant<UObject*, uint8*> Container;
-		Container.Set<UObject*>(Object);
-		if(FProperty* Property = InterchangePrivateNodeBase::FindPropertyByPathChecked(Container, Object->GetClass(), PropertyName.ToString()))
-		{
-			FString* PropertyValueAddress;
-			if (Container.IsType<UObject*>())
-			{
-				PropertyValueAddress = Property->ContainerPtrToValuePtr<FString>(Container.Get<UObject*>());
-			}
-			else
-			{
-				PropertyValueAddress = Property->ContainerPtrToValuePtr<FString>(Container.Get<uint8*>());
-			}
-
-			if (FObjectPropertyBase* ObjectProperty = CastField<FObjectPropertyBase>(Property))
-			{
-				ObjectProperty->SetObjectPropertyValue(PropertyValueAddress, FSoftObjectPath(ValueData).TryLoad());
-			}
-			else
-			{
-				*PropertyValueAddress = ValueData;
-			}
-		}
-		return true;
-	}
-	return false;
-}
-
-/**
- * Specialized version of ApplyAttributeToObject for bools.
- * If the target property is a FBoolProperty, treat the propertyg as a bitfield.
- */
-template<>
-inline bool UInterchangeBaseNode::ApplyAttributeToObject<bool>(const FName& NodeAttributeKey, UObject* Object, const FName PropertyName) const
-{
-	if (!Object)
-	{
-		return false;
-	}
-
-	bool bValueData;
-	if (GetAttribute<bool>(NodeAttributeKey, bValueData))
-	{
-		TVariant<UObject*, uint8*> Container;
-		Container.Set<UObject*>(Object);
-		if(FProperty* Property = InterchangePrivateNodeBase::FindPropertyByPathChecked(Container, Object->GetClass(), PropertyName.ToString()))
-		{
-			bool* PropertyValueAddress;
-			if (Container.IsType<UObject*>())
-			{
-				PropertyValueAddress = Property->ContainerPtrToValuePtr<bool>(Container.Get<UObject*>());
-			}
-			else
-			{
-				PropertyValueAddress = Property->ContainerPtrToValuePtr<bool>(Container.Get<uint8*>());
-			}
-
-			// Support for bitfields
-			if (FBoolProperty* BoolProperty = CastField<FBoolProperty>(Property))
-			{
-				BoolProperty->SetPropertyValue(PropertyValueAddress, bValueData);
-			}
-			else
-			{
-				*PropertyValueAddress = bValueData;
-			}
-		}
-		return true;
-	}
-	return false;
-}
-
-template<typename AttributeType>
-inline bool UInterchangeBaseNode::FillAttributeFromObject(const FName& NodeAttributeKey, UObject* Object, const FName PropertyName)
-{
-	TVariant<UObject*, uint8*> Container;
-	Container.Set<UObject*>(Object);
-	if(FProperty* Property = InterchangePrivateNodeBase::FindPropertyByPathChecked(Container, Object->GetClass(), PropertyName.ToString()))
-	{
-		AttributeType* PropertyValueAddress;
-		if (Container.IsType<UObject*>())
-		{
-			PropertyValueAddress = Property->ContainerPtrToValuePtr<AttributeType>(Container.Get<UObject*>());
-		}
-		else
-		{
-			PropertyValueAddress = Property->ContainerPtrToValuePtr<AttributeType>(Container.Get<uint8*>());
-		}
-
-		// Support for bitfields
-		if (FBoolProperty* BoolProperty = CastField<FBoolProperty>(Property))
-		{
-			const bool bPropertyValue = BoolProperty->GetPropertyValue(PropertyValueAddress);
-			return SetAttribute(NodeAttributeKey, bPropertyValue);
-		}
-		else
-		{
-			return SetAttribute(NodeAttributeKey, *PropertyValueAddress);
-		}
-	}
-
-	return false;
-}
-
-/**
- * Specialized version of FillAttributeFromObject for strings.
- * If the target property is a FObjectPropertyBase, treat the string as an object path.
- */
-template<>
-inline bool UInterchangeBaseNode::FillAttributeFromObject<FString>(const FName& NodeAttributeKey, UObject* Object, const FName PropertyName)
-{
-	TVariant<UObject*, uint8*> Container;
-	Container.Set<UObject*>(Object);
-	if(FProperty* Property = InterchangePrivateNodeBase::FindPropertyByPathChecked(Container, Object->GetClass(), PropertyName.ToString()))
-	{
-		FString* PropertyValueAddress;
-		if (Container.IsType<UObject*>())
-		{
-			PropertyValueAddress = Property->ContainerPtrToValuePtr<FString>(Container.Get<UObject*>());
-		}
-		else
-		{
-			PropertyValueAddress = Property->ContainerPtrToValuePtr<FString>(Container.Get<uint8*>());
-		}
-
-		if (FObjectPropertyBase* ObjectProperty = CastField<FObjectPropertyBase>(Property))
-		{
-			UObject* ObjectValue = ObjectProperty->GetObjectPropertyValue(PropertyValueAddress);
-			return SetAttribute(NodeAttributeKey, ObjectValue->GetPathName());
-		}
-		else
-		{
-			return SetAttribute(NodeAttributeKey, *PropertyValueAddress);
-		}
-	}
-
-	return false;
-}
-
-/**
- * Specialized version of FillAttributeFromObject for bools.
- * If the target property is a FBoolProperty, treat the property as a bitfield.
- */
-template<>
-inline bool UInterchangeBaseNode::FillAttributeFromObject<bool>(const FName& NodeAttributeKey, UObject* Object, const FName PropertyName)
-{
-	TVariant<UObject*, uint8*> Container;
-	Container.Set<UObject*>(Object);
-	if(FProperty* Property = InterchangePrivateNodeBase::FindPropertyByPathChecked(Container, Object->GetClass(), PropertyName.ToString()))
-	{
-		bool* PropertyValueAddress;
-		if (Container.IsType<UObject*>())
-		{
-			PropertyValueAddress = Property->ContainerPtrToValuePtr<bool>(Container.Get<UObject*>());
-		}
-		else
-		{
-			PropertyValueAddress = Property->ContainerPtrToValuePtr<bool>(Container.Get<uint8*>());
-		}
-
-		// Support for bitfields
-		if (FBoolProperty* BoolProperty = CastField<FBoolProperty>(Property))
-		{
-			const bool bPropertyValue = BoolProperty->GetPropertyValue(PropertyValueAddress);
-			return SetAttribute(NodeAttributeKey, bPropertyValue);
-		}
-		else
-		{
-			return SetAttribute(NodeAttributeKey, *PropertyValueAddress);
-		}
-	}
-
-	return false;
-}
diff --git a/Engine/Source/Runtime/Experimental/Interchange/Core/Public/Nodes/InterchangeFactoryBaseNode.h b/Engine/Source/Runtime/Experimental/Interchange/Core/Public/Nodes/InterchangeFactoryBaseNode.h
--- a/Engine/Source/Runtime/Experimental/Interchange/Core/Public/Nodes/InterchangeFactoryBaseNode.h
+++ b/Engine/Source/Runtime/Experimental/Interchange/Core/Public/Nodes/InterchangeFactoryBaseNode.h
@@ -20,19 +20,101 @@
 class UObject;
 struct FFrame;
 
+#define IMPLEMENT_NODE_ATTRIBUTE_DELEGATE_BY_PROPERTYNAME(AttributeName, AttributeType, ObjectType, PropertyName)		\
+	bool ApplyCustom##AttributeName##ToAsset(UObject* Asset) const														\
+	{																													\
+		return ApplyAttributeToObject<AttributeType>(Macro_Custom##AttributeName##Key.Key, Asset, PropertyName);		\
+	}																													\
+																														\
+	bool FillCustom##AttributeName##FromAsset(UObject* Asset)															\
+	{																													\
+		return FillAttributeFromObject<AttributeType>(Macro_Custom##AttributeName##Key.Key, Asset, PropertyName);		\
+	}
 
-namespace UE
+#define IMPLEMENT_NODE_ATTRIBUTE_SETTER(NodeClassName, AttributeName, AttributeType, AssetType)														\
+	FString OperationName = GetTypeName() + TEXT(".Set" #AttributeName);																			\
+	if(InterchangePrivateNodeBase::SetCustomAttribute<AttributeType>(*Attributes, Macro_Custom##AttributeName##Key, OperationName, AttributeValue))	\
+	{																																				\
+		if(bAddApplyDelegate)																														\
+		{																																			\
+			FName AttributeNameImpl = TEXT(""#AttributeName);																						\
+			AddApplyAndFillDelegates<AttributeType>(Macro_Custom##AttributeName##Key.Key, AssetType::StaticClass(), AttributeNameImpl);				\
+		}																																			\
+		return true;																																\
+	}																																				\
+	return false;
+
+/*
+ * Use this macro to implement a custom node attribute setter that will add delegate to apply the node attribute value to any UObject derive from
+ * the specified type class. Use this macro if the type class is different from the Node GetObjectCLass()
+ *
+ * The node should derive from UInterchangeFactoryBaseNode to apply a delegate
+ * The node need to implement 2 functions for the attribute
+ *  - ApplyCustom##AttributeName##ToAsset(UObject*) , see UE::Interchange::FApplyAttributeToAsset delegate signature
+ *  - FillCustom##AttributeName##FromAsset(UObject*) , see UE::Interchange::FFillAttributeToAsset delegate signature
+ */
+#define IMPLEMENT_NODE_ATTRIBUTE_SETTER_WITH_CUSTOM_DELEGATE_WITH_CUSTOM_CLASS(NodeClassName, AttributeName, AttributeType, AssetType)				\
+	FString OperationName = GetTypeName() + TEXT(".Set" #AttributeName);																			\
+	if(InterchangePrivateNodeBase::SetCustomAttribute<AttributeType>(*Attributes, Macro_Custom##AttributeName##Key, OperationName, AttributeValue))	\
+	{																																				\
+		if(bAddApplyDelegate)																														\
+		{																																			\
+			check(NodeClassName::StaticClass()->IsChildOf<UInterchangeFactoryBaseNode>());															\
+			UClass* AssetClass = AssetType::StaticClass();																							\
+			TArray<UE::Interchange::FApplyAttributeToAsset>& Delegates = ApplyCustomAttributeDelegates.FindOrAdd(AssetClass);						\
+			Delegates.Add(UE::Interchange::FApplyAttributeToAsset::CreateUObject(this, &NodeClassName::ApplyCustom##AttributeName##ToAsset));		\
+			TArray<UE::Interchange::FFillAttributeToAsset>& FillDelegates = FillCustomAttributeDelegates.FindOrAdd(AssetClass);						\
+			FillDelegates.Add(UE::Interchange::FFillAttributeToAsset::CreateUObject(this, &NodeClassName::FillCustom##AttributeName##FromAsset));	\
+		}																																			\
+		return true;																																\
+	}																																				\
+	return false;
+
+/*
+ * This macro instead of receiving a TypeClass use the factory node GetObjectClass() to register the delegate
+ * See documentation of IMPLEMENT_NODE_ATTRIBUTE_SETTER_WITH_CUSTOM_DELEGATE_WITH_CUSTOM_CLASS
+ */
+#define IMPLEMENT_NODE_ATTRIBUTE_SETTER_WITH_CUSTOM_DELEGATE(NodeClassName, AttributeName, AttributeType)											\
+	FString OperationName = GetTypeName() + TEXT(".Set" #AttributeName);																			\
+	if(InterchangePrivateNodeBase::SetCustomAttribute<AttributeType>(*Attributes, Macro_Custom##AttributeName##Key, OperationName, AttributeValue))	\
+	{																																				\
+		if(bAddApplyDelegate)																														\
+		{																																			\
+			check(NodeClassName::StaticClass()->IsChildOf<UInterchangeFactoryBaseNode>());															\
+			TArray<UE::Interchange::FApplyAttributeToAsset>& Delegates = ApplyCustomAttributeDelegates.FindOrAdd(GetObjectClass());					\
+			Delegates.Add(UE::Interchange::FApplyAttributeToAsset::CreateUObject(this, &NodeClassName::ApplyCustom##AttributeName##ToAsset));		\
+			TArray<UE::Interchange::FFillAttributeToAsset>& FillDelegates = FillCustomAttributeDelegates.FindOrAdd(GetObjectClass());				\
+			FillDelegates.Add(UE::Interchange::FFillAttributeToAsset::CreateUObject(this, &NodeClassName::FillCustom##AttributeName##FromAsset));	\
+		}																																			\
+		return true;																																\
+	}																																				\
+	return false;
+
+
+namespace UE::Interchange
 {
-	namespace Interchange
+
+	DECLARE_DELEGATE_RetVal_OneParam(bool, FApplyAttributeToAsset, UObject*);
+	DECLARE_DELEGATE_RetVal_OneParam(bool, FFillAttributeToAsset, UObject*);
+
+	struct INTERCHANGECORE_API FFactoryBaseNodeStaticData : public FBaseNodeStaticData
 	{
-		struct INTERCHANGECORE_API FFactoryBaseNodeStaticData : public FBaseNodeStaticData
-		{
-			static const FString& FactoryDependenciesBaseKey();
-		};
+		static const FString& FactoryDependenciesBaseKey();
+		static const FAttributeKey& ReimportStrategyFlagsKey();
+	};
 
-	} // namespace Interchange
-} // namespace UE
+} // namespace UE::Interchange
 
+UENUM()
+enum class EReimportStrategyFlags : uint8
+{
+	//Do not apply any property when re-importing, simply change the source data
+	ApplyNoProperties,
+	//Always apply all pipeline specified properties
+	ApplyPipelineProperties,
+	//Always apply all pipeline properties, but leave the properties modified in editor since the last import
+	ApplyEditorChangedProperties
+};
 
 /**
  * This struct is used to store and retrieve key value attributes. The attributes are store in a generic FAttributeStorage which serialize the value in a TArray64<uint8>
@@ -50,6 +132,20 @@ public:
 	virtual FString GetKeyDisplayName(const UE::Interchange::FAttributeKey& NodeAttributeKey) const override;
 	virtual FString GetAttributeCategory(const UE::Interchange::FAttributeKey& NodeAttributeKey) const override;
 
+	/**
+	 * Return the reimport strategy flags.
+	 *
+	 */
+	UFUNCTION(BlueprintCallable, Category = "Interchange | Node")
+	EReimportStrategyFlags GetReimportStrategyFlags() const;
+
+	/**
+	 * Change the reimport strategy flags.
+	 *
+	 */
+	UFUNCTION(BlueprintCallable, Category = "Interchange | Node")
+	bool SetReimportStrategyFlags(const EReimportStrategyFlags& ReimportStrategyFlags);
+
 	/**
 	 * Return the UClass of the object we represent so we can find factory/writer
 	 */
@@ -103,6 +199,44 @@ public:
 
 	static FString BuildFactoryNodeUid(const FString& TranslatedNodeUid);
 
+	/**
+	 * Adds the delegates that will read and write the attribute value to a UObject.
+	 * @param NodeAttributeKey	The key of the attribute for which we are adding the delegates.
+	 * @param ObjectClass		The class of the object that we will apply and fill on.
+	 * @param PropertyName		The name of the property of the ObjectClass that we will read and write to.
+	 */
+	template<typename AttributeType>
+	inline void AddApplyAndFillDelegates(const FName& NodeAttributeKey, UClass* ObjectClass, const FName PropertyName);
+
+	/**
+	 * Writes an attribute value to a UObject.
+	 * @param NodeAttributeKey	The key for the attribute to write.
+	 * @param Object			The object to write to.
+	 * @param PropertyName		The name of the property to write to on the Object.
+	 * @result					True if we succeeded.
+	 */
+	template<typename AttributeType>
+	inline bool ApplyAttributeToObject(const FName& NodeAttributeKey, UObject* Object, const FName PropertyName) const;
+
+	/**
+	 * Reads an attribute value from a UObject.
+	 * @param NodeAttributeKey	The key for the attribute to update.
+	 * @param Object			The object to read from.
+	 * @param PropertyName		The name of the property to read from on the Object.
+	 * @result					True if we succeeded.
+	 */
+	template<typename AttributeType>
+	inline bool FillAttributeFromObject(const FName& NodeAttributeKey, UObject* Object, const FName PropertyName);
+
+	/**
+	 * Each Attribute that was set and have a delegate set for the specified UObject->UClass will
+	 * get the delegate execute so it apply the attribute to the UObject property.
+	 * See the macros IMPLEMENT_NODE_ATTRIBUTE_SETTER at the top of the file to know how delegates are setup for property.
+	 */
+	void ApplyAllCustomAttributeToObject(UObject* Object) const;
+
+	void FillAllCustomAttributeFromObject(UObject* Object) const;
+
 protected:
 	/**
 	 * Those dependencies are use by the interchange parsing task to make sure the asset are created in the correct order.
@@ -111,9 +245,257 @@ protected:
 	 */
 	UE::Interchange::TArrayAttributeHelper<FString> FactoryDependencies;
 
+	/* This array hold the delegate to apply the attribute that has to be set on an UObject */
+	TMap<UClass*, TArray<UE::Interchange::FApplyAttributeToAsset>> ApplyCustomAttributeDelegates;
+
+	TMap<UClass*, TArray<UE::Interchange::FFillAttributeToAsset>> FillCustomAttributeDelegates;
 
 private:
 
 	const UE::Interchange::FAttributeKey Macro_CustomSubPathKey = UE::Interchange::FAttributeKey(TEXT("SubPath"));
 
-};
\ No newline at end of file
+};
+
+template<typename AttributeType>
+inline void UInterchangeFactoryBaseNode::AddApplyAndFillDelegates(const FName& NodeAttributeKey, UClass* ObjectClass, const FName PropertyName)
+{
+	TArray<UE::Interchange::FApplyAttributeToAsset>& ApplyDelegates = ApplyCustomAttributeDelegates.FindOrAdd(ObjectClass);
+	ApplyDelegates.Add(
+		UE::Interchange::FApplyAttributeToAsset::CreateLambda(
+			[NodeAttributeKey, PropertyName, this](UObject* Asset)
+			{
+				return ApplyAttributeToObject<AttributeType>(NodeAttributeKey, Asset, PropertyName);
+			})
+	);
+
+	TArray<UE::Interchange::FFillAttributeToAsset>& FillDelegates = FillCustomAttributeDelegates.FindOrAdd(ObjectClass);
+	FillDelegates.Add(
+		UE::Interchange::FApplyAttributeToAsset::CreateLambda(
+			[NodeAttributeKey, PropertyName, this](UObject* Asset)
+			{
+				return FillAttributeFromObject<AttributeType>(NodeAttributeKey, Asset, PropertyName);
+			})
+	);
+}
+
+template<typename AttributeType>
+inline bool UInterchangeFactoryBaseNode::ApplyAttributeToObject(const FName& NodeAttributeKey, UObject* Object, const FName PropertyName) const
+{
+	if (!Object)
+	{
+		return false;
+	}
+
+	AttributeType ValueData;
+	if (GetAttribute<AttributeType>(NodeAttributeKey, ValueData))
+	{
+		TVariant<UObject*, uint8*> Container;
+		Container.Set<UObject*>(Object);
+		if (FProperty* Property = InterchangePrivateNodeBase::FindPropertyByPathChecked(Container, Object->GetClass(), PropertyName.ToString()))
+		{
+			AttributeType* PropertyValueAddress;
+			if (Container.IsType<UObject*>())
+			{
+				PropertyValueAddress = Property->ContainerPtrToValuePtr<AttributeType>(Container.Get<UObject*>());
+			}
+			else
+			{
+				PropertyValueAddress = Property->ContainerPtrToValuePtr<AttributeType>(Container.Get<uint8*>());
+			}
+
+			*PropertyValueAddress = ValueData;
+		}
+		return true;
+	}
+	return false;
+}
+
+/**
+ * Specialized version of ApplyAttributeToObject for strings.
+ * If the target property is a FObjectPropertyBase, treat the string as an object path.
+ */
+template<>
+inline bool UInterchangeFactoryBaseNode::ApplyAttributeToObject<FString>(const FName& NodeAttributeKey, UObject* Object, const FName PropertyName) const
+{
+	if (!Object)
+	{
+		return false;
+	}
+
+	FString ValueData;
+	if (GetAttribute<FString>(NodeAttributeKey, ValueData))
+	{
+		TVariant<UObject*, uint8*> Container;
+		Container.Set<UObject*>(Object);
+		if (FProperty* Property = InterchangePrivateNodeBase::FindPropertyByPathChecked(Container, Object->GetClass(), PropertyName.ToString()))
+		{
+			FString* PropertyValueAddress;
+			if (Container.IsType<UObject*>())
+			{
+				PropertyValueAddress = Property->ContainerPtrToValuePtr<FString>(Container.Get<UObject*>());
+			}
+			else
+			{
+				PropertyValueAddress = Property->ContainerPtrToValuePtr<FString>(Container.Get<uint8*>());
+			}
+
+			if (FObjectPropertyBase* ObjectProperty = CastField<FObjectPropertyBase>(Property))
+			{
+				ObjectProperty->SetObjectPropertyValue(PropertyValueAddress, FSoftObjectPath(ValueData).TryLoad());
+			}
+			else
+			{
+				*PropertyValueAddress = ValueData;
+			}
+		}
+		return true;
+	}
+	return false;
+}
+
+/**
+ * Specialized version of ApplyAttributeToObject for bools.
+ * If the target property is a FBoolProperty, treat the propertyg as a bitfield.
+ */
+template<>
+inline bool UInterchangeFactoryBaseNode::ApplyAttributeToObject<bool>(const FName& NodeAttributeKey, UObject* Object, const FName PropertyName) const
+{
+	if (!Object)
+	{
+		return false;
+	}
+
+	bool bValueData;
+	if (GetAttribute<bool>(NodeAttributeKey, bValueData))
+	{
+		TVariant<UObject*, uint8*> Container;
+		Container.Set<UObject*>(Object);
+		if (FProperty* Property = InterchangePrivateNodeBase::FindPropertyByPathChecked(Container, Object->GetClass(), PropertyName.ToString()))
+		{
+			bool* PropertyValueAddress;
+			if (Container.IsType<UObject*>())
+			{
+				PropertyValueAddress = Property->ContainerPtrToValuePtr<bool>(Container.Get<UObject*>());
+			}
+			else
+			{
+				PropertyValueAddress = Property->ContainerPtrToValuePtr<bool>(Container.Get<uint8*>());
+			}
+
+			// Support for bitfields
+			if (FBoolProperty* BoolProperty = CastField<FBoolProperty>(Property))
+			{
+				BoolProperty->SetPropertyValue(PropertyValueAddress, bValueData);
+			}
+			else
+			{
+				*PropertyValueAddress = bValueData;
+			}
+		}
+		return true;
+	}
+	return false;
+}
+
+template<typename AttributeType>
+inline bool UInterchangeFactoryBaseNode::FillAttributeFromObject(const FName& NodeAttributeKey, UObject* Object, const FName PropertyName)
+{
+	TVariant<UObject*, uint8*> Container;
+	Container.Set<UObject*>(Object);
+	if (FProperty* Property = InterchangePrivateNodeBase::FindPropertyByPathChecked(Container, Object->GetClass(), PropertyName.ToString()))
+	{
+		AttributeType* PropertyValueAddress;
+		if (Container.IsType<UObject*>())
+		{
+			PropertyValueAddress = Property->ContainerPtrToValuePtr<AttributeType>(Container.Get<UObject*>());
+		}
+		else
+		{
+			PropertyValueAddress = Property->ContainerPtrToValuePtr<AttributeType>(Container.Get<uint8*>());
+		}
+
+		// Support for bitfields
+		if (FBoolProperty* BoolProperty = CastField<FBoolProperty>(Property))
+		{
+			const bool bPropertyValue = BoolProperty->GetPropertyValue(PropertyValueAddress);
+			return SetAttribute(NodeAttributeKey, bPropertyValue);
+		}
+		else
+		{
+			return SetAttribute(NodeAttributeKey, *PropertyValueAddress);
+		}
+	}
+
+	return false;
+}
+
+/**
+ * Specialized version of FillAttributeFromObject for strings.
+ * If the target property is a FObjectPropertyBase, treat the string as an object path.
+ */
+template<>
+inline bool UInterchangeFactoryBaseNode::FillAttributeFromObject<FString>(const FName& NodeAttributeKey, UObject* Object, const FName PropertyName)
+{
+	TVariant<UObject*, uint8*> Container;
+	Container.Set<UObject*>(Object);
+	if (FProperty* Property = InterchangePrivateNodeBase::FindPropertyByPathChecked(Container, Object->GetClass(), PropertyName.ToString()))
+	{
+		FString* PropertyValueAddress;
+		if (Container.IsType<UObject*>())
+		{
+			PropertyValueAddress = Property->ContainerPtrToValuePtr<FString>(Container.Get<UObject*>());
+		}
+		else
+		{
+			PropertyValueAddress = Property->ContainerPtrToValuePtr<FString>(Container.Get<uint8*>());
+		}
+
+		if (FObjectPropertyBase* ObjectProperty = CastField<FObjectPropertyBase>(Property))
+		{
+			UObject* ObjectValue = ObjectProperty->GetObjectPropertyValue(PropertyValueAddress);
+			return SetAttribute(NodeAttributeKey, ObjectValue->GetPathName());
+		}
+		else
+		{
+			return SetAttribute(NodeAttributeKey, *PropertyValueAddress);
+		}
+	}
+
+	return false;
+}
+
+/**
+ * Specialized version of FillAttributeFromObject for bools.
+ * If the target property is a FBoolProperty, treat the property as a bitfield.
+ */
+template<>
+inline bool UInterchangeFactoryBaseNode::FillAttributeFromObject<bool>(const FName& NodeAttributeKey, UObject* Object, const FName PropertyName)
+{
+	TVariant<UObject*, uint8*> Container;
+	Container.Set<UObject*>(Object);
+	if (FProperty* Property = InterchangePrivateNodeBase::FindPropertyByPathChecked(Container, Object->GetClass(), PropertyName.ToString()))
+	{
+		bool* PropertyValueAddress;
+		if (Container.IsType<UObject*>())
+		{
+			PropertyValueAddress = Property->ContainerPtrToValuePtr<bool>(Container.Get<UObject*>());
+		}
+		else
+		{
+			PropertyValueAddress = Property->ContainerPtrToValuePtr<bool>(Container.Get<uint8*>());
+		}
+
+		// Support for bitfields
+		if (FBoolProperty* BoolProperty = CastField<FBoolProperty>(Property))
+		{
+			const bool bPropertyValue = BoolProperty->GetPropertyValue(PropertyValueAddress);
+			return SetAttribute(NodeAttributeKey, bPropertyValue);
+		}
+		else
+		{
+			return SetAttribute(NodeAttributeKey, *PropertyValueAddress);
+		}
+	}
+
+	return false;
+}
