ChaosFlesh : Per-Frame updates
-Added framework to support per-frame updates from the game thread to physics thread.

diff --git a/Engine/Source/Runtime/Experimental/Chaos/Private/Chaos/Deformable/ChaosDeformableSolver.cpp b/Engine/Source/Runtime/Experimental/Chaos/Private/Chaos/Deformable/ChaosDeformableSolver.cpp
--- a/Engine/Source/Runtime/Experimental/Chaos/Private/Chaos/Deformable/ChaosDeformableSolver.cpp
+++ b/Engine/Source/Runtime/Experimental/Chaos/Private/Chaos/Deformable/ChaosDeformableSolver.cpp
@@ -35,7 +35,8 @@ DECLARE_CYCLE_STAT(TEXT("DeformableSolver.Advance"), STAT_DeformableSolver_Advan
 DEFINE_LOG_CATEGORY_STATIC(LogChaosDeformableSolver, Log, All);
 namespace Chaos::Softs
 {
-	FCriticalSection FDeformableSolver::PackageMutex;
+	FCriticalSection FDeformableSolver::PackageOutputMutex;
+	FCriticalSection FDeformableSolver::PackageInputMutex;
 
 
 	FDeformableSolver::FDeformableSolver(FDeformableSolverProperties InProp)
@@ -396,6 +397,40 @@ namespace Chaos::Softs
 	}
 
 
+	void FDeformableSolver::FGameThreadAccess::PushInputPackage(int32 InFrame, FDeformableDataMap&& InPackage)
+	{
+		Solver.PushInputPackage(InFrame, MoveTemp(InPackage));
+	}
+	void FDeformableSolver::PushInputPackage(int32 InFrame, FDeformableDataMap&& InPackage)
+	{
+		FScopeLock Lock(&PackageInputMutex);
+		TRACE_CPUPROFILER_EVENT_SCOPE(DeformableSolver_PushInputPackage);
+		InputPackages.Push(TUniquePtr< FDeformablePackage >(new FDeformablePackage(InFrame, MoveTemp(InPackage))));
+	}
+
+	TUniquePtr<FDeformablePackage> FDeformableSolver::FPhysicsThreadAccess::PullInputPackage()
+	{
+		return Solver.PullInputPackage();
+	}
+	TUniquePtr<FDeformablePackage> FDeformableSolver::PullInputPackage()
+	{
+		FScopeLock Lock(&PackageInputMutex);
+		TRACE_CPUPROFILER_EVENT_SCOPE(DeformableSolver_PullInputPackage);
+		if (InputPackages.Num())
+			return InputPackages.Pop();
+		return TUniquePtr<FDeformablePackage>(nullptr);
+	}
+
+	void FDeformableSolver::FPhysicsThreadAccess::UpdateProxyInputPackages()
+	{
+		return Solver.UpdateProxyInputPackages();
+	}
+	void FDeformableSolver::UpdateProxyInputPackages()
+	{
+		// todo(chaosflesh) : just clear the packages for now. 
+		while (PullInputPackage()) {}
+	}
+
 	void FDeformableSolver::FPhysicsThreadAccess::TickSimulation(FSolverReal DeltaTime)
 	{
 		Solver.TickSimulation(DeltaTime);
@@ -418,49 +453,48 @@ namespace Chaos::Softs
 		}
 
 
-		FOutputDataMap OutputBuffers;
+		FDeformableDataMap OutputBuffers;
 		for (TPair< FThreadingProxy::FKey, TUniquePtr<FThreadingProxy> > & BaseProxyPair : Proxies)
 		{
 			UpdateOutputState(*BaseProxyPair.Value);
 			if (FFleshThreadingProxy* Proxy = BaseProxyPair.Value->As<FFleshThreadingProxy>())
 			{
-				OutputBuffers.Add(Proxy->GetOwner(), TSharedPtr<FThreadingProxy::FOutputBuffer>(new FFleshThreadingProxy::FFleshOutputBuffer(*Proxy)));
+				OutputBuffers.Add(Proxy->GetOwner(), TSharedPtr<FThreadingProxy::FBuffer>(new FFleshThreadingProxy::FFleshOutputBuffer(*Proxy)));
 
 				if (Property.CacheToFile)
 				{
 					WriteFrame(*Proxy, DeltaTime);
 				}
 			}
-
 		}
 
-		PushPackage(Frame, MoveTemp(OutputBuffers));
+		PushOutputPackage(Frame, MoveTemp(OutputBuffers));
 	}
 
 
-	void FDeformableSolver::FPhysicsThreadAccess::PushPackage(int32 InFrame, FOutputDataMap&& InPackage)
+	void FDeformableSolver::FPhysicsThreadAccess::PushOutputPackage(int32 InFrame, FDeformableDataMap&& InPackage)
 	{
-		Solver.PushPackage(InFrame, MoveTemp(InPackage));
+		Solver.PushOutputPackage(InFrame, MoveTemp(InPackage));
 	}
-	void FDeformableSolver::PushPackage(int32 InFrame, FOutputDataMap&& InPackage)
+	void FDeformableSolver::PushOutputPackage(int32 InFrame, FDeformableDataMap&& InPackage)
 	{
-		FScopeLock Lock(&PackageMutex);
+		FScopeLock Lock(&PackageOutputMutex);
 		TRACE_CPUPROFILER_EVENT_SCOPE(DeformableSolver_PushPackage);
-		OutputPackages.Push(TUniquePtr< FOutputPackage >(new FOutputPackage(InFrame, MoveTemp(InPackage))));
+		OutputPackages.Push(TUniquePtr< FDeformablePackage >(new FDeformablePackage(InFrame, MoveTemp(InPackage))));
 	}
 
 
-	TUniquePtr<FOutputPackage>  FDeformableSolver::FGameThreadAccess::PullPackage()
+	TUniquePtr<FDeformablePackage>  FDeformableSolver::FGameThreadAccess::PullOutputPackage()
 	{
-		return Solver.PullPackage();
+		return Solver.PullOutputPackage();
 	}
-	TUniquePtr<FOutputPackage> FDeformableSolver::PullPackage()
+	TUniquePtr<FDeformablePackage> FDeformableSolver::PullOutputPackage()
 	{
-		FScopeLock Lock(&PackageMutex);
+		FScopeLock Lock(&PackageOutputMutex);
 		TRACE_CPUPROFILER_EVENT_SCOPE(DeformableSolver_PullPackage);
 		if (OutputPackages.Num())
 			return OutputPackages.Pop();
-		return TUniquePtr<FOutputPackage>(nullptr);
+		return TUniquePtr<FDeformablePackage>(nullptr);
 	}
 
 	void  FDeformableSolver::FGameThreadAccess::AddProxy(TUniquePtr<FThreadingProxy> InObject)
diff --git a/Engine/Source/Runtime/Experimental/Chaos/Public/Chaos/Deformable/ChaosDeformableSolver.h b/Engine/Source/Runtime/Experimental/Chaos/Public/Chaos/Deformable/ChaosDeformableSolver.h
--- a/Engine/Source/Runtime/Experimental/Chaos/Public/Chaos/Deformable/ChaosDeformableSolver.h
+++ b/Engine/Source/Runtime/Experimental/Chaos/Public/Chaos/Deformable/ChaosDeformableSolver.h
@@ -38,11 +38,13 @@ namespace Chaos::Softs
 			FPhysicsThreadAccess(FDeformableSolver& InSolver, const FPhysicsThreadAccessor&) : Solver(InSolver) {}
 
 			/* Simulation Advance */
+			void UpdateProxyInputPackages();
 			bool Advance(FSolverReal DeltaTime);
 			void Reset(const FDeformableSolverProperties&);
 			void TickSimulation(FSolverReal DeltaTime);
 			void UpdateOutputState(FThreadingProxy&);
-			void PushPackage(int32 Frame, FOutputDataMap&& Package);
+			TUniquePtr<FDeformablePackage> PullInputPackage();
+			void PushOutputPackage(int32 Frame, FDeformableDataMap&& Package);
 
 			/* Iteration Advance */
 			void InitializeSimulationObjects();
@@ -75,9 +77,12 @@ namespace Chaos::Softs
 		public:
 			FGameThreadAccess(FDeformableSolver& InSolver, const FGameThreadAccessor&) : Solver(InSolver) {}
 
+			int32 GetFrame() const { return Solver.GetFrame(); }
 			void AddProxy(TUniquePtr<FThreadingProxy> InObject);
 
-			TUniquePtr<FOutputPackage> PullPackage();
+			void PushInputPackage(int32 Frame, FDeformableDataMap&& InPackage);
+
+			TUniquePtr<FDeformablePackage> PullOutputPackage();
 
 		private:
 			FDeformableSolver& Solver;
@@ -86,11 +91,14 @@ namespace Chaos::Softs
 	protected:
 
 		/* Simulation Advance */
+		int32 GetFrame() const { return Frame; }
+		void UpdateProxyInputPackages();
 		bool Advance(FSolverReal DeltaTime);
 		void Reset(const FDeformableSolverProperties&);
 		void TickSimulation(FSolverReal DeltaTime);
 		void UpdateOutputState(FThreadingProxy&);
-		void PushPackage(int32 Frame, FOutputDataMap&& Package);
+		void PushOutputPackage(int32 Frame, FDeformableDataMap&& Package);
+		TUniquePtr<FDeformablePackage> PullInputPackage( );
 
 		/* Iteration Advance */
 		void InitializeSimulationObjects();
@@ -106,17 +114,20 @@ namespace Chaos::Softs
 
 		/*Game Thread API*/
 		void AddProxy(TUniquePtr<FThreadingProxy> InObject);
-		TUniquePtr<FOutputPackage> PullPackage();
+		TUniquePtr<FDeformablePackage> PullOutputPackage();
+		void PushInputPackage(int32 Frame, FDeformableDataMap&& InPackage);
 
 		const FDeformableSolverProperties& GetProperties() const { return Property; }
 
 	private:
 
 		// connections outside the solver.
-		static FCriticalSection	PackageMutex;
+		static FCriticalSection	PackageOutputMutex;
+		static FCriticalSection	PackageInputMutex;
 		TArray< TUniquePtr<FThreadingProxy> > UninitializedProxys;
 		TMap< FThreadingProxy::FKey, TUniquePtr<FThreadingProxy> > Proxies;
-		TArray< TUniquePtr<FOutputPackage>  > OutputPackages;
+		TArray< TUniquePtr<FDeformablePackage>  > InputPackages;
+		TArray< TUniquePtr<FDeformablePackage>  > OutputPackages;
 
 		// User Configuration
 		FDeformableSolverProperties Property;
diff --git a/Engine/Source/Runtime/Experimental/Chaos/Public/Chaos/Deformable/ChaosDeformableSolverProxy.h b/Engine/Source/Runtime/Experimental/Chaos/Public/Chaos/Deformable/ChaosDeformableSolverProxy.h
--- a/Engine/Source/Runtime/Experimental/Chaos/Public/Chaos/Deformable/ChaosDeformableSolverProxy.h
+++ b/Engine/Source/Runtime/Experimental/Chaos/Public/Chaos/Deformable/ChaosDeformableSolverProxy.h
@@ -32,20 +32,28 @@ namespace Chaos::Softs
 		const UObject* GetOwner() const { return Owner; }
 
 
-		class FOutputBuffer
+		class FBuffer
 		{
 			const UObject* Owner = nullptr;
 			FName TypeName = FName("");
 
 		public:
-			FOutputBuffer(const FThreadingProxy& Ref)
+			FBuffer(const FThreadingProxy& Ref)
 				: Owner(Ref.GetOwner())
 				, TypeName(Ref.BaseTypeName()) {}
 
-			FName BaseTypeName() { return TypeName; }
+			FBuffer(const UObject* InOwner,FName InTypeName)
+				: Owner(InOwner)
+				, TypeName(InTypeName) {}
+			virtual ~FBuffer() {}
+
+			FName BaseTypeName() const { return TypeName; }
 
 			template<class T>
 			T* As() { return T::Source::TypeName().IsEqual(BaseTypeName()) ? (T*)this : nullptr; }
+
+			template<class T>
+			T* As() const { return T::Source::TypeName().IsEqual(BaseTypeName()) ? (T*)this : nullptr; }
 		};
 	};
 
@@ -99,9 +107,26 @@ namespace Chaos::Softs
 			return Rest;
 		}
 
-		class FFleshOutputBuffer : public FThreadingProxy::FOutputBuffer
+		class FFleshInputBuffer : public FThreadingProxy::FBuffer
 		{
-			typedef FThreadingProxy::FOutputBuffer Super;
+			typedef FThreadingProxy::FBuffer Super;
+
+		public:
+			FFleshInputBuffer(const UObject* InOwner = nullptr)
+				: Super(InOwner, FFleshThreadingProxy::TypeName())
+			{
+				// Ref.Dynamic will have updated solver data 
+				//InputData.AddAttribute<FTransform>("Transforms", FTransformCollection::TransformGroup);
+				//InputData.CopyMatchingAttributesFrom(Ref.GetDynamicCollection());
+			}
+			virtual ~FFleshInputBuffer() {}
+
+			FManagedArrayCollection InputData;
+		};
+
+		class FFleshOutputBuffer : public FThreadingProxy::FBuffer
+		{
+			typedef FThreadingProxy::FBuffer Super;
 
 		public:
 			typedef FFleshThreadingProxy Source;
@@ -112,9 +137,11 @@ namespace Chaos::Softs
 				Dynamic.AddAttribute<FVector3f>("Vertex", FGeometryCollection::VerticesGroup);
 				Dynamic.CopyMatchingAttributesFrom(Ref.GetDynamicCollection());
 			}
+			virtual ~FFleshOutputBuffer() {}
 
 			FManagedArrayCollection Dynamic;
 		};
 
+
 	};
 }// namespace Chaos::Softs
\ No newline at end of file
diff --git a/Engine/Source/Runtime/Experimental/Chaos/Public/Chaos/Deformable/ChaosDeformableSolverTypes.h b/Engine/Source/Runtime/Experimental/Chaos/Public/Chaos/Deformable/ChaosDeformableSolverTypes.h
--- a/Engine/Source/Runtime/Experimental/Chaos/Public/Chaos/Deformable/ChaosDeformableSolverTypes.h
+++ b/Engine/Source/Runtime/Experimental/Chaos/Public/Chaos/Deformable/ChaosDeformableSolverTypes.h
@@ -62,20 +62,20 @@ namespace Chaos::Softs
 
 
 	/*Data Transfer*/
-	typedef TSharedPtr<FThreadingProxy::FOutputBuffer> FOutputDataMapValue; // OutputBuffer
-	typedef TMap<FThreadingProxy::FKey, FOutputDataMapValue > FOutputDataMap; // <const UObject*,OutputBuffer>
+	typedef TSharedPtr<const FThreadingProxy::FBuffer> FDataMapValue; // Buffer Pointer
+	typedef TMap<FThreadingProxy::FKey, FDataMapValue > FDeformableDataMap; // <const UObject*,FBufferSharedPtr>
 
-	struct CHAOS_API FOutputPackage {
-		FOutputPackage()
+	struct CHAOS_API FDeformablePackage {
+		FDeformablePackage()
 		{}
 
-		FOutputPackage(int32 InFrame, FOutputDataMap&& InMap)
+		FDeformablePackage(int32 InFrame, FDeformableDataMap&& InMap)
 			: Frame(InFrame)
 			, ObjectMap(InMap)
 		{}
 
 		int32 Frame = INDEX_NONE;
-		FOutputDataMap ObjectMap;
+		FDeformableDataMap ObjectMap;
 	};
 
 	/* Accessor for the Game Thread*/
