Reinstate workaround for vehicles slowly sliding down slopes at low speeds.

diff --git a/Engine/Source/Runtime/Experimental/Chaos/Private/Chaos/PBDSuspensionConstraints.cpp b/Engine/Source/Runtime/Experimental/Chaos/Private/Chaos/PBDSuspensionConstraints.cpp
--- a/Engine/Source/Runtime/Experimental/Chaos/Private/Chaos/PBDSuspensionConstraints.cpp
+++ b/Engine/Source/Runtime/Experimental/Chaos/Private/Chaos/PBDSuspensionConstraints.cpp
@@ -22,6 +22,15 @@ FAutoConsoleVariableRef CVarChaosSuspensionMaxPushoutVelocity(TEXT("p.Chaos.Susp
 float Chaos_Suspension_MaxPushout = 5.f;
 FAutoConsoleVariableRef CVarChaosSuspensionMaxPushout(TEXT("p.Chaos.Suspension.MaxPushout"), Chaos_Suspension_MaxPushout, TEXT("Chaos Suspension Max Pushout Value"));
 
+float Chaos_Suspension_SlopeThreshold = 0.707f;	// = Cos(SlopeAngle)
+FAutoConsoleVariableRef CVarChaosSuspensionSlopeThreshold(TEXT("p.Chaos.Suspension.SlopeThreshold"), Chaos_Suspension_SlopeThreshold, TEXT("Slope threshold below which the anti-slide on slope mechanism is employed, value = Cos(AlopeAngle), i.e. for 50 degree slope = 0.6428, 30 degree slope = 0.866"));
+
+float Chaos_Suspension_SlopeSpeedThreshold = 1.0f; // MPH
+FAutoConsoleVariableRef CVarChaosSuspensionSlopeSpeedThreshold(TEXT("p.Chaos.Suspension.SlopeSpeedThreshold"), Chaos_Suspension_SlopeSpeedThreshold, TEXT("Speed below which the anti-slide on slope mechanism is fully employed"));
+
+float Chaos_Suspension_SlopeSpeedBlendThreshold = 10.0f; // MPH
+FAutoConsoleVariableRef CVarChaosSuspensionSlopeSpeedBlendThreshold(TEXT("p.Chaos.Suspension.SlopeSpeedBlendThreshold"), Chaos_Suspension_SlopeSpeedBlendThreshold, TEXT("Speed below which the anti-slide on slope blend mechanism starts"));
+
 #if CHAOS_DEBUG_DRAW
 bool bChaos_Suspension_DebugDraw_Hardstop = false;
 FAutoConsoleVariableRef CVarChaosSuspensionDebugDrawHardstop(TEXT("p.Chaos.Suspension.DebugDraw.Hardstop"), bChaos_Suspension_DebugDraw_Hardstop, TEXT("Debug draw suspension hardstop manifold"));
@@ -314,28 +323,28 @@ namespace Chaos
 
 			FVec3 AxisWorld = BodyQ.RotateVector(SuspensionCoMAxis);
 
-			const FVec3& SurfaceNormal = Setting.Normal;
+			FVec3 SurfaceNormal = Setting.Normal;
 
 			const float MPHToCmS = 100000.f / 2236.94185f;
-			const float SpeedThreshold = 10.0f * MPHToCmS;
-			const float FortyFiveDegreesThreshold = 0.707f;
-
-			// @todo(chaos): do we need this?
-			//if (AxisWorld.Z > FortyFiveDegreesThreshold)
-			//{
-			//	if (Body.V().SquaredLength() < 1.0f)
-			//	{
-			//		AxisWorld = FVec3(0.f, 0.f, 1.f);
-			//	}
-			//	else
-			//	{
-			//		const FReal Speed = FMath::Abs(Body.V().Length());
-			//		if (Speed < SpeedThreshold)
-			//		{
-			//			AxisWorld = FMath::Lerp(FVec3(0.f, 0.f, 1.f), AxisWorld, Speed / SpeedThreshold);
-			//		}
-			//	}
-			//}
+			const float SpeedThreshold = Chaos_Suspension_SlopeSpeedThreshold * MPHToCmS;
+			const float SpeedBlendThreshold = Chaos_Suspension_SlopeSpeedBlendThreshold * MPHToCmS;
+
+			// Ingeniously blends the surface normal at low speeds to stop vehicles sliding slowly down steep slopes
+			if (SurfaceNormal.Z > Chaos_Suspension_SlopeThreshold)
+			{
+				if (Body.V().SquaredLength() < (SpeedThreshold * SpeedThreshold))
+				{
+					SurfaceNormal = FVec3(0.f, 0.f, 1.f);
+				}
+				else
+				{
+					const FReal Speed = FMath::Abs(Body.V().Length());
+					if (Speed < SpeedBlendThreshold)
+					{
+						SurfaceNormal = FMath::Lerp(FVec3(0.f, 0.f, 1.f), SurfaceNormal, Speed / SpeedBlendThreshold);
+					}
+				}
+			}
 
 			FReal Distance = FVec3::DotProduct(WorldSpaceX - T, AxisWorld);
 			if (Distance >= Setting.MaxLength)
