Constraints:
- socket based component handle to be able to constrain to bones/sockets
- disabled scale for parent constraints by default
- auto expand constraint view when creating a new constraint
- check for direct dependencies (hierarchy) between handles before creating constraints

diff --git a/Engine/Plugins/Animation/ControlRig/Source/ControlRig/Private/Constraints/ControlRigTransformableHandle.cpp b/Engine/Plugins/Animation/ControlRig/Source/ControlRig/Private/Constraints/ControlRigTransformableHandle.cpp
--- a/Engine/Plugins/Animation/ControlRig/Source/ControlRig/Private/Constraints/ControlRigTransformableHandle.cpp
+++ b/Engine/Plugins/Animation/ControlRig/Source/ControlRig/Private/Constraints/ControlRigTransformableHandle.cpp
@@ -136,11 +136,16 @@ FTickFunction* UTransformableControlHandle::GetTickFunction() const
 	return SkelMeshComponent ? &SkelMeshComponent->PrimaryComponentTick : nullptr;
 }
 
+uint32 UTransformableControlHandle::ComputeHash(const UControlRig* InControlRig, const FName& InControlName)
+{
+	return HashCombine(GetTypeHash(InControlRig), GetTypeHash(InControlName));
+}
+
 uint32 UTransformableControlHandle::GetHash() const
 {
 	if (ControlRig.IsValid() && ControlName != NAME_None)
 	{
-		return HashCombine(GetTypeHash(ControlRig.Get()), GetTypeHash(ControlName));
+		return ComputeHash(ControlRig.Get(), ControlName);
 	}
 	return 0;
 }
@@ -156,6 +161,43 @@ USkeletalMeshComponent* UTransformableControlHandle::GetSkeletalMesh() const
 	return ObjectBinding ? Cast<USkeletalMeshComponent>(ObjectBinding->GetBoundObject()) : nullptr;
 }
 
+bool UTransformableControlHandle::HasDirectDependencyWith(const UTransformableHandle& InOther) const
+{
+	const uint32 OtherHash = InOther.GetHash();
+	if (OtherHash == 0)
+	{
+		return false;
+	}
+
+	// check whether the other handle is one of the skeletal mesh parent
+	if (const USkeletalMeshComponent* SkeletalMeshComponent = GetSkeletalMesh())
+	{
+		for (const USceneComponent* Comp=SkeletalMeshComponent->GetAttachParent(); Comp!=nullptr; Comp=Comp->GetAttachParent() )
+		{
+			const uint32 AttachParentHash = GetTypeHash(Comp);
+			if (AttachParentHash == OtherHash)
+			{
+				return true;
+			}
+		}
+	}
+	
+	const FRigControlElement* ControlElement = GetControlElement();
+	if (!ControlElement)
+	{
+		return false;
+	}
+
+	// check whether the other handle is one of the control parent within the CR hierarchy
+	static constexpr bool bRecursive = true;
+	const FRigBaseElementParentArray AllParents = ControlRig->GetHierarchy()->GetParents(ControlElement, bRecursive);
+	return  AllParents.ContainsByPredicate([this, OtherHash](const FRigBaseElement* Parent)
+	{
+		const uint32 ParentHash = ComputeHash(ControlRig.Get(), Parent->GetName());
+		return ParentHash == OtherHash;		
+	});
+}
+
 FRigControlElement* UTransformableControlHandle::GetControlElement() const
 {
 	if (!ControlRig.IsValid() || ControlName == NAME_None)
diff --git a/Engine/Plugins/Animation/ControlRig/Source/ControlRig/Private/ControlRigModule.cpp b/Engine/Plugins/Animation/ControlRig/Source/ControlRig/Private/ControlRigModule.cpp
--- a/Engine/Plugins/Animation/ControlRig/Source/ControlRig/Private/ControlRigModule.cpp
+++ b/Engine/Plugins/Animation/ControlRig/Source/ControlRig/Private/ControlRigModule.cpp
@@ -57,9 +57,8 @@ void FControlRigModule::RegisterTransformableCustomization() const
 	{
 		if (const AControlRigShapeActor* ControlActor = Cast<AControlRigShapeActor>(InObject))
 		{
-			const uint32 ControlHash = HashCombine(
-				GetTypeHash(ControlActor->ControlRig.Get()),
-				GetTypeHash(ControlActor->ControlName));
+			const uint32 ControlHash = UTransformableControlHandle::ComputeHash(
+				ControlActor->ControlRig.Get(), ControlActor->ControlName);
 			return ControlHash;
 		}
 		return 0;
diff --git a/Engine/Plugins/Animation/ControlRig/Source/ControlRig/Private/Sequencer/MovieSceneControlRigParameterSection.cpp b/Engine/Plugins/Animation/ControlRig/Source/ControlRig/Private/Sequencer/MovieSceneControlRigParameterSection.cpp
--- a/Engine/Plugins/Animation/ControlRig/Source/ControlRig/Private/Sequencer/MovieSceneControlRigParameterSection.cpp
+++ b/Engine/Plugins/Animation/ControlRig/Source/ControlRig/Private/Sequencer/MovieSceneControlRigParameterSection.cpp
@@ -17,6 +17,7 @@
 #include "MovieSceneTimeHelpers.h"
 #include "Animation/AnimSequenceHelpers.h"
 #include "Components/SkeletalMeshComponent.h"
+#include "Constraints/ControlRigTransformableHandle.h"
 #include "UObject/UE5MainStreamObjectVersion.h"
 
 #if WITH_EDITOR
@@ -1424,8 +1425,8 @@ void UMovieSceneControlRigParameterSection::ReconstructChannelProxy()
 		const FConstraintsManagerController& Controller = FConstraintsManagerController::Get(ControlRig->GetWorld());
 		auto GetConstraints = [&Controller, this](const FName& InControlName)
 		{
-			static constexpr bool bSorted = true;			
-			const uint32 ControlHash = HashCombine(GetTypeHash(ControlRig.Get()), GetTypeHash(InControlName));
+			static constexpr bool bSorted = true;
+			const uint32 ControlHash = UTransformableControlHandle::ComputeHash(ControlRig.Get(), InControlName);
 			return Controller.GetParentConstraints(ControlHash, bSorted);
 		};
 
diff --git a/Engine/Plugins/Animation/ControlRig/Source/ControlRig/Public/Constraints/ControlRigTransformableHandle.h b/Engine/Plugins/Animation/ControlRig/Source/ControlRig/Public/Constraints/ControlRigTransformableHandle.h
--- a/Engine/Plugins/Animation/ControlRig/Source/ControlRig/Public/Constraints/ControlRigTransformableHandle.h
+++ b/Engine/Plugins/Animation/ControlRig/Source/ControlRig/Public/Constraints/ControlRigTransformableHandle.h
@@ -47,8 +47,10 @@ public:
 	virtual FTickFunction* GetTickFunction() const override;
 
 	/** Generates a hash value based on ControlRig and ControlName. */
+	static uint32 ComputeHash(const UControlRig* InControlRig, const FName& InControlName);
 	virtual uint32 GetHash() const override;
-	/** @todo document */
+	
+	/** Returns the underlying targeted object. */
 	virtual TWeakObjectPtr<UObject> GetTarget() const override;
 
 	/** Get the array of float channels for the specified section*/
@@ -69,8 +71,11 @@ public:
 	void UnregisterDelegates() const;
 	void RegisterDelegates();
 
+	/** Check for direct dependencies (ie hierarchy + skeletal mesh) with InOther. */
+	virtual bool HasDirectDependencyWith(const UTransformableHandle& InOther) const override;
+
 #if WITH_EDITOR
-	/** @todo document */
+	/** Returns labels used for UI. */
 	virtual FString GetLabel() const override;
 	virtual FString GetFullLabel() const override;
 #endif
diff --git a/Engine/Plugins/Animation/ControlRig/Source/ControlRigEditor/Private/EditMode/ControlRigEditMode.cpp b/Engine/Plugins/Animation/ControlRig/Source/ControlRigEditor/Private/EditMode/ControlRigEditMode.cpp
--- a/Engine/Plugins/Animation/ControlRig/Source/ControlRigEditor/Private/EditMode/ControlRigEditMode.cpp
+++ b/Engine/Plugins/Animation/ControlRig/Source/ControlRigEditor/Private/EditMode/ControlRigEditMode.cpp
@@ -3778,9 +3778,8 @@ void FControlRigEditMode::SetControlShapeTransform(
 	const FTransform ChildTransform = InGlobalTransform * InToWorldTransform;
 	FTransform LocalTransform = ControlRig->GetControlLocalTransform(InShapeActor->ControlName);
 	
-	const ETransformConstraintType ConstraintType = static_cast<ETransformConstraintType>(Constraint->GetType());
 	LocalTransform = FTransformConstraintUtils::ComputeRelativeTransform(LocalTransform,
-		ChildTransform, ParentTransform, ConstraintType);
+		ChildTransform, ParentTransform, Constraint);
 
 	FRigControlModifiedContext Context = InContext;
 	Context.bConstraintUpdate = false;
diff --git a/Engine/Plugins/Animation/ControlRig/Source/ControlRigEditor/Private/Editor/Constraints/SConstraintsWidget.cpp b/Engine/Plugins/Animation/ControlRig/Source/ControlRigEditor/Private/Editor/Constraints/SConstraintsWidget.cpp
--- a/Engine/Plugins/Animation/ControlRig/Source/ControlRigEditor/Private/Editor/Constraints/SConstraintsWidget.cpp
+++ b/Engine/Plugins/Animation/ControlRig/Source/ControlRigEditor/Private/Editor/Constraints/SConstraintsWidget.cpp
@@ -13,6 +13,7 @@
 #include "Selection.h"
 #include "SlateOptMacros.h"
 #include "LevelEditor.h"
+#include "Editor/SceneOutliner/Private/SSocketChooser.h"
 #include "TransformConstraint.h"
 #include "Modules/ModuleManager.h"
 #include "Widgets/Input/SButton.h"
@@ -279,28 +280,70 @@ void SDroppableConstraintItem::CreateConstraint(
 		return;	
 	}
 
+	// has socket?
+	USceneComponent* ComponentWithSockets = nullptr;
+	if(USceneComponent* ParentComponent = InParent->GetRootComponent())
+	{
+		if (ParentComponent->HasAnySockets())
+		{
+			ComponentWithSockets = ParentComponent;
+		}
+	}
+
 	// create constraints
-	UWorld* World = GCurrentLevelEditingViewportClient->GetWorld();
-	bool bCreated = false;
-	for (AActor* Child: Selection)
+	auto CreateConstraint = [InCreationDelegate, InConstraintType](
+		const TArray<AActor*>& Selection, AActor* InParent, const FName& InSocketName)
 	{
-		if (Child != InParent)
+		UWorld* World = GCurrentLevelEditingViewportClient->GetWorld();
+		bool bCreated = false;
+		for (AActor* Child: Selection)
 		{
-			FScopedTransaction Transaction(LOCTEXT("CreateConstraintKey", "Create Constraint Key"));
-			UTickableTransformConstraint* Constraint =
-				FTransformConstraintUtils::CreateAndAddFromActors(World, InParent, Child, InConstraintType);
-			if (Constraint)
+			if (Child != InParent)
 			{
-				FConstraintChannelHelper::SmartConstraintKey(Constraint);
-				bCreated = true;
+				FScopedTransaction Transaction(LOCTEXT("CreateConstraintKey", "Create Constraint Key"));
+				UTickableTransformConstraint* Constraint =
+					FTransformConstraintUtils::CreateAndAddFromActors(World, InParent, InSocketName, Child, InConstraintType);
+				if (Constraint)
+				{
+					FConstraintChannelHelper::SmartConstraintKey(Constraint);
+					bCreated = true;
+				}
 			}
 		}
+		
+		// update list
+		if (bCreated && InCreationDelegate.IsBound())
+		{
+			InCreationDelegate.Execute();
+		}
+	};
+	
+	// Show socket chooser if we have sockets to select
+	if (ComponentWithSockets != nullptr)
+	{		
+		const FLevelEditorModule& LevelEditorModule = FModuleManager::GetModuleChecked<FLevelEditorModule>("LevelEditor");
+		const TSharedPtr< ILevelEditor > LevelEditor = LevelEditorModule.GetFirstLevelEditor();
+
+		FVector2D SummonLocation = FSlateApplication::Get().GetCursorPos();
+		SummonLocation.Y += 4 * FSlateApplication::Get().GetCursorSize().Y;
+		
+		// Create as context menu
+		FSlateApplication::Get().PushMenu(
+			LevelEditor.ToSharedRef(),
+			FWidgetPath(),
+			SNew(SSocketChooserPopup)
+			.SceneComponent( ComponentWithSockets )
+			.OnSocketChosen_Lambda([CreateConstraint, Selection, InParent](FName InSocketName)
+			{
+				CreateConstraint(Selection, InParent, InSocketName);
+			}),
+			SummonLocation,
+			FPopupTransitionEffect( FPopupTransitionEffect::ContextMenu )
+			);
 	}
-
-	// update list
-	if (bCreated && InCreationDelegate.IsBound())
+	else
 	{
-		InCreationDelegate.Execute();
+		CreateConstraint(Selection, InParent, NAME_None);
 	}
 }
 
@@ -747,7 +790,7 @@ void SConstraintsEditionWidget::InvalidateConstraintList()
 	bNeedsRefresh = true;
 }
 
-void SConstraintsEditionWidget::RefreshConstraintList()
+int32 SConstraintsEditionWidget::RefreshConstraintList()
 {
 	// get constraints
 	UWorld* World = GCurrentLevelEditingViewportClient->GetWorld();
@@ -785,6 +828,8 @@ void SConstraintsEditionWidget::RefreshConstraintList()
 
 	// refresh tree
 	ListView->RequestListRefresh();
+	
+	return ListItems.Num();
 }
 
 void SConstraintsEditionWidget::OnActorSelectionChanged(const TArray<UObject*>& NewSelection, bool bForceRefresh)
diff --git a/Engine/Plugins/Animation/ControlRig/Source/ControlRigEditor/Private/Editor/Constraints/SConstraintsWidget.h b/Engine/Plugins/Animation/ControlRig/Source/ControlRigEditor/Private/Editor/Constraints/SConstraintsWidget.h
--- a/Engine/Plugins/Animation/ControlRig/Source/ControlRigEditor/Private/Editor/Constraints/SConstraintsWidget.h
+++ b/Engine/Plugins/Animation/ControlRig/Source/ControlRigEditor/Private/Editor/Constraints/SConstraintsWidget.h
@@ -209,7 +209,7 @@ public:
 	void InvalidateConstraintList();
 	
 	/** Rebuild the constraint list based on the current selection. */
-	void RefreshConstraintList();
+	int32 RefreshConstraintList();
 
 	/** Triggers a constraint list invalidation when selection in the level viewport. */
 	void OnActorSelectionChanged(const TArray<UObject*>& NewSelection, bool bForceRefresh);
diff --git a/Engine/Source/Editor/MovieSceneTools/Private/MovieSceneConstraintChannelHelper.cpp b/Engine/Source/Editor/MovieSceneTools/Private/MovieSceneConstraintChannelHelper.cpp
--- a/Engine/Source/Editor/MovieSceneTools/Private/MovieSceneConstraintChannelHelper.cpp
+++ b/Engine/Source/Editor/MovieSceneTools/Private/MovieSceneConstraintChannelHelper.cpp
@@ -70,8 +70,6 @@ void FCompensationEvaluator::ComputeLocalTransforms(
 	ChildGlobals.SetNum(NumFrames + 1);
 	SpaceGlobals.SetNum(NumFrames + 1);
 
-	const ETransformConstraintType ConstraintType = static_cast<ETransformConstraintType>(Constraint->GetType());
-
 	for (int32 Index = 0; Index < NumFrames + 1; ++Index)
 	{
 		const FFrameNumber FrameNumber = (Index == 0) ? InFrames[0] - 1 : InFrames[Index - 1];
@@ -100,8 +98,6 @@ void FCompensationEvaluator::ComputeLocalTransforms(
 		ChildGlobal = Handle->GetGlobalTransform();
 
 		const UTickableTransformConstraint* LastConstraint = GetLastActiveConstraint();
-		const ETransformConstraintType LastConstraintType = LastConstraint ?
-			static_cast<ETransformConstraintType>(LastConstraint->GetType()) : ETransformConstraintType::Parent;
 
 		// store constraint/parent space global transform
 		if (bToActive)
@@ -114,14 +110,14 @@ void FCompensationEvaluator::ComputeLocalTransforms(
 				{
 					SpaceGlobal = LastConstraint->GetParentGlobalTransform();
 					ChildLocal = FTransformConstraintUtils::ComputeRelativeTransform(
-						ChildLocal, ChildGlobal, SpaceGlobal, LastConstraintType);
+						ChildLocal, ChildGlobal, SpaceGlobal, LastConstraint);
 				}
 			}
 			else
 			{
 				SpaceGlobal = Constraint->GetParentGlobalTransform();
 				ChildLocal = FTransformConstraintUtils::ComputeRelativeTransform(
-					ChildLocal, ChildGlobal, SpaceGlobal, ConstraintType);
+					ChildLocal, ChildGlobal, SpaceGlobal, Constraint);
 			}
 		}
 		else
@@ -132,7 +128,7 @@ void FCompensationEvaluator::ComputeLocalTransforms(
 			{
 				SpaceGlobal = Constraint->GetParentGlobalTransform();
 				ChildLocal = FTransformConstraintUtils::ComputeRelativeTransform(
-					ChildLocal, ChildGlobal, SpaceGlobal, ConstraintType);
+					ChildLocal, ChildGlobal, SpaceGlobal, Constraint);
 			}
 			else
 			{
@@ -140,7 +136,7 @@ void FCompensationEvaluator::ComputeLocalTransforms(
 				{
 					SpaceGlobal = LastConstraint->GetParentGlobalTransform();
 					ChildLocal = FTransformConstraintUtils::ComputeRelativeTransform(
-						ChildLocal, ChildGlobal, SpaceGlobal, LastConstraintType);
+						ChildLocal, ChildGlobal, SpaceGlobal, LastConstraint);
 				}
 			}
 		}
@@ -216,12 +212,9 @@ void FCompensationEvaluator::ComputeLocalTransformsBeforeDeletion(
 		// store constraint/parent space global transform
 		if (const UTickableTransformConstraint* LastConstraint = GetLastActiveConstraint())
 		{
-			const ETransformConstraintType LastConstraintType =
-				static_cast<ETransformConstraintType>(LastConstraint->GetType());
-
 			SpaceGlobal = LastConstraint->GetParentGlobalTransform();
 			ChildLocal = FTransformConstraintUtils::ComputeRelativeTransform(
-				ChildLocal, ChildGlobal, SpaceGlobal, LastConstraintType);
+				ChildLocal, ChildGlobal, SpaceGlobal, LastConstraint);
 		}
 	}
 }
@@ -283,10 +276,8 @@ void FCompensationEvaluator::ComputeCompensation(UWorld* InWorld, const TSharedP
 	if (const UTickableTransformConstraint* LastConstraint = GetLastActiveConstraint())
 	{
 		SpaceGlobals[0] = LastConstraint->GetParentGlobalTransform();
-
-		const ETransformConstraintType LastConstraintType = static_cast<ETransformConstraintType>(LastConstraint->GetType());
 		ChildLocals[0] = FTransformConstraintUtils::ComputeRelativeTransform(ChildLocals[0],
-			ChildGlobals[0], SpaceGlobals[0], LastConstraintType);
+			ChildGlobals[0], SpaceGlobals[0], LastConstraint);
 	}
 	else // switch to parent space
 	{
diff --git a/Engine/Source/Runtime/Experimental/Animation/Constraints/Private/ConstraintsScripting.cpp b/Engine/Source/Runtime/Experimental/Animation/Constraints/Private/ConstraintsScripting.cpp
--- a/Engine/Source/Runtime/Experimental/Animation/Constraints/Private/ConstraintsScripting.cpp
+++ b/Engine/Source/Runtime/Experimental/Animation/Constraints/Private/ConstraintsScripting.cpp
@@ -13,12 +13,13 @@ UConstraintsManager* UConstraintsScriptingLibrary::GetManager(UWorld* InWorld)
 	return ConstraintsManager;
 }
 
-UTransformableComponentHandle* UConstraintsScriptingLibrary::CreateTransformableComponentHandle(UWorld* InWorld,USceneComponent* InSceneComponent)
+UTransformableComponentHandle* UConstraintsScriptingLibrary::CreateTransformableComponentHandle(
+	UWorld* InWorld, USceneComponent* InSceneComponent, const FName& InSocketName)
 {
 	UConstraintsManager* ConstraintsManager = UConstraintsManager::Get(InWorld);
 	if (ConstraintsManager)
 	{
-		UTransformableComponentHandle* Handle = FTransformConstraintUtils::CreateHandleForSceneComponent(InSceneComponent, ConstraintsManager);
+		UTransformableComponentHandle* Handle = FTransformConstraintUtils::CreateHandleForSceneComponent(InSceneComponent, InSocketName, ConstraintsManager);
 		return Handle;
 
 	}
diff --git a/Engine/Source/Runtime/Experimental/Animation/Constraints/Private/TransformConstraint.cpp b/Engine/Source/Runtime/Experimental/Animation/Constraints/Private/TransformConstraint.cpp
--- a/Engine/Source/Runtime/Experimental/Animation/Constraints/Private/TransformConstraint.cpp
+++ b/Engine/Source/Runtime/Experimental/Animation/Constraints/Private/TransformConstraint.cpp
@@ -570,7 +570,7 @@ uint32 UTickableParentConstraint::CalculateInputHash() const
 	Hash = HashCombine(Hash, GetTypeHash(ChildLocalTransform.GetRotation().Euler() ));
 	Hash = HashCombine(Hash, GetTypeHash(ChildLocalTransform.GetScale3D()));
 	
-	const FTransform ChildGlobalTransform = GetChildLocalTransform();
+	const FTransform ChildGlobalTransform = GetChildGlobalTransform();
 	Hash = HashCombine(Hash, GetTypeHash(ChildGlobalTransform.GetTranslation()));
 	Hash = HashCombine(Hash, GetTypeHash(ChildGlobalTransform.GetRotation().Euler() ));
 	Hash = HashCombine(Hash, GetTypeHash(ChildGlobalTransform.GetScale3D()));
@@ -610,7 +610,15 @@ FConstraintTickFunction::ConstraintFunction UTickableParentConstraint::GetFuncti
 		
 		FTransform TargetTransform = (!bMaintainOffset) ? ParentTransform : OffsetTransform * ParentTransform;
 		//apply weight if needed
-		LerpTransform(GetChildGlobalTransform(), TargetTransform);
+		const FTransform ChildGlobalTransform = GetChildGlobalTransform();
+		LerpTransform(ChildGlobalTransform, TargetTransform);
+
+		//remove scale?
+		if (!bScaling)
+		{
+			TargetTransform.SetScale3D(ChildGlobalTransform.GetScale3D());
+		}
+		
 		SetChildGlobalTransform(TargetTransform);
 	};
 }
@@ -747,7 +755,7 @@ FQuat UTickableLookAtConstraint::FindQuatBetweenNormals(const FVector& A, const
 namespace
 {
 
-UTransformableHandle* GetHandle(AActor* InActor, UObject* Outer)
+UTransformableHandle* GetHandle(AActor* InActor, const FName& InSocketName, UObject* Outer)
 {
 	// look for customized transform handle
 	const FTransformableRegistry& Registry = FTransformableRegistry::Get();
@@ -759,7 +767,7 @@ UTransformableHandle* GetHandle(AActor* InActor, UObject* Outer)
 	// need to make sure it's moveable
 	if (InActor->GetRootComponent())
 	{
-		return FTransformConstraintUtils::CreateHandleForSceneComponent(InActor->GetRootComponent(), Outer);
+		return FTransformConstraintUtils::CreateHandleForSceneComponent(InActor->GetRootComponent(), InSocketName, Outer);
 
 	}
 	return nullptr;
@@ -781,13 +789,17 @@ uint32 GetConstrainableHash(const AActor* InActor)
 	
 }
 
-UTransformableComponentHandle* FTransformConstraintUtils::CreateHandleForSceneComponent(USceneComponent* InSceneComponent, UObject* Outer)
+UTransformableComponentHandle* FTransformConstraintUtils::CreateHandleForSceneComponent(
+	USceneComponent* InSceneComponent,
+	const FName& InSocketName,
+	UObject* Outer)
 {
 	UTransformableComponentHandle* ComponentHandle = nullptr;
 	if (InSceneComponent)
 	{
 		ComponentHandle = NewObject<UTransformableComponentHandle>(Outer);
 		ComponentHandle->Component = InSceneComponent;
+		ComponentHandle->SocketName = InSocketName;
 		InSceneComponent->SetMobility(EComponentMobility::Movable);
 		ComponentHandle->RegisterDelegates();
 	}
@@ -866,6 +878,7 @@ UTickableTransformConstraint* FTransformConstraintUtils::CreateFromType(
 UTickableTransformConstraint* FTransformConstraintUtils::CreateAndAddFromActors(
 	UWorld* InWorld,
 	AActor* InParent,
+	const FName& InSocketName,
 	AActor* InChild,
 	const ETransformConstraintType InType,
 	const bool bMaintainOffset)
@@ -878,7 +891,7 @@ UTickableTransformConstraint* FTransformConstraintUtils::CreateAndAddFromActors(
 		UE_LOG(LogTemp, Error, TEXT("%s sanity check failed."), ErrorPrefix);
 		return nullptr;
 	}
-	
+
 	UConstraintsManager* ConstraintsManager = UConstraintsManager::Get(InWorld);
 	if (!ConstraintsManager)
 	{
@@ -886,24 +899,37 @@ UTickableTransformConstraint* FTransformConstraintUtils::CreateAndAddFromActors(
 		return nullptr;
 	}
 
-	UTransformableHandle* ParentHandle = GetHandle(InParent, ConstraintsManager);
+	UTransformableHandle* ParentHandle = GetHandle(InParent, InSocketName, ConstraintsManager);
 	if (!ParentHandle)
 	{
 		return nullptr;
 	}
 	
-	UTransformableHandle* ChildHandle = GetHandle(InChild, ConstraintsManager);
+	UTransformableHandle* ChildHandle = GetHandle(InChild, NAME_None, ConstraintsManager);
 	if (!ChildHandle)
 	{
 		return nullptr;
 	}
 
+	bool bDeleteHandles = false;
 	if (ChildHandle->GetHash() == ParentHandle->GetHash())
+	{
+		UE_LOG(LogTemp, Error, TEXT("%s handles are pointing at the same object."), ErrorPrefix);
+		bDeleteHandles = true;
+	}
+	
+	if (ParentHandle->HasDirectDependencyWith(*ChildHandle))
+	{
+#if WITH_EDITOR
+		UE_LOG(LogTemp, Error, TEXT("%s: %s has a direct dependecy to %s."), *ParentHandle->GetLabel(), *ChildHandle->GetLabel());
+#endif
+		bDeleteHandles = true;
+	}
+	
+	if (bDeleteHandles)
 	{
 		ChildHandle->MarkAsGarbage();
 		ParentHandle->MarkAsGarbage();
-		
-		UE_LOG(LogTemp, Error, TEXT("%s handles are pointing at the same object."), ErrorPrefix);
 		return nullptr;
 	}
 	
@@ -962,9 +988,15 @@ FTransform FTransformConstraintUtils::ComputeRelativeTransform(
 	const FTransform& InChildLocal,
 	const FTransform& InChildWorld,
 	const FTransform& InSpaceWorld,
-	const ETransformConstraintType InType)
+	const UTickableTransformConstraint* InConstraint)
 {
-	switch (InType)
+	if (!InConstraint)
+	{
+		return InChildWorld.GetRelativeTransform(InSpaceWorld);
+	}
+	
+	const ETransformConstraintType ConstraintType = static_cast<ETransformConstraintType>(InConstraint->GetType());
+	switch (ConstraintType)
 	{
 	case ETransformConstraintType::Translation:
 		{
@@ -992,7 +1024,16 @@ FTransform FTransformConstraintUtils::ComputeRelativeTransform(
 			return RelativeTransform;
 		}
 	case ETransformConstraintType::Parent:
-		return InChildWorld.GetRelativeTransform(InSpaceWorld);
+		{
+			FTransform RelativeTransform = InChildWorld.GetRelativeTransform(InSpaceWorld);
+			const UTickableParentConstraint* ParentConstraint = Cast<UTickableParentConstraint>(InConstraint);
+			const bool bScale = ParentConstraint ? ParentConstraint->IsScalingEnabled() : true;
+			if (!bScale)
+			{
+				RelativeTransform.SetScale3D(InChildLocal.GetScale3D());
+			}
+			return RelativeTransform;
+		}
 	case ETransformConstraintType::LookAt:
 		return InChildLocal;
 	default:
diff --git a/Engine/Source/Runtime/Experimental/Animation/Constraints/Private/TransformableHandle.cpp b/Engine/Source/Runtime/Experimental/Animation/Constraints/Private/TransformableHandle.cpp
--- a/Engine/Source/Runtime/Experimental/Animation/Constraints/Private/TransformableHandle.cpp
+++ b/Engine/Source/Runtime/Experimental/Animation/Constraints/Private/TransformableHandle.cpp
@@ -77,7 +77,7 @@ FTransform UTransformableComponentHandle::GetLocalTransform() const
 
 FTransform UTransformableComponentHandle::GetGlobalTransform() const
 {
-	return Component.IsValid() ? Component->GetComponentTransform() : FTransform::Identity;
+	return Component.IsValid() ? Component->GetSocketTransform(SocketName) : FTransform::Identity;
 }
 
 UObject* UTransformableComponentHandle::GetPrerequisiteObject() const
@@ -100,6 +100,30 @@ TWeakObjectPtr<UObject> UTransformableComponentHandle::GetTarget() const
 	return Component;
 }
 
+bool UTransformableComponentHandle::HasDirectDependencyWith(const UTransformableHandle& InOther) const
+{
+	const uint32 OtherHash = InOther.GetHash();
+	if (OtherHash == 0)
+	{
+		return false;
+	}
+
+	// check whether the other handle is one of the component's parent
+	if (Component.IsValid())
+	{
+		for (const USceneComponent* Comp=Component->GetAttachParent(); Comp!=nullptr; Comp=Comp->GetAttachParent() )
+		{
+			const uint32 AttachParentHash = GetTypeHash(Comp);
+			if (AttachParentHash == OtherHash)
+			{
+				return true;
+			}
+		}
+	}
+	
+	return false;
+}
+
 void UTransformableComponentHandle::UnregisterDelegates() const
 {
 #if WITH_EDITOR
@@ -339,13 +363,32 @@ FString UTransformableComponentHandle::GetLabel() const
 		return DummyLabel;
 	}
 
+	if (SocketName != NAME_None)
+	{
+		return SocketName.ToString(); 
+	}
+
 	const AActor* Actor = Component->GetOwner();
 	return Actor ? Actor->GetActorLabel() : Component->GetName();
 }
 
 FString UTransformableComponentHandle::GetFullLabel() const
 {
-	return GetLabel();
+	if (!Component.IsValid())
+	{
+		static const FString DummyLabel;
+		return DummyLabel;
+	}
+
+	const AActor* Actor = Component->GetOwner();
+	const FString ComponentLabel = Actor ? Actor->GetActorLabel() : Component->GetName();
+	
+	if (SocketName == NAME_None)
+	{
+		return ComponentLabel;
+	}
+
+	return FString::Printf(TEXT("%s/%s"), *ComponentLabel, *SocketName.ToString() );
 };
 
 #endif
\ No newline at end of file
diff --git a/Engine/Source/Runtime/Experimental/Animation/Constraints/Public/ConstraintsScripting.h b/Engine/Source/Runtime/Experimental/Animation/Constraints/Public/ConstraintsScripting.h
--- a/Engine/Source/Runtime/Experimental/Animation/Constraints/Public/ConstraintsScripting.h
+++ b/Engine/Source/Runtime/Experimental/Animation/Constraints/Public/ConstraintsScripting.h
@@ -39,10 +39,11 @@ public:
 	* Create the transformable handle that deals with getting and setting transforms on this scene component
 	* @param InWorld, the world you are in
 	* @param InSceneComponent World to create the constraint
+	* @param InSocketName Optional name of the socket to get the transform 
 	* @return returns the handle for this scene component
 	**/
 	UFUNCTION(BlueprintCallable, Category = "Editor Scripting | Sequencer Tools | Control Rig | Constraints")
-	static UTransformableComponentHandle* CreateTransformableComponentHandle(UWorld* InWorld, USceneComponent* InSceneComponent);
+	static UTransformableComponentHandle* CreateTransformableComponentHandle(UWorld* InWorld, USceneComponent* InSceneComponent, const FName& InSocketName);
 
 	/**
 	* Create Constraint based on the specified type.
diff --git a/Engine/Source/Runtime/Experimental/Animation/Constraints/Public/TransformConstraint.h b/Engine/Source/Runtime/Experimental/Animation/Constraints/Public/TransformConstraint.h
--- a/Engine/Source/Runtime/Experimental/Animation/Constraints/Public/TransformConstraint.h
+++ b/Engine/Source/Runtime/Experimental/Animation/Constraints/Public/TransformConstraint.h
@@ -254,6 +254,11 @@ public:
 	/** Returns the transform constraint function that the tick function will evaluate. */
 	virtual FConstraintTickFunction::ConstraintFunction GetFunction() const override;
 
+	bool IsScalingEnabled() const
+	{
+		return bScaling;
+	}
+	
 protected:
 	/** Cache data structure to store last child local/global transform. */
 	struct FDynamicCache
@@ -273,6 +278,10 @@ protected:
 	UPROPERTY(EditAnywhere, BlueprintReadWrite, Category="Offset", meta=(EditCondition="bMaintainOffset"))
 	FTransform OffsetTransform = FTransform::Identity;
 
+	/** Defines whether we propagate the parent scale. */
+	UPROPERTY(EditAnywhere, BlueprintReadWrite, Category="Properties")
+	bool bScaling = false;
+
 #if WITH_EDITOR
 public:
 	// UObject interface
@@ -324,7 +333,8 @@ struct CONSTRAINTS_API FTransformConstraintUtils
 		
 	/** Create a handle for the scene component.*/
 	static UTransformableComponentHandle* CreateHandleForSceneComponent(
-		USceneComponent* InSceneComponent, 
+		USceneComponent* InSceneComponent,
+		const FName& InSocketName,
 		UObject* Outer);
 
 	/** Creates a new transform constraint based on the InType. */
@@ -336,6 +346,7 @@ struct CONSTRAINTS_API FTransformConstraintUtils
 	static UTickableTransformConstraint* CreateAndAddFromActors(
 		UWorld* InWorld,
 		AActor* InParent,
+		const FName& InSocketName,
 		AActor* InChild,
 		const ETransformConstraintType InType,
 		const bool bMaintainOffset = true);
@@ -353,5 +364,5 @@ struct CONSTRAINTS_API FTransformConstraintUtils
 		const FTransform& InChildLocal,
 		const FTransform& InChildWorld,
 		const FTransform& InSpaceWorld,
-		const ETransformConstraintType InType);
+		const UTickableTransformConstraint* InConstraint);
 };
diff --git a/Engine/Source/Runtime/Experimental/Animation/Constraints/Public/TransformableHandle.h b/Engine/Source/Runtime/Experimental/Animation/Constraints/Public/TransformableHandle.h
--- a/Engine/Source/Runtime/Experimental/Animation/Constraints/Public/TransformableHandle.h
+++ b/Engine/Source/Runtime/Experimental/Animation/Constraints/Public/TransformableHandle.h
@@ -70,9 +70,12 @@ public:
 	/** Generates a hash value of the underlying transformable object. */
 	virtual uint32 GetHash() const PURE_VIRTUAL(GetHash, return 0;);
 
-	/** @todo document */
+	/** Returns the underlying targeted object. */
 	virtual TWeakObjectPtr<UObject> GetTarget() const PURE_VIRTUAL(GetTarget, return nullptr;);
 
+	/** Check for direct dependencies with InOther. */
+	virtual bool HasDirectDependencyWith(const UTransformableHandle& InOther) const PURE_VIRTUAL(HasDirectDependencyWith, return false;);
+	
 	FHandleModifiedEvent& HandleModified();
 
 #if WITH_EDITOR
@@ -117,9 +120,12 @@ public:
 	/** Generates a hash value of Component. */
 	virtual uint32 GetHash() const override;
 
-	/** @todo document */
+	/** Returns the underlying targeted object. */
 	virtual TWeakObjectPtr<UObject> GetTarget() const override;
 
+	/** Check for direct dependencies (ie hierarchy) with InOther. */
+	virtual bool HasDirectDependencyWith(const UTransformableHandle& InOther) const override;
+
 	/** Get the array of float channels for the specified section*/
 	virtual TArrayView<FMovieSceneFloatChannel*>  GetFloatChannels(const UMovieSceneSection* InSection) const override;
 	/** Get the array of double channels for the specified section*/
@@ -131,7 +137,9 @@ public:
 		const FFrameRate& InTickResolution,
 		UMovieSceneSection* InSection,
 		const bool bLocal = true) const override;
+
 #if WITH_EDITOR
+	/** Returns labels used for UI. */
 	virtual FString GetLabel() const override;
 	virtual FString GetFullLabel() const override;
 #endif
@@ -140,6 +148,10 @@ public:
 	UPROPERTY(BlueprintReadOnly, Category = "Object")
 	TWeakObjectPtr<USceneComponent> Component;
 
+	/** Optional socket name on Component. */
+	UPROPERTY(BlueprintReadOnly, Category = "Object")
+	FName SocketName = NAME_None;
+
 	/** Registers/Unregisters useful delegates to track changes in the Component's transform. */
 	void UnregisterDelegates() const;
 	void RegisterDelegates();
