Fix Broad Phase switch: Now we can safely switch from 4 to 2 broad phases

diff --git a/Engine/Source/Runtime/Experimental/Chaos/Private/Chaos/PBDRigidsEvolution.cpp b/Engine/Source/Runtime/Experimental/Chaos/Private/Chaos/PBDRigidsEvolution.cpp
--- a/Engine/Source/Runtime/Experimental/Chaos/Private/Chaos/PBDRigidsEvolution.cpp
+++ b/Engine/Source/Runtime/Experimental/Chaos/Private/Chaos/PBDRigidsEvolution.cpp
@@ -7,8 +7,6 @@
 #include "Chaos/ParticleHandle.h"
 #include "Chaos/SpatialAccelerationCollection.h"
 
-//PRAGMA_DISABLE_OPTIMIZATION
-
 CSV_DECLARE_CATEGORY_EXTERN(ChaosPhysicsTimers);
 
 int32 ChaosRigidsEvolutionApplyAllowEarlyOutCVar = 1;
@@ -37,9 +35,6 @@ namespace Chaos
 	CHAOS_API int32 AccelerationStructureUseDynamicTree = 1;
 	FAutoConsoleVariableRef CVarAccelerationStructureUseDynamicTree(TEXT("p.Chaos.AccelerationStructureUseDynamicTree"), AccelerationStructureUseDynamicTree, TEXT("Use a dynamic BVH tree structure for dynamic objects"));
 
-	CHAOS_API int32 AccelerationStructureUseDirtyTreeInsteadOfGrid = 1;
-	FAutoConsoleVariableRef CVarAccelerationStructureUseDirtyTreeInsteadOfGrid(TEXT("p.Chaos.AccelerationStructureUseDirtyTreeInsteadOfGrid"), AccelerationStructureUseDirtyTreeInsteadOfGrid, TEXT("Use a dynamic tree structure for dirty elements instead of a 2D grid"));
-
 	/** Console variable to enable the caching of the overlapping leaves if the dynamic tree is enable */
 	CHAOS_API int32 GAccelerationStructureCacheOverlappingLeaves = 1;
 	FAutoConsoleVariableRef CVarAccelerationStructureCacheOverlappingLeaves(TEXT("p.Chaos.AccelerationStructureCacheOverlappingLeaves"), GAccelerationStructureCacheOverlappingLeaves, TEXT("Set to 1: Cache the overlapping leaves for faster overlap query, any other value will disable the feature"));
@@ -50,43 +45,19 @@ namespace Chaos
 	CHAOS_API int32 AccelerationStructureTimeSlicingMaxBytesCopy = 100000;
 	FAutoConsoleVariableRef CVarAccelerationStructureTimeSlicingMaxBytesCopy(TEXT("p.Chaos.AccelerationStructureTimeSlicingMaxBytesCopy"), AccelerationStructureTimeSlicingMaxBytesCopy, TEXT("The Maximum number of bytes to copy to the external acceleration structure during Copy Time Slicing"));
 
-	struct FAccelerationConfig
-	{
-		int32 BroadphaseType;
-		int32 BVNumCells;
-		int32 MaxChildrenInLeaf;
-		int32 MaxTreeDepth;
-		int32 AABBMaxChildrenInLeaf;
-		int32 AABBMaxTreeDepth;
-		FRealSingle MaxPayloadSize;
-		int32 IterationsPerTimeSlice;
-
-		FAccelerationConfig()
-		{
-			BroadphaseType = 3;
-			BVNumCells = 35;
-			MaxChildrenInLeaf = 5;
-			MaxTreeDepth = 200;
-			AABBMaxChildrenInLeaf = 500;
-			AABBMaxTreeDepth = 200;
-			MaxPayloadSize = 100000;
-			IterationsPerTimeSlice = 4000;
-		}
-	} ConfigSettings;
+	CHAOS_API FBroadPhaseConfig BroadPhaseConfig;
 
-	FAutoConsoleVariableRef CVarBroadphaseIsTree(TEXT("p.BroadphaseType"), ConfigSettings.BroadphaseType, TEXT(""));
-	FAutoConsoleVariableRef CVarBoundingVolumeNumCells(TEXT("p.BoundingVolumeNumCells"), ConfigSettings.BVNumCells, TEXT(""));
-	FAutoConsoleVariableRef CVarMaxChildrenInLeaf(TEXT("p.MaxChildrenInLeaf"), ConfigSettings.MaxChildrenInLeaf, TEXT(""));
-	FAutoConsoleVariableRef CVarMaxTreeDepth(TEXT("p.MaxTreeDepth"), ConfigSettings.MaxTreeDepth, TEXT(""));
-	FAutoConsoleVariableRef CVarAABBMaxChildrenInLeaf(TEXT("p.AABBMaxChildrenInLeaf"), ConfigSettings.AABBMaxChildrenInLeaf, TEXT(""));
-	FAutoConsoleVariableRef CVarAABBMaxTreeDepth(TEXT("p.AABBMaxTreeDepth"), ConfigSettings.AABBMaxTreeDepth, TEXT(""));
-	FAutoConsoleVariableRef CVarMaxPayloadSize(TEXT("p.MaxPayloadSize"), ConfigSettings.MaxPayloadSize, TEXT(""));
-	FAutoConsoleVariableRef CVarIterationsPerTimeSlice(TEXT("p.IterationsPerTimeSlice"), ConfigSettings.IterationsPerTimeSlice, TEXT(""));
+	FAutoConsoleVariableRef CVarBroadphaseIsTree(TEXT("p.BroadphaseType"), BroadPhaseConfig.BroadphaseType, TEXT(""));
+	FAutoConsoleVariableRef CVarBoundingVolumeNumCells(TEXT("p.BoundingVolumeNumCells"), BroadPhaseConfig.BVNumCells, TEXT(""));
+	FAutoConsoleVariableRef CVarMaxChildrenInLeaf(TEXT("p.MaxChildrenInLeaf"), BroadPhaseConfig.MaxChildrenInLeaf, TEXT(""));
+	FAutoConsoleVariableRef CVarMaxTreeDepth(TEXT("p.MaxTreeDepth"), BroadPhaseConfig.MaxTreeDepth, TEXT(""));
+	FAutoConsoleVariableRef CVarAABBMaxChildrenInLeaf(TEXT("p.AABBMaxChildrenInLeaf"), BroadPhaseConfig.AABBMaxChildrenInLeaf, TEXT(""));
+	FAutoConsoleVariableRef CVarAABBMaxTreeDepth(TEXT("p.AABBMaxTreeDepth"), BroadPhaseConfig.AABBMaxTreeDepth, TEXT(""));
+	FAutoConsoleVariableRef CVarMaxPayloadSize(TEXT("p.MaxPayloadSize"), BroadPhaseConfig.MaxPayloadSize, TEXT(""));
+	FAutoConsoleVariableRef CVarIterationsPerTimeSlice(TEXT("p.IterationsPerTimeSlice"), BroadPhaseConfig.IterationsPerTimeSlice, TEXT(""));
 
 	struct FDefaultCollectionFactory : public ISpatialAccelerationCollectionFactory
 	{
-		FAccelerationConfig Config;
-
 		using BVType = TBoundingVolume<FAccelerationStructureHandle>;
 		using AABBTreeType = TAABBTree<FAccelerationStructureHandle, TAABBTreeLeafArray<FAccelerationStructureHandle>>;
 		using AABBDynamicTreeType = TAABBTree<FAccelerationStructureHandle, TAABBTreeLeafArray<FAccelerationStructureHandle>>;
@@ -102,7 +73,7 @@ namespace Chaos
 		{
 			TConstParticleView<FSpatialAccelerationCache> Empty;
 
-			const uint16 NumBuckets = ConfigSettings.BroadphaseType >= 3 ? 2 : 1;
+			const uint16 NumBuckets = BroadPhaseConfig.BroadphaseType >= FBroadPhaseConfig::TreeAndGrid ? 2 : 1;
 			auto Collection = new TSpatialAccelerationCollection<AABBTreeType, BVType, AABBTreeOfGridsType>();
 
 			for (uint16 BucketIdx = 0; BucketIdx < NumBuckets; ++BucketIdx)
@@ -121,7 +92,7 @@ namespace Chaos
 
 		virtual uint8 GetActiveBucketsMask() const
 		{
-			return ConfigSettings.BroadphaseType >= 3 ? 3 : 1;
+			return BroadPhaseConfig.BroadphaseType >= FBroadPhaseConfig::TreeAndGrid ? 3 : 1;
 		}
 
 		virtual bool IsBucketTimeSliced(uint16 BucketIdx) const
@@ -131,17 +102,17 @@ namespace Chaos
 			{
 			case 0:
 			{
-				if (ConfigSettings.BroadphaseType == 0)
+				if (BroadPhaseConfig.BroadphaseType == FBroadPhaseConfig::Grid)
 				{
 					// BVType
 					return false;
 				}
-				else if (ConfigSettings.BroadphaseType == 1 || ConfigSettings.BroadphaseType == 3)
+				else if (BroadPhaseConfig.BroadphaseType == FBroadPhaseConfig::Tree || BroadPhaseConfig.BroadphaseType == FBroadPhaseConfig::TreeAndGrid)
 				{
 					// AABBTreeType
 					return true;
 				}
-				else if (ConfigSettings.BroadphaseType == 4 || ConfigSettings.BroadphaseType == 2)
+				else if (BroadPhaseConfig.BroadphaseType == FBroadPhaseConfig::TreeOfGridAndGrid || BroadPhaseConfig.BroadphaseType == FBroadPhaseConfig::TreeOfGrid)
 				{
 					// AABBTreeOfGridsType
 					return true;
@@ -150,7 +121,7 @@ namespace Chaos
 			case 1:
 			{
 				// BVType
-				ensure(ConfigSettings.BroadphaseType == 3 || ConfigSettings.BroadphaseType == 4);
+				ensure(BroadPhaseConfig.BroadphaseType == FBroadPhaseConfig::TreeAndGrid || BroadPhaseConfig.BroadphaseType == FBroadPhaseConfig::TreeOfGridAndGrid);
 				return false;
 			}
 			default:
@@ -168,27 +139,27 @@ namespace Chaos
 			{
 			case 0:
 			{
-				if (ConfigSettings.BroadphaseType == 0)
+				if (BroadPhaseConfig.BroadphaseType == FBroadPhaseConfig::Grid)
 				{
-					return MakeUnique<BVType>(Particles, false, static_cast<FReal>(0), ConfigSettings.BVNumCells, ConfigSettings.MaxPayloadSize);
+					return MakeUnique<BVType>(Particles, false, static_cast<FReal>(0), BroadPhaseConfig.BVNumCells, BroadPhaseConfig.MaxPayloadSize);
 				}
-				else if (ConfigSettings.BroadphaseType == 1 || ConfigSettings.BroadphaseType == 3)
+				else if (BroadPhaseConfig.BroadphaseType == FBroadPhaseConfig::Tree || BroadPhaseConfig.BroadphaseType == FBroadPhaseConfig::TreeAndGrid)
 				{
 					if (bDynamicTree)
 					{
-						return MakeUnique<AABBDynamicTreeType>(Particles, ConfigSettings.MaxChildrenInLeaf, ConfigSettings.MaxTreeDepth, ConfigSettings.MaxPayloadSize, ForceFullBuild ? 0 : ConfigSettings.IterationsPerTimeSlice, true);
+						return MakeUnique<AABBDynamicTreeType>(Particles, BroadPhaseConfig.MaxChildrenInLeaf, BroadPhaseConfig.MaxTreeDepth, BroadPhaseConfig.MaxPayloadSize, ForceFullBuild ? 0 : BroadPhaseConfig.IterationsPerTimeSlice, true);
 					}
-					return MakeUnique<AABBTreeType>(Particles, ConfigSettings.MaxChildrenInLeaf, ConfigSettings.MaxTreeDepth, ConfigSettings.MaxPayloadSize, ForceFullBuild ? 0 : ConfigSettings.IterationsPerTimeSlice, false, AccelerationStructureUseDirtyTreeInsteadOfGrid == 1);
+					return MakeUnique<AABBTreeType>(Particles, BroadPhaseConfig.MaxChildrenInLeaf, BroadPhaseConfig.MaxTreeDepth, BroadPhaseConfig.MaxPayloadSize, ForceFullBuild ? 0 : BroadPhaseConfig.IterationsPerTimeSlice);
 				}
-				else if (ConfigSettings.BroadphaseType == 4 || ConfigSettings.BroadphaseType == 2)
+				else if (BroadPhaseConfig.BroadphaseType == FBroadPhaseConfig::TreeOfGridAndGrid || BroadPhaseConfig.BroadphaseType == FBroadPhaseConfig::TreeOfGrid)
 				{
-					return MakeUnique<AABBTreeOfGridsType>(Particles, ConfigSettings.AABBMaxChildrenInLeaf, ConfigSettings.AABBMaxTreeDepth, ConfigSettings.MaxPayloadSize);
+					return MakeUnique<AABBTreeOfGridsType>(Particles, BroadPhaseConfig.AABBMaxChildrenInLeaf, BroadPhaseConfig.AABBMaxTreeDepth, BroadPhaseConfig.MaxPayloadSize);
 				}
 			}
 			case 1:
 			{
-				ensure(ConfigSettings.BroadphaseType == 3 || ConfigSettings.BroadphaseType == 4);
-				return MakeUnique<BVType>(Particles, false, static_cast<FReal>(0), ConfigSettings.BVNumCells, ConfigSettings.MaxPayloadSize);
+				ensure(BroadPhaseConfig.BroadphaseType == FBroadPhaseConfig::TreeAndGrid || BroadPhaseConfig.BroadphaseType == FBroadPhaseConfig::TreeOfGridAndGrid);
+				return MakeUnique<BVType>(Particles, false, static_cast<FReal>(0), BroadPhaseConfig.BVNumCells, BroadPhaseConfig.MaxPayloadSize);
 			}
 			default:
 			{
@@ -956,17 +927,17 @@ namespace Chaos
 	void FPBDRigidsEvolutionBase::Serialize(FChaosArchive& Ar)
 	{
 		ensure(false);	//disabled transient data serialization. Need to rethink
-		int32 DefaultBroadphaseType = ConfigSettings.BroadphaseType;
+		int32 DefaultBroadphaseType = BroadPhaseConfig.BroadphaseType;
 
 		Ar.UsingCustomVersion(FExternalPhysicsCustomObjectVersion::GUID);
 		if (Ar.CustomVer(FExternalPhysicsCustomObjectVersion::GUID) >= FExternalPhysicsCustomObjectVersion::SerializeBroadphaseType)
 		{
-			Ar << ConfigSettings.BroadphaseType;
+			Ar << BroadPhaseConfig.BroadphaseType;
 		}
 		else
 		{
 			//older archives just assume type 3
-			ConfigSettings.BroadphaseType = 3;
+			BroadPhaseConfig.BroadphaseType = 3;
 		}
 
 		Particles.Serialize(Ar);
@@ -1022,7 +993,7 @@ namespace Chaos
 			FlushSpatialAcceleration();
 		}
 
-		ConfigSettings.BroadphaseType = DefaultBroadphaseType;
+		BroadPhaseConfig.BroadphaseType = DefaultBroadphaseType;
 	}
 
 	void FPBDRigidsEvolutionBase::SetParticleObjectState(FPBDRigidParticleHandle* Particle, EObjectStateType ObjectState)
diff --git a/Engine/Source/Runtime/Experimental/Chaos/Private/PBDRigidsSolver.cpp b/Engine/Source/Runtime/Experimental/Chaos/Private/PBDRigidsSolver.cpp
--- a/Engine/Source/Runtime/Experimental/Chaos/Private/PBDRigidsSolver.cpp
+++ b/Engine/Source/Runtime/Experimental/Chaos/Private/PBDRigidsSolver.cpp
@@ -7,6 +7,7 @@
 #include "Chaos/PBDCollisionConstraintsUtil.h"
 #include "Chaos/Utilities.h"
 #include "Chaos/ChaosDebugDraw.h"
+#include "Chaos/PBDRigidsEvolution.h"
 #include "ChaosStats.h"
 #include "ChaosSolversModule.h"
 #include "ChaosVisualDebugger/ChaosVisualDebuggerTrace.h"
@@ -495,7 +496,13 @@ namespace Chaos
 		UE_LOG(LogPBDRigidsSolver, Verbose, TEXT("FPBDRigidsSolver::RegisterObject()"));
 		auto& RigidBody_External = Proxy->GetGameThreadAPI();
 
-		if (RigidBody_External.Geometry() && RigidBody_External.Geometry()->HasBoundingBox() && RigidBody_External.Geometry()->BoundingBox().Extents().Max() >= MaxBoundsForTree)
+		if (BroadPhaseConfig.BroadphaseType == FBroadPhaseConfig::Tree)
+		{
+			FSpatialAccelerationIdx Idx = RigidBody_External.SpatialIdx();
+			Idx.Bucket = 0;
+			RigidBody_External.SetSpatialIdx(Idx);
+		}
+		else if (RigidBody_External.Geometry() && RigidBody_External.Geometry()->HasBoundingBox() && RigidBody_External.Geometry()->BoundingBox().Extents().Max() >= MaxBoundsForTree)
 		{
 			RigidBody_External.SetSpatialIdx(FSpatialAccelerationIdx{ 1,0 });
 		}
diff --git a/Engine/Source/Runtime/Experimental/Chaos/Public/Chaos/PBDRigidsEvolution.h b/Engine/Source/Runtime/Experimental/Chaos/Public/Chaos/PBDRigidsEvolution.h
--- a/Engine/Source/Runtime/Experimental/Chaos/Public/Chaos/PBDRigidsEvolution.h
+++ b/Engine/Source/Runtime/Experimental/Chaos/Public/Chaos/PBDRigidsEvolution.h
@@ -30,6 +30,41 @@ namespace Chaos
 {
 extern CHAOS_API int32 ForceNoCollisionIntoSQ;
 
+struct FBroadPhaseConfig
+{
+	enum 
+	{
+		Grid = 0,
+		Tree = 1,
+		TreeOfGrid = 2,
+		TreeAndGrid = 3,
+		TreeOfGridAndGrid = 4
+	};
+
+	int32 BroadphaseType;
+	int32 BVNumCells;
+	int32 MaxChildrenInLeaf;
+	int32 MaxTreeDepth;
+	int32 AABBMaxChildrenInLeaf;
+	int32 AABBMaxTreeDepth;
+	FRealSingle MaxPayloadSize;
+	int32 IterationsPerTimeSlice;
+
+	FBroadPhaseConfig()
+	{
+		BroadphaseType = 3;
+		BVNumCells = 35;
+		MaxChildrenInLeaf = 5;
+		MaxTreeDepth = 200;
+		AABBMaxChildrenInLeaf = 500;
+		AABBMaxTreeDepth = 200;
+		MaxPayloadSize = 100000;
+		IterationsPerTimeSlice = 4000;
+	}
+};
+
+extern CHAOS_API FBroadPhaseConfig BroadPhaseConfig;
+
 namespace Collisions
 {
 	void CHAOS_API ResetChaosCollisionCounters();
