Fix merged issue
Horde Issue 203756

diff --git a/Engine/Source/Runtime/Experimental/Chaos/Private/Chaos/PBDRigidsEvolution.cpp b/Engine/Source/Runtime/Experimental/Chaos/Private/Chaos/PBDRigidsEvolution.cpp
--- a/Engine/Source/Runtime/Experimental/Chaos/Private/Chaos/PBDRigidsEvolution.cpp
+++ b/Engine/Source/Runtime/Experimental/Chaos/Private/Chaos/PBDRigidsEvolution.cpp
@@ -35,6 +35,9 @@ namespace Chaos
 	CHAOS_API int32 AccelerationStructureUseDynamicTree = 1;
 	FAutoConsoleVariableRef CVarAccelerationStructureUseDynamicTree(TEXT("p.Chaos.AccelerationStructureUseDynamicTree"), AccelerationStructureUseDynamicTree, TEXT("Use a dynamic BVH tree structure for dynamic objects"));
 
+	CHAOS_API int32 AccelerationStructureUseDirtyTreeInsteadOfGrid = 1;
+	FAutoConsoleVariableRef CVarAccelerationStructureUseDirtyTreeInsteadOfGrid(TEXT("p.Chaos.AccelerationStructureUseDirtyTreeInsteadOfGrid"), AccelerationStructureUseDirtyTreeInsteadOfGrid, TEXT("Use a dynamic tree structure for dirty elements instead of a 2D grid"));
+
 	/** Console variable to enable the caching of the overlapping leaves if the dynamic tree is enable */
 	CHAOS_API int32 GAccelerationStructureCacheOverlappingLeaves = 1;
 	FAutoConsoleVariableRef CVarAccelerationStructureCacheOverlappingLeaves(TEXT("p.Chaos.AccelerationStructureCacheOverlappingLeaves"), GAccelerationStructureCacheOverlappingLeaves, TEXT("Set to 1: Cache the overlapping leaves for faster overlap query, any other value will disable the feature"));
