ChaosFlesh - Animation Updates
- Pull component space transforms and buffer on physics thread.

diff --git a/Engine/Source/Runtime/Experimental/Chaos/Private/Chaos/Deformable/ChaosDeformableSolver.cpp b/Engine/Source/Runtime/Experimental/Chaos/Private/Chaos/Deformable/ChaosDeformableSolver.cpp
--- a/Engine/Source/Runtime/Experimental/Chaos/Private/Chaos/Deformable/ChaosDeformableSolver.cpp
+++ b/Engine/Source/Runtime/Experimental/Chaos/Private/Chaos/Deformable/ChaosDeformableSolver.cpp
@@ -40,7 +40,10 @@ namespace Chaos::Softs
 
 
 	FDeformableSolver::FDeformableSolver(FDeformableSolverProperties InProp)
-		: Property(InProp)
+		: CurrentInputPackage(TUniquePtr < FDeformablePackage >(nullptr))
+		, PreviousInputPackage(TUniquePtr < FDeformablePackage >(nullptr))
+		, Property(InProp)
+
 	{
 		Reset(Property);
 	}
@@ -405,7 +408,7 @@ namespace Chaos::Softs
 	{
 		FScopeLock Lock(&PackageInputMutex);
 		TRACE_CPUPROFILER_EVENT_SCOPE(DeformableSolver_PushInputPackage);
-		InputPackages.Push(TUniquePtr< FDeformablePackage >(new FDeformablePackage(InFrame, MoveTemp(InPackage))));
+		BufferedInputPackages.Push(TUniquePtr< FDeformablePackage >(new FDeformablePackage(InFrame, MoveTemp(InPackage))));
 	}
 
 	TUniquePtr<FDeformablePackage> FDeformableSolver::FPhysicsThreadAccess::PullInputPackage()
@@ -416,8 +419,8 @@ namespace Chaos::Softs
 	{
 		FScopeLock Lock(&PackageInputMutex);
 		TRACE_CPUPROFILER_EVENT_SCOPE(DeformableSolver_PullInputPackage);
-		if (InputPackages.Num())
-			return InputPackages.Pop();
+		if (BufferedInputPackages.Num())
+			return BufferedInputPackages.Pop();
 		return TUniquePtr<FDeformablePackage>(nullptr);
 	}
 
@@ -427,8 +430,18 @@ namespace Chaos::Softs
 	}
 	void FDeformableSolver::UpdateProxyInputPackages()
 	{
-		// todo(chaosflesh) : just clear the packages for now. 
-		while (PullInputPackage()) {}
+		if (CurrentInputPackage)
+		{
+			PreviousInputPackage = TUniquePtr < FDeformablePackage >(CurrentInputPackage.Release());
+			CurrentInputPackage = TUniquePtr < FDeformablePackage >(nullptr);
+		}
+
+		TUniquePtr < FDeformablePackage > TailPackage = PullInputPackage();
+		while(TailPackage)
+		{
+			CurrentInputPackage = TUniquePtr < FDeformablePackage >(TailPackage.Release());
+			TailPackage = PullInputPackage();
+		}
 	}
 
 	void FDeformableSolver::FPhysicsThreadAccess::TickSimulation(FSolverReal DeltaTime)
@@ -480,7 +493,7 @@ namespace Chaos::Softs
 	{
 		FScopeLock Lock(&PackageOutputMutex);
 		TRACE_CPUPROFILER_EVENT_SCOPE(DeformableSolver_PushPackage);
-		OutputPackages.Push(TUniquePtr< FDeformablePackage >(new FDeformablePackage(InFrame, MoveTemp(InPackage))));
+		BufferedOutputPackages.Push(TUniquePtr< FDeformablePackage >(new FDeformablePackage(InFrame, MoveTemp(InPackage))));
 	}
 
 
@@ -492,8 +505,8 @@ namespace Chaos::Softs
 	{
 		FScopeLock Lock(&PackageOutputMutex);
 		TRACE_CPUPROFILER_EVENT_SCOPE(DeformableSolver_PullPackage);
-		if (OutputPackages.Num())
-			return OutputPackages.Pop();
+		if (BufferedOutputPackages.Num())
+			return BufferedOutputPackages.Pop();
 		return TUniquePtr<FDeformablePackage>(nullptr);
 	}
 
diff --git a/Engine/Source/Runtime/Experimental/Chaos/Public/Chaos/Deformable/ChaosDeformableSolver.h b/Engine/Source/Runtime/Experimental/Chaos/Public/Chaos/Deformable/ChaosDeformableSolver.h
--- a/Engine/Source/Runtime/Experimental/Chaos/Public/Chaos/Deformable/ChaosDeformableSolver.h
+++ b/Engine/Source/Runtime/Experimental/Chaos/Public/Chaos/Deformable/ChaosDeformableSolver.h
@@ -126,8 +126,10 @@ namespace Chaos::Softs
 		static FCriticalSection	PackageInputMutex;
 		TArray< TUniquePtr<FThreadingProxy> > UninitializedProxys;
 		TMap< FThreadingProxy::FKey, TUniquePtr<FThreadingProxy> > Proxies;
-		TArray< TUniquePtr<FDeformablePackage>  > InputPackages;
-		TArray< TUniquePtr<FDeformablePackage>  > OutputPackages;
+		TArray< TUniquePtr<FDeformablePackage>  > BufferedInputPackages;
+		TArray< TUniquePtr<FDeformablePackage>  > BufferedOutputPackages;
+		TUniquePtr < FDeformablePackage > CurrentInputPackage;
+		TUniquePtr < FDeformablePackage > PreviousInputPackage;
 
 		// User Configuration
 		FDeformableSolverProperties Property;
diff --git a/Engine/Source/Runtime/Experimental/Chaos/Public/Chaos/Deformable/ChaosDeformableSolverProxy.h b/Engine/Source/Runtime/Experimental/Chaos/Public/Chaos/Deformable/ChaosDeformableSolverProxy.h
--- a/Engine/Source/Runtime/Experimental/Chaos/Public/Chaos/Deformable/ChaosDeformableSolverProxy.h
+++ b/Engine/Source/Runtime/Experimental/Chaos/Public/Chaos/Deformable/ChaosDeformableSolverProxy.h
@@ -112,15 +112,13 @@ namespace Chaos::Softs
 			typedef FThreadingProxy::FBuffer Super;
 
 		public:
-			FFleshInputBuffer(const UObject* InOwner = nullptr)
+			FFleshInputBuffer(const TArray<FTransform> & InTransforms, const UObject* InOwner = nullptr)
 				: Super(InOwner, FFleshThreadingProxy::TypeName())
-			{
-				// Ref.Dynamic will have updated solver data 
-				//InputData.AddAttribute<FTransform>("Transforms", FTransformCollection::TransformGroup);
-				//InputData.CopyMatchingAttributesFrom(Ref.GetDynamicCollection());
-			}
+				, Transforms(InTransforms)
+			{}
 			virtual ~FFleshInputBuffer() {}
 
+			TArray<FTransform> Transforms;
 			FManagedArrayCollection InputData;
 		};
 
