Chaos : break and shock damage propapgation
- add per component and geo collection asset damage propagation data
- store damage propapagtion info at the proxy level ( sim parameters )
- refactor clustering code to properly account for procxy suim parameters
- add shock propagation model complementary to the break propagation model

diff --git a/Engine/Source/Runtime/Core/Public/UObject/UE5MainStreamObjectVersion.h b/Engine/Source/Runtime/Core/Public/UObject/UE5MainStreamObjectVersion.h
--- a/Engine/Source/Runtime/Core/Public/UObject/UE5MainStreamObjectVersion.h
+++ b/Engine/Source/Runtime/Core/Public/UObject/UE5MainStreamObjectVersion.h
@@ -268,6 +268,9 @@ struct CORE_API FUE5MainStreamObjectVersion
 		// Low resolution data in Height Field will store one height for (6x6) 36 cells
 		DecreaseLowResolutionHeightField,
 
+		// Add damage propagation settings to geometry collections
+		GeometryCollectionDamagePropagationData,
+
 		// -----<new versions can be added above this line>-------------------------------------------------
 		VersionPlusOne,
 		LatestVersion = VersionPlusOne - 1
diff --git a/Engine/Source/Runtime/Experimental/Chaos/Private/Chaos/PBDRigidClustering.cpp b/Engine/Source/Runtime/Experimental/Chaos/Private/Chaos/PBDRigidClustering.cpp
--- a/Engine/Source/Runtime/Experimental/Chaos/Private/Chaos/PBDRigidClustering.cpp
+++ b/Engine/Source/Runtime/Experimental/Chaos/Private/Chaos/PBDRigidClustering.cpp
@@ -623,8 +623,32 @@ namespace Chaos
 		{
 			return ActivatedChildren;
 		}
-		
+
+		// gather propagation information from the parent proxy
+		bool bUseDamagePropagation = false;
+		float BreakDamagePropagationFactor = 0.0f;
+		float ShockDamagePropagationFactor = 0.0f;
+		FPBDRigidClusteredParticleHandle* ClusteredParent = ClusteredParticle->CastToClustered();
+		if (const IPhysicsProxyBase* Proxy = ClusteredParent->PhysicsProxy())
+		{
+			if (Proxy->GetType() == EPhysicsProxyType::GeometryCollectionType)
+			{
+				const FGeometryCollectionPhysicsProxy* ConcreteProxy = static_cast<const FGeometryCollectionPhysicsProxy*>(Proxy);
+				bUseDamagePropagation = ConcreteProxy->GetSimParameters().bUseDamagePropagation;
+				BreakDamagePropagationFactor = ConcreteProxy->GetSimParameters().BreakDamagePropagationFactor;
+				ShockDamagePropagationFactor = ConcreteProxy->GetSimParameters().ShockDamagePropagationFactor;
+			}
+		}
+
 		TArray<FPBDRigidParticleHandle*>& Children = MChildren[ClusteredParticle];
+
+		// only used for propagation
+		TArray<FReal> AppliedStrains;
+		if (bUseDamagePropagation)
+		{
+			AppliedStrains.SetNum(Children.Num());
+		}
+
 		for (int32 ChildIdx = Children.Num() - 1; ChildIdx >= 0; --ChildIdx)
 		{
 			FPBDRigidClusteredParticleHandle* Child = Children[ChildIdx]->CastToClustered();
@@ -659,55 +683,59 @@ namespace Chaos
 				{
 					SendBreakingEvent(Child);
 				}
-
-				if (GraphPropagationBasedCollisionImpulseProcessing && !bForceRelease)
-				{
-					// keep max applied strain in external strain
-					// @todo(chaos) eventually shoudl get rid of collision impulse array and only use external strain
-					Child->SetExternalStrain(MaxAppliedStrain);
-				}
-				else
-				{
-					Child->ClearExternalStrain();	
-				}
 			}
-			else
+			if (bUseDamagePropagation)
 			{
-				// no need anymore of the external strain let's clear it
-				Child->ClearExternalStrain();				
+				AppliedStrains[ChildIdx] = MaxAppliedStrain;
 			}
+			Child->ClearExternalStrain();
 		}
 
-		// if necessary propagate through the graph
-		if (GraphPropagationBasedCollisionImpulseProcessing && !bForceRelease)
+		// if necessary propagate strain through the graph
+		if (bUseDamagePropagation)
 		{
-			for (FPBDRigidParticleHandle* ActivatedChild: ActivatedChildren)
+			for (int32 ChildIdx = 0; ChildIdx < Children.Num(); ChildIdx++)
 			{
-				FPBDRigidClusteredParticleHandle* ClusteredChild = ActivatedChild->CastToClustered();
+				FPBDRigidClusteredParticleHandle* ClusteredChild = Children[ChildIdx]->CastToClustered();
+
+				const FReal AppliedStrain = AppliedStrains[ChildIdx];
+				FReal PropagatedStrainPerConnection = 0.0f;
 
-				const FReal RemainingStrain = (ClusteredChild->GetExternalStrain() - ClusteredChild->Strain());
-				if (RemainingStrain > 0)
+				// @todo(chaos) : may not be optimal, but good enough for now
+				if (ActivatedChildren.Contains(Children[ChildIdx]))
 				{
-					const FReal AdjustedRemainingStrain = GraphPropagationBasedCollisionFactor * RemainingStrain;
-					// todo(chaos) : could do better and have something weighted on distance with a falloff maybe ?  
-					const FReal RemainingStrainPerSibling = AdjustedRemainingStrain / ClusteredChild->ConnectivityEdges().Num();
-					for (const TConnectivityEdge<FReal>& Edge: ClusteredChild->ConnectivityEdges())
+					// break damage propagation case: we only look at the broken pieces and propagate the strain remainder 
+					const FReal RemainingStrain = (AppliedStrain - ClusteredChild->Strain());
+					if (RemainingStrain > 0)
+					{
+						const FReal AdjustedRemainingStrain = (FReal)BreakDamagePropagationFactor * RemainingStrain;
+						// todo(chaos) : could do better and have something weighted on distance with a falloff maybe ?  
+						PropagatedStrainPerConnection = AdjustedRemainingStrain / ClusteredChild->ConnectivityEdges().Num();
+					}
+				}
+				else
+				{
+					// shock damage propagation case : for all the non broken pieces, proapagate the actual applied strain 
+					PropagatedStrainPerConnection = (FReal)ShockDamagePropagationFactor * AppliedStrain;
+				}
+
+				if (PropagatedStrainPerConnection > 0)
+				{
+					for (const TConnectivityEdge<FReal>& Edge : ClusteredChild->ConnectivityEdges())
 					{
 						if (Edge.Sibling)
 						{
 							if (FPBDRigidClusteredParticleHandle* ClusteredSibling = Edge.Sibling->CastToClustered())
 							{
 								// todo(chaos) this may currently be non optimal as we are in the apply loop and this may be cleareed right after
-								ClusteredSibling->SetExternalStrain(FMath::Max(ClusteredSibling->GetExternalStrain(), RemainingStrainPerSibling));
+								ClusteredSibling->SetExternalStrain(FMath::Max(ClusteredSibling->GetExternalStrain(), PropagatedStrainPerConnection));
 							}
 						}
 					}
 				}
-				// finally reset the external strain
-				ClusteredChild->ClearExternalStrain();
 			}
 		}
-		
+
 		if (ActivatedChildren.Num() > 0)
 		{
 			if (Children.Num() == 0)
@@ -1163,7 +1191,18 @@ namespace Chaos
 				const FReal AccumulatedImpulse = ContactHandle->GetAccumulatedImpulse().Size();
 				if (AccumulatedImpulse > UE_SMALL_NUMBER)
 				{
-					if (GraphPropagationBasedCollisionImpulseProcessing)
+					// gather propagation information from the parent proxy
+					bool bUseDamagePropagation = false;
+					if (const IPhysicsProxyBase* Proxy = Cluster->PhysicsProxy())
+					{
+						if (Proxy->GetType() == EPhysicsProxyType::GeometryCollectionType)
+						{
+							const FGeometryCollectionPhysicsProxy* ConcreteProxy = static_cast<const FGeometryCollectionPhysicsProxy*>(Proxy);
+							bUseDamagePropagation = ConcreteProxy->GetSimParameters().bUseDamagePropagation;
+						}
+					}
+
+					if (bUseDamagePropagation)
 					{
 						// propagation based breaking model start from the closest particle and propagate through the connection graph
 						// propagation logic is dealt when evaluating the strain
diff --git a/Engine/Source/Runtime/Experimental/Chaos/Public/GeometryCollection/GeometryCollectionSimulationCoreTypes.h b/Engine/Source/Runtime/Experimental/Chaos/Public/GeometryCollection/GeometryCollectionSimulationCoreTypes.h
--- a/Engine/Source/Runtime/Experimental/Chaos/Public/GeometryCollection/GeometryCollectionSimulationCoreTypes.h
+++ b/Engine/Source/Runtime/Experimental/Chaos/Public/GeometryCollection/GeometryCollectionSimulationCoreTypes.h
@@ -206,6 +206,9 @@ struct FSimulationParameters
 		, UseCCD(false)
 		, LinearDamping(0.01f)
 		, AngularDamping(0)
+		, bUseDamagePropagation(true)
+		, BreakDamagePropagationFactor(1.0f)
+		, ShockDamagePropagationFactor(0.0f)
 		, SimulationFilterData()
 		, QueryFilterData()
 		, UserData(nullptr)
@@ -252,6 +255,9 @@ struct FSimulationParameters
 		, UseCCD(Other.UseCCD)
 		, LinearDamping(Other.LinearDamping)
 		, AngularDamping(Other.AngularDamping)
+		, bUseDamagePropagation(Other.bUseDamagePropagation)
+		, BreakDamagePropagationFactor(Other.BreakDamagePropagationFactor)
+		, ShockDamagePropagationFactor(Other.ShockDamagePropagationFactor)
 		, SimulationFilterData(Other.SimulationFilterData)
 		, QueryFilterData(Other.QueryFilterData)
 		, UserData(Other.UserData)
@@ -322,6 +328,10 @@ struct FSimulationParameters
 	float LinearDamping;
 	float AngularDamping;
 
+	bool bUseDamagePropagation;
+	float BreakDamagePropagationFactor;
+	float ShockDamagePropagationFactor;
+
 	FCollisionFilterData SimulationFilterData;
 	FCollisionFilterData QueryFilterData;
 	void* UserData;
diff --git a/Engine/Source/Runtime/Experimental/GeometryCollectionEngine/Private/GeometryCollection/GeometryCollectionComponent.cpp b/Engine/Source/Runtime/Experimental/GeometryCollectionEngine/Private/GeometryCollection/GeometryCollectionComponent.cpp
--- a/Engine/Source/Runtime/Experimental/GeometryCollectionEngine/Private/GeometryCollection/GeometryCollectionComponent.cpp
+++ b/Engine/Source/Runtime/Experimental/GeometryCollectionEngine/Private/GeometryCollection/GeometryCollectionComponent.cpp
@@ -296,6 +296,9 @@ UGeometryCollectionComponent::UGeometryCollectionComponent(const FObjectInitiali
 
 	bWantsInitializeComponent = true;
 
+	// make sure older asset are using the default behaviour
+	DamagePropagationData.bEnabled = false;
+
 }
 
 Chaos::FPhysicsSolver* GetSolver(const UGeometryCollectionComponent& GeometryCollectionComponent)
@@ -2486,6 +2489,9 @@ void UGeometryCollectionComponent::RegisterAndInitializePhysicsProxy()
 		SimulationParameters.UseCCD = BodyInstance.bUseCCD;
 		SimulationParameters.LinearDamping = BodyInstance.LinearDamping;
 		SimulationParameters.AngularDamping = BodyInstance.AngularDamping;
+		SimulationParameters.bUseDamagePropagation = DamagePropagationData.bEnabled;
+		SimulationParameters.BreakDamagePropagationFactor = DamagePropagationData.BreakDamagePropagationFactor;
+		SimulationParameters.ShockDamagePropagationFactor = DamagePropagationData.ShockDamagePropagationFactor;
 		SimulationParameters.WorldTransform = GetComponentToWorld();
 		SimulationParameters.UserData = static_cast<void*>(&PhysicsUserData);
 
@@ -2666,6 +2672,9 @@ void UGeometryCollectionComponent::SetRestCollection(const UGeometryCollection*
 
 		// initialize the component per level damage threshold from the asset defaults 
 		DamageThreshold = RestCollection->DamageThreshold;
+
+		// initialize the component damage progataion data from the asset defaults 
+		DamagePropagationData = RestCollection->DamagePropagationData;
 		
 		//ResetDynamicCollection();
 	}
diff --git a/Engine/Source/Runtime/Experimental/GeometryCollectionEngine/Private/GeometryCollection/GeometryCollectionObject.cpp b/Engine/Source/Runtime/Experimental/GeometryCollectionEngine/Private/GeometryCollection/GeometryCollectionObject.cpp
--- a/Engine/Source/Runtime/Experimental/GeometryCollectionEngine/Private/GeometryCollection/GeometryCollectionObject.cpp
+++ b/Engine/Source/Runtime/Experimental/GeometryCollectionEngine/Private/GeometryCollection/GeometryCollectionObject.cpp
@@ -826,6 +826,12 @@ void UGeometryCollection::Serialize(FArchive& Ar)
 		PerClusterOnlyDamageThreshold = true; 
 	}
 
+	if (Ar.CustomVer(FUE5MainStreamObjectVersion::GUID) < FUE5MainStreamObjectVersion::GeometryCollectionDamagePropagationData)
+	{
+		// prior this version, damage propagation was not enabled by default
+		DamagePropagationData.bEnabled = false;
+	}
+
 #if WITH_EDITORONLY_DATA
 	if (bCreateSimulationData)
 	{
diff --git a/Engine/Source/Runtime/Experimental/GeometryCollectionEngine/Public/GeometryCollection/GeometryCollectionComponent.h b/Engine/Source/Runtime/Experimental/GeometryCollectionEngine/Public/GeometryCollection/GeometryCollectionComponent.h
--- a/Engine/Source/Runtime/Experimental/GeometryCollectionEngine/Public/GeometryCollection/GeometryCollectionComponent.h
+++ b/Engine/Source/Runtime/Experimental/GeometryCollectionEngine/Public/GeometryCollection/GeometryCollectionComponent.h
@@ -671,13 +671,17 @@ public:
 	UPROPERTY(EditAnywhere, BlueprintReadWrite, Category = "ChaosPhysics|Clustering")
 	int32 MaxClusterLevel;
 
-	UPROPERTY(EditAnywhere, BlueprintReadWrite, Category = "ChaosPhysics|Clustering", meta = (EditCondition = "!bUseSizeSpecificDamageThreshold"))
+	UPROPERTY(EditAnywhere, BlueprintReadWrite, Category = "ChaosPhysics|Damage", meta = (EditCondition = "!bUseSizeSpecificDamageThreshold"))
 	TArray<float> DamageThreshold;
 
 	/** Damage threshold for clusters at different levels. */
-	UPROPERTY(EditAnywhere, BlueprintReadWrite, Category = "ChaosPhysics|Clustering")
+	UPROPERTY(EditAnywhere, BlueprintReadWrite, Category = "ChaosPhysics|Damage")
 	bool bUseSizeSpecificDamageThreshold;
 
+	/** Data about how damage propagation shoudl behave. */
+	UPROPERTY(EditAnywhere, BlueprintReadWrite, Category = "ChaosPhysics|Damage")
+	FGeometryCollectionDamagePropagationData DamagePropagationData;
+
 	/** Allow removal on sleep for the instance if the rest collection has it enabled */
 	UPROPERTY(EditAnywhere, BlueprintReadWrite, Category = "ChaosPhysics|Removal")
 	bool bAllowRemovalOnSleep;
diff --git a/Engine/Source/Runtime/Experimental/GeometryCollectionEngine/Public/GeometryCollection/GeometryCollectionObject.h b/Engine/Source/Runtime/Experimental/GeometryCollectionEngine/Public/GeometryCollection/GeometryCollectionObject.h
--- a/Engine/Source/Runtime/Experimental/GeometryCollectionEngine/Public/GeometryCollection/GeometryCollectionObject.h
+++ b/Engine/Source/Runtime/Experimental/GeometryCollectionEngine/Public/GeometryCollection/GeometryCollectionObject.h
@@ -301,6 +301,26 @@ private:
 	bool bIsInitialized = false;
 };
 
+
+USTRUCT(BlueprintType)
+struct GEOMETRYCOLLECTIONENGINE_API FGeometryCollectionDamagePropagationData
+{
+public:
+	GENERATED_BODY()
+
+	/** Whether or not damage propagation is enabled. */
+	UPROPERTY(EditAnywhere, BlueprintReadWrite, Category = "Damage Propagation")
+	bool bEnabled = true;
+
+	/** factor of the remaining strain propagated through the connection graph after a piece breaks. */
+	UPROPERTY(EditAnywhere, BlueprintReadWrite, Category = "Damage Propagation")
+	float BreakDamagePropagationFactor = 1.0f;
+
+	/** factor of the received strain propagated throug the connection graph if the piece did not break. */
+	UPROPERTY(EditAnywhere, BlueprintReadWrite, Category = "Damage Propagation")
+	float ShockDamagePropagationFactor = 0.0f;
+};
+
 /**
 * UGeometryCollectionObject (UObject)
 *
@@ -436,15 +456,19 @@ public:
 	int32 MaxClusterLevel;
 
 	/** Damage threshold for clusters at different levels. */
-	UPROPERTY(EditAnywhere, Category = "Clustering")
+	UPROPERTY(EditAnywhere, Category = "Damage")
 	TArray<float> DamageThreshold;
 
 	/** compatibility check, when true, only cluster compute damage from parameters and propagate to direct children
 	 *  when false, each child will compute it's damage threshold allowing for more precise and intuitive destruction behavior
 	 */
-	UPROPERTY(EditAnywhere, Category = "Clustering")
+	UPROPERTY(EditAnywhere, Category = "Damage")
 	bool PerClusterOnlyDamageThreshold;
 
+	/** Data about how damage propagation shoudl behave. */
+	UPROPERTY(EditAnywhere, Category = "Damage")
+	FGeometryCollectionDamagePropagationData DamagePropagationData;
+
 	/** */
 	UPROPERTY(EditAnywhere, Category = "Clustering")
 	EClusterConnectionTypeEnum ClusterConnectionType;
