ChaosFlesh - Moved the ThreadingAPI out for clarity.

diff --git a/Engine/Source/Runtime/Experimental/Chaos/Private/Chaos/Deformable/ChaosDeformableSolver.cpp b/Engine/Source/Runtime/Experimental/Chaos/Private/Chaos/Deformable/ChaosDeformableSolver.cpp
--- a/Engine/Source/Runtime/Experimental/Chaos/Private/Chaos/Deformable/ChaosDeformableSolver.cpp
+++ b/Engine/Source/Runtime/Experimental/Chaos/Private/Chaos/Deformable/ChaosDeformableSolver.cpp
@@ -48,10 +48,6 @@ namespace Chaos::Softs
 		Reset(Property);
 	}
 
-	void FDeformableSolver::FPhysicsThreadAccess::Reset(const FDeformableSolverProperties& InProps)
-	{
-		Solver.Reset(InProps);
-	}
 	void FDeformableSolver::Reset(const FDeformableSolverProperties& InProps)
 	{
 		TRACE_CPUPROFILER_EVENT_SCOPE(DeformableSolver_Reset);
@@ -79,10 +75,6 @@ namespace Chaos::Softs
 		Time = 0.f;
 	}
 
-	bool FDeformableSolver::FPhysicsThreadAccess::Advance(FSolverReal DeltaTime)
-	{
-		return Solver.Advance(DeltaTime);
-	}
 	bool FDeformableSolver::Advance(FSolverReal DeltaTime)
 	{
 		SCOPE_CYCLE_COUNTER(STAT_DeformableSolver_Advance);
@@ -105,10 +97,6 @@ namespace Chaos::Softs
 		return false;
 	}
 
-	void FDeformableSolver::FPhysicsThreadAccess::InitializeSimulationObjects()
-	{
-		Solver.InitializeSimulationObjects();
-	}
 	void FDeformableSolver::InitializeSimulationObjects()
 	{
 		TRACE_CPUPROFILER_EVENT_SCOPE(DeformableSolver_InitializeSimulationObjects);
@@ -138,10 +126,6 @@ namespace Chaos::Softs
 		InitializeCollisionBodies();
 	}
 
-	void FDeformableSolver::FPhysicsThreadAccess::InitializeSimulationObject(FThreadingProxy& InProxy)
-	{
-		Solver.InitializeSimulationObject(InProxy);
-	}
 	void FDeformableSolver::InitializeSimulationObject(FThreadingProxy& InProxy)
 	{
 		TRACE_CPUPROFILER_EVENT_SCOPE(DeformableSolver_InitializeSimulationObject);
@@ -284,10 +268,6 @@ namespace Chaos::Softs
 		}
 	}
 
-	void FDeformableSolver::FPhysicsThreadAccess::InitializeCollisionBodies()
-	{
-		return Solver.InitializeCollisionBodies();
-	}
 	void FDeformableSolver::InitializeCollisionBodies()
 	{
 		TRACE_CPUPROFILER_EVENT_SCOPE(DeformableSolver_InitializeCollisionBodies);
@@ -302,10 +282,6 @@ namespace Chaos::Softs
 		}
 	}
 
-	void FDeformableSolver::FPhysicsThreadAccess::InitializeKinematicState(FThreadingProxy& InProxy)
-	{
-		Solver.InitializeKinematicState(InProxy);
-	}
 	void FDeformableSolver::InitializeKinematicState(FThreadingProxy& InProxy)
 	{
 		auto MKineticUpdate = [this](FSolverParticles& MParticles, const FSolverReal Dt, const FSolverReal MTime, const int32 Index)
@@ -335,10 +311,6 @@ namespace Chaos::Softs
 		Evolution->SetKinematicUpdateFunction(MKineticUpdate);
 	}
 
-	void FDeformableSolver::FPhysicsThreadAccess::InitializeSelfCollisionVariables()
-	{
-		Solver.InitializeSelfCollisionVariables();
-	}
 	void FDeformableSolver::InitializeSelfCollisionVariables()
 	{
 		TRACE_CPUPROFILER_EVENT_SCOPE(DeformableSolver_InitializeSelfCollisionVariables);
@@ -400,10 +372,6 @@ namespace Chaos::Softs
 	}
 
 
-	void FDeformableSolver::FGameThreadAccess::PushInputPackage(int32 InFrame, FDeformableDataMap&& InPackage)
-	{
-		Solver.PushInputPackage(InFrame, MoveTemp(InPackage));
-	}
 	void FDeformableSolver::PushInputPackage(int32 InFrame, FDeformableDataMap&& InPackage)
 	{
 		FScopeLock Lock(&PackageInputMutex);
@@ -411,10 +379,6 @@ namespace Chaos::Softs
 		BufferedInputPackages.Push(TUniquePtr< FDeformablePackage >(new FDeformablePackage(InFrame, MoveTemp(InPackage))));
 	}
 
-	TUniquePtr<FDeformablePackage> FDeformableSolver::FPhysicsThreadAccess::PullInputPackage()
-	{
-		return Solver.PullInputPackage();
-	}
 	TUniquePtr<FDeformablePackage> FDeformableSolver::PullInputPackage()
 	{
 		FScopeLock Lock(&PackageInputMutex);
@@ -424,10 +388,6 @@ namespace Chaos::Softs
 		return TUniquePtr<FDeformablePackage>(nullptr);
 	}
 
-	void FDeformableSolver::FPhysicsThreadAccess::UpdateProxyInputPackages()
-	{
-		return Solver.UpdateProxyInputPackages();
-	}
 	void FDeformableSolver::UpdateProxyInputPackages()
 	{
 		if (CurrentInputPackage)
@@ -444,10 +404,6 @@ namespace Chaos::Softs
 		}
 	}
 
-	void FDeformableSolver::FPhysicsThreadAccess::TickSimulation(FSolverReal DeltaTime)
-	{
-		Solver.TickSimulation(DeltaTime);
-	}
 	void FDeformableSolver::TickSimulation(FSolverReal DeltaTime)
 	{
 		TRACE_CPUPROFILER_EVENT_SCOPE(DeformableSolver_TickSimulation);
@@ -484,11 +440,6 @@ namespace Chaos::Softs
 		PushOutputPackage(Frame, MoveTemp(OutputBuffers));
 	}
 
-
-	void FDeformableSolver::FPhysicsThreadAccess::PushOutputPackage(int32 InFrame, FDeformableDataMap&& InPackage)
-	{
-		Solver.PushOutputPackage(InFrame, MoveTemp(InPackage));
-	}
 	void FDeformableSolver::PushOutputPackage(int32 InFrame, FDeformableDataMap&& InPackage)
 	{
 		FScopeLock Lock(&PackageOutputMutex);
@@ -496,11 +447,6 @@ namespace Chaos::Softs
 		BufferedOutputPackages.Push(TUniquePtr< FDeformablePackage >(new FDeformablePackage(InFrame, MoveTemp(InPackage))));
 	}
 
-
-	TUniquePtr<FDeformablePackage>  FDeformableSolver::FGameThreadAccess::PullOutputPackage()
-	{
-		return Solver.PullOutputPackage();
-	}
 	TUniquePtr<FDeformablePackage> FDeformableSolver::PullOutputPackage()
 	{
 		FScopeLock Lock(&PackageOutputMutex);
@@ -510,19 +456,11 @@ namespace Chaos::Softs
 		return TUniquePtr<FDeformablePackage>(nullptr);
 	}
 
-	void  FDeformableSolver::FGameThreadAccess::AddProxy(TUniquePtr<FThreadingProxy> InObject)
-	{
-		return Solver.AddProxy(TUniquePtr<FThreadingProxy>(InObject.Release()));
-	}
 	void FDeformableSolver::AddProxy(TUniquePtr<FThreadingProxy> InObject)
 	{
 		UninitializedProxys.Add(TUniquePtr< FThreadingProxy>(InObject.Release()));
 	}
 
-	void  FDeformableSolver::FPhysicsThreadAccess::UpdateOutputState(FThreadingProxy& InProxy)
-	{
-		Solver.UpdateOutputState(InProxy);
-	}
 	void FDeformableSolver::UpdateOutputState(FThreadingProxy& InProxy)
 	{
 		TRACE_CPUPROFILER_EVENT_SCOPE(DeformableSolver_UpdateOutputState);
@@ -545,10 +483,6 @@ namespace Chaos::Softs
 		}
 	}
 
-	void  FDeformableSolver::FPhysicsThreadAccess::WriteFrame(FThreadingProxy& InProxy, const FSolverReal DeltaTime)
-	{
-		Solver.WriteFrame(InProxy, DeltaTime);
-	}
 	void FDeformableSolver::WriteFrame(FThreadingProxy& InProxy, const FSolverReal DeltaTime)
 	{
 		TRACE_CPUPROFILER_EVENT_SCOPE(DeformableSolver_WriteFrame);
@@ -570,10 +504,6 @@ namespace Chaos::Softs
 		}
 	}
 
-	void  FDeformableSolver::FPhysicsThreadAccess::WriteTrisGEO(const FSolverParticles& Particles, const TArray<TVec3<int32>>& Mesh)
-	{
-		Solver.WriteTrisGEO(Particles, Mesh);
-	}
 	void FDeformableSolver::WriteTrisGEO(const FSolverParticles& Particles, const TArray<TVec3<int32>>& Mesh)
 	{
 		FString file = FPaths::ProjectDir();
diff --git a/Engine/Source/Runtime/Experimental/Chaos/Private/Chaos/Deformable/ChaosDeformableSolverThreadingWrapper.cpp b/Engine/Source/Runtime/Experimental/Chaos/Private/Chaos/Deformable/ChaosDeformableSolverThreadingWrapper.cpp
new file mode 100644
--- /dev/null
+++ b/Engine/Source/Runtime/Experimental/Chaos/Private/Chaos/Deformable/ChaosDeformableSolverThreadingWrapper.cpp
@@ -0,0 +1,120 @@
+// Copyright Epic Games, Inc. All Rights Reserved.
+
+#include "Chaos/Deformable/ChaosDeformableSolver.h"
+#include "Chaos/Deformable/ChaosDeformableSolverTypes.h"
+#include "Chaos/Deformable/ChaosDeformableSolverProxy.h"
+
+#include "Chaos/TriangleMesh.h"
+#include "Chaos/PBDAltitudeSpringConstraints.h"
+#include "Chaos/PBDBendingConstraints.h"
+#include "Chaos/PBDSpringConstraints.h"
+#include "Chaos/PBDVolumeConstraint.h"
+#include "Chaos/PBDSpringConstraints.h"
+#include "Chaos/PBDTriangleMeshCollisions.h"
+#include "Chaos/PBDCollisionSpringConstraints.h"
+#include "Chaos/PBDParticles.h"
+#include "Chaos/PBDSoftsSolverParticles.h"
+#include "Chaos/PBDTetConstraints.h"
+#include "Chaos/PerParticleGravity.h"
+#include "Chaos/XPBDCorotatedConstraints.h"
+#include "Chaos/XPBDVolumeConstraints.h"
+#include "Chaos/XPBDCorotatedFiberConstraints.h"
+#include "Chaos/Plane.h"
+#include "Chaos/Utilities.h"
+#include "Chaos/PBDEvolution.h"
+#include "Containers/StringConv.h"
+#include "CoreMinimal.h"
+#include "GeometryCollection/Facades/CollectionKinematicBindingFacade.h"
+#include "Misc/FileHelper.h"
+#include "HAL/PlatformFileManager.h"
+#include "HAL/IConsoleManager.h"
+#include "Misc/Paths.h"
+
+
+namespace Chaos::Softs
+{
+	void FDeformableSolver::FPhysicsThreadAccess::Reset(const FDeformableSolverProperties& InProps)
+	{
+		Solver.Reset(InProps);
+	}
+
+	bool FDeformableSolver::FPhysicsThreadAccess::Advance(FSolverReal DeltaTime)
+	{
+		return Solver.Advance(DeltaTime);
+	}
+
+	void FDeformableSolver::FPhysicsThreadAccess::InitializeSimulationObjects()
+	{
+		Solver.InitializeSimulationObjects();
+	}
+
+	void FDeformableSolver::FPhysicsThreadAccess::InitializeSimulationObject(FThreadingProxy& InProxy)
+	{
+		Solver.InitializeSimulationObject(InProxy);
+	}
+
+	void FDeformableSolver::FPhysicsThreadAccess::InitializeCollisionBodies()
+	{
+		return Solver.InitializeCollisionBodies();
+	}
+
+	void FDeformableSolver::FPhysicsThreadAccess::InitializeKinematicState(FThreadingProxy& InProxy)
+	{
+		Solver.InitializeKinematicState(InProxy);
+	}
+
+	void FDeformableSolver::FPhysicsThreadAccess::InitializeSelfCollisionVariables()
+	{
+		Solver.InitializeSelfCollisionVariables();
+	}
+
+	void FDeformableSolver::FGameThreadAccess::PushInputPackage(int32 InFrame, FDeformableDataMap&& InPackage)
+	{
+		Solver.PushInputPackage(InFrame, MoveTemp(InPackage));
+	}
+
+	TUniquePtr<FDeformablePackage> FDeformableSolver::FPhysicsThreadAccess::PullInputPackage()
+	{
+		return Solver.PullInputPackage();
+	}
+
+	void FDeformableSolver::FPhysicsThreadAccess::UpdateProxyInputPackages()
+	{
+		return Solver.UpdateProxyInputPackages();
+	}
+
+	void FDeformableSolver::FPhysicsThreadAccess::TickSimulation(FSolverReal DeltaTime)
+	{
+		Solver.TickSimulation(DeltaTime);
+	}
+
+	void FDeformableSolver::FPhysicsThreadAccess::PushOutputPackage(int32 InFrame, FDeformableDataMap&& InPackage)
+	{
+		Solver.PushOutputPackage(InFrame, MoveTemp(InPackage));
+	}
+
+	TUniquePtr<FDeformablePackage>  FDeformableSolver::FGameThreadAccess::PullOutputPackage()
+	{
+		return Solver.PullOutputPackage();
+	}
+
+	void  FDeformableSolver::FGameThreadAccess::AddProxy(TUniquePtr<FThreadingProxy> InObject)
+	{
+		return Solver.AddProxy(TUniquePtr<FThreadingProxy>(InObject.Release()));
+	}
+
+	void  FDeformableSolver::FPhysicsThreadAccess::UpdateOutputState(FThreadingProxy& InProxy)
+	{
+		Solver.UpdateOutputState(InProxy);
+	}
+
+	void  FDeformableSolver::FPhysicsThreadAccess::WriteFrame(FThreadingProxy& InProxy, const FSolverReal DeltaTime)
+	{
+		Solver.WriteFrame(InProxy, DeltaTime);
+	}
+
+	void  FDeformableSolver::FPhysicsThreadAccess::WriteTrisGEO(const FSolverParticles& Particles, const TArray<TVec3<int32>>& Mesh)
+	{
+		Solver.WriteTrisGEO(Particles, Mesh);
+	}
+}; // Namespace Chaos::Softs
\ No newline at end of file
