Dataflow - Context now can switch between multi/single threaded modes

diff --git a/Engine/Plugins/Experimental/Dataflow/Source/DataflowEditor/Private/Dataflow/DataflowGraphEditor.cpp b/Engine/Plugins/Experimental/Dataflow/Source/DataflowEditor/Private/Dataflow/DataflowGraphEditor.cpp
--- a/Engine/Plugins/Experimental/Dataflow/Source/DataflowEditor/Private/Dataflow/DataflowGraphEditor.cpp
+++ b/Engine/Plugins/Experimental/Dataflow/Source/DataflowEditor/Private/Dataflow/DataflowGraphEditor.cpp
@@ -115,7 +115,7 @@ void SDataflowGraphEditor::EvaluateNode()
 	{
 		FDataflowEditorCommands::FGraphEvaluationCallback LocalEvaluateCallback = [](FDataflowNode* Node, FDataflowOutput* Out)
 		{
-			Dataflow::FContext Context(FGameTime::GetTimeSinceAppStart().GetRealTimeSeconds());
+			Dataflow::FContextThreaded Context(FGameTime::GetTimeSinceAppStart().GetRealTimeSeconds());
 			Node->Evaluate(Context, Out);
 		};
 
diff --git a/Engine/Source/Runtime/Experimental/Dataflow/Core/Private/Dataflow/DataflowInputOutput.cpp b/Engine/Source/Runtime/Experimental/Dataflow/Core/Private/Dataflow/DataflowInputOutput.cpp
--- a/Engine/Source/Runtime/Experimental/Dataflow/Core/Private/Dataflow/DataflowInputOutput.cpp
+++ b/Engine/Source/Runtime/Experimental/Dataflow/Core/Private/Dataflow/DataflowInputOutput.cpp
@@ -129,3 +129,20 @@ void FDataflowOutput::Invalidate()
 	}
 }
 
+bool FDataflowOutput::EvaluateImpl(Dataflow::FContext& Context) const
+{
+	// check if the cache has a valid version
+	if(Context.HasData(CacheKey(), OwningNode->LastModifiedTimestamp))
+	{
+		return true;
+	}
+	// if not, evaluate
+	OwningNode->Evaluate(Context, this);
+	// Validation
+	if (!Context.HasData(CacheKey()))
+	{
+		ensureMsgf(false, TEXT("Failed to evaluate output (%s:%s)"), *OwningNode->GetName().ToString(), *GetName().ToString());
+		return false;
+	}
+	return true;
+}
\ No newline at end of file
diff --git a/Engine/Source/Runtime/Experimental/Dataflow/Core/Private/Dataflow/DataflowNodeParameters.cpp b/Engine/Source/Runtime/Experimental/Dataflow/Core/Private/Dataflow/DataflowNodeParameters.cpp
new file mode 100644
--- /dev/null
+++ b/Engine/Source/Runtime/Experimental/Dataflow/Core/Private/Dataflow/DataflowNodeParameters.cpp
@@ -0,0 +1,19 @@
+﻿// Copyright Epic Games, Inc. All Rights Reserved.
+
+#include "Dataflow/DataflowNodeParameters.h"
+#include "Dataflow/DataflowInputOutput.h"
+
+namespace Dataflow
+{
+	bool FContextSingle::Evaluate(const FDataflowOutput& Connection)
+	{
+		return Connection.EvaluateImpl(*this);
+	}
+
+	bool FContextThreaded::Evaluate(const FDataflowOutput& Connection)
+	{
+		Connection.OutputLock->Lock(); ON_SCOPE_EXIT { Connection.OutputLock->Unlock(); };
+		return Connection.EvaluateImpl(*this);
+	}
+}
+
diff --git a/Engine/Source/Runtime/Experimental/Dataflow/Core/Public/Dataflow/DataflowInputOutput.h b/Engine/Source/Runtime/Experimental/Dataflow/Core/Public/Dataflow/DataflowInputOutput.h
--- a/Engine/Source/Runtime/Experimental/Dataflow/Core/Public/Dataflow/DataflowInputOutput.h
+++ b/Engine/Source/Runtime/Experimental/Dataflow/Core/Public/Dataflow/DataflowInputOutput.h
@@ -89,11 +89,7 @@ struct DATAFLOWCORE_API FDataflowOutput : public FDataflowConnection
 	GENERATED_USTRUCT_BODY()
 
 	friend struct FDataflowConnection;
-
-	mutable uint32 CacheKeyValue = UINT_MAX;
-	mutable Dataflow::TCacheValue Cache;
-
-
+	
 	TArray< FDataflowInput* > Connections;
 
 	size_t PassthroughOffsetAddress = INDEX_NONE;
@@ -101,11 +97,8 @@ struct DATAFLOWCORE_API FDataflowOutput : public FDataflowConnection
 public:
 	static FDataflowOutput NoOpOutput;
 	
-	mutable TSharedPtr<FCriticalSection> OutputLock; // @todo(dataflow) find alternatives
-	// Why is this made into a SharedPtr? FCriticalSection cannot be copied - making FCriticalSection a member of
-	// FDataflowOutput deletes the copy constructor of FDataflowOutput. But some other part of the code needs
-	// FDataflowOutput to be copyable. Hence wrapped it in a sharedPtr inspired by NiagaraEmitter.h:227.
-
+	mutable TSharedPtr<FCriticalSection> OutputLock;
+	
 	FDataflowOutput(const Dataflow::FOutputParameters& Param = {}, FGuid InGuid = FGuid::NewGuid());
 
 	TArray<FDataflowInput*>& GetConnections();
@@ -156,6 +149,8 @@ public:
 		return Default;
 	}
 
+	bool EvaluateImpl(Dataflow::FContext& Context) const;
+	
 	template<class T>
 	bool Evaluate(Dataflow::FContext& Context) const;
 
@@ -197,29 +192,15 @@ TFuture<const T&> FDataflowInput::GetValueParallel(Dataflow::FContext& Context,
 template<class T>
 bool FDataflowOutput::Evaluate(Dataflow::FContext& Context) const
 {
-	OutputLock->Lock(); ON_SCOPE_EXIT { this->OutputLock->Unlock(); };
-
 	check(OwningNode);
  
 	if (OwningNode->bActive)
 	{
-		// check if the cache has a valid version
-		if(Context.HasData(CacheKey(), OwningNode->LastModifiedTimestamp))
-		{
-			return true;
-		}
-		// if not, evaluate
-		OwningNode->Evaluate(Context, this);
-		// Validation
-		if (!Context.HasData(CacheKey()))
-		{
-			ensureMsgf(false, TEXT("Failed to evaluate output (%s:%s)"), *OwningNode->GetName().ToString(), *GetName().ToString());
-			return false;
-		}
-		return true;
+		return Context.Evaluate(*this);
 	}
 	else if(const FDataflowInput* PassthroughInput = OwningNode->FindInput(GetPassthroughRealAddress()))
 	{
+		// @todo(dataflow) would be nice if the passthrough does not overwrite the existing cache value.
 		T PassthroughData = PassthroughInput->GetValue<T>(Context, *reinterpret_cast<const T*>(PassthroughInput->RealAddress()));
 		SetValue<T>(PassthroughData, Context);
 		return true;
diff --git a/Engine/Source/Runtime/Experimental/Dataflow/Core/Public/Dataflow/DataflowNodeParameters.h b/Engine/Source/Runtime/Experimental/Dataflow/Core/Public/Dataflow/DataflowNodeParameters.h
--- a/Engine/Source/Runtime/Experimental/Dataflow/Core/Public/Dataflow/DataflowNodeParameters.h
+++ b/Engine/Source/Runtime/Experimental/Dataflow/Core/Public/Dataflow/DataflowNodeParameters.h
@@ -8,16 +8,17 @@
 #include "GenericPlatform/GenericPlatformCriticalSection.h"
 
 class  UDataflow;
+struct FDataflowOutput;
 
 namespace Dataflow
 {
-	struct FContextCacheBase 
+	struct FContextCacheElementBase 
 	{
-		FContextCacheBase(FProperty* InProperty = nullptr, uint64 InTimestamp = 0)
+		FContextCacheElementBase(FProperty* InProperty = nullptr, uint64 InTimestamp = 0)
 			: Property(InProperty)
 			, Timestamp(InTimestamp)
 		{}
-		virtual ~FContextCacheBase() {}
+		virtual ~FContextCacheElementBase() {}
 
 		template<typename T>
 		const T& GetTypedData(const FProperty* PropertyIn) const;
@@ -27,15 +28,15 @@ namespace Dataflow
 	};
 
 	template<class T>
-	struct FContextCache : public FContextCacheBase 
+	struct FContextCacheElement : public FContextCacheElementBase 
 	{
-		FContextCache(FProperty* InProperty, const T& InData, uint64 Timestamp)
-			: FContextCacheBase(InProperty, Timestamp)
+		FContextCacheElement(FProperty* InProperty, const T& InData, uint64 Timestamp)
+			: FContextCacheElementBase(InProperty, Timestamp)
 			, Data(InData)
 		{}
 
-		FContextCache(FProperty* InProperty, T&& InData, uint64 Timestamp)
-			: FContextCacheBase(InProperty, Timestamp)
+		FContextCacheElement(FProperty* InProperty, T&& InData, uint64 Timestamp)
+			: FContextCacheElementBase(InProperty, Timestamp)
 			, Data(InData)
 		{}
 		
@@ -43,34 +44,37 @@ namespace Dataflow
 	};
 
 	template<class T>
-	const T& FContextCacheBase::GetTypedData(const FProperty* PropertyIn) const
+	const T& FContextCacheElementBase::GetTypedData(const FProperty* PropertyIn) const
 	{
 		check(PropertyIn);
 		// check(PropertyIn->IsA<T>()); // @todo(dataflow) compile error for non-class T; find alternatives
 		check(Property->SameType(PropertyIn));
-		return static_cast<const FContextCache<T>&>(*this).Data;
+		return static_cast<const FContextCacheElement<T>&>(*this).Data;
 	}
-	
+
+	struct FContextCache : public TMap<int64, TUniquePtr<FContextCacheElementBase>>
+	{
+		// @todo(dataflow) make an API for FContextCache
+	};
+
 	class DATAFLOWCORE_API FContext
 	{
-		TMap<int64, TUniquePtr<FContextCacheBase>> DataStore;
-		TSharedPtr<FCriticalSection> CacheLock;
-		
+	protected:
 		FContext(FContext&&) = default;
 		FContext& operator=(FContext&&) = default;
 		
 		FContext(const FContext&) = delete;
 		FContext& operator=(const FContext&) = delete;
 
-	public:
 
+	public:
 		FContext(float InTime, FString InType = FString(""))
 			: Timestamp(InTime)
 			, Type(StaticType().Append(InType))
-		{
-			CacheLock = MakeShared<FCriticalSection>();
-		}
+		{}
 
+		virtual ~FContext() {}
+		
 		float Timestamp = 0.f;
 		FString Type;
 		static FString StaticType() { return "FContext"; }
@@ -90,51 +94,116 @@ namespace Dataflow
 			return nullptr;
 		}
 
+		virtual void SetDataImpl(int64 Key, TUniquePtr<FContextCacheElementBase>&& DataStoreEntry) = 0;
+		
 		template<typename T>
 		void SetData(size_t Key, FProperty* Property, const T& Value)
 		{
-			CacheLock->Lock(); ON_SCOPE_EXIT { CacheLock->Unlock(); };
-			
 			int64 IntKey = (int64)Key;
-			TUniquePtr<FContextCache<T>> DataStoreEntry = MakeUnique<FContextCache<T>>(Property, Value, FPlatformTime::Cycles64());
-			DataStore.Emplace(IntKey, MoveTemp(DataStoreEntry));
+			TUniquePtr<FContextCacheElement<T>> DataStoreEntry = MakeUnique<FContextCacheElement<T>>(Property, Value, FPlatformTime::Cycles64());
+
+			SetDataImpl(IntKey, MoveTemp(DataStoreEntry));
 		}
 
 		template<typename T>
 		void SetData(size_t Key, FProperty* Property, T&& Value)
 		{
-			CacheLock->Lock(); ON_SCOPE_EXIT { CacheLock->Unlock(); };
-			
 			int64 IntKey = (int64)Key;
-			TUniquePtr<FContextCache<T>> DataStoreEntry = MakeUnique<FContextCache<T>>(Property, Forward<T>(Value), FPlatformTime::Cycles64());
-			DataStore.Emplace(IntKey, MoveTemp(DataStoreEntry));
+			TUniquePtr<FContextCacheElement<T>> DataStoreEntry = MakeUnique<FContextCacheElement<T>>(Property, Forward<T>(Value), FPlatformTime::Cycles64());
+
+			SetDataImpl(IntKey, MoveTemp(DataStoreEntry));
 		}
+
 		
+		virtual TUniquePtr<FContextCacheElementBase>* GetDataImpl(int64 Key) = 0;
+
 		template<class T>
 		const T& GetData(size_t Key, FProperty* Property, const T& Default = T())
 		{
-			CacheLock->Lock(); ON_SCOPE_EXIT { CacheLock->Unlock(); };
-			
-			if (TUniquePtr<FContextCacheBase>* Cache = DataStore.Find(Key))
+			if (TUniquePtr<FContextCacheElementBase>* Cache = GetDataImpl(Key))
 			{
 				return (*Cache)->GetTypedData<T>(Property);
 			}
 			return Default;
 		}
 
+		
+		virtual bool HasDataImpl(int64 Key, uint64 StoredAfter = 0) = 0;
+		
 		bool HasData(size_t Key, uint64 StoredAfter = 0)
 		{
-			CacheLock->Lock(); ON_SCOPE_EXIT { CacheLock->Unlock(); };
-			
 			int64 IntKey = (int64)Key;
-			return DataStore.Contains(IntKey) && DataStore[Key]->Timestamp >= StoredAfter;
+			return HasDataImpl(Key, StoredAfter);
 		}
+		
+		virtual bool Evaluate(const FDataflowOutput& Connection) = 0;
 	};
 
+	class DATAFLOWCORE_API FContextSingle : public FContext
+	{
+		FContextCache DataStore;
 
-	struct TCacheValue {
-		TCacheValue() {}
+	public:
+
+		FContextSingle(float InTime, FString InType = FString(""))
+			: FContext(InTime, InType)
+		{}
+
+		virtual void SetDataImpl(int64 Key, TUniquePtr<FContextCacheElementBase>&& DataStoreEntry) override
+		{
+			DataStore.Emplace(Key, MoveTemp(DataStoreEntry));
+		}
+
+		virtual TUniquePtr<FContextCacheElementBase>* GetDataImpl(int64 Key) override
+		{
+			return DataStore.Find(Key);
+		}
+
+		virtual bool HasDataImpl(int64 Key, uint64 StoredAfter = 0) override
+		{
+			return DataStore.Contains(Key) && DataStore[Key]->Timestamp >= StoredAfter;
+		}
+
+		virtual bool Evaluate(const FDataflowOutput& Connection) override;
 	};
+	
+	class DATAFLOWCORE_API FContextThreaded : public FContext
+	{
+		FContextCache DataStore;
+		TSharedPtr<FCriticalSection> CacheLock;
+
+	public:
+
+		FContextThreaded(float InTime, FString InType = FString(""))
+			: FContext(InTime, InType)
+		{
+			CacheLock = MakeShared<FCriticalSection>();
+		}
+
+		virtual void SetDataImpl(int64 Key, TUniquePtr<FContextCacheElementBase>&& DataStoreEntry) override
+		{
+			CacheLock->Lock(); ON_SCOPE_EXIT { CacheLock->Unlock(); };
+			
+			DataStore.Emplace(Key, MoveTemp(DataStoreEntry));
+		}
+
+		virtual TUniquePtr<FContextCacheElementBase>* GetDataImpl(int64 Key) override
+		{
+			CacheLock->Lock(); ON_SCOPE_EXIT { CacheLock->Unlock(); };
+			
+			return DataStore.Find(Key);
+		}
+
+		virtual bool HasDataImpl(int64 Key, uint64 StoredAfter = 0) override
+		{
+			CacheLock->Lock(); ON_SCOPE_EXIT { CacheLock->Unlock(); };
+			
+			return DataStore.Contains(Key) && DataStore[Key]->Timestamp >= StoredAfter;
+		}
+
+		virtual bool Evaluate(const FDataflowOutput& Connection) override;
+	};
+
 }
 
 FORCEINLINE uint32 GetTypeHash(const Dataflow::FContext& Context)
diff --git a/Engine/Source/Runtime/Experimental/Dataflow/Engine/Private/Dataflow/DataflowObjectInterface.cpp b/Engine/Source/Runtime/Experimental/Dataflow/Engine/Private/Dataflow/DataflowObjectInterface.cpp
--- a/Engine/Source/Runtime/Experimental/Dataflow/Engine/Private/Dataflow/DataflowObjectInterface.cpp
+++ b/Engine/Source/Runtime/Experimental/Dataflow/Engine/Private/Dataflow/DataflowObjectInterface.cpp
@@ -1,12 +1,3 @@
-// Copyright Epic Games, Inc. All Rights Reserved.
+﻿// Copyright Epic Games, Inc. All Rights Reserved.
 
 #include "Dataflow/DataflowObjectInterface.h"
-
-namespace Dataflow
-{
-	FEngineContext::FEngineContext(UObject* InOwner, UDataflow* InGraph, float InTime, FString InType)
-			: FContext(InTime, StaticType().Append(InType) )
-			, Owner(InOwner)
-			, Graph(InGraph)
-		{}
-}
\ No newline at end of file
diff --git a/Engine/Source/Runtime/Experimental/Dataflow/Engine/Public/Dataflow/DataflowObjectInterface.h b/Engine/Source/Runtime/Experimental/Dataflow/Engine/Public/Dataflow/DataflowObjectInterface.h
--- a/Engine/Source/Runtime/Experimental/Dataflow/Engine/Public/Dataflow/DataflowObjectInterface.h
+++ b/Engine/Source/Runtime/Experimental/Dataflow/Engine/Public/Dataflow/DataflowObjectInterface.h
@@ -12,14 +12,29 @@ class UObject;
 
 namespace Dataflow
 {
-
-	class DATAFLOWENGINE_API FEngineContext : public FContext
+	template<class Base = FContextSingle>
+	class DATAFLOWENGINE_API TEngineContext : public Base
 	{
 	public:
-		FEngineContext(UObject* Owner, UDataflow* InGraph, float InTime, FString InType = "");
-		static FString StaticType() { return "FEngineContext"; }
-
+		TEngineContext(UObject* InOwner, UDataflow* InGraph, float InTime, FString InType)
+				: Base(InTime, StaticType().Append(InType))
+				, Owner(InOwner)
+				, Graph(InGraph)
+		{}
+	
+		static FString StaticType()
+		{
+			return "TEngineContext";
+		}
+	
 		UObject* Owner = nullptr;
 		UDataflow* Graph = nullptr;
 	};
+
+	template class TEngineContext<FContextSingle>;
+	template class TEngineContext<FContextThreaded>;
+
+	typedef TEngineContext<FContextSingle> FEngineContext;
+	typedef TEngineContext<FContextThreaded> FEngineContextThreaded;
+
 }
\ No newline at end of file
