HeadlessChaos : fix broken unit test caused by CL 21520742

diff --git a/Engine/Source/Runtime/Experimental/Chaos/Public/GeometryCollection/GeometryCollectionSimulationCoreTypes.h b/Engine/Source/Runtime/Experimental/Chaos/Public/GeometryCollection/GeometryCollectionSimulationCoreTypes.h
--- a/Engine/Source/Runtime/Experimental/Chaos/Public/GeometryCollection/GeometryCollectionSimulationCoreTypes.h
+++ b/Engine/Source/Runtime/Experimental/Chaos/Public/GeometryCollection/GeometryCollectionSimulationCoreTypes.h
@@ -206,7 +206,7 @@ struct FSimulationParameters
 		, UseCCD(false)
 		, LinearDamping(0.01f)
 		, AngularDamping(0)
-		, bUseDamagePropagation(true)
+		, bUseDamagePropagation(false)
 		, BreakDamagePropagationFactor(1.0f)
 		, ShockDamagePropagationFactor(0.0f)
 		, SimulationFilterData()
