Small optimization: Remove unnecessary work
Not computing the raycast position if not necessary

diff --git a/Engine/Source/Runtime/Experimental/Chaos/Public/Chaos/BoundingVolume.h b/Engine/Source/Runtime/Experimental/Chaos/Public/Chaos/BoundingVolume.h
--- a/Engine/Source/Runtime/Experimental/Chaos/Public/Chaos/BoundingVolume.h
+++ b/Engine/Source/Runtime/Experimental/Chaos/Public/Chaos/BoundingVolume.h
@@ -473,8 +473,7 @@ private:
 	template <typename SQVisitor, bool bPruneDuplicates = true>
 	bool RaycastImp(const TVector<T, d>& Start, FQueryFastData& CurData, SQVisitor& Visitor, const FVec3& Dir, const FVec3 InvDir, const bool bParallel[3]) const
 	{
-		TVector<T, d> TmpPosition;
-		T TOI;
+		T TOI, ExitTime;
 
 		for (const auto& Elem : MGlobalPayloads)
 		{
@@ -485,7 +484,7 @@ private:
 
 			const auto& InstanceBounds = Elem.Bounds;
 			if (InstanceBounds.RaycastFast(Start,
-				Dir, InvDir, bParallel, CurData.CurrentLength, CurData.InvCurrentLength, TOI, TmpPosition))
+				Dir, InvDir, bParallel, CurData.CurrentLength, CurData.InvCurrentLength, TOI, ExitTime))
 			{
 				TSpatialVisitorData<TPayloadType> VisitData(Elem.Payload, true, InstanceBounds);
 				const bool bContinue = Visitor.VisitRaycast(VisitData, CurData);
@@ -506,7 +505,7 @@ private:
 
 			const auto& InstanceBounds = Elem.Bounds;
 			if (InstanceBounds.RaycastFast(Start, Dir,
-				InvDir, bParallel, CurData.CurrentLength, CurData.InvCurrentLength, TOI, TmpPosition))
+				InvDir, bParallel, CurData.CurrentLength, CurData.InvCurrentLength, TOI, ExitTime))
 			{
 				TSpatialVisitorData<TPayloadType> VisitData(Elem.Payload, true, InstanceBounds);
 				const bool bContinue = Visitor.VisitRaycast(VisitData, CurData);
@@ -573,7 +572,7 @@ private:
 					}
 					const auto& InstanceBounds = Elem.Bounds;
 					if (InstanceBounds.RaycastFast(Start,
-							Dir, InvDir, bParallel, CurData.CurrentLength, CurData.InvCurrentLength, TOI, TmpPosition))
+							Dir, InvDir, bParallel, CurData.CurrentLength, CurData.InvCurrentLength, TOI, ExitTime))
 					{
 						TSpatialVisitorData<TPayloadType> VisitData(Elem.Payload, true, InstanceBounds);
 						const bool bContinue = Visitor.VisitRaycast(VisitData, CurData);
@@ -1364,4 +1363,4 @@ extern template class TBoundingVolume<int32, FReal, 3>;
 extern template class TBoundingVolume<FAccelerationStructureHandle, FReal, 3>;
 #endif
 
-}
\ No newline at end of file
+}
diff --git a/Engine/Source/Runtime/PhysicsCore/Private/SQAccelerator.cpp b/Engine/Source/Runtime/PhysicsCore/Private/SQAccelerator.cpp
--- a/Engine/Source/Runtime/PhysicsCore/Private/SQAccelerator.cpp
+++ b/Engine/Source/Runtime/PhysicsCore/Private/SQAccelerator.cpp
@@ -252,10 +252,9 @@ private:
 				if (SQ != ESQType::Overlap)
 				{
 					//todo: use fast raycast
-					Chaos::FReal TmpTime;
-					FVec3 TmpPos;
+					Chaos::FReal TmpTime, TmpExitTime;
 					const FVec3 InflatedBoundsTraceStart = SQ == ESQType::Raycast ? StartPoint : QueryGeomWorldBounds.Center();
-					if (!InflatedWorldBounds.RaycastFast(InflatedBoundsTraceStart, CurData->Dir, CurData->InvDir, CurData->bParallel, CurData->CurrentLength, CurData->InvCurrentLength, TmpTime, TmpPos))
+					if (!InflatedWorldBounds.RaycastFast(InflatedBoundsTraceStart, CurData->Dir, CurData->InvDir, CurData->bParallel, CurData->CurrentLength, CurData->InvCurrentLength, TmpTime, TmpExitTime))
 					{
 						continue;
 					}
