Fix Broad Phases, the following types have only one broadphase bucket:
- Tree
- Grid
- TreeOfGrid

diff --git a/Engine/Source/Runtime/Experimental/Chaos/Private/PBDRigidsSolver.cpp b/Engine/Source/Runtime/Experimental/Chaos/Private/PBDRigidsSolver.cpp
--- a/Engine/Source/Runtime/Experimental/Chaos/Private/PBDRigidsSolver.cpp
+++ b/Engine/Source/Runtime/Experimental/Chaos/Private/PBDRigidsSolver.cpp
@@ -496,7 +496,7 @@ namespace Chaos
 		UE_LOG(LogPBDRigidsSolver, Verbose, TEXT("FPBDRigidsSolver::RegisterObject()"));
 		auto& RigidBody_External = Proxy->GetGameThreadAPI();
 
-		if (BroadPhaseConfig.BroadphaseType == FBroadPhaseConfig::Tree)
+		if (BroadPhaseConfig.BroadphaseType < FBroadPhaseConfig::TreeAndGrid)
 		{
 			FSpatialAccelerationIdx Idx = RigidBody_External.SpatialIdx();
 			Idx.Bucket = 0;
