Chaos - Constraint Solver Refactor

diff --git a/Engine/Source/Runtime/Experimental/Chaos/Private/Chaos/Collision/SolverCollisionContainer.cpp b/Engine/Source/Runtime/Experimental/Chaos/Private/Chaos/Collision/SolverCollisionContainer.cpp
--- a/Engine/Source/Runtime/Experimental/Chaos/Private/Chaos/Collision/SolverCollisionContainer.cpp
+++ b/Engine/Source/Runtime/Experimental/Chaos/Private/Chaos/Collision/SolverCollisionContainer.cpp
@@ -6,6 +6,7 @@
 #include "Chaos/CollisionResolution.h"
 #include "Chaos/DebugDrawQueue.h"
 #include "Chaos/Evolution/SolverBodyContainer.h"
+#include "Chaos/Island/IslandManager.h"
 #include "Chaos/Particle/ParticleUtilities.h"
 #include "Chaos/PBDCollisionConstraints.h"
 #include "Chaos/PBDCollisionConstraintsContact.h"
@@ -71,48 +72,84 @@ namespace Chaos
 		{
 		}
 
-		FPBDCollisionSolver& GetSolver() { return Solver; }
-		FPBDCollisionConstraint* GetConstraint() { return Constraint; }
+		FPBDCollisionSolver& GetSolver()
+		{ 
+			return Solver;
+		}
 
-		bool IsManifold() const { return bIsManifold; }
-		bool IsIncrementalManifold() const { return bIsIncrementalManifold; }
+		FPBDCollisionConstraint* GetConstraint()
+		{ 
+			return Constraint;
+		}
 
-		void PreGatherInput(
-			const FReal Dt,
-			FPBDCollisionConstraint& InConstraint,
-			FSolverBodyContainer& SolverBodyContainer)
+		void SetConstraint(FPBDCollisionConstraint& InConstraint)
 		{
 			Constraint = &InConstraint;
-			SolverBodyContainer.FindOrAdd(Constraint->GetParticle0(), Dt);
-			SolverBodyContainer.FindOrAdd(Constraint->GetParticle1(), Dt);
+		}
+
+		void SetBodies(FSolverBody& Body0, FSolverBody& Body1)
+		{
+			Solver.SetSolverBodies(Body0, Body1);
+
+			// We should try to remove this - the Constraint should not need to know about solver objects
+			Constraint->SetSolverBodies(&Body0, &Body1);
+		}
+
+		void Unbind()
+		{
+			if (Constraint != nullptr)
+			{
+				Constraint->SetSolverBodies(nullptr, nullptr);
+				Constraint = nullptr;
+			}
+			Solver.ResetSolverBodies();
+		}
+
+		void Reset()
+		{
+			Unbind();
+
+			bIsManifold = true;
+			bIsIncrementalManifold = false;
+			Solver.Reset();
+		}
+
+		bool IsManifold() const
+		{ 
+			return bIsManifold;
+		}
+		
+		bool IsIncrementalManifold() const
+		{ 
+			return bIsIncrementalManifold;
 		}
 
 		/**
 		 * @brief Modify solver stiffness when we have bodies with large mass differences
 		*/
 		FReal CalculateSolverStiffness(
-			const FSolverBody* Body0,
-			const FSolverBody* Body1,
+			const FSolverBody& Body0,
+			const FSolverBody& Body1,
 			const FReal MassRatio1,
 			const FReal MassRatio2)
 		{
 			// Adjust the solver stiffness if one body is more than MassRatio1 times the mass of the other and the heavier one is on top.
 			// Solver stiffness will be equal to 1 when the mass ratio is MassRatio1.
 			// Solver stiffness will be equal to 0 when the mass ratio is MassRatio2.
-			if (Body0->IsDynamic() && Body1->IsDynamic() && (MassRatio1 > 0) && (MassRatio2 > MassRatio1))
+			if (Body0.IsDynamic() && Body1.IsDynamic() && (MassRatio1 > 0) && (MassRatio2 > MassRatio1))
 			{
 				// Find heavy body and the mass ratio
 				const FSolverBody* HeavyBody;
 				FReal MassRatio;
-				if (Body0->InvM() < Body1->InvM())
+				if (Body0.InvM() < Body1.InvM())
 				{
-					HeavyBody = Body0;
-					MassRatio = Body1->InvM() / Body0->InvM();
+					HeavyBody = &Body0;
+					MassRatio = Body1.InvM() / Body0.InvM();
 				}
 				else
 				{
-					HeavyBody = Body1;
-					MassRatio = Body0->InvM() / Body1->InvM();
+					HeavyBody = &Body1;
+					MassRatio = Body0.InvM() / Body1.InvM();
 				}
 
 				if (MassRatio > MassRatio1)
@@ -125,7 +162,7 @@ namespace Chaos
 					if (FMath::Abs(WorldNormal.Z) > FReal(0.3))
 					{
 						// Which body is on the top? (Normal always points away from second body - see FContactPoint)
-						const FSolverBody* TopBody = (WorldNormal.Z > 0) ? Body0 : Body1;
+						const FSolverBody* TopBody = (WorldNormal.Z > 0) ? &Body0 : &Body1;
 						if (TopBody == HeavyBody)
 						{
 							// The heavy body is on top - reduce the solver stiffness
@@ -140,25 +177,13 @@ namespace Chaos
 			return Constraint->GetStiffness();
 		}
 
-		void BindBodies(
-			const FReal Dt,
-			const int32 Particle0Level,
-			const int32 Particle1Level,
-			FSolverBodyContainer& SolverBodyContainer,
+		void GatherInput(
+			const FReal InDt,
 			const FPBDCollisionSolverSettings& SolverSettings)
 		{
-			// Find the solver bodies for the particles we constrain. This will add them to the container
-			// if they aren't there already, and ensure that they are populated with the latest data.
-			FSolverBody* Body0 = SolverBodyContainer.FindOrAdd(Constraint->GetParticle0(), Dt);
-			FSolverBody* Body1 = SolverBodyContainer.FindOrAdd(Constraint->GetParticle1(), Dt);
-
-			Body0->SetLevel(Particle0Level);
-			Body1->SetLevel(Particle1Level);
-		
-			Solver.SetSolverBodies(*Body0, *Body1);
-
-			// We should try to remove this - the Constraint should not need to know about solver objects
-			Constraint->SetSolverBodies(Body0, Body1);
+			FSolverReal Dt = FSolverReal(InDt);
+			const FConstraintSolverBody& Body0 = Solver.SolverBody0();
+			const FConstraintSolverBody& Body1 = Solver.SolverBody1();
 
 			// Friction values. Static and Dynamic friction are applied in the position solve for most shapes.
 			// For quadratic shapes, we run dynamic friction in the velocity solve for better rolling behaviour.
@@ -188,7 +213,7 @@ namespace Chaos
 
 			Solver.SetFriction(PositionStaticFriction, PositionDynamicFriction, VelocityDynamicFriction);
 
-			const FReal SolverStiffness = CalculateSolverStiffness(Body0, Body1, Chaos_PBDCollisionSolver_AutoStiffness_MassRatio1, Chaos_PBDCollisionSolver_AutoStiffness_MassRatio2);
+			const FReal SolverStiffness = CalculateSolverStiffness(Body0.SolverBody(), Body1.SolverBody(), Chaos_PBDCollisionSolver_AutoStiffness_MassRatio1, Chaos_PBDCollisionSolver_AutoStiffness_MassRatio2);
 			Solver.SetStiffness(FSolverReal(SolverStiffness));
 
 			Solver.SolverBody0().SetInvMScale(Constraint->GetInvMassScale0());
@@ -197,7 +222,7 @@ namespace Chaos
 			Solver.SolverBody1().SetInvIScale(Constraint->GetInvInertiaScale1());
 
 			bIsManifold = Constraint->GetUseManifold();
-			bIsIncrementalManifold = bIsManifold && Constraint->GetUseIncrementalCollisionDetection();
+			UpdateManifoldPoints(InDt);
 		}
 
 		void UpdateManifoldPoints(
@@ -348,10 +373,7 @@ namespace Chaos
 					Dt);
 			}
 
-			Constraint->SetSolverBodies(nullptr, nullptr);
-			Constraint->SetSolverIndex(INDEX_NONE);
-			Constraint = nullptr;
-			Solver.ResetSolverBodies();
+			Unbind();
 		}
 
 
@@ -369,63 +391,157 @@ namespace Chaos
 	//////////////////////////////////////////////////////////////////////////////////////////////////
 
 
-	FPBDCollisionSolverContainer::FPBDCollisionSolverContainer()
-		: FConstraintSolverContainer()
+	FPBDCollisionContainerSolver::FPBDCollisionContainerSolver(const FPBDCollisionConstraints& InConstraintContainer, const int32 InPriority)
+		: FConstraintContainerSolver(InPriority)
+		, ConstraintContainer(InConstraintContainer)
 		, bPerIterationCollisionDetection(false)
 		, bDeferredCollisionDetection(false)
 	{
 	}
 
-	FPBDCollisionSolverContainer::~FPBDCollisionSolverContainer()
+	FPBDCollisionContainerSolver::~FPBDCollisionContainerSolver()
 	{
 	}
 
-	void FPBDCollisionSolverContainer::Reset(const int32 MaxCollisions)
+	void FPBDCollisionContainerSolver::Reset(const int32 MaxCollisions)
 	{
-		CollisionSolvers.Reset(MaxCollisions);
+		CollisionSolvers.Reset(0);
+		CollisionSolvers.Reserve(MaxCollisions);
 	}
-	
-	void FPBDCollisionSolverContainer::SetNum(const int32 MaxCollisions)
+
+	void FPBDCollisionContainerSolver::AddConstraints()
 	{
-		CollisionSolvers.SetNum(MaxCollisions, false);
+		for (FPBDCollisionConstraintHandle* ConstraintHandle : ConstraintContainer.GetConstraintHandles())
+		{
+			FPBDCollisionConstraint& Constraint = ConstraintHandle->GetContact();
+
+			AddConstraint(Constraint);
+		}
+	}
+
+	void FPBDCollisionContainerSolver::AddConstraints(const TArrayView<FPBDIslandConstraint>& IslandConstraints)
+	{
+		for (FPBDIslandConstraint& IslandConstraint : IslandConstraints)
+		{
+			// @todo(chaos): unsafe version of As
+			FPBDCollisionConstraint& Constraint = IslandConstraint.GetConstraint()->As<FPBDCollisionConstraintHandle>()->GetContact();
+
+			AddConstraint(Constraint);
+		}
 	}
 
-	void FPBDCollisionSolverContainer::PreAddConstraintSolver(const FReal Dt, FPBDCollisionConstraint& Constraint, FSolverBodyContainer& SolverBodyContainer, int32& SolverIndex)
+	void FPBDCollisionContainerSolver::AddConstraint(FPBDCollisionConstraint& Constraint)
 	{
-		// This container is required to allocate pointers that are valid for the whole tick,
-		// so we cannot allow the container to resize during the tick. See Reset()
-		check(SolverIndex < CollisionSolvers.Num());
-		Constraint.SetSolverIndex(SolverIndex);
+		FPBDCollisionSolverAdapter& CollisionSolver = CollisionSolvers[CollisionSolvers.AddDefaulted()];
+		CollisionSolver.Reset();
+		CollisionSolver.SetConstraint(Constraint);
+	}
 
-		FPBDCollisionSolverAdapter& CollisionSolver = CollisionSolvers[SolverIndex];
-		CollisionSolver.GetSolver().Reset();
 
-		CollisionSolver.PreGatherInput(Dt, Constraint, SolverBodyContainer);
+	void FPBDCollisionContainerSolver::AddBodies(FSolverBodyContainer& SolverBodyContainer)
+	{
+		for (FPBDCollisionSolverAdapter& CollisionSolver : CollisionSolvers)
+		{
+			FPBDCollisionConstraint* Constraint = CollisionSolver.GetConstraint();
+			check(Constraint != nullptr);
+
+			// Find the solver bodies for the particles we constrain. This will add them to the container
+			// if they aren't there already, and ensure that they are populated with the latest data.
+			FSolverBody* Body0 = SolverBodyContainer.FindOrAdd(Constraint->GetParticle0());
+			FSolverBody* Body1 = SolverBodyContainer.FindOrAdd(Constraint->GetParticle1());
+			check(Body0 != nullptr);
+			check(Body1 != nullptr);
 
-		++SolverIndex;
+			CollisionSolver.SetBodies(*Body0, *Body1);
+		}
 	}
 
-	void FPBDCollisionSolverContainer::AddConstraintSolver(const FReal Dt, FPBDCollisionConstraint& Constraint, const int32 Particle0Level, const int32 Particle1Level, FSolverBodyContainer& SolverBodyContainer, const FPBDCollisionSolverSettings& SolverSettings)
+	void FPBDCollisionContainerSolver::GatherInput(const FReal Dt)
 	{
-		// This container is required to allocate pointers that are valid for the whole tick,
-		// so we cannot allow the container to resize during the tick. See Reset()
-		const int32 SolverIndex = Constraint.GetSolverIndex();
-		check(SolverIndex < CollisionSolvers.Num());
+		if (bDeferredCollisionDetection)
+		{
+			return;
+		}
+
+		for (FPBDCollisionSolverAdapter& CollisionSolver : CollisionSolvers)
+		{
+			CollisionSolver.GatherInput(Dt, ConstraintContainer.GetSolverSettings());
+
+			// We need to run collision every iteration if we are not using manifolds, or are using incremental manifolds
+			if (!CollisionSolver.IsManifold() || CollisionSolver.IsIncrementalManifold())
+			{
+				bPerIterationCollisionDetection = true;
+			}
+		}
+	}
+
+	void FPBDCollisionContainerSolver::GatherInput(const FReal Dt, const int32 BeginIndex, const int32 EndIndex)
+	{
+		if (bDeferredCollisionDetection)
+		{
+			return;
+		}
 
-		FPBDCollisionSolverAdapter& CollisionSolver = CollisionSolvers[SolverIndex];
-		check(CollisionSolver.GetConstraint() == &Constraint);
+		check(BeginIndex >= 0);
+		check(EndIndex <= NumSolvers());
+
+		bool bWantPerIterationCollisionDetection = false;
+		for (int32 ConstraintIndex = BeginIndex; ConstraintIndex < EndIndex; ++ConstraintIndex)
+		{
+			FPBDCollisionSolverAdapter& CollisionSolver = CollisionSolvers[ConstraintIndex];
+			
+			CollisionSolver.GatherInput(Dt, ConstraintContainer.GetSolverSettings());
+
+			// We need to run collision every iteration if we are not using manifolds, or are using incremental manifolds
+			if (!CollisionSolver.IsManifold() || CollisionSolver.IsIncrementalManifold())
+			{
+				bWantPerIterationCollisionDetection = true;
+			}
+		}
+
+		if (bWantPerIterationCollisionDetection)
+		{
+			// We should lock here? We only ever set to true or do nothing so I think it doesn't matter if this happens on multiple threads...
+			bPerIterationCollisionDetection = true;
+		}
+	}
+
+	void FPBDCollisionContainerSolver::ScatterOutput(const FReal Dt)
+	{
+		for (FPBDCollisionSolverAdapter& CollisionSolver : CollisionSolvers)
+		{
+			CollisionSolver.ScatterOutput(Dt);
+		}
+	}
 
-		CollisionSolver.BindBodies(Dt, Particle0Level, Particle1Level, SolverBodyContainer, SolverSettings);
+	void FPBDCollisionContainerSolver::ScatterOutput(const FReal Dt, const int32 BeginIndex, const int32 EndIndex)
+	{
+		check(BeginIndex >= 0);
+		check(EndIndex <= NumSolvers());
 
-		if (!bDeferredCollisionDetection)
+		for (int32 ConstraintIndex = BeginIndex; ConstraintIndex < EndIndex; ++ConstraintIndex)
 		{
-			CollisionSolver.UpdateManifoldPoints(Dt);
+			FPBDCollisionSolverAdapter& CollisionSolver = CollisionSolvers[ConstraintIndex];
+			CollisionSolver.ScatterOutput(Dt);
 		}
+	}
+
+	void FPBDCollisionContainerSolver::ApplyPositionConstraints(const FReal Dt, const int32 It, const int32 NumIts)
+	{
+		SolvePositionImpl(Dt, It, NumIts, 0, CollisionSolvers.Num(), ConstraintContainer.GetSolverSettings());
+	}
+
+	void FPBDCollisionContainerSolver::ApplyVelocityConstraints(const FReal Dt, const int32 It, const int32 NumIts)
+	{
+		SolveVelocityImpl(Dt, It, NumIts, 0, CollisionSolvers.Num(), ConstraintContainer.GetSolverSettings());
+	}
 
-		bPerIterationCollisionDetection |= (!CollisionSolver.IsManifold() || CollisionSolver.IsIncrementalManifold());
+	void FPBDCollisionContainerSolver::ApplyProjectionConstraints(const FReal Dt, const int32 It, const int32 NumIts)
+	{
+		// Not supported for collisions
 	}
 
-	void FPBDCollisionSolverContainer::UpdatePositionShockPropagation(const FReal Dt, const int32 It, const int32 NumIts, const int32 BeginIndex, const int32 EndIndex, const FPBDCollisionSolverSettings& SolverSettings)
+	void FPBDCollisionContainerSolver::UpdatePositionShockPropagation(const FReal Dt, const int32 It, const int32 NumIts, const int32 BeginIndex, const int32 EndIndex, const FPBDCollisionSolverSettings& SolverSettings)
 	{
 		// If this is the first shock propagation iteration, enable it on each solver
 		const bool bEnableShockPropagation = (It == NumIts - SolverSettings.NumPositionShockPropagationIterations);
@@ -438,7 +554,7 @@ namespace Chaos
 		}
 	}
 
-	void FPBDCollisionSolverContainer::UpdateVelocityShockPropagation(const FReal Dt, const int32 It, const int32 NumIts, const int32 BeginIndex, const int32 EndIndex, const FPBDCollisionSolverSettings& SolverSettings)
+	void FPBDCollisionContainerSolver::UpdateVelocityShockPropagation(const FReal Dt, const int32 It, const int32 NumIts, const int32 BeginIndex, const int32 EndIndex, const FPBDCollisionSolverSettings& SolverSettings)
 	{
 		// Set/reset the shock propagation based on current iteration. The position solve may
 		// have left the bodies with a mass scale and we want to change or reset it.
@@ -459,28 +575,8 @@ namespace Chaos
 		}
 	}
 
-	bool FPBDCollisionSolverContainer::SolvePositionSerial(const FReal Dt, const int32 It, const int32 NumIts, const int32 BeginIndex, const int32 EndIndex, const FPBDCollisionSolverSettings& SolverSettings)
-	{
-		return SolvePositionImpl(Dt, It, NumIts, BeginIndex, EndIndex, SolverSettings, false);
-	}
-
-	bool FPBDCollisionSolverContainer::SolveVelocitySerial(const FReal Dt, const int32 It, const int32 NumIts, const int32 BeginIndex, const int32 EndIndex, const FPBDCollisionSolverSettings& SolverSettings)
-	{
-		return SolveVelocityImpl(Dt, It, NumIts, BeginIndex, EndIndex, SolverSettings, false);
-	}
-
-	bool FPBDCollisionSolverContainer::SolvePositionParallel(const FReal Dt, const int32 It, const int32 NumIts, const int32 BeginIndex, const int32 EndIndex, const FPBDCollisionSolverSettings& SolverSettings)
-	{
-		return SolvePositionImpl(Dt, It, NumIts, BeginIndex, EndIndex, SolverSettings, true);
-	}
-
-	bool FPBDCollisionSolverContainer::SolveVelocityParallel(const FReal Dt, const int32 It, const int32 NumIts, const int32 BeginIndex, const int32 EndIndex, const FPBDCollisionSolverSettings& SolverSettings)
-	{
-		return SolveVelocityImpl(Dt, It, NumIts, BeginIndex, EndIndex, SolverSettings, true);
-	}
-
 	// @todo(chaos): parallel version of SolvePosition
-	bool FPBDCollisionSolverContainer::SolvePositionImpl(const FReal Dt, const int32 It, const int32 NumIts, const int32 BeginIndex, const int32 EndIndex, const FPBDCollisionSolverSettings& SolverSettings, const bool bParallel)
+	bool FPBDCollisionContainerSolver::SolvePositionImpl(const FReal Dt, const int32 It, const int32 NumIts, const int32 BeginIndex, const int32 EndIndex, const FPBDCollisionSolverSettings& SolverSettings)
 	{
 		SCOPE_CYCLE_COUNTER(STAT_Collisions_Apply);
 		if (!bChaos_PBDCollisionSolver_Position_SolveEnabled)
@@ -508,16 +604,16 @@ namespace Chaos
 		// Apply the position correction
 		if (bApplyStaticFriction)
 		{
-			return SolvePositionWithFrictionImpl(Dt, BeginIndex, EndIndex, MaxPushOut, bParallel);
+			return SolvePositionWithFrictionImpl(Dt, BeginIndex, EndIndex, MaxPushOut);
 		}
 		else
 		{
-			return SolvePositionNoFrictionImpl(Dt, BeginIndex, EndIndex, MaxPushOut, bParallel);
+			return SolvePositionNoFrictionImpl(Dt, BeginIndex, EndIndex, MaxPushOut);
 		}
 	}
 
 	// Solve position with friction (last few iterations each tick)
-	bool FPBDCollisionSolverContainer::SolvePositionWithFrictionImpl(const FReal InDt, const int32 BeginIndex, const int32 EndIndex, const FReal InMaxPushOut, const bool bParallel)
+	bool FPBDCollisionContainerSolver::SolvePositionWithFrictionImpl(const FReal InDt, const int32 BeginIndex, const int32 EndIndex, const FReal InMaxPushOut)
 	{
 		if (EndIndex == BeginIndex)
 		{
@@ -526,20 +622,16 @@ namespace Chaos
 		const FSolverReal Dt = FSolverReal(InDt);
 		const FSolverReal MaxPushOut = FSolverReal(InMaxPushOut);
 
-		InnerPhysicsParallelForRange(EndIndex - BeginIndex, [&](int32 StartRangeIndex, int32 EndRangeIndex)
+		for (int32 SolverIndex = BeginIndex; SolverIndex < EndIndex; ++SolverIndex)
 		{
-			int32 LocalStartIndex = StartRangeIndex + BeginIndex;
-			int32 LocalEndIndex = EndRangeIndex + BeginIndex;
-			for (int32 SolverIndex = LocalStartIndex; SolverIndex < LocalEndIndex; ++SolverIndex)
-			{
-				CollisionSolvers[SolverIndex].GetSolver().SolvePositionWithFriction(Dt, MaxPushOut);
-			}
-		}, Chaos::LargeBatchSize, !bParallel);
+			CollisionSolvers[SolverIndex].GetSolver().SolvePositionWithFriction(Dt, MaxPushOut);
+		}
+
 		return true;
 	}
 
 	// Solve position without friction (first few iterations each tick)
-	bool FPBDCollisionSolverContainer::SolvePositionNoFrictionImpl(const FReal InDt, const int32 BeginIndex, const int32 EndIndex, const FReal InMaxPushOut, const bool bParallel)
+	bool FPBDCollisionContainerSolver::SolvePositionNoFrictionImpl(const FReal InDt, const int32 BeginIndex, const int32 EndIndex, const FReal InMaxPushOut)
 	{
 		if (EndIndex == BeginIndex)
 		{
@@ -548,19 +640,15 @@ namespace Chaos
 		const FSolverReal Dt = FSolverReal(InDt);
 		const FSolverReal MaxPushOut = FSolverReal(InMaxPushOut);
 
-		InnerPhysicsParallelForRange(EndIndex - BeginIndex, [&](int32 StartRangeIndex, int32 EndRangeIndex)
+		for (int32 SolverIndex = BeginIndex; SolverIndex < EndIndex; ++SolverIndex)
 		{
-			int32 LocalStartIndex = StartRangeIndex + BeginIndex;
-			int32 LocalEndIndex = EndRangeIndex + BeginIndex;
-			for (int32 SolverIndex = LocalStartIndex; SolverIndex < LocalEndIndex; ++SolverIndex)
-			{
-				CollisionSolvers[SolverIndex].GetSolver().SolvePositionNoFriction(Dt, MaxPushOut);
-			}
-		}, Chaos::LargeBatchSize, !bParallel);
+			CollisionSolvers[SolverIndex].GetSolver().SolvePositionNoFriction(Dt, MaxPushOut);
+		}
+
 		return true;
 	}
 
-	bool FPBDCollisionSolverContainer::SolveVelocityImpl(const FReal InDt, const int32 It, const int32 NumIts, const int32 BeginIndex, const int32 EndIndex, const FPBDCollisionSolverSettings& SolverSettings, const bool bParallel)
+	bool FPBDCollisionContainerSolver::SolveVelocityImpl(const FReal InDt, const int32 It, const int32 NumIts, const int32 BeginIndex, const int32 EndIndex, const FPBDCollisionSolverSettings& SolverSettings)
 	{
 		SCOPE_CYCLE_COUNTER(STAT_Collisions_ApplyPushOut);
 		if (!bChaos_PBDCollisionSolver_Velocity_SolveEnabled)
@@ -586,33 +674,7 @@ namespace Chaos
 		return bNeedsAnotherIteration;
 	}
 
-	void FPBDCollisionSolverContainer::ScatterOutput(const FReal Dt, const int32 BeginIndex, const int32 EndIndex)
-	{
-		return ScatterOutputImpl(Dt, BeginIndex, EndIndex, false);
-	}
-
-	void FPBDCollisionSolverContainer::ScatterOutputImpl(const FReal Dt, const int32 BeginIndex, const int32 EndIndex, const bool bParallel)
-	{
-		SCOPE_CYCLE_COUNTER(STAT_Collisions_Scatter);
-		check(BeginIndex >= 0);
-		check(EndIndex <= CollisionSolvers.Num());
-		if (EndIndex == BeginIndex)
-		{
-			return;
-		}
-
-		InnerPhysicsParallelForRange(EndIndex - BeginIndex, [&](int32 LocalStartIndex, int32 LocalEndIndex)
-		{
-			int32 LoopStartIndex = LocalStartIndex + BeginIndex;
-			int32 LoopEndIndex = LocalEndIndex + BeginIndex;
-			for (int32 SolverIndex = LoopStartIndex; SolverIndex < LoopEndIndex; ++SolverIndex)
-			{
-				CollisionSolvers[SolverIndex].ScatterOutput(Dt);
-			}
-		}, Chaos::LargeBatchSize);
-	}
-
-	void FPBDCollisionSolverContainer::UpdateCollisions(const FReal InDt, const int32 BeginIndex, const int32 EndIndex)
+	void FPBDCollisionContainerSolver::UpdateCollisions(const FReal InDt, const int32 BeginIndex, const int32 EndIndex)
 	{
 		const FSolverReal Dt = FSolverReal(InDt);
 
diff --git a/Engine/Source/Runtime/Experimental/Chaos/Private/Chaos/Evolution/PBDMinEvolution.cpp b/Engine/Source/Runtime/Experimental/Chaos/Private/Chaos/Evolution/PBDMinEvolution.cpp
--- a/Engine/Source/Runtime/Experimental/Chaos/Private/Chaos/Evolution/PBDMinEvolution.cpp
+++ b/Engine/Source/Runtime/Experimental/Chaos/Private/Chaos/Evolution/PBDMinEvolution.cpp
@@ -3,8 +3,9 @@
 #include "Chaos/Evolution/PBDMinEvolution.h"
 #include "Chaos/Collision/NarrowPhase.h"
 #include "Chaos/Collision/ParticlePairCollisionDetector.h"
+#include "Chaos/Evolution/SolverConstraintContainer.h"
 #include "Chaos/PBDCollisionConstraints.h"
-#include "Chaos/PBDConstraintRule.h"
+#include "Chaos/PBDConstraintContainer.h"
 #include "Chaos/PBDRigidsSOAs.h"
 #include "Chaos/PerParticleAddImpulses.h"
 #include "Chaos/PerParticleEtherDrag.h"
@@ -42,9 +43,7 @@ namespace Chaos
 	DECLARE_CYCLE_STAT(TEXT("MinEvolution::Gather"), STAT_MinEvolution_Gather, STATGROUP_ChaosMinEvolution);
 	DECLARE_CYCLE_STAT(TEXT("MinEvolution::Scatter"), STAT_MinEvolution_Scatter, STATGROUP_ChaosMinEvolution);
 	DECLARE_CYCLE_STAT(TEXT("MinEvolution::ApplyConstraintsPhase1"), STAT_MinEvolution_ApplyConstraintsPhase1, STATGROUP_ChaosMinEvolution);
-	DECLARE_CYCLE_STAT(TEXT("MinEvolution::UpdateVelocities"), STAT_MinEvolution_UpdateVelocites, STATGROUP_ChaosMinEvolution);
 	DECLARE_CYCLE_STAT(TEXT("MinEvolution::ApplyConstraintsPhase2"), STAT_MinEvolution_ApplyConstraintsPhase2, STATGROUP_ChaosMinEvolution);
-	DECLARE_CYCLE_STAT(TEXT("MinEvolution::ApplyCorrections"), STAT_MinEvolution_ApplyCorrections, STATGROUP_ChaosMinEvolution);
 	DECLARE_CYCLE_STAT(TEXT("MinEvolution::ApplyConstraintsPhase3"), STAT_MinEvolution_ApplyConstraintsPhase3, STATGROUP_ChaosMinEvolution);
 	DECLARE_CYCLE_STAT(TEXT("MinEvolution::DetectCollisions"), STAT_MinEvolution_DetectCollisions, STATGROUP_ChaosMinEvolution);
 
@@ -68,9 +67,6 @@ namespace Chaos
 		, CollisionDetector(InCollisionDetector)
 		, ParticlePrevXs(InPrevX)
 		, ParticlePrevRs(InPrevR)
-		, SolverType(EConstraintSolverType::QuasiPbd)
-		, NumApplyIterations(0)
-		, NumApplyPushOutIterations(0)
 		, NumPositionIterations(0)
 		, NumVelocityIterations(0)
 		, NumProjectionIterations(0)
@@ -80,14 +76,24 @@ namespace Chaos
 	{
 	}
 
-	void FPBDMinEvolution::AddConstraintRule(FSimpleConstraintRule* Rule)
+	FPBDMinEvolution::~FPBDMinEvolution()
 	{
-		check(Rule != nullptr);
-		if (Rule != nullptr)
-		{
-			const uint32 ContainerId = (uint32)ConstraintRules.Add(Rule);
-			Rule->BindToDatas(SolverData, ContainerId);
-		}
+	}
+
+	void FPBDMinEvolution::AddConstraintContainer(FPBDConstraintContainer& InContainer, const int32 Priority)
+	{
+		// Do not add twice
+		check(InContainer.GetContainerId() == INDEX_NONE);
+
+		const int32 ContainerId = ConstraintContainers.Add(&InContainer);
+		InContainer.SetContainerId(ContainerId);
+
+		ConstraintSolver.SetConstraintSolver(ContainerId, InContainer.CreateSceneSolver(Priority));
+	}
+
+	void FPBDMinEvolution::SetConstraintContainerPriority(const int32 ContainerId, const int32 Priority)
+	{
+		ConstraintSolver.SetConstraintSolverPriority(ContainerId, Priority);
 	}
 
 	void FPBDMinEvolution::Advance(const FReal StepDt, const int32 NumSteps, const FReal RewindDt)
@@ -132,40 +138,16 @@ namespace Chaos
 
 		ApplyKinematicTargets(Dt, StepFraction);
 
-		if (PostIntegrateCallback != nullptr)
-		{
-			PostIntegrateCallback();
-		}
-
 		DetectCollisions(Dt);
 
-		if (PostDetectCollisionsCallback != nullptr)
-		{
-			PostDetectCollisionsCallback();
-		}
-
 		if (Dt > 0)
 		{
 			GatherInput(Dt);
 
 			ApplyConstraintsPhase1(Dt);
 
-			if (PostApplyCallback != nullptr)
-			{
-				PostApplyCallback();
-			}
-
-			UpdateVelocities(Dt);
-
 			ApplyConstraintsPhase2(Dt);
 
-			if (PostApplyPushOutCallback != nullptr)
-			{
-				PostApplyPushOutCallback();
-			}
-
-			ApplyCorrections(Dt);
-
 			ApplyConstraintsPhase3(Dt);
 
 			ScatterOutput(Dt);
@@ -177,9 +159,9 @@ namespace Chaos
 	{
 		SCOPE_CYCLE_COUNTER(STAT_MinEvolution_PrepareTick);
 
-		for (FSimpleConstraintRule* ConstraintRule : ConstraintRules)
+		for (FPBDConstraintContainer* ConstraintContainer : ConstraintContainers)
 		{
-			ConstraintRule->PrepareTick();
+			ConstraintContainer->PrepareTick();
 		}
 	}
 
@@ -187,9 +169,9 @@ namespace Chaos
 	{
 		SCOPE_CYCLE_COUNTER(STAT_MinEvolution_UnprepareTick);
 
-		for (FSimpleConstraintRule* ConstraintRule : ConstraintRules)
+		for (FPBDConstraintContainer* ConstraintContainer : ConstraintContainers)
 		{
-			ConstraintRule->UnprepareTick();
+			ConstraintContainer->UnprepareTick();
 		}
 	}
 
@@ -394,13 +376,9 @@ namespace Chaos
 	{
 		SCOPE_CYCLE_COUNTER(STAT_MinEvolution_DetectCollisions);
 
-		// @todo(ccaulfield): doesn't need to be every frame
-		PrioritizedConstraintRules = ConstraintRules;
-		PrioritizedConstraintRules.StableSort();
-
-		for (FSimpleConstraintRule* ConstraintRule : PrioritizedConstraintRules)
+		for (FPBDConstraintContainer* ConstraintContainer : ConstraintContainers)
 		{
-			ConstraintRule->UpdatePositionBasedState(Dt);
+			ConstraintContainer->UpdatePositionBasedState(Dt);
 		}
 
 		CollisionDetector.DetectCollisions(Dt, nullptr);
@@ -412,25 +390,19 @@ namespace Chaos
 	{
 		SCOPE_CYCLE_COUNTER(STAT_MinEvolution_Gather);
 
-		SolverData.GetBodyContainer().Reset(Particles.GetAllParticlesView().Num());
-
-		for (FSimpleConstraintRule* ConstraintRule : ConstraintRules)
-		{
-			ConstraintRule->GatherSolverInput(Dt);
-		}
+		ConstraintSolver.Reset();
+		ConstraintSolver.AddConstraintsAndBodies();
+		ConstraintSolver.GatherBodies(Dt);
+		ConstraintSolver.GatherConstraints(Dt);
 	}
 
 	void FPBDMinEvolution::ScatterOutput(FReal Dt)
 	{
 		SCOPE_CYCLE_COUNTER(STAT_MinEvolution_Scatter);
 
-		for (FSimpleConstraintRule* ConstraintRule : ConstraintRules)
-		{
-			ConstraintRule->ScatterSolverOutput(Dt);
-		}
+		ConstraintSolver.ScatterConstraints(Dt);
+		ConstraintSolver.ScatterBodies(Dt);
 
-		SolverData.GetBodyContainer().ScatterOutput();
-	
 		for (auto& Particle : Particles.GetActiveParticlesView())
 		{
 			Particle.Handle()->AuxilaryValue(ParticlePrevXs) = Particle.X();
@@ -444,92 +416,20 @@ namespace Chaos
 	{
 		SCOPE_CYCLE_COUNTER(STAT_MinEvolution_ApplyConstraintsPhase1);
 
-		const int32 NumIterationsPhase1 = (SolverType == EConstraintSolverType::QuasiPbd) ? NumPositionIterations : NumApplyIterations;
-
-		for (int32 i = 0; i < NumIterationsPhase1; ++i)
-		{
-			bool bNeedsAnotherIteration = Chaos_MinEvolution_ForceMaxConstraintIterations;
-			for (FSimpleConstraintRule* ConstraintRule : PrioritizedConstraintRules)
-			{
-				bNeedsAnotherIteration |= ConstraintRule->ApplyConstraints(Dt, i, NumIterationsPhase1);
-			}
-
-			if (!bNeedsAnotherIteration)
-			{
-				break;
-			}
-		}
-	}
-
-	void FPBDMinEvolution::UpdateVelocities(FReal Dt)
-	{
-		SCOPE_CYCLE_COUNTER(STAT_MinEvolution_UpdateVelocites);
-
-		// @todo(chaos): clean this up - the two solvers calculate implicit velocity differently because 
-		// QPBD accumulates transform deltas and the StandardPBD applies transform changes directly
-		if (SolverType == EConstraintSolverType::StandardPbd)
-		{
-			for (FSolverBodyAdapter& SolverBody : SolverData.GetBodyContainer().GetBodies())
-			{
-				const FVec3 V = FVec3::CalculateVelocity(SolverBody.GetSolverBody().X(), SolverBody.GetSolverBody().P(), Dt);
-				const FVec3 W = FRotation3::CalculateAngularVelocity(SolverBody.GetSolverBody().R(), SolverBody.GetSolverBody().Q(), Dt);
-				SolverBody.GetSolverBody().SetV(V);
-				SolverBody.GetSolverBody().SetW(W);
-			}
-		}
-		else
-		{ 
-			SolverData.GetBodyContainer().SetImplicitVelocities(Dt);
-		}
+		ConstraintSolver.ApplyPositionConstraints(Dt, NumPositionIterations);
 	}
 
 	void FPBDMinEvolution::ApplyConstraintsPhase2(FReal Dt)
 	{
 		SCOPE_CYCLE_COUNTER(STAT_MinEvolution_ApplyConstraintsPhase2);
 
-		const int32 NumIterationsPhase2 = (SolverType == EConstraintSolverType::QuasiPbd) ? NumVelocityIterations : 0;
-
-		for (int32 It = 0; It < NumIterationsPhase2; ++It)
-		{
-			bool bNeedsAnotherIteration = false;
-			for (FSimpleConstraintRule* ConstraintRule : PrioritizedConstraintRules)
-			{
-				bNeedsAnotherIteration |= ConstraintRule->ApplyPushOut(Dt, It, NumIterationsPhase2);
-			}
-
-			if (!bNeedsAnotherIteration)
-			{
-				break;
-			}
-		}
-	}
-
-	void FPBDMinEvolution::ApplyCorrections(FReal Dt)
-	{
-		SCOPE_CYCLE_COUNTER(STAT_MinEvolution_ApplyCorrections);
-
-		SolverData.GetBodyContainer().ApplyCorrections();
-		SolverData.GetBodyContainer().UpdateRotationDependentState();
+		ConstraintSolver.ApplyVelocityConstraints(Dt, NumVelocityIterations);
 	}
 
 	void FPBDMinEvolution::ApplyConstraintsPhase3(FReal Dt)
 	{
 		SCOPE_CYCLE_COUNTER(STAT_MinEvolution_ApplyConstraintsPhase3);
 
-		const int32 NumIterationsPhase3 = (SolverType == EConstraintSolverType::QuasiPbd) ? NumProjectionIterations : NumApplyPushOutIterations;
-
-		for (int32 It = 0; It < NumIterationsPhase3; ++It)
-		{
-			bool bNeedsAnotherIteration = false;
-			for (FSimpleConstraintRule* ConstraintRule : PrioritizedConstraintRules)
-			{
-				bNeedsAnotherIteration |= ConstraintRule->ApplyProjection(Dt, It, NumIterationsPhase3);
-			}
-
-			if (!bNeedsAnotherIteration)
-			{
-				break;
-			}
-		}
+		ConstraintSolver.ApplyProjectionConstraints(Dt, NumProjectionIterations);
 	}
 }
diff --git a/Engine/Source/Runtime/Experimental/Chaos/Private/Chaos/Evolution/SolverBodyContainer.cpp b/Engine/Source/Runtime/Experimental/Chaos/Private/Chaos/Evolution/SolverBodyContainer.cpp
--- a/Engine/Source/Runtime/Experimental/Chaos/Private/Chaos/Evolution/SolverBodyContainer.cpp
+++ b/Engine/Source/Runtime/Experimental/Chaos/Private/Chaos/Evolution/SolverBodyContainer.cpp
@@ -24,14 +24,15 @@ namespace Chaos
 		}
 	}
 
-	int32 FSolverBodyContainer::AddParticle(FGenericParticleHandle InParticle, const FReal Dt)
+	int32 FSolverBodyContainer::AddParticle(FGenericParticleHandle InParticle)
 	{
 		// No array resizing allowed (we want fixed pointers)
 		check(NumItems() < MaxItems());
-		return SolverBodies.Emplace(FSolverBodyAdapter(InParticle, Dt));
+
+		return SolverBodies.AddElement(FSolverBodyAdapter(InParticle));
 	}
 
-	FSolverBody* FSolverBodyContainer::FindOrAdd(FGenericParticleHandle InParticle, const FReal Dt)
+	FSolverBody* FSolverBodyContainer::FindOrAdd(FGenericParticleHandle InParticle)
 	{
 		// For dynamic bodies, we store a cookie on the Particle that holds the solver body index
 		// For kinematics we cannot do this because the kinematic may be in multiple islands and 
@@ -43,7 +44,8 @@ namespace Chaos
 			if (InParticle->IsDynamic())
 			{
 				// First time we have seen this particle, so add it
-				ItemIndex = AddParticle(InParticle, Dt);
+				check(!bLocked);
+				ItemIndex = AddParticle(InParticle);
 				InParticle->SetSolverBodyIndex(ItemIndex);
 			}
 			else // Not Dynamic
@@ -56,14 +58,30 @@ namespace Chaos
 				else
 				{
 					// First time we have seen this particle, so add it
-					ItemIndex = AddParticle(InParticle, Dt);
+					check(!bLocked);
+					ItemIndex = AddParticle(InParticle);
 					ParticleToIndexMap.Add(InParticle, ItemIndex);
 				}
 			}			
 		}
 	
 		check(ItemIndex != INDEX_NONE);
-		return &SolverBodies[ItemIndex].GetSolverBody();
+		
+		FSolverBody* SolverBody = &SolverBodies[ItemIndex].GetSolverBody();
+
+		return SolverBody;
+	}
+
+	void FSolverBodyContainer::GatherInput(const FReal Dt, const int32 BeginIndex, const int32 EndIndex)
+	{
+		check((BeginIndex >= 0) && (BeginIndex <= EndIndex));
+		check((EndIndex >= 0) && (EndIndex <= SolverBodies.Num()));
+
+		for (int32 SolverBodyIndex = BeginIndex; SolverBodyIndex < EndIndex; ++SolverBodyIndex)
+		{
+			FSolverBodyAdapter& SolverBody = SolverBodies[SolverBodyIndex];
+			SolverBody.GatherInput(Dt);
+		}
 	}
 
 	void FSolverBodyContainer::ScatterOutput()
@@ -74,6 +92,18 @@ namespace Chaos
 		}
 	}
 
+	void FSolverBodyContainer::ScatterOutput(const int32 BeginIndex, const int32 EndIndex)
+	{
+		check((BeginIndex >= 0) && (BeginIndex <= EndIndex));
+		check((EndIndex >= 0) && (EndIndex <= SolverBodies.Num()));
+
+		for (int32 SolverBodyIndex = BeginIndex; SolverBodyIndex < EndIndex; ++SolverBodyIndex)
+		{
+			FSolverBodyAdapter& SolverBody = SolverBodies[SolverBodyIndex];
+			SolverBody.ScatterOutput();
+		}
+	}
+
 	void FSolverBodyContainer::SetImplicitVelocities(FReal Dt)
 	{
 		for (FSolverBodyAdapter& SolverBody : SolverBodies)
diff --git a/Engine/Source/Runtime/Experimental/Chaos/Private/Chaos/Island/IslandGraph.cpp b/Engine/Source/Runtime/Experimental/Chaos/Private/Chaos/Island/IslandGraph.cpp
--- a/Engine/Source/Runtime/Experimental/Chaos/Private/Chaos/Island/IslandGraph.cpp
+++ b/Engine/Source/Runtime/Experimental/Chaos/Private/Chaos/Island/IslandGraph.cpp
@@ -204,6 +204,21 @@ void FIslandGraph<NodeType, EdgeType, IslandType, OwnerType>::RemoveNode(const N
 	}
 }
 
+template<typename NodeType, typename EdgeType, typename IslandType, typename OwnerType>
+int32 FIslandGraph<NodeType, EdgeType, IslandType, OwnerType>::GetNodeItemLevel(const NodeType& NodeItem) const
+{
+	if (const int32* PNodeIndex = ItemNodes.Find(NodeItem))
+	{
+		const int32 NodeIndex = *PNodeIndex;
+		if (GraphNodes.IsValidIndex(NodeIndex))
+		{
+			const FGraphNode& GraphNode = GraphNodes[NodeIndex];
+			return GraphNode.LevelIndex;
+		}
+	}
+	return 0;
+}
+
 template<typename NodeType, typename EdgeType, typename IslandType, typename OwnerType>
 void FIslandGraph<NodeType, EdgeType, IslandType, OwnerType>::ResetEdges()
 {
@@ -516,7 +531,7 @@ void FIslandGraph<NodeType, EdgeType, IslandType, OwnerType>::InitSorting()
 }
 	
 template<typename NodeType, typename EdgeType, typename IslandType, typename OwnerType>
-void FIslandGraph<NodeType, EdgeType, IslandType, OwnerType>::UpdateLevels(const int32 NodeIndex, const int32 ContainerId)
+void FIslandGraph<NodeType, EdgeType, IslandType, OwnerType>::UpdateLevels(const int32 NodeIndex)
 {
 	if(GraphNodes.IsValidIndex(NodeIndex))
 	{
@@ -530,8 +545,8 @@ void FIslandGraph<NodeType, EdgeType, IslandType, OwnerType>::UpdateLevels(const
 			ensure(!GraphEdge.bValidEdge || GraphIslands.IsValidIndex(GraphEdge.IslandIndex));
 #endif
 
-			// Do nothing if the edge is not coming from the same container or if the island is sleeping 
-			if (GraphEdge.bValidEdge && (GraphEdge.ItemContainer == ContainerId) && (GraphEdge.LevelIndex == INDEX_NONE) && !GraphIslands[GraphEdge.IslandIndex].bIsSleeping)
+			// Do nothing if the edge is assigned a level or if the island is sleeping 
+			if (GraphEdge.bValidEdge && (GraphEdge.LevelIndex == INDEX_NONE) && !GraphIslands[GraphEdge.IslandIndex].bIsSleeping)
 			{
 				const int32 OtherIndex = (NodeIndex == GraphEdge.FirstNode) ?
 							GraphEdge.SecondNode : GraphEdge.FirstNode;
@@ -553,7 +568,7 @@ void FIslandGraph<NodeType, EdgeType, IslandType, OwnerType>::UpdateLevels(const
 }
 
 template<typename NodeType, typename EdgeType, typename IslandType, typename OwnerType>
-void FIslandGraph<NodeType, EdgeType, IslandType, OwnerType>::ComputeLevels(const int32 ContainerId)
+void FIslandGraph<NodeType, EdgeType, IslandType, OwnerType>::ComputeLevels()
 {
 	// First enqueue all the static/kinematic nodes for levels 0
 	NodeQueue.Reset();
@@ -565,7 +580,7 @@ void FIslandGraph<NodeType, EdgeType, IslandType, OwnerType>::ComputeLevels(cons
 			if(!GraphNode.bValidNode)
 			{
 				GraphNode.LevelIndex = 0;
-				UpdateLevels(NodeIndex, ContainerId);
+				UpdateLevels(NodeIndex);
 			}
 		}
 	}
@@ -577,13 +592,13 @@ void FIslandGraph<NodeType, EdgeType, IslandType, OwnerType>::ComputeLevels(cons
 		NodeIndex = NodeQueue.First();
 		NodeQueue.PopFirst();
 
-		UpdateLevels(NodeIndex, ContainerId);
+		UpdateLevels(NodeIndex);
 	}
 
 	// An isolated island that is only dynamics will not have been processed above, put everything without a level into level zero
 	for(auto& GraphEdge : GraphEdges)
 	{
-		if(GraphEdge.bValidEdge && (GraphEdge.ItemContainer == ContainerId))
+		if(GraphEdge.bValidEdge)
 		{
 			GraphEdge.LevelIndex = FGenericPlatformMath::Max(GraphEdge.LevelIndex, 0);
 		}
@@ -621,7 +636,7 @@ int32 FIslandGraph<NodeType, EdgeType, IslandType, OwnerType>::PickColor(const F
 }
 
 template<typename NodeType, typename EdgeType, typename IslandType, typename OwnerType>
-void FIslandGraph<NodeType, EdgeType, IslandType, OwnerType>::UpdateColors(const int32 NodeIndex, const int32 ContainerId, const int32 MinEdges)
+void FIslandGraph<NodeType, EdgeType, IslandType, OwnerType>::UpdateColors(const int32 NodeIndex, const int32 MinEdges)
 {
 	if (GraphNodes.IsValidIndex(NodeIndex))
 	{
@@ -638,7 +653,7 @@ void FIslandGraph<NodeType, EdgeType, IslandType, OwnerType>::UpdateColors(const
 #endif
 
 			// Do nothing if the edge is not coming from the same container or if the island is sleeping 
-			if (GraphEdge.bValidEdge && (GraphEdge.ItemContainer == ContainerId) && (GraphEdge.ColorIndex == INDEX_NONE) &&
+			if (GraphEdge.bValidEdge && (GraphEdge.ColorIndex == INDEX_NONE) &&
 				!GraphIslands[GraphEdge.IslandIndex].bIsSleeping && (GraphIslands[GraphEdge.IslandIndex].NumEdges > MinEdges))
 			{
 				// Get the opposite node index for the given edge
@@ -659,7 +674,7 @@ void FIslandGraph<NodeType, EdgeType, IslandType, OwnerType>::UpdateColors(const
 }
 
 template<typename NodeType, typename EdgeType, typename IslandType, typename OwnerType>
-void FIslandGraph<NodeType, EdgeType, IslandType, OwnerType>::ComputeColors(const int32 ContainerId, const int32 MinEdges)
+void FIslandGraph<NodeType, EdgeType, IslandType, OwnerType>::ComputeColors(const int32 MinEdges)
 {
 	GraphCounter = (GraphCounter + 1) % MaxCount;
 
@@ -680,7 +695,7 @@ void FIslandGraph<NodeType, EdgeType, IslandType, OwnerType>::ComputeColors(cons
 					NodeIndex = NodeQueue.First();
 					NodeQueue.PopFirst();
 
-					UpdateColors(NodeIndex, ContainerId, MinEdges);
+					UpdateColors(NodeIndex, MinEdges);
 				}
 			}
 		}
@@ -959,6 +974,6 @@ void FIslandGraph<NodeType, EdgeType, IslandType, OwnerType>::ResetIslands()
 	}
 }
 
-template class FIslandGraph<FGeometryParticleHandle*, FConstraintHandleHolder, FPBDIslandSolver*, FPBDIslandManager>;
+template class FIslandGraph<FGeometryParticleHandle*, FConstraintHandleHolder, FPBDIsland*, FPBDIslandManager>;
 template class FIslandGraph<int32, int32, int32, TNullIslandGraphOwner<int32, int32>>;
 }
diff --git a/Engine/Source/Runtime/Experimental/Chaos/Private/Chaos/Island/IslandGroup.cpp b/Engine/Source/Runtime/Experimental/Chaos/Private/Chaos/Island/IslandGroup.cpp
--- a/Engine/Source/Runtime/Experimental/Chaos/Private/Chaos/Island/IslandGroup.cpp
+++ b/Engine/Source/Runtime/Experimental/Chaos/Private/Chaos/Island/IslandGroup.cpp
@@ -1,38 +1,96 @@
 // Copyright Epic Games, Inc. All Rights Reserved.
 #include "Chaos/Island/IslandGroup.h"
+
+#include "Chaos/ConstraintHandle.h"
+#include "Chaos/Island/IslandManager.h"
 #include "Chaos/Island/SolverIsland.h"
+#include "Chaos/PBDConstraintContainer.h"
 
 namespace Chaos
 {
-	
-FPBDIslandGroup::FPBDIslandGroup(const int32 GroupIndexIn ) : FPBDIslandSolverData(GroupIndexIn),
-	 IslandSolvers(), ParticlesCount(0), ConstraintsCount(0)
-{}
+	FPBDIslandConstraintGroupSolver::FPBDIslandConstraintGroupSolver(FPBDIslandManager& InIslandManager)
+		: FPBDConstraintGroupSolver()
+		, IslandManager(InIslandManager)
+		, NumParticles(0)
+		, NumConstraints(0)
+	{
+	}
 
-void FPBDIslandGroup::ReserveIslands(const int32 NumSolvers)
-{
-	IslandSolvers.Reserve(NumSolvers);
-}
+	void FPBDIslandConstraintGroupSolver::SetConstraintSolverImpl(int32 ContainerId)
+	{
+		if (ContainerId >= NumContainerConstraints.Num())
+		{
+			NumContainerConstraints.SetNumZeroed(ContainerId + 1);
+		}
+	}
 
-void FPBDIslandGroup::AddIsland(FPBDIslandSolver* IslandSolver)
-{
-	if (IslandSolver)
+	void FPBDIslandConstraintGroupSolver::ResetImpl()
 	{
-		IslandSolvers.Add(IslandSolver);
+		NumParticles = 0;
+		NumConstraints = 0;
+		for (int32 ContainerIndex = 0; ContainerIndex < ConstraintContainerSolvers.Num(); ++ContainerIndex)
+		{
+			NumContainerConstraints[ContainerIndex] = 0;
+		}
+
+		Islands.Reset();
+	}
+
+	void FPBDIslandConstraintGroupSolver::AddIsland(FPBDIsland* Island)
+	{
+		if (Island)
+		{
+			Islands.Add(Island);
+
+			NumParticles += Island->GetNumParticles();
+
+			for (int32 ContainerIndex = 0; ContainerIndex < ConstraintContainerSolvers.Num(); ++ContainerIndex)
+			{
+				const int32 NumIslandConstraints = Island->GetConstraints(ContainerIndex).Num();
+				NumContainerConstraints[ContainerIndex] += NumIslandConstraints;
+				NumConstraints += NumIslandConstraints;
+			}
+		}
+	}
+
+	void FPBDIslandConstraintGroupSolver::AddConstraintsImpl()
+	{
+		// Initialize buffer sizes for solver particles and constraint solvers
+		SolverBodyContainer.Reset(NumParticles);
+		for (int32 ContainerIndex = 0; ContainerIndex < ConstraintContainerSolvers.Num(); ++ContainerIndex)
+		{
+			if (ConstraintContainerSolvers[ContainerIndex] != nullptr)
+			{
+				ConstraintContainerSolvers[ContainerIndex]->Reset(NumContainerConstraints[ContainerIndex]);
+			}
+		}
+
+		// Add all the constraints to the solvers, but do not collect data from the constraints yet (each constraint type has its own solver). 
+		// This also sets up the SolverBody array, but does not gather the body data.
+		for (int32 ContainerIndex = 0; ContainerIndex < ConstraintContainerSolvers.Num(); ++ContainerIndex)
+		{
+			if (ConstraintContainerSolvers[ContainerIndex] != nullptr)
+			{
+				for (FPBDIsland* Island : Islands)
+				{
+					// const_cast is because MakeArrayView doesn't handle a const array of pointers-to-non-const.
+					TArrayView<FPBDIslandConstraint> IslandConstraints = MakeArrayView(Island->GetConstraints(ContainerIndex));
+					ConstraintContainerSolvers[ContainerIndex]->AddConstraints(IslandConstraints);
+				}
+			}
+		}
+	}
+
+	void FPBDIslandConstraintGroupSolver::GatherBodiesImpl(const FReal Dt, const int32 BeginBodyIndex, const int32 EndBodyIndex)
+	{
+		// @todo(chaos): optimize
+		for (int32 SolverBodyIndex = BeginBodyIndex; SolverBodyIndex < EndBodyIndex; ++SolverBodyIndex)
+		{
+			FSolverBodyAdapter& SolverBody = SolverBodyContainer.GetItem(SolverBodyIndex);
+			const int32 Level = IslandManager.GetParticleLevel(SolverBody.GetParticle()->Handle());
+			SolverBody.GetSolverBody().SetLevel(Level);
+		}
 	}
-}
-	
-void FPBDIslandGroup::ClearIslands()
-{
-	IslandSolvers.Reset();
-}
 
-void FPBDIslandGroup::InitGroup()
-{
-	IslandSolvers.SetNum(0,false);
-	ParticlesCount = 0;
-	ConstraintsCount = 0;
-}
-	
 
 }
\ No newline at end of file
diff --git a/Engine/Source/Runtime/Experimental/Chaos/Private/Chaos/Island/IslandManager.cpp b/Engine/Source/Runtime/Experimental/Chaos/Private/Chaos/Island/IslandManager.cpp
--- a/Engine/Source/Runtime/Experimental/Chaos/Private/Chaos/Island/IslandManager.cpp
+++ b/Engine/Source/Runtime/Experimental/Chaos/Private/Chaos/Island/IslandManager.cpp
@@ -14,32 +14,30 @@
 
 //PRAGMA_DISABLE_OPTIMIZATION
 
+namespace Chaos
+{
+
 /** Cvar to enable/disable the island sleeping */
 bool bChaosSolverSleepEnabled = true;
 FAutoConsoleVariableRef CVarChaosSolverSleepEnabled(TEXT("p.Chaos.Solver.SleepEnabled"), bChaosSolverSleepEnabled, TEXT(""));
 
-/** Cvar to control the number of island groups used in the solver. The total number will be NumThreads * IslandGroupsMultiplier */
-Chaos::FRealSingle GChaosSolverIslandGroupsMultiplier = 1;
-FAutoConsoleVariableRef CVarSolverIslandGroupsMultiplier(TEXT("p.Chaos.Solver.IslandGroupsMultiplier"), GChaosSolverIslandGroupsMultiplier, TEXT("Total number of island groups in the solver will be NumThreads * IslandGroupsMultiplier.[def:1]"));
-
 /** Cvar to override the sleep counter threshold if necessary */
 int32 ChaosSolverCollisionDefaultSleepCounterThresholdCVar = 20;
 FAutoConsoleVariableRef CVarChaosSolverCollisionDefaultSleepCounterThreshold(TEXT("p.ChaosSolverCollisionDefaultSleepCounterThreshold"), ChaosSolverCollisionDefaultSleepCounterThresholdCVar, TEXT("Default counter threshold for sleeping.[def:20]"));
 
 /** Cvar to override the sleep linear threshold if necessary */
-Chaos::FRealSingle ChaosSolverCollisionDefaultLinearSleepThresholdCVar = 0.001f; // .001 unit mass cm
+FRealSingle ChaosSolverCollisionDefaultLinearSleepThresholdCVar = 0.001f; // .001 unit mass cm
 FAutoConsoleVariableRef CVarChaosSolverCollisionDefaultLinearSleepThreshold(TEXT("p.ChaosSolverCollisionDefaultLinearSleepThreshold"), ChaosSolverCollisionDefaultLinearSleepThresholdCVar, TEXT("Default linear threshold for sleeping.[def:0.001]"));
 
 /** Cvar to override the sleep angular threshold if necessary */
-Chaos::FRealSingle ChaosSolverCollisionDefaultAngularSleepThresholdCVar = 0.0087f;  //~1/2 unit mass degree
+FRealSingle ChaosSolverCollisionDefaultAngularSleepThresholdCVar = 0.0087f;  //~1/2 unit mass degree
 FAutoConsoleVariableRef CVarChaosSolverCollisionDefaultAngularSleepThreshold(TEXT("p.ChaosSolverCollisionDefaultAngularSleepThreshold"), ChaosSolverCollisionDefaultAngularSleepThresholdCVar, TEXT("Default angular threshold for sleeping.[def:0.0087]"));
 
 bool bChaosSolverValidateGraph = (CHAOS_CONSTRAINTHANDLE_DEBUG_ENABLED != 0);
 FAutoConsoleVariableRef CVarChaosSolverValidateGraph(TEXT("p.Chaos.Solver.ValidateGraph"), bChaosSolverValidateGraph, TEXT(""));
 
-namespace Chaos
-{
-	extern int32 GSingleThreadedPhysics;
+extern int32 GSingleThreadedPhysics;
+
 	
 /** Check if a particle is dynamic or sleeping */
 FORCEINLINE bool IsDynamicParticle(const FGeometryParticleHandle* ParticleHandle)
@@ -108,81 +106,6 @@ inline bool IsIslandSleeping(const FReal MaxLinearSpeed2, const FReal MaxAngular
 	return false;
 }
 
-	
-/**  Update all the island particles/constraints sleep state to be consistent with the island one */
-inline void UpdateSleepState(FPBDIslandSolver* IslandSolver, FPBDRigidsSOAs& Particles)
-{
-	if(IslandSolver)
-	{
-		// Sleeping flag has already been computed by the island graph 
-		const bool bIsSleeping = IslandSolver->IsSleeping();
-		
-		bool bNeedRebuild = false;
-		for (auto& IslandParticle : IslandSolver->GetParticles())
-		{
-			if(IslandParticle->CastToRigidParticle() && !IslandParticle->CastToRigidParticle()->Disabled())
-			{
-				// If not sleeping we activate the sleeping particles
-				if (!bIsSleeping && IslandParticle->Sleeping())
-				{
-					Particles.ActivateParticle(IslandParticle, true);
-
-					// When we wake particles, we have skipped their integrate step which causes some issues:
-					//	- we have zero velocity (no gravity or external forces applied)
-					//	- the world transforms cached in the ShapesArray will be at the last post-integrate positions
-					//	  which doesn't match what the velocity is telling us
-					// This causes problems for the solver - essentially we have an "initial overlap" situation.
-					// @todo(chaos): We could just run (partial) integrate here for this particle, but we don't know about the Evolution - fix this
-					for (const TUniquePtr<FPerShapeData>& Shape : IslandParticle->ShapesArray())
-					{
-						Shape->UpdateLeafWorldTransform(IslandParticle);
-					}
-
-					bNeedRebuild = true;
-				}
-				// If sleeping we deactivate the dynamic particles
-				else if (bIsSleeping && !IslandParticle->Sleeping())
-				{
-					Particles.DeactivateParticle(IslandParticle, true);
-					bNeedRebuild = true;
-				}
-			}
-		}
-		if(bNeedRebuild)
-		{
-			Particles.RebuildViews();
-		}
-		// Island constraints are updating their sleeping flag + awaken one 
-		for (auto& IslandConstraint : IslandSolver->GetConstraints())
-		{
-			if (FConstraintHandle* ConstraintHandle = IslandConstraint.Get())
-			{
-				ConstraintHandle->SetIsSleeping(bIsSleeping);
-			}
-		}
-	}
-}
-
-/** Update all the island particles sync state to be consistent with the island one */
-inline void UpdateSyncState(FPBDIslandSolver* IslandSolver, FPBDRigidsSOAs& Particles)
-{
-	if(IslandSolver)
-	{
-		bool bNeedsResim = false;
-		for (auto& IslandParticle : IslandSolver->GetParticles())
-		{
-			//if even one particle is soft/hard desync we must resim the entire island (when resim is used)
-			//seems cheap enough so just always do it, if slow pass resim template in here
-			if (IslandParticle->SyncState() != ESyncState::InSync)
-			{
-				bNeedsResim = true;
-				break;
-			}
-		}
-		IslandSolver->SetNeedsResim(bNeedsResim);
-	}
-}
-
 /** Add all the graph particle and constraints to the solver islands*/
 inline void PopulateIslands(FPBDIslandManager::GraphType* IslandGraph)
 {
@@ -190,10 +113,10 @@ inline void PopulateIslands(FPBDIslandManager::GraphType* IslandGraph)
 	{
 		if (IslandGraph->GraphIslands.IsValidIndex(IslandIndex))
 		{
-			FPBDIslandSolver* IslandSolver = IslandGraph->GraphIslands[IslandIndex].IslandItem;
-			if (!IslandSolver->IsSleeping() || IslandSolver->SleepingChanged())
+			FPBDIsland* Island = IslandGraph->GraphIslands[IslandIndex].IslandItem;
+			if (!Island->IsSleeping() || Island->SleepingChanged())
 			{
-				IslandSolver->AddParticle(GraphNode.NodeItem);
+				Island->AddParticle(GraphNode.NodeItem);
 			}
 		}
 	};
@@ -201,9 +124,12 @@ inline void PopulateIslands(FPBDIslandManager::GraphType* IslandGraph)
 	{
 		GraphNode.NodeIslands.Reset();
 
-		// If the node is valid : only one island
+		// If the node is valid : one and only one island
 		if (GraphNode.bValidNode)
 		{
+#if CHAOS_CONSTRAINTHANDLE_DEBUG_ENABLED
+			ensure(GraphNode.IslandIndex != INDEX_NONE);
+#endif
 			AddNodeToIsland(GraphNode.IslandIndex, GraphNode);
 			GraphNode.NodeIslands.Add(GraphNode.IslandIndex);
 		}
@@ -213,7 +139,13 @@ inline void PopulateIslands(FPBDIslandManager::GraphType* IslandGraph)
 			// First we compute the unique particle set of islands
 			for (auto& NodeEdge : GraphNode.NodeEdges)
 			{
-				GraphNode.NodeIslands.Add(IslandGraph->GraphEdges[NodeEdge].IslandIndex);
+				// The check for INDEX_NONE here filters out kinematic-kinematic constraints which may still
+				// be in the graph (a dynamic was converted to kinematic)
+				const int32 EdgeIslandIndex = IslandGraph->GraphEdges[NodeEdge].IslandIndex;
+				if (EdgeIslandIndex != INDEX_NONE)
+				{
+					GraphNode.NodeIslands.Add(IslandGraph->GraphEdges[NodeEdge].IslandIndex);
+				}
 			}
 			// Loop over the set of islands and add the particle to the solver island
 			for (auto& NodeIsland : GraphNode.NodeIslands)
@@ -229,24 +161,33 @@ inline void PopulateIslands(FPBDIslandManager::GraphType* IslandGraph)
 	{
 		if (IslandGraph->GraphIslands.IsValidIndex(GraphEdge.IslandIndex))
 		{
-			FPBDIslandSolver* IslandSolver = IslandGraph->GraphIslands[GraphEdge.IslandIndex].IslandItem;
+			FPBDIsland* Island = IslandGraph->GraphIslands[GraphEdge.IslandIndex].IslandItem;
 
 			// Note: we also populate the island if the sleeping state just changed. This is because we use this
 			// data to sync the constraint's sleeping state in SyncIslands.
 			// @todo(chaos): This is not great - we need a more robust way to sync the sleep state of islands/constraints
-			if (!IslandSolver->IsSleeping() || IslandSolver->SleepingChanged())
+			if (!Island->IsSleeping() || Island->SleepingChanged())
 			{
 #if CHAOS_CONSTRAINTHANDLE_DEBUG_ENABLED
 				ensure(GraphEdge.bValidEdge);
 #endif
-				IslandSolver->AddConstraint(GraphEdge.EdgeItem);
+				// The sub sort key is for consistent ordering of constraints when both Level and Color are the same.
+				// Without this the ordering within a level/color would depend on the order they were added to the graph
+				// which depends on collision detection order, which is multi-threaded and not deterministic (by default).
+				// The sort key is based on Particle's NodeIndex because particles remain in the graph with the same index 
+				// as long as they are active. We need to pack the key into 32 bits so we truncate the particle indices. 
+				// This should matter only very rarely. If we hit a conflict we may see a slight non-deterministic behaviour 
+				// difference because the solve  order may vary from frame to frame, but it will still solve ok.
+				// DESYNC WARNING! See above
+				const uint32 SubSortKey = (((uint32(GraphEdge.FirstNode) & 0xFFFF) << 16) | (uint32(GraphEdge.SecondNode) & 0xFFFF));
+				Island->AddConstraint(GraphEdge.EdgeItem, GraphEdge.LevelIndex, GraphEdge.ColorIndex, SubSortKey);
 			}
 		}
 	}
 }
 	
 /** Compute sleeping thresholds given a solver island  */
-inline bool ComputeSleepingThresholds(FPBDIslandSolver* IslandSolver,
+inline bool ComputeSleepingThresholds(FPBDIsland* Island,
 	const TArrayCollectionArray<TSerializablePtr<FChaosPhysicsMaterial>>& PerParticleMaterialAttributes,
 	const THandleArray<FChaosPhysicsMaterial>& SolverPhysicsMaterials, FReal& LinearSleepingThreshold,
 	FReal& AngularSleepingThreshold, FReal& MaxLinearSpeed2, FReal& MaxAngularSpeed2, int32& SleepCounterThreshold)
@@ -258,8 +199,9 @@ inline bool ComputeSleepingThresholds(FPBDIslandSolver* IslandSolver,
 	SleepCounterThreshold = 0;
 
 	bool bHaveSleepThreshold = false;
-	for (FGeometryParticleHandle* ParticleHandle : IslandSolver->GetParticles())
+	for (FPBDIslandParticle& IslandParticle : Island->GetParticles())
 	{
+		FGeometryParticleHandle* ParticleHandle = IslandParticle.GetParticle();
 		if (ParticleHandle)
 		{
 			if (FPBDRigidParticleHandle* PBDRigid = ParticleHandle->CastToRigidParticle())
@@ -301,26 +243,25 @@ inline bool ComputeSleepingThresholds(FPBDIslandSolver* IslandSolver,
 }
 
 FPBDIslandManager::FPBDIslandManager() 
-	: IslandSolvers()
+	: Islands()
 	, IslandGraph(MakeUnique<GraphType>())
+	, MaxNumConstraintContainers(0)
 	, MaxParticleIndex(INDEX_NONE)
 	, bIslandsPopulated(false)
 	, bEndTickCalled(true)
 {
 	IslandGraph->SetOwner(this);
-
-	InitializeGroups();
 }
 
 FPBDIslandManager::FPBDIslandManager(const TParticleView<FPBDRigidParticles>& PBDRigids) 
-	: IslandSolvers()
+	: Islands()
 	, IslandGraph(MakeUnique<GraphType>())
+	, MaxNumConstraintContainers(0)
 	, MaxParticleIndex(INDEX_NONE)
 	, bIslandsPopulated(false)
 	, bEndTickCalled(true)
 {
 	InitializeGraph(PBDRigids);
-	InitializeGroups();
 }
 
 FPBDIslandManager::~FPBDIslandManager()
@@ -332,27 +273,14 @@ void FPBDIslandManager::Reset()
 	RemoveConstraints();
 	RemoveParticles();
 
-	IslandSolvers.Reset();
-
-	for (auto& IslandGroup : IslandGroups)
-	{
-		IslandGroup->NumConstraints() = 0;
-	}
+	Islands.Reset();
 }
 
-void FPBDIslandManager::InitializeGroups()
+void FPBDIslandManager::AddConstraintContainer(FPBDConstraintContainer& ConstraintContainer)
 {
-	// @todo(chaos): is the number of worker threads a good indicator of how many threads we get in the solver loop? (Currently uses ParallelFor)
-	// Check for use of the "-onethread" command line arg, and physics threading disabled (GetNumWorkerThreads() is not affected by these)
-	const int32 NumWorkerThreads = (FApp::ShouldUseThreadingForPerformance() && !GSingleThreadedPhysics) ? FMath::Min(FTaskGraphInterface::Get().GetNumWorkerThreads(), Chaos::MaxNumWorkers) : 0;
-	const int32 MaxIslandGroups = FMath::Max(1, FMath::CeilToInt32(FReal(NumWorkerThreads) * GChaosSolverIslandGroupsMultiplier));
+	const int32 ContainerId = ConstraintContainer.GetContainerId();
 
-	IslandGroups.SetNum(MaxIslandGroups, false);
-	
-	for(int32 GroupIndex = 0, NumGroups = IslandGroups.Num(); GroupIndex < NumGroups; ++GroupIndex)
-	{
-		IslandGroups[GroupIndex] = MakeUnique<FPBDIslandGroup>(GroupIndex);
-	}
+	MaxNumConstraintContainers = FMath::Max(MaxNumConstraintContainers, ContainerId + 1);
 }
 
 void FPBDIslandManager::InitializeGraph(const TParticleView<FPBDRigidParticles>& PBDRigids)
@@ -413,12 +341,12 @@ void FPBDIslandManager::InitializeGraph(const TParticleView<FPBDRigidParticles>&
 
 void FPBDIslandManager::EndTick()
 {
-	// IslandSolvers are transient for use during the solver phase - clear them at the end of the frame
-	for (auto& IslandSolver : IslandSolvers)
+	// Islands are transient for use during the solver phase - clear them at the end of the frame
+	for (auto& Island : Islands)
 	{
-		IslandSolver->ClearParticles();
-		IslandSolver->ClearConstraints();
-		IslandSolver->ResetSleepingChanged();
+		Island->ClearParticles();
+		Island->ClearConstraints();
+		Island->ResetSleepingChanged();
 	}
 
 	// We can no longer use the Island's particle and constraint lists
@@ -433,13 +361,9 @@ void FPBDIslandManager::RemoveParticles()
 
 	IslandGraph->ResetNodes();
 
-	for (auto& IslandSolver : IslandSolvers)
-	{
-		IslandSolver->ClearParticles();
-	}
-	for (auto& IslandGroup : IslandGroups)
+	for (auto& Island : Islands)
 	{
-		IslandGroup->NumParticles() = 0;
+		Island->ClearParticles();
 	}
 }
 
@@ -447,13 +371,9 @@ void FPBDIslandManager::RemoveConstraints()
 {
 	IslandGraph->ResetEdges();
 
-	for(auto& IslandSolver : IslandSolvers)
+	for(auto& Island : Islands)
 	{
-		IslandSolver->ClearConstraints();
-	}
-	for(auto& IslandGroup : IslandGroups)
-	{
-		IslandGroup->NumConstraints() = 0;
+		Island->ClearConstraints();
 	}
 }
 
@@ -462,7 +382,7 @@ int32 FPBDIslandManager::ReserveParticles(const int32 NumParticles)
 	const int32 MaxIndex = IslandGraph->MaxNumNodes();
 	IslandGraph->ReserveNodes(NumParticles);
 	
-	IslandSolvers.Reserve(NumParticles);
+	Islands.Reserve(NumParticles);
 	IslandIndexing.Reserve(NumParticles);
 	
 	return FMath::Max(0,NumParticles - MaxIndex);
@@ -583,7 +503,7 @@ void FPBDIslandManager::RemoveConstraint(const uint32 ContainerId, FConstraintHa
 {
 	if (ConstraintHandle)
 	{
-		const int32 EdgeIndex = ConstraintHandle->ConstraintGraphIndex();
+		const int32 EdgeIndex = ConstraintHandle->GetConstraintGraphIndex();
 		if (IslandGraph->GraphEdges.IsValidIndex(EdgeIndex))
 		{
 			IslandGraph->RemoveEdge(EdgeIndex);
@@ -623,7 +543,7 @@ void FPBDIslandManager::RemoveParticleConstraints(FGeometryParticleHandle* Parti
 							{
 #if CHAOS_CONSTRAINTHANDLE_DEBUG_ENABLED
 								// All constraints should know their graph index
-								ensure(GraphEdge.EdgeItem->ConstraintGraphIndex() != INDEX_NONE);
+								ensure(GraphEdge.EdgeItem->GetConstraintGraphIndex() != INDEX_NONE);
 #endif
 
 								RemoveConstraint(ContainerId, GraphEdge.EdgeItem.Get());
@@ -661,11 +581,7 @@ void FPBDIslandManager::EnableParticle(FGeometryParticleHandle* ChildParticle, c
 				// and using the parent island index. if the island index we directly
 				// update the solver island without waiting the next sync
 				IslandIndex = ParentPBDRigid->IslandIndex();
-				ChildPBDRigid->SetSleeping(ParentPBDRigid->Sleeping());	
-				if (IslandSolvers.IsValidIndex(IslandIndex))
-				{
-					IslandSolvers[IslandIndex]->AddParticle(ChildParticle);
-				}
+				ChildPBDRigid->SetSleeping(ParentPBDRigid->Sleeping());	// @todo(chaos): this feels dodgy
 			}
 			// We add the child particle to the graph 
 			AddParticle(ChildParticle, IslandIndex);
@@ -682,19 +598,24 @@ void FPBDIslandManager::DisableParticle(FGeometryParticleHandle* ParticleHandle)
 	}
 }
 
+int32 FPBDIslandManager::GetParticleLevel(FGeometryParticleHandle* ParticleHandle) const
+{
+	return IslandGraph->GetNodeItemLevel(ParticleHandle);
+}
+
 void FPBDIslandManager::ResetIslands(const TParticleView<FPBDRigidParticles>& PBDRigids)
 {
 	// @todo 
 }
 
-inline bool SolverIslandSortPredicate(const TUniquePtr<FPBDIslandSolver>& SolverIslandL, const TUniquePtr<FPBDIslandSolver>& SolverIslandR)
+inline bool SolverIslandSortPredicate(const TUniquePtr<FPBDIsland>& SolverIslandL, const TUniquePtr<FPBDIsland>& SolverIslandR)
 {
-	return SolverIslandL->NumConstraints() < SolverIslandR->NumConstraints();
+	return SolverIslandL->GetNumConstraints() < SolverIslandR->GetNumConstraints();
 }
 	
-void FPBDIslandManager::SyncIslands(FPBDRigidsSOAs& Particles, const int32 NumContainers)
+void FPBDIslandManager::SyncIslands(FPBDRigidsSOAs& Particles)
 {
-	IslandSolvers.Reserve(IslandGraph->MaxNumIslands());
+	Islands.Reserve(IslandGraph->MaxNumIslands());
 	IslandIndexing.SetNum(IslandGraph->MaxNumIslands(),false);
 	SortedIslands.SetNum(IslandGraph->MaxNumIslands(),false);
 	int32 LocalIsland = 0;
@@ -709,43 +630,39 @@ void FPBDIslandManager::SyncIslands(FPBDRigidsSOAs& Particles, const int32 NumCo
 
 			// If the island item is not set on the graph we check if the solver island
 			// at the right index is already there. If not we create a new solver island.
-			if(!IslandSolvers.IsValidIndex(IslandIndex))
+			if(!Islands.IsValidIndex(IslandIndex))
 			{
-				IslandSolvers.EmplaceAt(IslandIndex, MakeUnique<FPBDIslandSolver>(this, LocalIsland));
+				Islands.EmplaceAt(IslandIndex, MakeUnique<FPBDIsland>(MaxNumConstraintContainers));
 			}
-			FPBDIslandSolver*& IslandSolver = IslandGraph->GraphIslands[IslandIndex].IslandItem;
-			IslandSolver = IslandSolvers[IslandIndex].Get();
-			IslandSolver->ResizeConstraintsCounts(NumContainers);
+			FPBDIsland*& Island = IslandGraph->GraphIslands[IslandIndex].IslandItem;
+			Island = Islands[IslandIndex].Get();
 			
 			// We then transfer the persistent flag and the graph dense index to the solver island
-			IslandSolver->SetIsPersistent(IslandGraph->GraphIslands[IslandIndex].bIsPersistent);
-			IslandSolver->SetIsSleeping(IslandGraph->GraphIslands[IslandIndex].bIsSleeping);
-			IslandSolver->GetIslandIndex() = LocalIsland;
+			Island->SetIsPersistent(IslandGraph->GraphIslands[IslandIndex].bIsPersistent);
+			Island->SetIsSleeping(IslandGraph->GraphIslands[IslandIndex].bIsSleeping);
+			Island->SetIslandIndex(LocalIsland);
 
 			// We update the IslandIndexing to retrieve the graph sparse and persistent index from the dense one.
 			IslandIndexing[LocalIsland] = IslandIndex;
 			SortedIslands[LocalIsland] = IslandIndex;
 			LocalIsland++;
 
-			// Clear the island solver ready for repopulation . This leaves the sleeping islands empty, but they 
-			// are not used until they wake at which point they will be filled again.
-			IslandSolver->ReserveParticles(IslandGraph->GraphIslands[IslandIndex].NumNodes);
-			IslandSolver->ReserveConstraints(IslandGraph->GraphIslands[IslandIndex].NumEdges);
+			Island->ReserveParticles(IslandGraph->GraphIslands[IslandIndex].NumNodes);
 
 			// Reset of the sleep counter if the island is :
 			// - Non persistent since we are starting incrementing the counter once the island is persistent and if values below the threshold
 			// - Sleeping since as soon as it wakes up we could start incrementing the counter as well
 			// - Just woken because we may not have been asleep for a whole frame to hit the
-			if (!IslandSolver->IsPersistent() || IslandSolver->IsSleeping() || IslandSolver->SleepingChanged())
+			if (!Island->IsPersistent() || Island->IsSleeping() || Island->SleepingChanged())
 			{ 
-				IslandSolver->SetSleepCounter(0);
+				Island->SetSleepCounter(0);
 			}
 			// We reset the persistent flag to be true on the island graph. 
 			IslandGraph->GraphIslands[IslandIndex].bIsPersistent = true;
 		}
-		else if (IslandSolvers.IsValidIndex(IslandIndex))
+		else if (Islands.IsValidIndex(IslandIndex))
 		{
-			IslandSolvers.RemoveAt(IslandIndex);
+			Islands.RemoveAt(IslandIndex);
 		}
 	}
 
@@ -754,13 +671,13 @@ void FPBDIslandManager::SyncIslands(FPBDRigidsSOAs& Particles, const int32 NumCo
 
 	PopulateIslands(IslandGraph.Get());
 
-	// Update of the sync and sleep state for each islands
-	for(auto& IslandSolver : IslandSolvers)
+	// Update of the sync and sleep state for each island
+	for(auto& Island : Islands)
 	{
-		if(!IslandSolver->IsSleeping() || IslandSolver->SleepingChanged())
+		if(!Island->IsSleeping() || Island->SleepingChanged())
 		{
-			UpdateSyncState(IslandSolver.Get(), Particles);
-			UpdateSleepState(IslandSolver.Get(), Particles);
+			Island->UpdateSyncState(Particles);
+			Island->PropagateSleepState(Particles);
 		}
 	}
 
@@ -770,78 +687,34 @@ void FPBDIslandManager::SyncIslands(FPBDRigidsSOAs& Particles, const int32 NumCo
 
 	// We can now safely use the GetIslandParticles and GetIslandConstraints accessors
 	bIslandsPopulated = true;
-
-	// Build all the island groups
-	BuildGroups(NumContainers);
-
 }
 
-void FPBDIslandManager::BuildGroups(const int32 NumContainers)
+void FPBDIslandManager::UpdateGraph()
 {
-	// No Need to sort for now since it is adding a cost for not improving much the end result
-	TSparseArray<TUniquePtr<FPBDIslandSolver>>& LocalIslands = IslandSolvers;
-	SortedIslands.Sort([&LocalIslands](const int32 IslandIndexA, const int32 IslandIndexB) -> bool
-		{ return LocalIslands[IslandIndexA]->NumConstraints() > LocalIslands[IslandIndexB]->NumConstraints();});
+	// Merge the graph islands if required
+	IslandGraph->UpdateGraph();
+}
 
-	int32 NumConstraints = 0;
-	for (TUniquePtr<FPBDIslandSolver>& IslandSolver : IslandSolvers)
-	{
-		if (!IslandSolver->IsSleeping())
-		{
-			NumConstraints += IslandSolver->NumConstraints();
-		}
-	}
+void FPBDIslandManager::UpdateIslands(FPBDRigidsSOAs& Particles)
+{
+	// Assign levels, colors, etc
+	IslandGraph->InitSorting();
+	IslandGraph->ComputeLevels();
+	//IslandGraph->ComputeColors();
 	
-	const int32 NumGroups = IslandGroups.Num();
-	const int32 GroupSize = NumConstraints / NumGroups + 1;
-
-	for(TUniquePtr<FPBDIslandGroup>& IslandGroup : IslandGroups)
-	{
-		IslandGroup->InitGroup();
-		IslandGroup->ResizeConstraintsCounts(NumContainers);
-	}
+	// Sync the graph islands with the solver islands objects
+	SyncIslands(Particles);
 
-	int32 GroupIndex = 0;
-	int32 GroupOffset = 0;
-	for(int32& SortedIndex : SortedIslands)
+	// Sort each island's constraints based on level, color, etc
+	for (TUniquePtr<FPBDIsland>& Island : Islands)
 	{
-		if( FPBDIslandSolver* IslandSolver = IslandSolvers[SortedIndex].Get())
+		if (!Island->IsSleeping())
 		{
-			if (!IslandSolver->IsSleeping())
-			{
-				checkf(GroupIndex < NumGroups, TEXT("Index %d is less than NumGroups %d"), GroupIndex, NumGroups);
-				IslandGroups[GroupIndex]->AddIsland(IslandSolver);
-				IslandGroups[GroupIndex]->NumParticles() += IslandSolver->NumParticles();
-				IslandGroups[GroupIndex]->NumConstraints() += IslandSolver->NumConstraints();
-
-				check(IslandSolver->NumContainerIds() == IslandGroups[GroupIndex]->NumContainerIds());
-				for(int32 ContainerIndex = 0; ContainerIndex < IslandSolver->NumContainerIds(); ++ContainerIndex)
-				{
-					IslandGroups[GroupIndex]->ConstraintCount(ContainerIndex) +=
-						IslandSolver->ConstraintCount(ContainerIndex);
-				}
-			
-				IslandSolver->SetGroupIndex(GroupIndex);
-				GroupOffset += IslandSolver->NumConstraints();
-
-				// @todo: In theory we should never need the second check	
-				if(GroupOffset > GroupSize && (GroupIndex + 1) < NumGroups)
-				{
-					GroupIndex++;
-					GroupOffset = 0;
-				}
-			}
+			// Sort constraints based on level
+			// @todo(chaos): can we eliminate this step? We effectively did it already above
+			Island->SortConstraints();
 		}
 	}
-}
-	
-void FPBDIslandManager::UpdateIslands(const TParticleView<FPBDRigidParticles>& PBDRigids, FPBDRigidsSOAs& Particles, const int32 NumContainers)
-{
-	// Merge the graph islands if required
-	IslandGraph->UpdateGraph();
-	
-	// Sync the graph islands with the solver islands objects
-	SyncIslands(Particles, NumContainers);
 
 	ValidateIslands();
 }
@@ -870,8 +743,8 @@ bool FPBDIslandManager::SleepInactive(const int32 IslandIndex,
 {
 	// Only the persistent islands could start sleeping
 	const int32 GraphIndex = GetGraphIndex(IslandIndex);
-	if (!bChaosSolverSleepEnabled || !IslandSolvers.IsValidIndex(GraphIndex) ||
-		(IslandSolvers.IsValidIndex(GraphIndex) && !IslandSolvers[GraphIndex]->IsPersistent()) )
+	if (!bChaosSolverSleepEnabled || !Islands.IsValidIndex(GraphIndex) ||
+		(Islands.IsValidIndex(GraphIndex) && !Islands[GraphIndex]->IsPersistent()) )
 	{
 		return false;
 	}
@@ -880,15 +753,15 @@ bool FPBDIslandManager::SleepInactive(const int32 IslandIndex,
 	int32 SleepCounterThreshold = 0, NumDynamicParticles = 0;
 
 	// Compute of the linear/angular velocities + thresholds to make islands sleeping 
-	if (ComputeSleepingThresholds(IslandSolvers[GraphIndex].Get(), PerParticleMaterialAttributes,
+	if (ComputeSleepingThresholds(Islands[GraphIndex].Get(), PerParticleMaterialAttributes,
 		SolverPhysicsMaterials, LinearSleepingThreshold, AngularSleepingThreshold,
 		MaxLinearSpeed2, MaxAngularSpeed2, SleepCounterThreshold))
 	{
-		int32 SleepCounter = IslandSolvers[GraphIndex]->GetSleepCounter();
+		int32 SleepCounter = Islands[GraphIndex]->GetSleepCounter();
 		const bool bSleepingIsland = IsIslandSleeping(MaxLinearSpeed2, MaxAngularSpeed2, LinearSleepingThreshold,
 			AngularSleepingThreshold, SleepCounterThreshold, SleepCounter);
 		
-		IslandSolvers[GraphIndex]->SetSleepCounter(SleepCounter);
+		Islands[GraphIndex]->SetSleepCounter(SleepCounter);
 
 		return bSleepingIsland;
 	}
@@ -909,39 +782,23 @@ void FPBDIslandManager::SetIslandSleeping(FPBDRigidsSOAs& Particles, const int32
 	// NOTE: This Sleep function is called from the graph update to set the state of the 
 	// island and all particles in it.
 	const int32 GraphIndex = GetGraphIndex(IslandIndex);
-	if (IslandSolvers.IsValidIndex(GraphIndex))
+	if (Islands.IsValidIndex(GraphIndex))
 	{
-		const bool bWasSleeping = IslandSolvers[GraphIndex]->IsSleeping();
+		const bool bWasSleeping = Islands[GraphIndex]->IsSleeping();
 		if (bWasSleeping != bIsSleeping)
 		{
-			IslandSolvers[GraphIndex]->SetIsSleeping(bIsSleeping);
 			IslandGraph->GraphIslands[GraphIndex].bIsSleeping = bIsSleeping;
-			UpdateSleepState(IslandSolvers[GraphIndex].Get(),Particles);
+			Islands[GraphIndex]->SetIsSleeping(bIsSleeping);
+			Islands[GraphIndex]->PropagateSleepState(Particles);
 		}
 
 		// Reset the sleep counter with every wake call, even if already awake
 		if (!bIsSleeping)
 		{
-			IslandSolvers[GraphIndex]->SetSleepCounter(0);
+			Islands[GraphIndex]->SetSleepCounter(0);
 		}
 	}
 }
-	
-const TArray<FGeometryParticleHandle*>& FPBDIslandManager::GetIslandParticles(const int32 IslandIndex) const
-{
-#if CHAOS_CONSTRAINTHANDLE_DEBUG_ENABLED
-	ensure(bIslandsPopulated);
-#endif
-	return IslandSolvers[GetGraphIndex(IslandIndex)]->GetParticles();
-}
-
-const TArray<FConstraintHandleHolder>& FPBDIslandManager::GetIslandConstraints(const int32 IslandIndex) const
-{
-#if CHAOS_CONSTRAINTHANDLE_DEBUG_ENABLED
-	ensure(bIslandsPopulated);
-#endif
-	return IslandSolvers[GetGraphIndex(IslandIndex)]->GetConstraints();
-}
 
 TArray<int32> FPBDIslandManager::FindParticleIslands(const FGeometryParticleHandle* ParticleHandle) const
 {
@@ -998,7 +855,7 @@ TArray<int32> FPBDIslandManager::FindParticleIslands(const FGeometryParticleHand
 
 int32 FPBDIslandManager::GetConstraintIsland(const FConstraintHandle* ConstraintHandle) const
 {
-	const int32 EdgeIndex = ConstraintHandle->ConstraintGraphIndex();
+	const int32 EdgeIndex = ConstraintHandle->GetConstraintGraphIndex();
 	if (IslandGraph->GraphEdges.IsValidIndex(EdgeIndex))
 	{
 		int32 GraphIslandIndex = IslandGraph->GraphEdges[EdgeIndex].IslandIndex;
@@ -1013,7 +870,7 @@ int32 FPBDIslandManager::GetConstraintIsland(const FConstraintHandle* Constraint
 
 bool FPBDIslandManager::IslandNeedsResim(const int32 IslandIndex) const
 {
-	return IslandSolvers[GetGraphIndex(IslandIndex)]->NeedsResim();
+	return Islands[GetGraphIndex(IslandIndex)]->NeedsResim();
 }
 
 #if CHAOS_CONSTRAINTHANDLE_DEBUG_ENABLED
@@ -1094,13 +951,16 @@ bool FPBDIslandManager::DebugCheckConstraintNotInGraph(const FConstraintHandle*
 
 	for (int32 IslandIndex = 0; IslandIndex < NumIslands(); ++IslandIndex)
 	{
-		const FPBDIslandSolver* IslandSolver = GetSolverIsland(IslandIndex);
-		for (const FConstraintHandleHolder& HandleHolder : IslandSolver->GetConstraints())
+		const FPBDIsland* Island = GetIsland(IslandIndex);
+		for (int32 ContainerIndex = 0; ContainerIndex < Island->GetNumConstraintContainers(); ++ContainerIndex)
 		{
-			if (HandleHolder.Get() == ConstraintHandle)
+			for (const FPBDIslandConstraint& Constraint : Island->GetConstraints(ContainerIndex))
 			{
-				bSuccess = false;
-				ensure(false);
+				if (Constraint.GetConstraint() == ConstraintHandle)
+				{
+					bSuccess = false;
+					ensure(false);
+				}
 			}
 		}
 	}
@@ -1177,24 +1037,29 @@ void FPBDIslandManager::DebugCheckParticleIslands(const FGeometryParticleHandle*
 		TArray<int32> NodeSolverIslands;
 		for (int32 IslandIndex = 0; IslandIndex < NumIslands(); ++IslandIndex)
 		{
-			const FPBDIslandSolver* IslandSolver = GetSolverIsland(IslandIndex);
-			for (const FGeometryParticleHandle* IslandParticleHandle : IslandSolver->GetParticles())
+			const FPBDIsland* Island = GetIsland(IslandIndex);
+			for (const FPBDIslandParticle& IslandParticle : Island->GetParticles())
 			{
+				const FGeometryParticleHandle* IslandParticleHandle = IslandParticle.GetParticle();
+
 				if (IslandParticleHandle == ParticleHandle)
 				{
 					NodeSolverIslands.AddUnique(IslandIndex);
 				}
 			}
-			for (const FConstraintHandleHolder& HandleHolder : IslandSolver->GetConstraints())
+			for (int32 ContainerIndex = 0; ContainerIndex < Island->GetNumConstraintContainers(); ++ContainerIndex)
 			{
-				if ((HandleHolder.GetParticle0() == ParticleHandle) || (HandleHolder.GetParticle1() == ParticleHandle))
+				for (const FPBDIslandConstraint& IslandConstraint : Island->GetConstraints(ContainerIndex))
 				{
-					NodeSolverIslands.AddUnique(IslandIndex);
+					if ((IslandConstraint.GetConstraint()->GetConstrainedParticles()[0] == ParticleHandle) || (IslandConstraint.GetConstraint()->GetConstrainedParticles()[1] == ParticleHandle))
+					{
+						NodeSolverIslands.AddUnique(IslandIndex);
+					}
 				}
 			}
 		}
 
-		// NOTE: Particles are always assigneed to an island, but if we do not have any constraints, the node may not be in the IslandSolver
+		// NOTE: Particles are always assigneed to an island, but if we do not have any constraints, the node may not be in the Island
 		if (FConstGenericParticleHandle(ParticleHandle)->IsDynamic())
 		{
 			// Dynamic particles should be in 1 island
@@ -1226,12 +1091,18 @@ void FPBDIslandManager::DebugCheckConstraintIslands(const FConstraintHandle* Con
 	TArray<int32> EdgeSolverIslands;
 	for (int32 SolverIslandIndex = 0; SolverIslandIndex < NumIslands(); ++SolverIslandIndex)
 	{
-		const FPBDIslandSolver* IslandSolver = GetSolverIsland(SolverIslandIndex);
-		for (const FConstraintHandleHolder& HandleHolder : IslandSolver->GetConstraints())
+		const FPBDIsland* Island = GetIsland(SolverIslandIndex);
+		for (int32 ContainerIndex = 0; ContainerIndex < Island->GetNumConstraintContainers(); ++ContainerIndex)
 		{
-			if (HandleHolder.Get() == ConstraintHandle)
+			for (const FPBDIslandConstraint& IslandConstraint : Island->GetConstraints(ContainerIndex))
 			{
-				EdgeSolverIslands.AddUnique(SolverIslandIndex);
+				// Make sure all the constraints are in the correct list according to type
+				ensure(ContainerIndex == IslandConstraint.GetConstraint()->GetContainerId());
+
+				if (IslandConstraint.GetConstraint() == ConstraintHandle)
+				{
+					EdgeSolverIslands.AddUnique(SolverIslandIndex);
+				}
 			}
 		}
 	}
diff --git a/Engine/Source/Runtime/Experimental/Chaos/Private/Chaos/Island/SolverIsland.cpp b/Engine/Source/Runtime/Experimental/Chaos/Private/Chaos/Island/SolverIsland.cpp
--- a/Engine/Source/Runtime/Experimental/Chaos/Private/Chaos/Island/SolverIsland.cpp
+++ b/Engine/Source/Runtime/Experimental/Chaos/Private/Chaos/Island/SolverIsland.cpp
@@ -1,19 +1,29 @@
 // Copyright Epic Games, Inc. All Rights Reserved.
 #include "Chaos/Island/SolverIsland.h"
-#include "Chaos/PBDConstraintRule.h"
+#include "Chaos/PBDConstraintContainer.h"
+#include "Chaos/PBDRigidsSOAs.h"
 
 namespace Chaos
 {
 	
-FPBDIslandSolver::FPBDIslandSolver(const FPBDIslandManager* IslandManagerIn, const int32 IslandIndexIn) : FPBDIslandSolverData(IslandIndexIn),
-	IslandManager(IslandManagerIn), bIsSleeping(false), bNeedsResim(false), bIsPersistent(true), SleepCounter(0),
-	IslandParticles(), IslandConstraints()
-{}
+	FPBDIsland::FPBDIsland(const int32 MaxConstraintContainers)
+	: IslandIndex(INDEX_NONE)
+	, bIsSleeping(false)
+	, bNeedsResim(false)
+	, bIsPersistent(true)
+	, SleepCounter(0)
+	, IslandParticles()
+	, IslandConstraintsByType()
+	, NumConstraints(0)
+{
+	IslandConstraintsByType.SetNum(MaxConstraintContainers);
+}
 
-void FPBDIslandSolver::UpdateParticles()
+void FPBDIsland::UpdateParticles()
 {
-	for (FGenericParticleHandle ParticleHandle : IslandParticles)
+	for (FPBDIslandParticle& IslandParticle : IslandParticles)
 	{
+		FGenericParticleHandle ParticleHandle = IslandParticle.GetParticle();
 		if (ParticleHandle.IsValid() && ParticleHandle->IsDynamic())
 		{
 			ParticleHandle->SetIslandIndex(IslandIndex);
@@ -21,12 +31,12 @@ void FPBDIslandSolver::UpdateParticles()
 	}
 }
 	
-void FPBDIslandSolver::ClearParticles()
+void FPBDIsland::ClearParticles()
 {
 	IslandParticles.Reset();
 }
 
-void FPBDIslandSolver::AddParticle(FGenericParticleHandle ParticleHandle)
+void FPBDIsland::AddParticle(FGenericParticleHandle ParticleHandle)
 {
 	if (ParticleHandle.IsValid())
 	{
@@ -38,105 +48,112 @@ void FPBDIslandSolver::AddParticle(FGenericParticleHandle ParticleHandle)
 	}
 }
 
-void FPBDIslandSolver::RemoveParticle(FGenericParticleHandle ParticleHandle)
-{
-	if (ParticleHandle.IsValid())
-	{
-		if (ParticleHandle->IsDynamic())
-		{
-			ParticleHandle->SetIslandIndex(INDEX_NONE);
-		}
-		IslandParticles.Remove(ParticleHandle->Handle());
-	}
-}
-
-void FPBDIslandSolver::ReserveParticles(const int32 NumParticles)
+void FPBDIsland::ReserveParticles(const int32 NumParticles)
 {
 	ClearParticles();
 	IslandParticles.Reserve(NumParticles);
 }
 
-void FPBDIslandSolver::AddConstraint(FConstraintHandle* ConstraintHandle)
+void FPBDIsland::AddConstraint(FConstraintHandle* ConstraintHandle, const int32 Level, const int32 Color, const uint32 SubSortKey)
 {
 	if (ConstraintHandle)
 	{
 		const int32 ContainerId = ConstraintHandle->GetContainerId();
-		if (ConstraintCounts.IsValidIndex(ContainerId))
-		{
-			IslandConstraints.Add(ConstraintHandle);
 
-			ConstraintCounts[ContainerId]++;
-		}
+		// Array should already be appropriately sized for all containers
+		IslandConstraintsByType[ContainerId].Emplace(ConstraintHandle, Level, Color, SubSortKey);
+
+		++NumConstraints;
 	}
 }
 
-void FPBDIslandSolver::RemoveConstraint(FConstraintHandle* ConstraintHandle)
+void FPBDIsland::ClearConstraints()
 {
-	// @todo(chaos): store the Island Constraint Index as a cookie on the constraint
-	if (ConstraintHandle)
+	for (TArray<FPBDIslandConstraint>& Constraints : IslandConstraintsByType)
 	{
-		IslandConstraints.Remove(ConstraintHandle);
+		Constraints.Reset();
 	}
+	NumConstraints = 0;
 }
 
-void FPBDIslandSolver::ReserveConstraints(const int32 NumConstraints)
+void FPBDIsland::SortConstraints()
 {
-	ClearConstraints();
-	IslandConstraints.Reserve(NumConstraints);
-}
-
-void FPBDIslandSolver::ClearConstraints()
-{
-	IslandConstraints.Reset();
+	for (TArray<FPBDIslandConstraint>& IslandConstraints : IslandConstraintsByType)
+	{
+		if (IslandConstraints.Num() > 1)
+		{
+			Algo::Sort(IslandConstraints,
+				[](const FPBDIslandConstraint& L, const FPBDIslandConstraint& R) -> bool
+				{
+					return L.GetSortKey() < R.GetSortKey();
+				});
+		}
+	}
 }
 
-// we should remove the one in constraint allocator once this one will be used
-inline bool ConstraintSortPredicate(const FConstraintHandle& L, const FConstraintHandle& R)
+void FPBDIsland::PropagateSleepState(FPBDRigidsSOAs& ParticleSOAs)
 {
-	const FPBDCollisionConstraint* CollisionConstraintL = L.As<FPBDCollisionConstraint>();
-	const FPBDCollisionConstraint* CollisionConstraintR = R.As<FPBDCollisionConstraint>();
-
-	if(CollisionConstraintL && CollisionConstraintR)
+	bool bNeedRebuild = false;
+	for (FPBDIslandParticle& IslandParticle : IslandParticles)
 	{
-		//sort constraints by the smallest particle idx in them first
-		//if the smallest particle idx is the same for both, use the other idx
-
-		if (CollisionConstraintL->GetCCDType() != CollisionConstraintR->GetCCDType())
+		FGeometryParticleHandle* Particle = IslandParticle.GetParticle();
+		if (Particle->CastToRigidParticle() && !Particle->CastToRigidParticle()->Disabled())
 		{
-			return CollisionConstraintL->GetCCDType() < CollisionConstraintR->GetCCDType();
+			// If not sleeping we activate the sleeping particles
+			if (!bIsSleeping && Particle->Sleeping())
+			{
+				ParticleSOAs.ActivateParticle(Particle, true);
+
+				// When we wake particles, we have skipped their integrate step which causes some issues:
+				//	- we have zero velocity (no gravity or external forces applied)
+				//	- the world transforms cached in the ShapesArray will be at the last post-integrate positions
+				//	  which doesn't match what the velocity is telling us
+				// This causes problems for the solver - essentially we have an "initial overlap" situation.
+				// @todo(chaos): We could just run (partial) integrate here for this particle, but we don't know about the Evolution - fix this
+				for (const TUniquePtr<FPerShapeData>& Shape : Particle->ShapesArray())
+				{
+					Shape->UpdateLeafWorldTransform(Particle);
+				}
+
+				bNeedRebuild = true;
+			}
+			// If sleeping we deactivate the dynamic particles
+			else if (bIsSleeping && !Particle->Sleeping())
+			{
+				ParticleSOAs.DeactivateParticle(Particle, true);
+				bNeedRebuild = true;
+			}
 		}
+	}
 
-		const FParticleID ParticleIdxsL[] = { CollisionConstraintL->GetParticle0()->ParticleID(), CollisionConstraintL->GetParticle1()->ParticleID() };
-		const FParticleID ParticleIdxsR[] = { CollisionConstraintR->GetParticle0()->ParticleID(), CollisionConstraintL->GetParticle1()->ParticleID() };
-
-		const int32 MinIdxL = ParticleIdxsL[0] < ParticleIdxsL[1] ? 0 : 1;
-		const int32 MinIdxR = ParticleIdxsR[0] < ParticleIdxsR[1] ? 0 : 1;
+	if (bNeedRebuild)
+	{
+		ParticleSOAs.RebuildViews();
+	}
 
-		if(ParticleIdxsL[MinIdxL] < ParticleIdxsR[MinIdxR])
-		{
-			return true;
-		} 
-		else if(ParticleIdxsL[MinIdxL] == ParticleIdxsR[MinIdxR])
+	// Island constraints are updating their sleeping flag + awaken one 
+	for (TArray<FPBDIslandConstraint>& TypedIslandConstraints : IslandConstraintsByType)
+	{
+		for (FPBDIslandConstraint& IslandConstraint : TypedIslandConstraints)
 		{
-			return ParticleIdxsL[!MinIdxL] < ParticleIdxsR[!MinIdxR];
+			IslandConstraint.GetConstraint()->SetIsSleeping(bIsSleeping);
 		}
-
-		return false;
 	}
-	return false;
 }
 
-inline bool ConstraintHolderSortPredicate(const FConstraintHandleHolder& L, const FConstraintHandleHolder& R)
+void FPBDIsland::UpdateSyncState(FPBDRigidsSOAs& Particles)
 {
-	return ConstraintSortPredicate(*L.Get(), *R.Get());
-}
-
-
-void FPBDIslandSolver::SortConstraints()
-{
-	if(!IsSleeping())
+	bNeedsResim = false;
+	for (FPBDIslandParticle& IslandParticle : IslandParticles)
 	{
-		IslandConstraints.Sort(ConstraintHolderSortPredicate);
+		// If even one particle is soft/hard desync we must resim the entire island (when resim is used)
+		// seems cheap enough so just always do it, if slow pass resim template in here
+		FGeometryParticleHandle* Particle = IslandParticle.GetParticle();
+		if (Particle->SyncState() != ESyncState::InSync)
+		{
+			bNeedsResim = true;
+			break;
+		}
 	}
 }
 
diff --git a/Engine/Source/Runtime/Experimental/Chaos/Private/Chaos/Joint/PBDJointSolverGaussSeidel.cpp b/Engine/Source/Runtime/Experimental/Chaos/Private/Chaos/Joint/PBDJointSolverGaussSeidel.cpp
--- a/Engine/Source/Runtime/Experimental/Chaos/Private/Chaos/Joint/PBDJointSolverGaussSeidel.cpp
+++ b/Engine/Source/Runtime/Experimental/Chaos/Private/Chaos/Joint/PBDJointSolverGaussSeidel.cpp
@@ -86,15 +86,11 @@ namespace Chaos
 
 	void FPBDJointSolver::Init(
 		const FReal Dt,
-		const FSolverBodyPtrPair& SolverBodyPair,
 		const FPBDJointSolverSettings& SolverSettings,
 		const FPBDJointSettings& JointSettings,
 		const FRigidTransform3& XL0,
 		const FRigidTransform3& XL1)
 	{
-		SolverBodies[0] = *SolverBodyPair[0];
-		SolverBodies[1] = *SolverBodyPair[1];
-
 		LocalConnectorXs[0] = XL0;
 		LocalConnectorXs[1] = XL1;
 
