Chaos - Constraint Solver Refactor

diff --git a/Engine/Source/Runtime/Experimental/Chaos/Private/Chaos/PBDCollisionConstraints.cpp b/Engine/Source/Runtime/Experimental/Chaos/Private/Chaos/PBDCollisionConstraints.cpp
--- a/Engine/Source/Runtime/Experimental/Chaos/Private/Chaos/PBDCollisionConstraints.cpp
+++ b/Engine/Source/Runtime/Experimental/Chaos/Private/Chaos/PBDCollisionConstraints.cpp
@@ -11,12 +11,12 @@
 #include "Chaos/Defines.h"
 #include "Chaos/Evolution/SolverBodyContainer.h"
 #include "Chaos/GeometryQueries.h"
+#include "Chaos/Island/IslandManager.h"
 #include "Chaos/SpatialAccelerationCollection.h"
 #include "Chaos/PBDRigidsSOAs.h"
 #include "Chaos/CastingUtilities.h"
 #include "ChaosLog.h"
 #include "ChaosStats.h"
-#include "Chaos/Evolution/SolverDatas.h"
 #include "ProfilingDebugging/ScopedTimers.h"
 #include "Algo/Sort.h"
 #include "Algo/StableSort.h"
@@ -100,18 +100,16 @@ namespace Chaos
 		const TArrayCollectionArray<TSerializablePtr<FChaosPhysicsMaterial>>& InPhysicsMaterials,
 		const TArrayCollectionArray<TUniquePtr<FChaosPhysicsMaterial>>& InPerParticlePhysicsMaterials,
 		const THandleArray<FChaosPhysicsMaterial>* const InSimMaterials,
-		const int32 InApplyPairIterations /*= 1*/,
-		const int32 InApplyPushOutPairIterations /*= 1*/,
-		const FReal InRestitutionThreshold /*= (FReal)2000*/)
+		const int32 NumCollisionsPerBlock,
+		const FReal InRestitutionThreshold)
 		: FPBDConstraintContainer(FConstraintContainerHandle::StaticType())
 		, Particles(InParticles)
+		, ConstraintAllocator(NumCollisionsPerBlock)
 		, NumActivePointConstraints(0)
 		, MCollided(Collided)
 		, MPhysicsMaterials(InPhysicsMaterials)
 		, MPerParticlePhysicsMaterials(InPerParticlePhysicsMaterials)
 		, SimMaterials(InSimMaterials)
-		, MApplyPairIterations(InApplyPairIterations)
-		, MApplyPushOutPairIterations(InApplyPushOutPairIterations)
 		, RestitutionThreshold(InRestitutionThreshold)	// @todo(chaos): expose as property
 		, bEnableCollisions(true)
 		, bEnableRestitution(true)
@@ -122,7 +120,6 @@ namespace Chaos
 		, GravityDirection(FVec3(0,0,-1))
 		, GravitySize(980)
 		, SolverSettings()
-		, SolverType(EConstraintSolverType::QuasiPbd)
 	{
 	}
 
@@ -254,10 +251,6 @@ namespace Chaos
 		CollisionMaterial.ResetMaterialModifications();
 	}
 
-	void FPBDCollisionConstraints::UpdatePositionBasedState(const FReal Dt)
-	{
-	}
-
 	void FPBDCollisionConstraints::BeginFrame()
 	{
 		ConstraintAllocator.BeginFrame();
@@ -367,308 +360,17 @@ namespace Chaos
 		}
 	}
 
-	Collisions::FContactParticleParameters FPBDCollisionConstraints::GetContactParticleParameters(const FReal Dt)
-	{
-		return { 
-			(CollisionRestitutionThresholdOverride >= 0.0f) ? CollisionRestitutionThresholdOverride * Dt : RestitutionThreshold * Dt,
-			CollisionCanAlwaysDisableContacts ? true : (CollisionCanNeverDisableContacts ? false : bCanDisableContacts),
-			&MCollided,
-
-		};
-	}
-
-	Collisions::FContactIterationParameters FPBDCollisionConstraints::GetContactIterationParameters(const FReal Dt, const int32 Iteration, const int32 NumIterations, const int32 NumPairIterations, bool& bNeedsAnotherIteration)
-	{
-		return {
-			Dt, 
-			Iteration, 
-			NumIterations, 
-			NumPairIterations, 
-			SolverType, 
-			&bNeedsAnotherIteration
-		};
-	}
-
-	void FPBDCollisionConstraints::SetNumIslandConstraints(const int32 NumIslandConstraints, FPBDIslandSolverData& SolverData)
-	{
-		if (SolverType == EConstraintSolverType::QuasiPbd)
-		{
-			FPBDCollisionSolverContainer& SolverContainer = GetConstraintSolverContainer(SolverData);
-			SolverContainer.SetNum(NumIslandConstraints);
-			SolverContainer.SetIsDeferredCollisionDetection(DetectorSettings.bDeferNarrowPhase);
-		}
-		else
-		{
-			SolverData.GetConstraintHandles(ContainerId).Reset(NumIslandConstraints);
-		}
-	}
-
-	FPBDCollisionSolverContainer& FPBDCollisionConstraints::GetConstraintSolverContainer(FPBDIslandSolverData& SolverData)
-	{
-		check(SolverType == EConstraintSolverType::QuasiPbd);
-		return SolverData.GetConstraintContainer<FPBDCollisionSolverContainer>(ContainerId);
-	}
-
-	void FPBDCollisionConstraints::PreGatherInput(const FReal Dt, FPBDCollisionConstraint& Constraint, FPBDIslandSolverData& SolverData)
-	{
-		if (SolverType == EConstraintSolverType::QuasiPbd)
-		{
-			FPBDCollisionSolverContainer& SolverContainer = GetConstraintSolverContainer(SolverData);
-			SolverContainer.PreAddConstraintSolver(Dt, Constraint, SolverData.GetBodyContainer(), SolverData.GetConstraintIndex(ContainerId));
-		}
-	}
-
-	void FPBDCollisionConstraints::GatherInput(const FReal Dt, FPBDCollisionConstraint& Constraint, const int32 Particle0Level, const int32 Particle1Level, FPBDIslandSolverData& SolverData)
-	{
-		if (SolverType == EConstraintSolverType::QuasiPbd)
-		{
-			// We shouldn't be adding disabled constraints to the solver list. The check needs to be at caller site or we should return success/fail - see TPBDConstraintColorRule::GatherSolverInput
-			check(Constraint.IsEnabled() && !Constraint.IsProbe());
-
-			FPBDCollisionSolverContainer& SolverContainer = GetConstraintSolverContainer(SolverData);
-			SolverContainer.AddConstraintSolver(Dt, Constraint, Particle0Level, Particle1Level, SolverData.GetBodyContainer(), SolverSettings);
-		}
-		else
-		{
-			LegacyGatherInput(Dt, Constraint, Particle0Level, Particle1Level, SolverData);
-		}
-	}
-
-	void FPBDCollisionConstraints::PreGatherInput(const FReal Dt, FPBDIslandSolverData& SolverData)
-	{
-		if (SolverType == EConstraintSolverType::QuasiPbd)
-		{
-			for (FPBDCollisionConstraint* Constraint : GetConstraints())
-			{
-				if (Constraint->IsEnabled())
-				{
-					PreGatherInput(Dt, *Constraint, SolverData);
-				}
-			}
-		}
-	}
-
-	void FPBDCollisionConstraints::GatherInput(const FReal Dt, FPBDIslandSolverData& SolverData)
+	void FPBDCollisionConstraints::AddConstraintsToGraph(FPBDIslandManager& IslandManager)
 	{
-		if (SolverType == EConstraintSolverType::QuasiPbd)
-		{
-			for (FPBDCollisionConstraint* Constraint : GetConstraints())
-			{
-				if (Constraint->IsEnabled() && !Constraint->IsProbe())
-				{
-					GatherInput(Dt, *Constraint, INDEX_NONE, INDEX_NONE, SolverData);
-				}
-			}
-		}
-		else
+		for (FPBDCollisionConstraintHandle* ConstraintHandle : GetConstraints())
 		{
-			for (FPBDCollisionConstraint* Constraint : GetConstraints())
+			if (ConstraintHandle->IsEnabled())
 			{
-				if (Constraint->IsEnabled() && !Constraint->IsProbe())
-				{
-					LegacyGatherInput(Dt, *Constraint, INDEX_NONE, INDEX_NONE, SolverData);
-				}
+				IslandManager.AddConstraint(GetContainerId(), ConstraintHandle, ConstraintHandle->GetConstrainedParticles());
 			}
 		}
 	}
 
-	void FPBDCollisionConstraints::ScatterOutput(const FReal Dt, const int32 BeginIndex, const int32 EndIndex, FPBDIslandSolverData& SolverData)
-	{
-		if (SolverType == EConstraintSolverType::QuasiPbd)
-		{
-			GetConstraintSolverContainer(SolverData).ScatterOutput(Dt, BeginIndex, EndIndex);
-		}
-		else
-		{
-			LegacyScatterOutput(Dt, BeginIndex, EndIndex, SolverData);
-		}
-	}
-
-	void FPBDCollisionConstraints::ScatterOutput(const FReal Dt, FPBDIslandSolverData& SolverData)
-	{
-		if (SolverType == EConstraintSolverType::QuasiPbd)
-		{
-			FPBDCollisionSolverContainer& SolverContainer = GetConstraintSolverContainer(SolverData);
-			SolverContainer.ScatterOutput(Dt, 0, SolverContainer.NumSolvers());
-		}
-		else
-		{
-			LegacyScatterOutput(Dt, 0, SolverData.GetConstraintHandles(ContainerId).Num(), SolverData);
-		}
-	}
-
-	// Simple Rule version
-	bool FPBDCollisionConstraints::ApplyPhase1(const FReal Dt, const int32 It, const int32 NumIts, FPBDIslandSolverData& SolverData)
-	{
-		SCOPE_CYCLE_COUNTER(STAT_Collisions_Apply);
-
-		if (SolverType == EConstraintSolverType::QuasiPbd)
-		{
-			FPBDCollisionSolverContainer& SolverContainer = GetConstraintSolverContainer(SolverData);
-			return SolverContainer.SolvePositionSerial(Dt, It, NumIts, 0, SolverContainer.NumSolvers(), SolverSettings);
-		}
-		else
-		{
-			return LegacyApplyPhase1Serial(Dt, It, NumIts, 0, SolverData.GetConstraintHandles(ContainerId).Num(), SolverData);
-		}
-	}
-
-	// Island Rule version
-	bool FPBDCollisionConstraints::ApplyPhase1Serial(const FReal Dt, const int32 It, const int32 NumIts, FPBDIslandSolverData& SolverData)
-	{
-		SCOPE_CYCLE_COUNTER(STAT_Collisions_Apply);
-
-		if (SolverType == EConstraintSolverType::QuasiPbd)
-		{
-			FPBDCollisionSolverContainer& SolverContainer = GetConstraintSolverContainer(SolverData);
-			return SolverContainer.SolvePositionSerial(Dt, It, NumIts, 0, SolverContainer.NumSolvers(), SolverSettings);
-		}
-		else
-		{
-			return LegacyApplyPhase1Serial(Dt, It, NumIts, 0, SolverData.GetConstraintHandles(ContainerId).Num(), SolverData);
-		}
-	}
-
-	// Color Rule version
-	bool FPBDCollisionConstraints::ApplyPhase1Serial(const FReal Dt, const int32 It, const int32 NumIts, const int32 BeginIndex, const int32 EndIndex, FPBDIslandSolverData& SolverData)
-	{
-		SCOPE_CYCLE_COUNTER(STAT_Collisions_Apply);
-
-		if (SolverType == EConstraintSolverType::QuasiPbd)
-		{
-			FPBDCollisionSolverContainer& SolverContainer = GetConstraintSolverContainer(SolverData);
-			return SolverContainer.SolvePositionSerial(Dt, It, NumIts, BeginIndex, EndIndex, SolverSettings);
-		}
-		else
-		{
-			return LegacyApplyPhase1Serial(Dt, It, NumIts, BeginIndex, EndIndex, SolverData);
-		}
-	}
-
-	// Color Rule version
-	bool FPBDCollisionConstraints::ApplyPhase1Parallel(const FReal Dt, const int32 It, const int32 NumIts, const int32 BeginIndex, const int32 EndIndex, FPBDIslandSolverData& SolverData)
-	{
-		SCOPE_CYCLE_COUNTER(STAT_Collisions_Apply);
-
-		if (SolverType == EConstraintSolverType::QuasiPbd)
-		{
-			FPBDCollisionSolverContainer& SolverContainer = GetConstraintSolverContainer(SolverData);
-			return SolverContainer.SolvePositionParallel(Dt, It, NumIts, BeginIndex, EndIndex, SolverSettings);
-		}
-		else
-		{
-			return LegacyApplyPhase1Parallel(Dt, It, NumIts, BeginIndex, EndIndex, SolverData);
-		}
-	}
-
-	// Simple Rule version
-	bool FPBDCollisionConstraints::ApplyPhase2(const FReal Dt, const int32 It, const int32 NumIts, FPBDIslandSolverData& SolverData)
-	{
-		SCOPE_CYCLE_COUNTER(STAT_Collisions_ApplyPushOut);
-
-		if (SolverType == EConstraintSolverType::QuasiPbd)
-		{
-			FPBDCollisionSolverContainer& SolverContainer = GetConstraintSolverContainer(SolverData);
-			return SolverContainer.SolveVelocitySerial(Dt, It, NumIts, 0, SolverContainer.NumSolvers(), SolverSettings);
-		}
-
-		return false;
-	}
-
-	// Island Rule version
-	bool FPBDCollisionConstraints::ApplyPhase2Serial(const FReal Dt, const int32 It, const int32 NumIts, FPBDIslandSolverData& SolverData)
-	{
-		SCOPE_CYCLE_COUNTER(STAT_Collisions_ApplyPushOut);
-
-		if (SolverType == EConstraintSolverType::QuasiPbd)
-		{
-			FPBDCollisionSolverContainer& SolverContainer = GetConstraintSolverContainer(SolverData);
-			return SolverContainer.SolveVelocitySerial(Dt, It, NumIts, 0, SolverContainer.NumSolvers(), SolverSettings);
-		}
-
-		return false;
-	}
-
-	// Color Rule version
-	bool FPBDCollisionConstraints::ApplyPhase2Serial(const FReal Dt, const int32 It, const int32 NumIts, const int32 BeginIndex, const int32 EndIndex, FPBDIslandSolverData& SolverData)
-	{
-		SCOPE_CYCLE_COUNTER(STAT_Collisions_ApplyPushOut);
-
-		if (SolverType == EConstraintSolverType::QuasiPbd)
-		{
-			FPBDCollisionSolverContainer& SolverContainer = GetConstraintSolverContainer(SolverData);
-			return SolverContainer.SolveVelocitySerial(Dt, It, NumIts, BeginIndex, EndIndex, SolverSettings);
-		}
-
-		return false;
-	}
-
-	// Color Rule version
-	bool FPBDCollisionConstraints::ApplyPhase2Parallel(const FReal Dt,  const int32 It, const int32 NumIts, const int32 BeginIndex, const int32 EndIndex, FPBDIslandSolverData& SolverData)
-	{
-		SCOPE_CYCLE_COUNTER(STAT_Collisions_ApplyPushOut);
-
-		if (SolverType == EConstraintSolverType::QuasiPbd)
-		{
-			FPBDCollisionSolverContainer& SolverContainer = GetConstraintSolverContainer(SolverData);
-			return SolverContainer.SolveVelocityParallel(Dt, It, NumIts, BeginIndex, EndIndex, SolverSettings);
-		}
-
-		return false;
-	}
-
-	void FPBDCollisionConstraints::LegacyGatherInput(const FReal Dt, FPBDCollisionConstraint& Constraint, const int32 Particle0Level, const int32 Particle1Level, FPBDIslandSolverData& SolverData)
-	{
-		SolverData.GetConstraintHandles(ContainerId).Add(&Constraint);
-
-		FSolverBody* SolverBody0 = SolverData.GetBodyContainer().FindOrAdd(Constraint.Particle[0], Dt);
-		FSolverBody* SolverBody1 = SolverData.GetBodyContainer().FindOrAdd(Constraint.Particle[1], Dt);
-
-		SolverBody0->SetLevel(Particle0Level);
-		SolverBody1->SetLevel(Particle1Level);
-
-		Constraint.SetSolverBodies(SolverBody0, SolverBody1);
-
-		Constraint.AccumulatedImpulse = FVec3(0);
-	}
-
-	void FPBDCollisionConstraints::LegacyScatterOutput(const FReal Dt, const int32 BeginIndex, const int32 EndIndex, FPBDIslandSolverData& SolverData)
-	{
-		for (int32 Index = BeginIndex; Index < EndIndex; ++Index)
-		{
-			FPBDCollisionConstraint* Constraint = SolverData.GetConstraintHandle<FPBDCollisionConstraint>(ContainerId,Index);
-			Constraint->SetSolverBodies(nullptr, nullptr);
-		}
-	}
-
-	bool FPBDCollisionConstraints::LegacyApplyPhase1Serial(const FReal Dt, const int32 Iterations, const int32 NumIterations, const int32 BeginIndex, const int32 EndIndex, FPBDIslandSolverData& SolverData)
-	{
-		bool bNeedsAnotherIteration = false;
-		if (MApplyPairIterations > 0)
-		{
-			NumActivePointConstraints = 0;
-			const Collisions::FContactParticleParameters ParticleParameters = GetContactParticleParameters(Dt);
-			const Collisions::FContactIterationParameters IterationParameters = GetContactIterationParameters(Dt, Iterations, NumIterations, MApplyPairIterations, bNeedsAnotherIteration);
-
-			for (int32 Index = BeginIndex; Index < EndIndex; ++Index)
-			{
-				FPBDCollisionConstraint* Constraint = SolverData.GetConstraintHandle<FPBDCollisionConstraint>(ContainerId,Index);
-
-				if ((!Constraint->GetDisabled()) && (!Constraint->GetIsProbe()))
-				{
-					Collisions::Apply(*Constraint, IterationParameters, ParticleParameters);
-					++NumActivePointConstraints;
-				}
-			}
-		}
-		return bNeedsAnotherIteration;
-	}
-
-	bool FPBDCollisionConstraints::LegacyApplyPhase1Parallel(const FReal Dt, const int32 Iterations, const int32 NumIterations, const int32 BeginIndex, const int32 EndIndex, FPBDIslandSolverData& SolverData)
-	{
-		return LegacyApplyPhase1Serial(Dt, Iterations, NumIterations, BeginIndex, EndIndex, SolverData);
-	}
-
 	const FPBDCollisionConstraint& FPBDCollisionConstraints::GetConstraint(int32 Index) const
 	{
 		check(Index < NumConstraints());
diff --git a/Engine/Source/Runtime/Experimental/Chaos/Private/Chaos/PBDCollisionConstraintsContact.cpp b/Engine/Source/Runtime/Experimental/Chaos/Private/Chaos/PBDCollisionConstraintsContact.cpp
--- a/Engine/Source/Runtime/Experimental/Chaos/Private/Chaos/PBDCollisionConstraintsContact.cpp
+++ b/Engine/Source/Runtime/Experimental/Chaos/Private/Chaos/PBDCollisionConstraintsContact.cpp
@@ -1,213 +1,8 @@
 // Copyright Epic Games, Inc. All Rights Reserved.
 
 #include "Chaos/PBDCollisionConstraintsContact.h"
-#include "Chaos/Collision/CollisionSolver.h"
-#include "Chaos/Collision/PBDCollisionSolver.h"
-#include "Chaos/CollisionResolution.h"
-#include "Chaos/CollisionResolutionUtil.h"
-#include "Chaos/Defines.h"
-#include "Chaos/Particle/ParticleUtilities.h"
-#include "Chaos/Utilities.h"
 
-#if INTEL_ISPC
-#include "PBDCollisionConstraints.ispc.generated.h"
-#endif
-
-/**
- * @brief This file contains collision solver code for the StandardPbd solver (which will be removed at some point)
- * @see PBDCollisionSolver.cpp for the new solver
-*/
-
-namespace Chaos
-{
-	namespace Collisions
-	{
-		void Update(FPBDCollisionConstraint& Constraint, const FReal Dt)
-		{
-			// NOTE: These are actor transforms, not CoM transforms
-			// \todo(chaos): see if we can easily switch to CoM transforms now in collision loop (shapes are held in actor space)
-			const FSolverBody& Body0 = *Constraint.GetSolverBody0();
-			const FSolverBody& Body1 = *Constraint.GetSolverBody1();
-			const FRigidTransform3 Transform0 = FRigidTransform3(Body0.ActorP(), Body0.ActorQ());
-			const FRigidTransform3 Transform1 = FRigidTransform3(Body1.ActorP(), Body1.ActorQ());
-
-			Constraint.ResetPhi(Constraint.GetCullDistance());
-			UpdateConstraintFromGeometry<ECollisionUpdateType::Deepest>(Constraint, Transform0, Transform1, Dt);
-		}
-
-		// A PBD collision penetration correction.
-		// Currently only used by RBAN
-		FVec3 ApplyContact2(FPBDCollisionConstraint& Constraint,
-			const FContactIterationParameters& IterationParameters,
-			const FContactParticleParameters& ParticleParameters)
-		{
-			FSolverBody& Body0 = *Constraint.GetSolverBody0();
-			FSolverBody& Body1 = *Constraint.GetSolverBody1();
-
-			const FVec3 ContactLocation = Constraint.CalculateWorldContactLocation();
-			const FVec3 ContactNormal = Constraint.CalculateWorldContactNormal();
-			const FReal ContactFriction = Constraint.GetDynamicFriction();
-			const FReal ContactRestitution = Constraint.GetRestitution();
-			FReal ContactPhi = Constraint.GetPhi();
-			FReal ContactRestitutionPadding = Constraint.GetRestitutionPadding();
-
-			FVec3 VectorToPoint0 = ContactLocation - Body0.P();
-			FVec3 VectorToPoint1 = ContactLocation - Body1.P();
-
-			if (ContactPhi < 0)
-			{
-				*IterationParameters.NeedsAnotherIteration = true;
-
-				bool bApplyResitution = (ContactRestitution > 0.0f);
-				bool bHaveRestitutionPadding = (ContactRestitutionPadding > 0.0f);
-				bool bApplyFriction = (ContactFriction > 0) && (IterationParameters.Dt > UE_SMALL_NUMBER);
-
-				// If we have restitution, pad the constraint by an amount that enforces the outgoing velocity constraint
-				// Really this should be per contact point, not per constraint.
-				// NOTE: once we have calculated a padding, it is locked in for the rest of the iterations, and automatically
-				// included in the Phi we get back from collision detection. The first time we calculate it, we must also
-				// add the padding to the Phi (since it was from pre-padded collision detection).
-				if (bApplyResitution && !bHaveRestitutionPadding)
-				{
-					const FVec3 BodyV0 = Body0.V() + (FVec3(Body0.DP()) / IterationParameters.Dt);
-					const FVec3 BodyW0 = Body0.W() + (FVec3(Body0.DQ()) / IterationParameters.Dt);
-					const FVec3 BodyV1 = Body1.V() + (FVec3(Body1.DP()) / IterationParameters.Dt);
-					const FVec3 BodyW1 = Body1.W() + (FVec3(Body1.DQ()) / IterationParameters.Dt);
-					FVec3 CV0 = BodyV0 + FVec3::CrossProduct(BodyW0, VectorToPoint0);
-					FVec3 CV1 = BodyV1 + FVec3::CrossProduct(BodyW1, VectorToPoint1);
-					FVec3 CV = CV0 - CV1;
-					FReal CVNormal = FVec3::DotProduct(CV, ContactNormal);
-
-					// No restitution below threshold normal velocity (CVNormal is negative here)
-					if (CVNormal < -ParticleParameters.RestitutionVelocityThreshold)
-					{
-						ContactRestitutionPadding = -(1.0f + ContactRestitution) * CVNormal * IterationParameters.Dt + ContactPhi;
-						ContactPhi -= ContactRestitutionPadding;
-
-						Constraint.SetRestitutionPadding(ContactRestitutionPadding);
-					}
-				}
-			
-				FMatrix33 ContactInvI =
-					(Body0.IsDynamic() ? ComputeFactorMatrix3(VectorToPoint0, FMatrix33(Body0.InvI()), FReal(Body0.InvM())) : FMatrix33(0)) +
-					(Body1.IsDynamic() ? ComputeFactorMatrix3(VectorToPoint1, FMatrix33(Body1.InvI()), FReal(Body1.InvM())) : FMatrix33(0));
-
-				// Calculate the normal correction
-				FVec3 NormalError = ContactPhi * ContactNormal;
-				FReal NormalImpulseDenominator = FVec3::DotProduct(ContactNormal, ContactInvI * ContactNormal);
-				FVec3 NormalImpulseNumerator = -NormalError;
-				FVec3 NormalCorrection = NormalImpulseNumerator / NormalImpulseDenominator;
-
-				// Calculate lateral correction, clamped to the friction cone. Kinda.
-				FVec3 LateralCorrection = FVec3(0);
-				if (bApplyFriction)
-				{
-					// @todo(ccaulfield): use initial velocity (as for restitution) and accumulate friction force per contact point
-					FVec3 V0 = FVec3::CalculateVelocity(Body0.X(), Body0.CorrectedP(), IterationParameters.Dt);
-					FVec3 W0 = FRotation3::CalculateAngularVelocity(Body0.R(), Body0.CorrectedQ(), IterationParameters.Dt);
-					FVec3 V1 = FVec3::CalculateVelocity(Body1.X(), Body1.CorrectedP(), IterationParameters.Dt);
-					FVec3 W1 = FRotation3::CalculateAngularVelocity(Body1.R(), Body1.CorrectedQ(), IterationParameters.Dt);
-					FVec3 CV0 = V0 + FVec3::CrossProduct(W0, VectorToPoint0);
-					FVec3 CV1 = V1 + FVec3::CrossProduct(W1, VectorToPoint1);
-					FVec3 CV = CV0 - CV1;
-					FReal CVNormal = FVec3::DotProduct(CV, ContactNormal);
-					if (CVNormal < 0.0f)
-					{
-						FVec3 CVLateral = CV - CVNormal * ContactNormal;
-						FReal CVLateralMag = CVLateral.Size();
-						if (CVLateralMag > UE_KINDA_SMALL_NUMBER)
-						{
-							FVec3 DirLateral = CVLateral / CVLateralMag;
-							FVec3 LateralImpulseNumerator = -CVLateral * IterationParameters.Dt;
-							FReal LateralImpulseDenominator = FVec3::DotProduct(DirLateral, ContactInvI * DirLateral);
-							LateralCorrection = LateralImpulseNumerator / LateralImpulseDenominator;
-							FReal LateralImpulseMag = LateralCorrection.Size();
-							FReal NormalImpulseMag = NormalCorrection.Size();
-							if (LateralImpulseMag > ContactFriction * NormalImpulseMag)
-							{
-								LateralCorrection *= ContactFriction * NormalImpulseMag / LateralImpulseMag;
-							}
-						}
-					}
-				}
-
-				// Net Correction
-				FVec3 DX = NormalCorrection + LateralCorrection;
-				
-				if (Body0.IsDynamic())
-				{
-					FVec3 DP0 = Body0.InvM() * DX;
-					FVec3 DR0 = Utilities::Multiply(FMatrix33(Body0.InvI()), FVec3::CrossProduct(VectorToPoint0, DX));
-					Body0.ApplyTransformDelta(DP0, DR0);
-					Body0.ApplyCorrections();
-					Body0.UpdateRotationDependentState();
-				}
-				if (Body1.IsDynamic())
-				{
-					FVec3 DP1 = Body1.InvM() * -DX;
-					FVec3 DR1 = Utilities::Multiply(FMatrix33(Body1.InvI()), FVec3::CrossProduct(VectorToPoint1, -DX));
-					Body1.ApplyTransformDelta(DP1, DR1);
-					Body1.ApplyCorrections();
-					Body1.UpdateRotationDependentState();
-				}
-
-				return DX;
-			}
-
-			return FVec3(0);
-		}
-
-
-		void Apply(FPBDCollisionConstraint& Constraint, const FContactIterationParameters& IterationParameters, const FContactParticleParameters& ParticleParameters)
-		{
-			check(IterationParameters.SolverType == EConstraintSolverType::StandardPbd);
-
-			for (int32 PairIt = 0; PairIt < IterationParameters.NumPairIterations; ++PairIt)
-			{
-				// Update the contact information based on current particles' positions
-				if (Constraint.GetUseIncrementalCollisionDetection() || (Constraint.GetManifoldPoints().Num() == 0))
-				{
-					Collisions::Update(Constraint, IterationParameters.Dt);
-				}
-				else
-				{
-					Constraint.UpdateManifoldContacts();
-				}
-
-				// Permanently disable a constraint that is beyond the cull distance
-				if (Constraint.GetPhi() >= Constraint.GetCullDistance())
-				{
-					if (ParticleParameters.bCanDisableContacts)
-					{
-						Constraint.SetDisabled(true);
-					}
-					return;
-				}
-
-				// Do not early out here in the case of Accumulated impulse solve
-				// @todo(chaos): remove this early out when we settle on manifolds
-				const bool bIsAccumulatingImpulses = Constraint.GetUseManifold();
-				if (Constraint.GetPhi() >= 0.0f && !bIsAccumulatingImpulses)
-				{
-					return;
-				}
-
-				// @todo(chaos): fix the collided flag - it will sometimes be set if clipping is enabled, even if there was no contact...
-				// @todo(chaos): this doesn't seem to be being set or used...remove Collided and the Aux Collided array if so. If it is supposed 
-				// to work, it should be set via the Scatter method of the SolverBody
-				if (ParticleParameters.Collided)
-				{
-					FGenericParticleHandle Particle0 = FGenericParticleHandle(Constraint.GetParticle0());
-					FGenericParticleHandle Particle1 = FGenericParticleHandle(Constraint.GetParticle1());
-					Particle0->AuxilaryValue(*ParticleParameters.Collided) = true;
-					Particle1->AuxilaryValue(*ParticleParameters.Collided) = true;
-				}
-
-				Constraint.AccumulatedImpulse += ApplyContact2(Constraint, IterationParameters, ParticleParameters);
-			}
-		}
-
-	} // Collisions
-
-}// Chaos
+//
+// TO BE REMOVED
+//
 
diff --git a/Engine/Source/Runtime/Experimental/Chaos/Private/Chaos/PBDConstraintColor.cpp b/Engine/Source/Runtime/Experimental/Chaos/Private/Chaos/PBDConstraintColor.cpp
--- a/Engine/Source/Runtime/Experimental/Chaos/Private/Chaos/PBDConstraintColor.cpp
+++ b/Engine/Source/Runtime/Experimental/Chaos/Private/Chaos/PBDConstraintColor.cpp
@@ -2,821 +2,6 @@
 
 #include "Chaos/PBDConstraintColor.h"
 
-#include "Chaos/PBDConstraintGraph.h"
-#include "Chaos/ConstraintHandle.h"
-#include "Chaos/Framework/Parallel.h"
-#include "Chaos/ParticleHandle.h"
-#include "Chaos/PBDRigidParticles.h"
-#include "Chaos/PBDCollisionConstraints.h"
-
-#include "ProfilingDebugging/ScopedTimers.h"
-#include "ChaosStats.h"
-#include "Containers/Queue.h"
-#include "ChaosLog.h"
-
-#include <memory>
-#include <queue>
-#include <sstream>
-#include "Containers/RingBuffer.h"
-
-using namespace Chaos;
-
-DECLARE_CYCLE_STAT(TEXT("FPBDConstraintColor::ComputeColors"), STAT_Constraint_ComputeColor, STATGROUP_Chaos);
-DECLARE_CYCLE_STAT(TEXT("FPBDConstraintColor::ComputeContactGraph"), STAT_Constraint_ComputeContactGraph, STATGROUP_Chaos);
-DECLARE_CYCLE_STAT(TEXT("FPBDConstraintColor::ComputeContactGraphGBF"), STAT_Constraint_ComputeContactGraphGBF, STATGROUP_Chaos);
-DECLARE_CYCLE_STAT(TEXT("FPBDConstraintColor::ComputeIslandColoring"), STAT_Constraint_ComputeIslandColoring, STATGROUP_Chaos);
-
-bool bUseContactGraphGBF = false;
-FAutoConsoleVariableRef CVarUseContactGraphGBF(TEXT("p.Chaos.UseContactGraphGBF"), bUseContactGraphGBF, TEXT(""));
-
-void FPBDConstraintColor::ComputeIslandColoring(const int32 Island, const FPBDConstraintGraph& ConstraintGraph, uint32 ContainerId)
-{
-	SCOPE_CYCLE_COUNTER(STAT_Constraint_ComputeIslandColoring);
-
-	// We need to sort the particles for the coloring to be deterministic
-	// @todo(chaos): get rid of this sort and solve the problem at a higher level
-	//const TArray<TGeometryParticleHandle<FReal, 3>*>& IslandParticles = ConstraintGraph.GetIslandParticles(Island);
-	TArray<TGeometryParticleHandle<FReal, 3>*> IslandParticles = ConstraintGraph.GetIslandParticles(Island);
-	IslandParticles.Sort([](const FGeometryParticleHandle& A, const FGeometryParticleHandle& B)
-		{
-			return A.ParticleID() < B.ParticleID();
-		});
-
-	FLevelToColorToConstraintListMap& LevelToColorToConstraintListMap = IslandData[Island].LevelToColorToConstraintListMap;
-	int32& MaxColor = IslandData[Island].MaxColor;
-	int32& MaxEdges = IslandData[Island].NumEdges;
-
-	const int32 MaxLevel = IslandData[Island].MaxLevel;
-	
-	LevelToColorToConstraintListMap.Reset();
-	LevelToColorToConstraintListMap.SetNum(MaxLevel + 1);
-	MaxColor = -1;
-	MaxEdges = 0;
-
-	TSet<int32> ProcessedNodes;
-	TArray<int32> NodesToProcess;
-
-	for (const TGeometryParticleHandle<FReal, 3>* Particle :IslandParticles)
-	{
-		if (!ConstraintGraph.GetParticleNodes().Find(Particle))
-		{
-			continue;
-		}
-
-		const int32 ParticleNodeIndex = ConstraintGraph.GetParticleNodes()[Particle];
-
-		const bool bIsParticleDynamic = Particle->CastToRigidParticle() && Particle->ObjectState() == EObjectStateType::Dynamic;
-		if (ProcessedNodes.Contains(ParticleNodeIndex) || bIsParticleDynamic == false)
-		{
-			continue;
-		}
-
-		NodesToProcess.Add(ParticleNodeIndex);
-
-		while (NodesToProcess.Num())
-		{
-			const int32 NodeIndex = NodesToProcess.Last();
-			const typename FPBDConstraintGraph::FGraphNode& GraphNode = ConstraintGraph.GetGraphNodes()[NodeIndex];
-			FGraphNodeColor& ColorNode = Nodes[NodeIndex];
-
-			NodesToProcess.SetNum(NodesToProcess.Num() - 1, /*bAllowShrinking=*/false);
-			ProcessedNodes.Add(NodeIndex); 
-
-
-			for (const int32 EdgeIndex : GraphNode.NodeEdges)
-			{
-				const typename FPBDConstraintGraph::FGraphEdge& GraphEdge = ConstraintGraph.GetGraphEdges()[EdgeIndex];
-				FGraphEdgeColor& ColorEdge = Edges[EdgeIndex];
-
-				// If this is not from our rule, ignore it
-				if (GraphEdge.ItemContainer != ContainerId)
-				{
-					continue;
-				}
-
-				// If edge has been colored skip it
-				if (ColorEdge.Color >= 0)
-				{
-					continue;
-				}
-
-				// Get index to the other node on the edge
-				int32 OtherNodeIndex = INDEX_NONE;
-				if (GraphEdge.FirstNode == NodeIndex)
-				{
-					OtherNodeIndex = GraphEdge.SecondNode;
-				}
-				if (GraphEdge.SecondNode == NodeIndex)
-				{
-					OtherNodeIndex = GraphEdge.FirstNode;
-				}
-
-				// Find next color that is not used already at this node
-				while (ColorNode.UsedColors.Contains(ColorNode.NextColor))
-				{
-					ColorNode.NextColor++;
-				}
-				int32 ColorToUse = ColorNode.NextColor;
-
-				// Exclude colors used by the other node (but still allow this node to use them for other edges)
-				if (OtherNodeIndex != INDEX_NONE)
-				{
-					FGraphNodeColor& OtherColorNode = Nodes[OtherNodeIndex];
-
-					const typename FPBDConstraintGraph::FGraphNode& OtherGraphNode = ConstraintGraph.GetGraphNodes()[OtherNodeIndex];
-					const bool bIsOtherGraphNodeDynamic = OtherGraphNode.NodeItem->CastToRigidParticle() && OtherGraphNode.NodeItem->ObjectState() == EObjectStateType::Dynamic;
-					if (bIsOtherGraphNodeDynamic)
-					{
-						while (OtherColorNode.UsedColors.Contains(ColorToUse) || ColorNode.UsedColors.Contains(ColorToUse))
-						{
-							ColorToUse++;
-						}
-					}
-				}
-
-				// Assign color and set as used at this node
-				MaxColor = FMath::Max(ColorToUse, MaxColor);
-				ColorNode.UsedColors.Add(ColorToUse);
-				ColorEdge.Color = ColorToUse;
-
-				// Bump color to use next time, but only if we weren't forced to use a different color by the other node
-				if ((ColorToUse == ColorNode.NextColor) && bIsParticleDynamic == true)
-				{
-					ColorNode.NextColor++;
-				}
-
-				int32 Level = ColorEdge.Level;
-
-				if ((Level < 0) || (Level >= LevelToColorToConstraintListMap.Num()))
-				{
-					UE_LOG(LogChaos, Error, TEXT("\t **** Level is out of bounds!!!!  Level - %d, LevelToColorToConstraintListMap.Num() - %d"), Level, LevelToColorToConstraintListMap.Num());
-					continue;
-				}
-
-				if (!LevelToColorToConstraintListMap[Level].Contains(ColorEdge.Color))
-				{
-					LevelToColorToConstraintListMap[Level].Add(ColorEdge.Color, {});
-				}
-
-				LevelToColorToConstraintListMap[Level][ColorEdge.Color].Add(GraphEdge.EdgeItem);
-				++MaxEdges;
-
-				if (OtherNodeIndex != INDEX_NONE)
-				{
-					const typename FPBDConstraintGraph::FGraphNode& OtherGraphNode = ConstraintGraph.GetGraphNodes()[OtherNodeIndex];
-					const bool bIsOtherGraphNodeDynamic = OtherGraphNode.NodeItem->CastToRigidParticle() && OtherGraphNode.NodeItem->ObjectState() == EObjectStateType::Dynamic;
-					if (bIsOtherGraphNodeDynamic)
-					{
-						FGraphNodeColor& OtherColorNode = Nodes[OtherNodeIndex];
-
-						// Mark other node as not allowing use of this color
-						if (bIsParticleDynamic)
-						{
-							OtherColorNode.UsedColors.Add(ColorEdge.Color);
-						}
-
-						// Queue other node for processing
-						if (!ProcessedNodes.Contains(OtherNodeIndex))
-						{
-							ensure(OtherGraphNode.IslandIndex == GraphNode.IslandIndex);
-							checkSlow(IslandParticles.Find(OtherGraphNode.NodeItem) != INDEX_NONE);
-							NodesToProcess.Add(OtherNodeIndex);
-						}
-					}
-				}
-			}
-		}
-	}
-}
-
-void FPBDConstraintColor::ComputeContactGraph(const int32 Island, const FPBDConstraintGraph& ConstraintGraph, uint32 ContainerId)
-{
-	SCOPE_CYCLE_COUNTER(STAT_Constraint_ComputeContactGraph);
-	const TArray<FConstraintHandleHolder>& IslandConstraints = ConstraintGraph.GetIslandConstraints(Island);
-
-	IslandData[Island].MaxLevel = IslandConstraints.Num() ? 0 : -1;
-	
-	struct FLevelNodePair
-	{
-	FLevelNodePair()
-		: Level(INDEX_NONE)
-		, NodeIndex(INDEX_NONE)
-	{}
-	FLevelNodePair(int32 InLevel, int32 InNodeIndex)
-		: Level(InLevel)
-		, NodeIndex(InNodeIndex)
-	{}
-
-	int32 Level;
-	int32 NodeIndex;
-	};
-	TRingBuffer<FLevelNodePair> NodeQueue(100);
-
-	for(const TGeometryParticleHandle<FReal, 3> * Particle : ConstraintGraph.GetIslandParticles(Island))
-	{
-		const int32* NodeIndexPtr = ConstraintGraph.GetParticleNodes().Find(Particle);
-		const bool bIsParticleDynamic = Particle->CastToRigidParticle() && Particle->ObjectState() == EObjectStateType::Dynamic;
-
-		if (NodeIndexPtr)
-		{
-			// To be consistent with the solver body interface, we need to return a level given a particle. But internally we compute levels on nodes, therefore a mapping from nodes to particles is needed. 
-			if (bIsParticleDynamic)
-			{
-				NodeToParticle[*NodeIndexPtr] = Particle;
-			}
-		}
-		// We're only interested in static particles here to generate the graph (graph of dynamic objects touching static)
-		if(bIsParticleDynamic == false && NodeIndexPtr)
-		{
-			const int32 NodeIndex = *NodeIndexPtr;
-			const typename FPBDConstraintGraph::FGraphNode& GraphNode = ConstraintGraph.GetGraphNodes()[NodeIndex];
-			
-
-			for (int32 EdgeIndex : GraphNode.NodeEdges)
-			{
-				const typename FPBDConstraintGraph::FGraphEdge& GraphEdge = ConstraintGraph.GetGraphEdges()[EdgeIndex];
-
-				// If this is not from our rule, ignore it
-				if (GraphEdge.ItemContainer != ContainerId)
-				{
-					continue;
-				}
-
-				// Find adjacent node
-				int32 OtherNode = INDEX_NONE;
-				if (GraphEdge.FirstNode == NodeIndex)
-				{
-					OtherNode = GraphEdge.SecondNode;
-				}
-				else if (GraphEdge.SecondNode == NodeIndex)
-				{
-					OtherNode = GraphEdge.FirstNode;
-				}
-
-				// If we have a node, add it to the queue only if it matches our island. Statics have no island and can be touching dynamics of many islands
-				// so we need to pick out only the edges that lead to the requested island to correctly build the graph. Implicitly all further edges must
-				// be of the same island so we only need to do this check for level 1
-				if (OtherNode != INDEX_NONE && ConstraintGraph.GetGraphNodes()[OtherNode].IslandIndex == ConstraintGraph.GetGraphIndex(Island))
-				{
-					Edges[EdgeIndex].Level = 0;
-					NodeQueue.Emplace(1, OtherNode);
-				}
-			}
-		}
-	}
-
-	FLevelNodePair Current;
-	while(!NodeQueue.IsEmpty())
-	{
-		Current = NodeQueue.First();
-		NodeQueue.PopFrontNoCheck();
-
-		const int32 Level = Current.Level;
-		const int32 NodeIndex = Current.NodeIndex;
-		const TGeometryParticleHandle<FReal, 3> *Particle = NodeToParticle[NodeIndex];
-		if (Particle)
-		{
-			ParticleToLevel[Particle->UniqueIdx().Idx] = Level;
-		} 
-		const typename FPBDConstraintGraph::FGraphNode& GraphNode = ConstraintGraph.GetGraphNodes()[NodeIndex];
-		
-		for (int32 EdgeIndex : GraphNode.NodeEdges)
-		{
-			const typename FPBDConstraintGraph::FGraphEdge& GraphEdge = ConstraintGraph.GetGraphEdges()[EdgeIndex];
-			FGraphEdgeColor& ColorEdge = Edges[EdgeIndex];
-
-			// If this is not from our rule, ignore it
-			if (GraphEdge.ItemContainer != ContainerId)
-			{
-				continue;
-			}
-
-			// If we have already been assigned a level, move on
-			if (ColorEdge.Level >= 0)
-			{
-				continue;
-			}
-
-			// Find adjacent node and recurse
-			int32 OtherNode = INDEX_NONE;
-			if (GraphEdge.FirstNode == NodeIndex)
-			{
-				OtherNode = GraphEdge.SecondNode;
-			}
-			else if (GraphEdge.SecondNode == NodeIndex)
-			{
-				OtherNode = GraphEdge.FirstNode;
-			}
-
-			// Assign the level and update max level for the island if required
-			// NOTE: if we hit a non-dynamic particle (node), it will contain all of
-			// the contacts (edges) for dynamic particles interacting with it. They
-			// may not all be in the same island, which is ok (e.g., two separated
-			// boxes sat on a large plane). We need to ignore edges that are in other islands
-			// @todo(chaos): we should probably store the island index with each edge.
-			if (ConstraintGraph.GetGraphNodes()[OtherNode].IslandIndex == ConstraintGraph.GetGraphIndex(Island))
-			{
-				ColorEdge.Level = Level;
-				IslandData[Island].MaxLevel = FGenericPlatformMath::Max(IslandData[Island].MaxLevel, ColorEdge.Level);
-
-				// If we have an other node, append it to our queue on the next level
-				if (OtherNode != INDEX_NONE)
-				{
-					NodeQueue.Emplace(ColorEdge.Level + 1, OtherNode);
-				}
-			}
-		}
-	}
-
-	// An isolated island that is only dynamics will not have been processed above, put everything without a level into level zero
-	// #BGTODO this can surely be done as we build the edges, after this function everything will be at least level 0 so we can probably construct them
-	//         in that level to avoid a potentially large iteration here.
-	for(const FConstraintHandle* IslandConstraint : IslandConstraints)
-	{
-		if(IslandConstraint)
-		{
-			const int32 EdgeIndex = IslandConstraint->ConstraintGraphIndex();
-			if(Edges.IsValidIndex(EdgeIndex))
-			{
-				check(Edges[EdgeIndex].Level <= IslandData[Island].MaxLevel);
-				if(Edges[EdgeIndex].Level < 0)
-				{
-					Edges[EdgeIndex].Level = 0;
-				}
-			}
-		} 
-	}
-
-	check(IslandData[Island].MaxLevel >= 0 || !IslandConstraints.Num());
-}
-
-void FPBDConstraintColor::ComputeContactGraphGBF(const FReal Dt, const int32 Island, const FPBDConstraintGraph& ConstraintGraph, const uint32 ContainerId)
-{
-	SCOPE_CYCLE_COUNTER(STAT_Constraint_ComputeContactGraphGBF);
-	FGBFContactGraphData& Data = IslandToGBFContactGraphData[Island];
-	// First create an island-local graph. The mapping is stored in DynamicGlobalToIslandLocalNodeIndicesArray and KinematicGlobalToIslandLocalNodeIndices.
-	ComputeGlobalToIslandLocalNodeMapping(Island, ConstraintGraph, DynamicGlobalToIslandLocalNodeIndicesArray, Data.KinematicGlobalToIslandLocalNodeIndices);
-	CollectIslandDirectedEdges(Dt, Island, ConstraintGraph, DynamicGlobalToIslandLocalNodeIndicesArray, Data.KinematicGlobalToIslandLocalNodeIndices, Data.DigraphEdges);
-
-	// Build compressed sparse row (CSR) representation of the directed graph (digraph).
-	const int32 NumNodes=ConstraintGraph.GetIslandParticles(Island).Num();
-	BuildGraphCSR(Data.DigraphEdges, NumNodes, Data.ChildrenCSRIndices, Data.ChildrenCSRValues, Data.CSRCurrentIndices);
-
-	// Root nodes: nodes have no parents.
-	ComputeIsRootNode(Data.DigraphEdges, NumNodes, Data.IsRootNode);
-
-	// Collapse directed graph (digraph) to direct acyclic graph (DAG). This is a many-to-one mapping.
-	CollapseDigraphToDAG(Data.ChildrenCSRIndices, Data.ChildrenCSRValues, Data.IsRootNode, Data.DigraphToDAGIndices, Data.DAGToDigraphIndices, Data.Visited, Data.TraversalStack, Data.NumOnStack);
-
-	// Compute the DAG edges from the digraph edges 
-	ComputeDAGEdges(Data.DigraphEdges, Data.DigraphToDAGIndices, Data.DAGEdges);
-	BuildGraphCSR(Data.DAGEdges, NumNodes, Data.DAGChildrenCSRIndices, Data.DAGChildrenCSRValues, Data.CSRCurrentIndices);
-
-	// Topological sort on DAG nodes
-	TopologicalSortDAG(Data.DAGChildrenCSRIndices, Data.DAGChildrenCSRValues, Data.DAGToDigraphIndices, Data.SortedDAGNodes, Data.Visited);
-
-	AssignDAGNodeLevels(Data.DAGChildrenCSRIndices, Data.DAGChildrenCSRValues, Data.SortedDAGNodes, Data.DAGNodeLevels);
-	AssignEdgeLevels(Island, ConstraintGraph, DynamicGlobalToIslandLocalNodeIndicesArray, Data.KinematicGlobalToIslandLocalNodeIndices, Data.DigraphToDAGIndices, Data.DAGNodeLevels);
-	UpdateParticleToLevel(Island, ConstraintGraph, DynamicGlobalToIslandLocalNodeIndicesArray, Data.KinematicGlobalToIslandLocalNodeIndices, Data.DigraphToDAGIndices, Data.DAGNodeLevels);
-}
-
-void FPBDConstraintColor::ComputeGlobalToIslandLocalNodeMapping(const int32 Island, const FPBDConstraintGraph& ConstraintGraph, TArray<int32>& DynamicGlobalToIslandLocalNodeIndices, TMap<int32, int32>& KinematicGlobalToIslandLocalNodeIndices) const
-{
-	const TArray<FGeometryParticleHandle*>& IslandParticles = ConstraintGraph.GetIslandParticles(Island);
-	const int32 NumIslandParticles = IslandParticles.Num();
-	KinematicGlobalToIslandLocalNodeIndices.Reset();
-	for(int32 GraphNodeI = 0; GraphNodeI < NumIslandParticles; GraphNodeI++)
-	{
-		const TGeometryParticleHandle<FReal, 3>* Particle = IslandParticles[GraphNodeI];
-		const int32* NodeIndexPtr = ConstraintGraph.GetParticleNodes().Find(Particle);
-		if (NodeIndexPtr)
-		{
-			const int32 NodeIndex = *NodeIndexPtr;
-			const bool IsDynamic = ConstraintGraph.GetGraphNodes()[NodeIndex].IslandIndex != INDEX_NONE;
-			if (IsDynamic)
-			{
-				DynamicGlobalToIslandLocalNodeIndices[NodeIndex] = GraphNodeI;
-			}
-			else
-			{
-				KinematicGlobalToIslandLocalNodeIndices.FindOrAdd(NodeIndex) = GraphNodeI;
-			}
-		}
-	}
-}
-
-int32 FPBDConstraintColor::GetIslandLocalNodeIdx(const int32 GlobalNodeIndex, const TArray<int32>& DynamicGlobalToIslandLocalNodeIndices, const TMap<int32, int32>& KinematicGlobalToIslandLocalNodeIndices, const bool IsNodeDynamic) const
-{
-	if (IsNodeDynamic)
-	{
-		return DynamicGlobalToIslandLocalNodeIndices[GlobalNodeIndex];
-	}
-	else
-	{
-		return KinematicGlobalToIslandLocalNodeIndices[GlobalNodeIndex];
-	}
-}
-
-void FPBDConstraintColor::CollectIslandDirectedEdges(const FReal Dt, const int32 Island, const FPBDConstraintGraph& ConstraintGraph, const TArray<int32>& DynamicGlobalToIslandLocalNodeIndices, const TMap<int32, int32>& KinematicGlobalToIslandLocalNodeIndices, TArray<FDirectedEdge>& DirectedEdges) const
-{
-	DirectedEdges.Reset();
-	const int32 NumIslandParticles = ConstraintGraph.GetIslandParticles(Island).Num();
-
-	for (const FConstraintHandle* IslandConstraint : ConstraintGraph.GetIslandConstraints(Island))
-	{
-		if(IslandConstraint && ConstraintGraph.GetGraphEdges().IsValidIndex(IslandConstraint->ConstraintGraphIndex()))
-		{
-			const typename FPBDConstraintGraph::FGraphEdge& GraphEdge = ConstraintGraph.GetGraphEdges()[IslandConstraint->ConstraintGraphIndex()];
-
-			const FPBDCollisionConstraintHandle *CollisionHandle = GraphEdge.EdgeItem->As<FPBDCollisionConstraintHandle>();
-			if (CollisionHandle)
-			{
-				const bool IsFirstNodeDynamic = ConstraintGraph.GetGraphNodes()[GraphEdge.FirstNode].IslandIndex != INDEX_NONE;
-				const int32 FirstLocalNodeIndex = GetIslandLocalNodeIdx(GraphEdge.FirstNode, DynamicGlobalToIslandLocalNodeIndices, KinematicGlobalToIslandLocalNodeIndices, IsFirstNodeDynamic);
-				const bool IsSecondNodeDynamic = ConstraintGraph.GetGraphNodes()[GraphEdge.SecondNode].IslandIndex != INDEX_NONE;
-				const int32 SecondLocalNodeIndex = GetIslandLocalNodeIdx(GraphEdge.SecondNode, DynamicGlobalToIslandLocalNodeIndices, KinematicGlobalToIslandLocalNodeIndices, IsSecondNodeDynamic);
-
-				const ECollisionConstraintDirection ConstraintDirection = CollisionHandle->GetContact().GetConstraintDirection(Dt);
-				if(ConstraintDirection == ECollisionConstraintDirection::Particle0ToParticle1)
-				{
-					DirectedEdges.Push(FDirectedEdge(FirstLocalNodeIndex, SecondLocalNodeIndex));
-				}
-				else if(ConstraintDirection == ECollisionConstraintDirection::Particle1ToParticle0)
-				{
-					DirectedEdges.Push(FDirectedEdge(SecondLocalNodeIndex, FirstLocalNodeIndex));
-				}
-			}
-		}
-	}
-}
-
-/**
- * The children of node i is stored from CSRValues[CSRIndices[i]] to CSRValues[CSRIndices[i+1]]
- */
-void FPBDConstraintColor::BuildGraphCSR(const TArray<FDirectedEdge>& GraphEdges, const int32 NumNodes, TArray<int32>& CSRIndices, TArray<int32>& CSRValues, TArray<int32>& CurrentIndices) const
-{
-	CSRIndices.Init(0, NumNodes + 1);
-	if (NumNodes == 0)
-	{
-		return;
-	}
-	const int32 NumEdges = GraphEdges.Num();
-	CSRValues.Init(-1, NumEdges);
-	// Collect the number of children of each node and store them in CSRIndices[i+1]
-	for(const FDirectedEdge& Edge : GraphEdges)
-	{
-		CSRIndices[Edge.FirstNode + 1] += 1;
-	}
-	// CurrentIndices[i] is the first available vacancy for node i to populate its children
-	CurrentIndices.Init(0, NumNodes);
-	CurrentIndices[0] = 0;
-	if (NumNodes > 1)
-	{
-		CurrentIndices[1] = CSRIndices[1];
-		for(int32 i = 2; i < NumNodes; i++)
-		{
-			CSRIndices[i] += CSRIndices[i-1];
-			CurrentIndices[i] = CSRIndices[i];
-		}
-	}
-
-	// Populate edges
-	CSRIndices[NumNodes] = NumEdges;
-	for(const FDirectedEdge& Edge : GraphEdges)
-	{
-		CSRValues[CurrentIndices[Edge.FirstNode]] = Edge.SecondNode;
-		CurrentIndices[Edge.FirstNode]++;
-	}
-}
-
-void FPBDConstraintColor::ComputeIsRootNode(const TArray<FDirectedEdge>& GraphEdges, const int32 NumNodes, TArray<bool> &IsRootNode) const
-{
-	IsRootNode.SetNum(NumNodes);
-	for (bool& B : IsRootNode)
-	{
-		B = true;
-	}
-	for(const FDirectedEdge& Edge : GraphEdges)
-	{
-		if (IsRootNode[Edge.SecondNode])
-		{
-			IsRootNode[Edge.SecondNode] = false;
-		}
-	}
-}
-
-void FPBDConstraintColor::CollapseDigraphToDAG(const TArray<int32>& ChildrenCSRIndices, const TArray<int32>& ChildrenCSRValues, const TArray<bool>& IsRootNode, TArray<int32>& DigraphToDAGIndices, TArray<TArray<int32>>& DAGToDigraphIndices, TArray<bool>& Visited, TArray<int32>& TraversalStack, TArray<int32>& NumOnStack) const
-{
-	const int32 NumDigraphNodes = ChildrenCSRIndices.Num() - 1;
-	DigraphToDAGIndices.SetNum(NumDigraphNodes);
-	DAGToDigraphIndices.SetNum(NumDigraphNodes);
-
-	// Initialize identity mapping between Digraph and DAG
-	for(int32 i = 0; i < NumDigraphNodes; i++)
-	{
-		DigraphToDAGIndices[i] = i;
-		DAGToDigraphIndices[i].SetNum(1);
-		DAGToDigraphIndices[i][0] = i; 
-	}
-
-	// If Visited[**DigraphI**] == true, the node DigraphI and its children are processed and will not be traversed in the future.
-	Visited.Init(false, NumDigraphNodes);
-
-	// Every **Digraph** node will be pushed to TraversalStack before they are processed and will be popped from the stack after the process finished.
-	TraversalStack.Reserve(64);
-	TraversalStack.Reset();
-
-	// NumOnStack[**DAGI**] is the number of Digraph nodes that are collapsed to DAGI and are currently on TraversalStack.
-	NumOnStack.Init(0, NumDigraphNodes);
-
-
-	TFunctionRef<void(const int32)> CollapseLoop = [&DigraphToDAGIndices, &DAGToDigraphIndices, &TraversalStack, &NumOnStack](const int32 DigraphI)
-	{
-		const int32 DAGI = DigraphToDAGIndices[DigraphI];
-		// Go through the traversal stack in reverse order until we find another node that is pointed to DAGI. Collapse all the nodes traversed this way to DAGI.
-		// We collapse these nodes to DAGI, but we don't remove them from the traversal stack. We will still need to visit their children later.
-		for(int32 j = TraversalStack.Num() - 1; DigraphToDAGIndices[TraversalStack[j]] != DAGI; j--)
-		{
-			const int32 DigraphJ = TraversalStack[j];
-			const int32 DAGJ = DigraphToDAGIndices[DigraphJ];
-			// Merge DAGI and DAGJ. Move all the nodes in DAGToDigraphIndices[DAGJ] to DAGToDigraphIndices[DAGI].
-			const int32 Num = DAGToDigraphIndices[DAGJ].Num();
-			if (Num > 0)
-			{
-				for(int32 k = 0; k < Num; k++)
-				{
-					const int32 DigraphK = DAGToDigraphIndices[DAGJ][k];
-					DigraphToDAGIndices[DigraphK] = DAGI;
-					DAGToDigraphIndices[DAGI].Push(DigraphK);
-				}
-				DAGToDigraphIndices[DAGJ].Empty();
-				NumOnStack[DAGJ] -= Num;
-				NumOnStack[DAGI] += Num;
-			}
-		}
-	};
-
-	// For simplicity, traverse on **Digraph** nodes so we don't need to change graph edges as we collapsed digraph to DAG.
-	TFunctionRef<void(const int32)> TraverseAndCollapse = [&DigraphToDAGIndices, &NumOnStack, &CollapseLoop, &TraversalStack, &Visited, &ChildrenCSRIndices, &ChildrenCSRValues, &TraverseAndCollapse](const int32 DigraphI)
-	{
-		int32 DAGI = DigraphToDAGIndices[DigraphI];
-		// If there are already Digraph nodes that point to DAGI on the stack, collapse the loop.
-		if (NumOnStack[DAGI] > 0)
-		{
-			CollapseLoop(DigraphI);
-		}
-		else
-		{
-			const int32 NumEdges = ChildrenCSRIndices[DigraphI + 1] - ChildrenCSRIndices[DigraphI];
-			// For non-leaf nodes
-			if (NumEdges > 0)
-			{
-				TraversalStack.Push(DigraphI);
-				NumOnStack[DAGI]++;
-				for(int32 j = ChildrenCSRIndices[DigraphI]; j < ChildrenCSRIndices[DigraphI + 1]; j++)
-				{
-					const int32 DigraphJ = ChildrenCSRValues[j];
-					if (Visited[DigraphJ])
-					{
-						continue;
-					}
-					TraverseAndCollapse(DigraphJ);
-				}
-				// Since there might have been collapsing of the loops and thereby the mapped DAG node might be modified, DAGI needs to be updated here.
-				DAGI = DigraphToDAGIndices[DigraphI];
-				// Popping from the stack
-				NumOnStack[DAGI]--;
-				TraversalStack.Pop(false/*bAllowShrinking*/); 
-			}
-			Visited[DigraphI]=true;
-		}
-	};
-
-	// First traverse all the root nodes that have no parents.
-	for (int32 DigraphI = 0; DigraphI < NumDigraphNodes; DigraphI++)
-	{
-		if (IsRootNode[DigraphI])
-		{
-			TraverseAndCollapse(DigraphI);
-		}
-	}
-
-	// Next traverse the rest loops.
-	for(int32 DigraphI = 0; DigraphI < NumDigraphNodes; DigraphI++)
-	{
-		if(!Visited[DigraphI])
-		{
-			TraverseAndCollapse(DigraphI);
-		}
-	}
-}
-
-void FPBDConstraintColor::ComputeDAGEdges(const TArray<FDirectedEdge>& DigraphEdges, const TArray<int32>& DigraphToDAGIndices, TArray<FDirectedEdge> &DAGEdges) const
-{
-	DAGEdges.Reset(DigraphEdges.Num());
-	for (const FDirectedEdge &Edge : DigraphEdges)
-	{
-		const int32 DAGFirst = DigraphToDAGIndices[Edge.FirstNode];
-		const int32 DAGSecond = DigraphToDAGIndices[Edge.SecondNode];
-		// Remove self loops
-		if (DAGFirst != DAGSecond)
-		{
-			DAGEdges.Push({DAGFirst, DAGSecond});
-		}
-	}
-}
-
-void FPBDConstraintColor::TopologicalSortDAG(const TArray<int32>& DAGChildrenCSRIndices, const TArray<int32>& DAGChildrenCSRValues, const TArray<TArray<int32>>& DAGToDigraphIndices, TArray<int32>& SortedDAGNodes, TArray<bool>& Visited) const
-{
-	const int32 NumNodes = DAGToDigraphIndices.Num();
-	SortedDAGNodes.Reset(NumNodes);
-	Visited.Init(false, NumNodes);
-	TFunctionRef<void(const int32)> Traverse = [&DAGChildrenCSRIndices, &DAGChildrenCSRValues, &SortedDAGNodes, &Visited, &Traverse](const int32 DAGNodeI)
-	{
-		Visited[DAGNodeI] = true;
-		for (int32 i = DAGChildrenCSRIndices[DAGNodeI]; i < DAGChildrenCSRIndices[DAGNodeI + 1]; i++)
-		{
-			const int32 Child = DAGChildrenCSRValues[i];
-			if (!Visited[Child])
-			{
-				Traverse(Child);
-			}
-		}
-		// Push DAGNodeI after all its children are pushed.
-		SortedDAGNodes.Push(DAGNodeI);
-	};
-	for(int32 i = 0; i < NumNodes; i++)
-	{
-		const bool IsValidDAGNode = DAGToDigraphIndices[i].Num() > 0; 
-		if (!IsValidDAGNode)
-		{
-			continue;
-		}
-		if (Visited[i])
-		{
-			continue;
-		}
-		Traverse(i);
-	}
-}
-
-void FPBDConstraintColor::AssignDAGNodeLevels(const TArray<int32>& DAGChildrenCSRIndices, const TArray<int32>& DAGChildrenCSRValues, const TArray<int32>& SortedDAGNodes, TArray<int32> &DAGNodeLevels) const
-{
-	const int32 NumNodes = DAGChildrenCSRIndices.Num() - 1;
-	DAGNodeLevels.Init(0, NumNodes);
-	for (int32 i = SortedDAGNodes.Num() - 1; i >= 0; i--)
-	{
-		const int32 DAGNodeI = SortedDAGNodes[i];
-		for (int32 j = DAGChildrenCSRIndices[DAGNodeI]; j < DAGChildrenCSRIndices[DAGNodeI + 1]; j++)
-		{
-			const int32 Child = DAGChildrenCSRValues[j];
-			// Level of DAGI is the max distance to root nodes
-			DAGNodeLevels[Child] = FMath::Max(DAGNodeLevels[Child], DAGNodeLevels[DAGNodeI] + 1);
-		}
-	}
-}
-
-void FPBDConstraintColor::AssignEdgeLevels(const int32 Island, const FPBDConstraintGraph& ConstraintGraph, const TArray<int32>& DynamicGlobalToIslandLocalNodeIndices, const TMap<int32, int32>& KinematicGlobalToIslandLocalNodeIndices, const TArray<int32>& DigraphToDAGIndices, const TArray<int32> DAGNodeLevels)
-{
-	for(const FConstraintHandle* IslandConstraint : ConstraintGraph.GetIslandConstraints(Island))
-	{
-		if(IslandConstraint && ConstraintGraph.GetGraphEdges().IsValidIndex(IslandConstraint->ConstraintGraphIndex()))
-		{
-			const int32 EdgeIndex = IslandConstraint->ConstraintGraphIndex();
-			const typename FPBDConstraintGraph::FGraphEdge& GraphEdge = ConstraintGraph.GetGraphEdges()[EdgeIndex];
-
-			const bool IsFirstNodeDynamic = ConstraintGraph.GetGraphNodes()[GraphEdge.FirstNode].IslandIndex != INDEX_NONE;
-			const int32 FirstLocalNodeIndex = GetIslandLocalNodeIdx(GraphEdge.FirstNode, DynamicGlobalToIslandLocalNodeIndices, KinematicGlobalToIslandLocalNodeIndices, IsFirstNodeDynamic);
-			const int32 FirstLevel = DAGNodeLevels[DigraphToDAGIndices[FirstLocalNodeIndex]];
-			const bool IsSecondNodeDynamic = ConstraintGraph.GetGraphNodes()[GraphEdge.SecondNode].IslandIndex != INDEX_NONE;
-			const int32 SecondLocalNodeIndex = GetIslandLocalNodeIdx(GraphEdge.SecondNode, DynamicGlobalToIslandLocalNodeIndices, KinematicGlobalToIslandLocalNodeIndices, IsSecondNodeDynamic);
-			const int32 SecondLevel = DAGNodeLevels[DigraphToDAGIndices[SecondLocalNodeIndex]];
-
-			// Edge level is the max of node level
-			const int32 EdgeLevel = FMath::Max(FirstLevel, SecondLevel);
-			Edges[EdgeIndex].Level = EdgeLevel;
-			IslandData[Island].MaxLevel = FMath::Max(IslandData[Island].MaxLevel, EdgeLevel);
-		}
-	}
-}
-
-void FPBDConstraintColor::UpdateParticleToLevel(const int32 Island, const FPBDConstraintGraph& ConstraintGraph, const TArray<int32>& DynamicGlobalToIslandLocalNodeIndices, const TMap<int32, int32>& KinematicGlobalToIslandLocalNodeIndices, const TArray<int32>& DigraphToDAGIndices, const TArray<int32>& DAGNodeLevels)
-{
-	for(const TGeometryParticleHandle<FReal, 3> * Particle : ConstraintGraph.GetIslandParticles(Island))
-	{
-		const int32* NodeIndexPtr = ConstraintGraph.GetParticleNodes().Find(Particle);
-		if (NodeIndexPtr)
-		{
-			const int32 NodeIndex = *NodeIndexPtr;
-			const int32 ParticleIdx = Particle->UniqueIdx().Idx;
-			const typename FPBDConstraintGraph::FGraphNode& GraphNode = ConstraintGraph.GetGraphNodes()[NodeIndex];
-			const bool IsNodeDynamic = GraphNode.IslandIndex != INDEX_NONE;
-			if (IsNodeDynamic)
-			{
-				const int32 IslandDigraphI = GetIslandLocalNodeIdx(NodeIndex, DynamicGlobalToIslandLocalNodeIndices, KinematicGlobalToIslandLocalNodeIndices, IsNodeDynamic);
-				const int32 Level = DAGNodeLevels[DigraphToDAGIndices[IslandDigraphI]];
-				ParticleToLevel[ParticleIdx] = Level;	
-			}
-			else
-			{
-				// Set kinematic particle levels to 0
-				ParticleToLevel[ParticleIdx] = 0;
-			}
-		}
-	}
-}
-
-int32 FPBDConstraintColor::GetParticleLevel(const FGeometryParticleHandle* ParticleHandle) const
-{
-	// todo(chaos) the index check should nopt be necessarybut right ParticleToLevel is not as larghe as the largest uniqueIdx in the graph( because of MaxParticleIndex not reflecting this )
-	const int32 UniqueIdx = ParticleHandle->UniqueIdx().Idx;
-	if (ParticleToLevel.IsValidIndex(UniqueIdx))
-	{
-		return ParticleToLevel[ParticleHandle->UniqueIdx().Idx];
-	}
-	return 0;
-}
-
-void FPBDConstraintColor::InitializeColor(const FPBDConstraintGraph& ConstraintGraph)
-{
-	// The Number of nodes is large and fairly constant so persist rather than resetting every frame
-	if (Nodes.Num() != ConstraintGraph.GetGraphNodes().GetMaxIndex())
-	{
-		// Nodes need to grow when the nodes of the constraint graph grows
-		Nodes.AddDefaulted(ConstraintGraph.GetGraphNodes().GetMaxIndex() - Nodes.Num());
-	}
-
-	for (auto& Node : Nodes)
-	{
-		Node.NextColor = 0;
-		Node.UsedColors.Empty();
-	}
-
-	// edges are not persistent right now so we still reset them
-	Edges.Reset();
-	IslandData.Reset();
-
-	Edges.SetNum(ConstraintGraph.GetGraphEdges().GetMaxIndex());
-	IslandData.SetNum(ConstraintGraph.NumIslands());
-	ParticleToLevel.Init(0, ConstraintGraph.GetMaxParticleUniqueIdx() + 1);
-	DynamicGlobalToIslandLocalNodeIndicesArray.Init(-1, ConstraintGraph.GetIslandGraph()->MaxNumNodes());
-	NodeToParticle.Init(nullptr, ConstraintGraph.GetIslandGraph()->MaxNumNodes());
-	IslandToGBFContactGraphData.SetNum(ConstraintGraph.NumIslands());
-}
-
-void FPBDConstraintColor::ComputeColor(const FReal Dt, const int32 Island, const FPBDConstraintGraph& ConstraintGraph, uint32 ContainerId)
-{
-	SCOPE_CYCLE_COUNTER(STAT_Constraint_ComputeColor);
-	if (bUseContactGraph)
-	{
-		if (bUseContactGraphGBF)
-		{
-			// Hack: extra cautious to make sure ComputeContactGraphGBF is only called once per frame.
-			// Ideally this would be ContainerType == CollisionConstraints
-			if (ContainerId == 0) 
-			{
-				ComputeContactGraphGBF(Dt, Island, ConstraintGraph, ContainerId);
-			}
-		}
-		else
-		{
-			ComputeContactGraph(Island, ConstraintGraph, ContainerId);
-		}
-	}
-	else
-	{
-		for (FGraphEdgeColor& Edge : Edges)
-		{
-			Edge.Level = 0;
-		}
-		IslandData[Island].MaxLevel = 0;
-	}
-	ComputeIslandColoring(Island, ConstraintGraph, ContainerId);
-}
-
-const typename FPBDConstraintColor::FLevelToColorToConstraintListMap& FPBDConstraintColor::GetIslandLevelToColorToConstraintListMap(int32 Island) const
-{
-	if (Island < IslandData.Num())
-	{
-		return IslandData[Island].LevelToColorToConstraintListMap;
-	}
-	return EmptyLevelToColorToConstraintListMap;
-}
-
-int FPBDConstraintColor::GetIslandMaxColor(int32 Island) const
-{
-	if (Island < IslandData.Num())
-	{
-		return IslandData[Island].MaxColor;
-	}
-	return -1;
-}
-
-int FPBDConstraintColor::GetIslandMaxLevel(int32 Island) const
-{
-	if (Island < IslandData.Num())
-	{
-		return IslandData[Island].MaxLevel;
-	}
-	return -1;
-}
+//
+// TO BE REMOVED
+//
\ No newline at end of file
diff --git a/Engine/Source/Runtime/Experimental/Chaos/Private/Chaos/PBDConstraintContainer.cpp b/Engine/Source/Runtime/Experimental/Chaos/Private/Chaos/PBDConstraintContainer.cpp
--- a/Engine/Source/Runtime/Experimental/Chaos/Private/Chaos/PBDConstraintContainer.cpp
+++ b/Engine/Source/Runtime/Experimental/Chaos/Private/Chaos/PBDConstraintContainer.cpp
@@ -1,5 +1,6 @@
 // Copyright Epic Games, Inc. All Rights Reserved.
 #include "Chaos/PBDConstraintContainer.h"
+#include "Chaos/Evolution/SolverConstraintContainer.h"
 
 namespace Chaos
 {
@@ -13,13 +14,4 @@ namespace Chaos
 	{
 	}
 
-	int32 FPBDIndexedConstraintContainer::GetConstraintIndex(const FIndexedConstraintHandle* ConstraintHandle) const
-	{
-		return ConstraintHandle->GetConstraintIndex();
-	}
-
-	void FPBDIndexedConstraintContainer::SetConstraintIndex(FIndexedConstraintHandle* ConstraintHandle, int32 ConstraintIndex) const
-	{
-		ConstraintHandle->ConstraintIndex = ConstraintIndex;
-	}
 }
diff --git a/Engine/Source/Runtime/Experimental/Chaos/Private/Chaos/PBDConstraintRule.cpp b/Engine/Source/Runtime/Experimental/Chaos/Private/Chaos/PBDConstraintRule.cpp
--- a/Engine/Source/Runtime/Experimental/Chaos/Private/Chaos/PBDConstraintRule.cpp
+++ b/Engine/Source/Runtime/Experimental/Chaos/Private/Chaos/PBDConstraintRule.cpp
@@ -2,676 +2,6 @@
 
 #include "Chaos/PBDConstraintRule.h"
 
-#include "Chaos/Island/SolverIsland.h"
-#include "Chaos/PBDCollisionConstraints.h"
-#include "Chaos/PBDJointConstraints.h"
-#include "Chaos/PBDNullConstraints.h"
-#include "Chaos/PBDPositionConstraints.h"
-#include "Chaos/PBDSuspensionConstraints.h"
-#include "Chaos/PBDRigidDynamicSpringConstraints.h"
-#include "Chaos/PBDRigidSpringConstraints.h"
-
-//PRAGMA_DISABLE_OPTIMIZATION
-
-namespace Chaos
-{
-	int32 ChaosShockPropagationVelocityPerLevelIterations = 1;
-	int32 ChaosShockPropagationPositionPerLevelIterations = 1;
-	FAutoConsoleVariableRef CVarChaosShockPropagationPositionPerLevelIterations(TEXT("p.Chaos.ShockPropagation.Position.PerLevelIterations"), ChaosShockPropagationPositionPerLevelIterations, TEXT(""));
-	FAutoConsoleVariableRef CVarChaosShockPropagationVelocityPerLevelIterations(TEXT("p.Chaos.ShockPropagation.Velocity.PerLevelIterations"), ChaosShockPropagationVelocityPerLevelIterations, TEXT(""));
-
-	int32 ChaosCollisionColorMinParticles = 2000;
-	FAutoConsoleVariableRef CVarChaosCollisionColorMinParticles(TEXT("p.Chaos.Collision.Color.MinParticles"), ChaosCollisionColorMinParticles, TEXT(""));
-
-	/** Console variable to disable the levels computation on each island */
-	int32 ChaosDisableIslandLevels = false;
-	FAutoConsoleVariableRef CVarChaosDisableIslandLevels(TEXT("p.Chaos.Islands.DisableLevels"), ChaosDisableIslandLevels, TEXT(""));
-
-	/** console variable to disable the colors computation on each island */
-	int32 ChaosDisableIslandColors = true;
-	FAutoConsoleVariableRef CVarChaosDisableIslandColors(TEXT("p.Chaos.Islands.DisableColors"), ChaosDisableIslandColors, TEXT(""));
-	
-	DECLARE_CYCLE_STAT(TEXT("FPBDConstraintRule::Island_Gather"), STAT_Constraints_Island_Gather, STATGROUP_Chaos);
-	DECLARE_CYCLE_STAT(TEXT("FPBDConstraintRule::Color_Gather"), STAT_Constraints_Color_Gather, STATGROUP_Chaos);
-	DECLARE_CYCLE_STAT(TEXT("FPBDConstraintRule::AddToGraph"), STAT_Constraints_AddToGraph, STATGROUP_Chaos);
-
-	template<class ConstraintType>
-	TSimpleConstraintRule<ConstraintType>::TSimpleConstraintRule(int32 InPriority, FConstraints& InConstraints)
-		: FSimpleConstraintRule(InPriority)
-		, Constraints(InConstraints)
-	{
-	}
-
-	template<class ConstraintType>
-	TSimpleConstraintRule<ConstraintType>::~TSimpleConstraintRule()
-	{
-	}
-
-	template<class ConstraintType>
-	void TSimpleConstraintRule<ConstraintType>::PrepareTick()
-	{
-		Constraints.PrepareTick();
-	}
-
-	template<class ConstraintType>
-	void TSimpleConstraintRule<ConstraintType>::UnprepareTick()
-	{
-		Constraints.UnprepareTick();
-	}
-
-	template<class ConstraintType>
-	void TSimpleConstraintRule<ConstraintType>::UpdatePositionBasedState(const FReal Dt)
-	{
-		return Constraints.UpdatePositionBasedState(Dt);
-	}
-
-	template<class ConstraintType>
-	void TSimpleConstraintRule<ConstraintType>::BindToDatas(FPBDIslandSolverData& InSolverDatas, const uint32 InContainerId)
-	{
-		Constraints.SetContainerId(InContainerId);
-
-		SolverData = &InSolverDatas;
-		if (SolverData != nullptr)
-		{
-			SolverData->template AddConstraintDatas<ConstraintType>(Constraints.GetContainerId());
-		}
-	}
-
-	template<class ConstraintType>
-	void TSimpleConstraintRule<ConstraintType>::GatherSolverInput(const FReal Dt)
-	{
-		if(SolverData)
-		{
-			SolverData->GetConstraintIndex(Constraints.GetContainerId()) = 0;
-			Constraints.SetNumIslandConstraints(Constraints.NumConstraints(), *SolverData);
-			Constraints.PreGatherInput(Dt, *SolverData);
-			Constraints.GatherInput(Dt, *SolverData);
-		}
-	}
-
-	template<class ConstraintType>
-	void TSimpleConstraintRule<ConstraintType>::ScatterSolverOutput(const FReal Dt)
-	{
-		if(SolverData) Constraints.ScatterOutput(Dt, *SolverData);
-	}
-
-	template<class ConstraintType>
-	bool TSimpleConstraintRule<ConstraintType>::ApplyConstraints(const FReal Dt, const int32 It, const int32 NumIts)
-	{
-		return SolverData ? Constraints.ApplyPhase1(Dt, It, NumIts, *SolverData) : false;
-	}
-
-	template<class ConstraintType>
-	bool TSimpleConstraintRule<ConstraintType>::ApplyPushOut(const FReal Dt, const int32 It, const int32 NumIts)
-	{
-		return SolverData ? Constraints.ApplyPhase2(Dt, It, NumIts, *SolverData) : false;
-	}
-
-	template<class ConstraintType>
-	bool TSimpleConstraintRule<ConstraintType>::ApplyProjection(const FReal Dt, const int32 It, const int32 NumIts)
-	{
-		return SolverData ? Constraints.ApplyPhase3(Dt, It, NumIts, *SolverData) : false;
-	}
-
-	template<class ConstraintType>
-	TPBDConstraintGraphRuleImpl<ConstraintType>::TPBDConstraintGraphRuleImpl(FConstraints& InConstraints, int32 InPriority)
-		: FPBDConstraintGraphRule(InPriority)
-		, Constraints(InConstraints)
-		, ConstraintGraph(nullptr)
-	{
-	}
-	
-	template<class ConstraintType>
-	void TPBDConstraintGraphRuleImpl<ConstraintType>::BindToGraph(FPBDConstraintGraph& InContactGraph, uint32 InContainerId)
-	{
-		Constraints.SetContainerId(InContainerId);
-		ConstraintGraph = &InContactGraph;
-	}
-
-	template<class ConstraintType>
-	void TPBDConstraintGraphRuleImpl<ConstraintType>::UpdatePositionBasedState(const FReal Dt)
-	{
-		Constraints.UpdatePositionBasedState(Dt);
-	}
-	
-	template<class ConstraintType>
-	void TPBDConstraintGraphRuleImpl<ConstraintType>::AddToGraph()
-	{
-		SCOPE_CYCLE_COUNTER(STAT_Constraints_AddToGraph);
-		ConstraintGraph->ReserveConstraints(Constraints.NumConstraints());
-
-		for (typename FConstraints::FConstraintContainerHandle * ConstraintHandle : Constraints.GetConstraintHandles())
-		{
-			// NOTE: For now, we include probe constraints as well so that
-			// they can trigger collision events
-			if (ConstraintHandle->IsEnabled())
-			{
-				bool DisabledParticle0 = (ConstraintHandle->GetConstrainedParticles()[0] && ConstraintHandle->GetConstrainedParticles()[0]->CastToRigidParticle()) ? ConstraintHandle->GetConstrainedParticles()[0]->CastToRigidParticle()->Disabled() : false;
-				bool DisabledParticle1 = (ConstraintHandle->GetConstrainedParticles()[1] && ConstraintHandle->GetConstrainedParticles()[1]->CastToRigidParticle()) ? ConstraintHandle->GetConstrainedParticles()[1]->CastToRigidParticle()->Disabled() : false;
-
-				if (!DisabledParticle0 && !DisabledParticle1)
-				{
-					ConstraintGraph->AddConstraint(GetContainerId(), ConstraintHandle, ConstraintHandle->GetConstrainedParticles());
-				}
-			}
-		}
-	}
-
-	template<class ConstraintType>
-	TPBDConstraintIslandRule<ConstraintType>::TPBDConstraintIslandRule(FConstraints& InConstraints, int32 InPriority)
-		: TPBDConstraintGraphRuleImpl<ConstraintType>(InConstraints, InPriority)
-	{
-	}
-
-	template<class ConstraintType>
-	TPBDConstraintIslandRule<ConstraintType>::~TPBDConstraintIslandRule()
-	{
-	}
-
-	template<class ConstraintType>
-	void TPBDConstraintIslandRule<ConstraintType>::GatherSolverInput(const FReal Dt, int32 GroupIndex)
-	{
-		SCOPE_CYCLE_COUNTER(STAT_Constraints_Island_Gather);
-		if(FPBDIslandGroup* IslandGroup = ConstraintGraph->GetIslandGroup(GroupIndex))
-		{
-			// This will reset the number of constraints inside the solver datas. For now we keep this function since according
-			// to the constraints we can use the handles, the indices or the container. when only the container will be used
-			// we can replace this call by :
-			// IslandSolver->template GetConstraintContainer<typename ConstraintType::FSolverConstraintContainerType>(ContainerId)->Reset();
-			Constraints.SetNumIslandConstraints(IslandGroup->ConstraintCount(Constraints.GetContainerId()), *IslandGroup);
-			IslandGroup->InitConstraintIndex(Constraints.GetContainerId());
-			
-			for (FPBDIslandSolver* IslandSolver : IslandGroup->GetIslands())
-			{
-				if(!IslandSolver->IsSleeping() && !IslandSolver->IsUsingCache())
-				{
-					const TArray<FConstraintHandleHolder>& IslandConstraints = ConstraintGraph->GetIslandConstraints(IslandSolver->GetIslandIndex());
-					
-					if (IslandConstraints.Num())
-					{
-						for (int32 Index = 0; Index < IslandConstraints.Num(); ++Index)
-						{
-							FConstraintHandle* ConstraintHandle = IslandConstraints[Index];
-							if (ConstraintHandle->GetContainerId() == GetContainerId())
-							{
-								FConstraintContainerHandle* Constraint = ConstraintHandle->As<FConstraintContainerHandle>();
-
-								// Note we are building the SolverBodies as we go, in the order that we visit them. Each constraint
-								// references two bodies, so we won't strictly be accessing only in cache order, but it's about as good as it can be.
-								if (Constraint->IsEnabled() && !Constraint->IsProbe())
-								{
-									// @todo(chaos): we should provide Particle Levels in the island rule as well (see TPBDConstraintColorRule)
-									Constraint->PreGatherInput(Dt, *IslandGroup);
-								}
-							}
-						}
-
-					        auto Gather = [&](const int32 Index)
-						{
-							FConstraintHandle* ConstraintHandle = IslandConstraints[Index];
-							if (ConstraintHandle->GetContainerId() == GetContainerId())
-							{
-								FConstraintContainerHandle* Constraint = ConstraintHandle->As<FConstraintContainerHandle>();
-
-								// Note we are building the SolverBodies as we go, in the order that we visit them. Each constraint
-								// references two bodies, so we won't strictly be accessing only in cache order, but it's about as good as it can be.
-								if (Constraint->IsEnabled() && !Constraint->IsProbe())
-								{
-									// @todo(chaos): we should provide Particle Levels in the island rule as well (see TPBDConstraintColorRule)
-									Constraint->GatherInput(Dt, INDEX_NONE, INDEX_NONE, *IslandGroup);
-								}
-							}
-						};
-						for (int32 Index = 0; Index < IslandConstraints.Num(); ++Index)
-						{
-							Gather(Index);
-						}
-						
-					}
-				}
-			}
-		}
-	}
-
-	template<class ConstraintType>
-	void TPBDConstraintIslandRule<ConstraintType>::ScatterSolverOutput(const FReal Dt, int32 GroupIndex)
-	{
-		if(FPBDIslandGroup* IslandGroup = ConstraintGraph->GetIslandGroup(GroupIndex))
-		{
-			Constraints.ScatterOutput(Dt, *IslandGroup);
-		}
-	}
-	
-	template<class ConstraintType>
-	bool TPBDConstraintIslandRule<ConstraintType>::ApplyConstraints(const FReal Dt, int32 GroupIndex, const int32 It, const int32 NumIts)
-	{
-		FPBDIslandGroup* IslandGroup = ConstraintGraph->GetIslandGroup(GroupIndex);
-		return IslandGroup ? Constraints.ApplyPhase1Serial(Dt, It, NumIts, *IslandGroup) : false;
-	}
-
-	template<class ConstraintType>
-	bool TPBDConstraintIslandRule<ConstraintType>::ApplyPushOut(const FReal Dt, int32 GroupIndex, const int32 It, const int32 NumIts)
-	{
-		FPBDIslandGroup* IslandGroup = ConstraintGraph->GetIslandGroup(GroupIndex);
-		return IslandGroup  ? Constraints.ApplyPhase2Serial(Dt, It, NumIts, *IslandGroup) : false;
-	}
-
-	template<class ConstraintType>
-	bool TPBDConstraintIslandRule<ConstraintType>::ApplyProjection(const FReal Dt, int32 GroupIndex, const int32 It, const int32 NumIts)
-	{
-		FPBDIslandGroup* IslandGroup = ConstraintGraph->GetIslandGroup(GroupIndex);
-		return IslandGroup ? Constraints.ApplyPhase3Serial(Dt, It, NumIts, *IslandGroup) : false;
-	}
-
-	template<class ConstraintType>
-	void TPBDConstraintIslandRule<ConstraintType>::InitializeAccelerationStructures()
-	{
-		ConstraintGraph->template AddConstraintDatas<ConstraintType>(Constraints.GetContainerId());
-	}
-
-	template<class ConstraintType>
-	void TPBDConstraintIslandRule<ConstraintType>::UpdateAccelerationStructures(const FReal Dt, const int32 GroupIndex)
-	{
-	}
-	
-	template<class ConstraintType>
-	int32 TPBDConstraintGraphRuleImpl<ConstraintType>::NumConstraints() const
-	{ 
-		return Constraints.NumConstraints(); 
-	}
-
-	template<class ConstraintType>
-	TPBDConstraintColorRule<ConstraintType>::TPBDConstraintColorRule(FConstraints& InConstraints, int32 InPriority)
-		: TPBDConstraintIslandRule<ConstraintType>(InConstraints, InPriority)
-		, ConstraintSets()
-	{
-	}
-
-	template<class ConstraintType>
-	TPBDConstraintColorRule<ConstraintType>::~TPBDConstraintColorRule()
-	{
-	}
-
-	template<class ConstraintType>
-	void TPBDConstraintColorRule<ConstraintType>::GatherSolverInput(const FReal Dt, int32 GroupIndex)
-	{
-		SCOPE_CYCLE_COUNTER(STAT_Constraints_Color_Gather);
-		FPBDConstraintGraph* LocalGraph = ConstraintGraph;
-		auto GetParticleLevel = [LocalGraph](TGeometryParticleHandle<FReal, 3>* ConstrainedParticle) -> int32
-		{
-			int32 ParticleLevel = INDEX_NONE;
-			if(FPBDRigidParticleHandle* PBDRigid = ConstrainedParticle->CastToRigidParticle())
-			{
-				if(LocalGraph->GetIslandGraph()->GraphNodes.IsValidIndex(PBDRigid->ConstraintGraphIndex()))
-				{
-					ParticleLevel = FMath::Max( LocalGraph->GetIslandGraph()->GraphNodes[PBDRigid->ConstraintGraphIndex()].LevelIndex, 0);
-				}
-			}
-			return ParticleLevel;
-		};
-		if (IsSortingEnabled())
-		{
-			if (FPBDIslandGroup* IslandGroup = ConstraintGraph->GetIslandGroup(GroupIndex))
-			{
-				bool bResizeDone = false;
-				TArray<TPair<int32, int32>>& IslandConstraintSets = ConstraintSets[GroupIndex];
-				IslandConstraintSets.Reset();
-				int32 ConstraintSetBegin = 0, ConstraintSetEnd = 0;
-				
-				for (FPBDIslandSolver* IslandSolver : IslandGroup->GetIslands())
-				{
-					ConstraintSetBegin = ConstraintSetEnd;
-
-					if(!IslandSolver->IsSleeping() && !IslandSolver->IsUsingCache())
-					{
-						if(!bResizeDone)
-						{
-							IslandGroup->InitConstraintIndex(Constraints.GetContainerId());
-							Constraints.SetNumIslandConstraints(IslandGroup->ConstraintCount(Constraints.GetContainerId()), *IslandGroup);
-							bResizeDone = true;
-						}
-						const int32 Island = IslandSolver->GetIslandIndex();
-				
-						const int32 GraphIndex = ConstraintGraph->GetGraphIndex(Island);
-						if(ConstraintGraph->GetIslandGraph()->GraphIslands.IsValidIndex(GraphIndex))
-						{
-							const int32 MaxColor = IsSortingUsingColors() ? FMath::Max(1,ConstraintGraph->GetIslandGraph()->GraphIslands[GraphIndex].MaxColors+1) : 1;
-							const int32 MaxLevel = IsSortingUsingLevels() ? FMath::Max(1,ConstraintGraph->GetIslandGraph()->GraphIslands[GraphIndex].MaxLevels+1) : 1;
-							
-							if (IsSortingUsingColors())
-							{
-								IslandConstraintSets.Reserve(IslandConstraintSets.Num()+MaxLevel * MaxColor);	// Pessimistic array size - we could store the actual required size in coloring algorithm
-							}
-
-							for (int32 Level = 0; Level < MaxLevel; ++Level)
-							{
-								for (int32 Color = 0; Color < MaxColor; ++Color)
-								{
-									const int32 OffsetIndex = IslandOffsets[Island] + Level * MaxColor + Color;
-									const int32 OffsetBegin = (OffsetIndex == 0) ? 0 : ConstraintOffsets[OffsetIndex-1];
-									const int32 OffsetEnd = ConstraintOffsets[OffsetIndex];
-									
-									if (OffsetEnd != OffsetBegin)
-									{
-										for (int32 ConstraintIndex = OffsetBegin; ConstraintIndex < OffsetEnd; ++ConstraintIndex)
-										{
-											FConstraintContainerHandle* Constraint = SortedConstraints[ConstraintIndex];
-											if (Constraint->IsEnabled() && !Constraint->IsProbe())
-											{
-												Constraint->PreGatherInput(Dt, *IslandGroup);
-											}
-										}
-										// Calculate the range of indices for this color as a set of independent contacts
-										TPair<int32, int32> ColorConstrainSet(ConstraintSetEnd, ConstraintSetEnd);
-										InnerPhysicsParallelForRange(OffsetEnd - OffsetBegin, [&](int32 StartIndex, int32 EndIndex)
-										{
-											for (int32 Index = StartIndex; Index < EndIndex; ++Index)
-											{
-												int32 ConstraintIndex = Index + OffsetBegin;
-												FConstraintContainerHandle* Constraint = SortedConstraints[ConstraintIndex];
-												if (Constraint->IsEnabled() && !Constraint->IsProbe())
-												{
-													// Levels that should be assigned to the bodies for shock propagation
-													// @todo(chaos): optimize the lookup
-													const TVector<TGeometryParticleHandle<FReal, 3>*, 2> ConstrainedParticles = Constraint->GetConstrainedParticles();
-													const int32 Particle0Level = GetParticleLevel(ConstrainedParticles[0]);
-													const int32 Particle1Level = GetParticleLevel(ConstrainedParticles[1]);
-
-													// Note we are building the SolverBodies as we go, in the order that we visit them. Each constraint
-													// references two bodies, so we won't strictly be accessing only in cache order, but it's about as good as it can be.
-													Constraint->GatherInput(Dt, Particle0Level, Particle1Level, *IslandGroup);
-												}
-											}
-										}, Chaos::SmallBatchSize);
-										for (int32 ConstraintIndex = OffsetBegin; ConstraintIndex < OffsetEnd; ++ConstraintIndex)
-										{
-											FConstraintContainerHandle* Constraint = SortedConstraints[ConstraintIndex];
-											if (Constraint->IsEnabled() && !Constraint->IsProbe())
-											{
-												// Update the current constraint set of this color
-												ColorConstrainSet.Value = ++ConstraintSetEnd;
-											}
-										}
-										// Remember the set of constraints of this color
-										if (IsSortingUsingColors()) 
-										{
-											IslandConstraintSets.Add(ColorConstrainSet);
-										}
-									}
-								}
-							}
-						}
-					}
-				}
-
-				// If we aren't coloring, we have a single group of all constraints in the IslandGroup (they have been created in level order above)
-				if (!IsSortingUsingColors())
-				{
-					const TPair<int32, int32> ConstrainSet(0, ConstraintSetEnd);
-					IslandConstraintSets.Add(ConstrainSet);
-				}
-			}
-		}
-		else
-		{
-			Base::GatherSolverInput(Dt, GroupIndex);
-		}
-	}
-
-	template<class ConstraintType>
-	void TPBDConstraintColorRule<ConstraintType>::ScatterSolverOutput(const FReal Dt, int32 GroupIndex)
-	{
-		if (IsSortingEnabled())
-		{
-			if (FPBDIslandGroup* IslandGroup = ConstraintGraph->GetIslandGroup(GroupIndex))
-			{
-				const TArray<TPair<int32, int32>>& IslandConstraintSets = ConstraintSets[GroupIndex];
-				for (const TPair<int32, int32>& ConstraintSet : IslandConstraintSets)
-				{
-					Constraints.ScatterOutput(Dt, ConstraintSet.Key, ConstraintSet.Value, *IslandGroup);
-				}
-			}
-		}
-		else
-		{
-			return Base::ScatterSolverOutput(Dt, GroupIndex);
-		}
-	}
-
-	template<class ConstraintType>
-	bool TPBDConstraintColorRule<ConstraintType>::ApplyConstraints(const FReal Dt, int32 GroupIndex, const int32 It, const int32 NumIts)
-	{
-		if (IsSortingEnabled())
-		{
-			bool bNeedsAnotherIteration = false;
-			if (FPBDIslandGroup* IslandGroup = ConstraintGraph->GetIslandGroup(GroupIndex))
-			{
-				const TArray<TPair<int32, int32>>& IslandConstraintSets = ConstraintSets[GroupIndex];
-				if (!IsSortingUsingColors())
-				{
-					for (const TPair<int32, int32>& ConstraintSet : IslandConstraintSets)
-					{
-						bNeedsAnotherIteration |= Constraints.ApplyPhase1Serial(Dt, It, NumIts, ConstraintSet.Key, ConstraintSet.Value, *IslandGroup);
-					}
-				}
-				else
-				{
-					for (const TPair<int32, int32>& ConstraintSet : IslandConstraintSets)
-					{
-						bNeedsAnotherIteration |= Constraints.ApplyPhase1Parallel(Dt, It, NumIts, ConstraintSet.Key, ConstraintSet.Value, *IslandGroup);
-					}
-				}
-			}
-			return bNeedsAnotherIteration;
-		}
-		else
-		{
-			return Base::ApplyConstraints(Dt, GroupIndex, It, NumIts);
-		}
-	}
-
-	template<class ConstraintType>
-	bool TPBDConstraintColorRule<ConstraintType>::ApplyPushOut(const FReal Dt, int32 GroupIndex, const int32 It, const int32 NumIts)
-	{
-		if (IsSortingEnabled())
-		{
-			bool bNeedsAnotherIteration = false;
-			if (FPBDIslandGroup* IslandGroup = ConstraintGraph->GetIslandGroup(GroupIndex))
-			{
-				const TArray<TPair<int32, int32>>& IslandConstraintSets = ConstraintSets[GroupIndex];
-				if (!IsSortingUsingColors())
-				{
-					for (const TPair<int32, int32>& ConstraintSet : IslandConstraintSets)
-					{
-						bNeedsAnotherIteration |= Constraints.ApplyPhase2Serial(Dt, It, NumIts, ConstraintSet.Key, ConstraintSet.Value, *IslandGroup);
-					}
-				}
-				else
-				{
-					for (const TPair<int32, int32>& ConstraintSet : IslandConstraintSets)
-					{
-						bNeedsAnotherIteration |= Constraints.ApplyPhase2Parallel(Dt, It, NumIts, ConstraintSet.Key, ConstraintSet.Value, *IslandGroup);
-					}
-				}
-			}
-			return bNeedsAnotherIteration;
-		}
-		else
-		{
-			return Base::ApplyPushOut(Dt, GroupIndex, It, NumIts);
-		}
-	}
-
-	template<class ConstraintType>
-	void TPBDConstraintColorRule<ConstraintType>::InitializeAccelerationStructures()
-	{
-		if (IsSortingEnabled())
-		{
-			ConstraintSets.SetNum(ConstraintGraph->NumGroups());
-			ConstraintGraph->template AddConstraintDatas<ConstraintType>(Constraints.GetContainerId());
-		}
-		else
-		{
-			Base::InitializeAccelerationStructures();
-		}
-	}
-
-	template<class ConstraintType>
-	void TPBDConstraintColorRule<ConstraintType>::UpdateAccelerationStructures(const FReal Dt, const int32 GroupIndex)
-	{}
-	
-	template<class ConstraintType>
-	void TPBDConstraintColorRule<ConstraintType>::SortConstraints()
-	{
-		// Compute levels for each constraints
-		ComputeLevels();
-
-		// Compute colors for each constraints
-		ComputeColors();
-
-		// Populate the sorted constraints based on the island/level/color
-		PopulateConstraints();
-	}
-	
-	template<class ConstraintType>
-    bool TPBDConstraintColorRule<ConstraintType>::IsSortingEnabled() const 
-    {
-        return IsSortingUsingColors() || IsSortingUsingLevels();
-    }
-	
-	template<class ConstraintType>
-    bool TPBDConstraintColorRule<ConstraintType>::IsSortingUsingLevels() const 
-    {
-    	return !ChaosDisableIslandLevels;
-    }
-    
-    template<class ConstraintType>
-    bool TPBDConstraintColorRule<ConstraintType>::IsSortingUsingColors() const
-    {
-    	return !ChaosDisableIslandColors;
-    }
-	
-	template<class ConstraintType>
-	void TPBDConstraintColorRule<ConstraintType>::ComputeLevels()
-	{
-		if (IsSortingUsingLevels())
-		{
-			ConstraintGraph->GetIslandGraph()->ComputeLevels(Constraints.GetContainerId());
-		}
-	}
-
-	template<class ConstraintType>
-	void TPBDConstraintColorRule<ConstraintType>::ComputeColors()
-	{
-		if (IsSortingUsingColors())
-		{
-			ConstraintGraph->GetIslandGraph()->ComputeColors(Constraints.GetContainerId(), ChaosCollisionColorMinParticles);
-		}
-	}
-
-	template<class ConstraintType>
-	void TPBDConstraintColorRule<ConstraintType>::ForEachEdges(TFunctionRef<void(const int32, const FPBDConstraintGraph::GraphType::FGraphEdge&)> InFunction)
-	{
-		for(auto& GraphEdge : ConstraintGraph->GetIslandGraph()->GraphEdges)
-		{
-			if (GraphEdge.IslandIndex != INDEX_NONE)
-			{
-				auto& GraphIsland = ConstraintGraph->GetIslandGraph()->GraphIslands[GraphEdge.IslandIndex];
-				if(GraphEdge.ItemContainer == Constraints.GetContainerId() && !GraphIsland.IslandItem->IsSleeping())
-				{
-					const int32 EdgeColor = IsSortingUsingColors() ? FMath::Max(0,GraphEdge.ColorIndex) : 0;
-					const int32 EdgeLevel = IsSortingUsingLevels() ? FMath::Max(0,GraphEdge.LevelIndex) : 0;
-					const int32 MaxColors = IsSortingUsingColors() ? FMath::Max(1,GraphIsland.MaxColors+1) : 1;
-
-					const int32 IslandOffset = IslandOffsets[GraphIsland.IslandItem->GetIslandIndex()];
-					const int32 OffsetIndex = IslandOffset + EdgeLevel * MaxColors + EdgeColor;
-
-					InFunction(OffsetIndex,GraphEdge);
-				}
-			}
-		}
-	}
-
-	template<class ConstraintType>
-	void TPBDConstraintColorRule<ConstraintType>::PopulateConstraints()
-	{
-		if(IsSortingEnabled())
-		{
-			SortedConstraints.SetNum(ConstraintGraph->GetIslandGraph()->GraphEdges.Num(), false);
-			IslandOffsets.SetNum(ConstraintGraph->NumIslands(), false);
-
-			// We first fill the island offsets that will be used in the gather
-			int32 IslandOffset = 0;
-			for(auto& GraphIsland : ConstraintGraph->GetIslandGraph()->GraphIslands)
-			{
-				if(!GraphIsland.IslandItem->IsSleeping())
-				{
-					IslandOffsets[GraphIsland.IslandItem->GetIslandIndex()] = IslandOffset;
-					
-					const int32 MaxColors = IsSortingUsingColors() ? FMath::Max(1,GraphIsland.MaxColors+1) : 1;
-					const int32 MaxLevels = IsSortingUsingLevels() ? FMath::Max(1,GraphIsland.MaxLevels+1) : 1;
-
-					IslandOffset += MaxLevels * MaxColors;
-				}
-			}
-			// Initialization of the constraint offsets and the offsets counters
-			ConstraintOffsets.Reset();
-			OffsetCounters.Reset();
-			ConstraintOffsets.SetNumZeroed(IslandOffset, false);
-			OffsetCounters.SetNumZeroed(IslandOffset, false);
-			
-			auto& LocalConstraintOffsets = ConstraintOffsets;
-			auto& LocalOffsetCounters = OffsetCounters;
-			auto& LocalSortedConstraints = SortedConstraints;
-
-			// Build of the constraint offsets that will be used to know where the sorted constraints will
-			// be inserted into the flat array  
-			ForEachEdges( [&LocalConstraintOffsets](
-				const int32 OffsetIndex, const FPBDConstraintGraph::GraphType::FGraphEdge& GraphEdge)
-			{
-				++LocalConstraintOffsets[OffsetIndex];
-			});
-
-			for(int32 OffsetIndex = 0, OffsetEnd = ConstraintOffsets.Num()-1; OffsetIndex < OffsetEnd; ++OffsetIndex)
-			{
-				ConstraintOffsets[OffsetIndex+1] += ConstraintOffsets[OffsetIndex];
-			}
-
-			// Insert the constraint handles in the right order : Island|Level|Color to be processed by the Gather
-			ForEachEdges( [&LocalConstraintOffsets, &LocalOffsetCounters, &LocalSortedConstraints](
-				const int32 OffsetIndex, const FPBDConstraintGraph::GraphType::FGraphEdge& GraphEdge)
-			{
-				const int32 OffsetValue = (OffsetIndex == 0) ? 0 : LocalConstraintOffsets[OffsetIndex-1];
-				LocalSortedConstraints[OffsetValue+LocalOffsetCounters[OffsetIndex]] = GraphEdge.EdgeItem->As<FConstraintContainerHandle>();
-				
-				++LocalOffsetCounters[OffsetIndex];
-			});
-		}
-	}
-
-	template<class ConstraintType>
-	void TPBDConstraintColorRule<ConstraintType>::SetUseContactGraph(const bool bInUseContactGraph)
-	{}
-
-
-	template class TSimpleConstraintRule<FPBDCollisionConstraints>;
-	template class TSimpleConstraintRule<FPBDJointConstraints>;
-	template class TSimpleConstraintRule<FPBDNullConstraints>;
-	template class TSimpleConstraintRule<FPBDRigidSpringConstraints>;
-
-	template class TPBDConstraintGraphRuleImpl<FPBDCollisionConstraints>;
-	template class TPBDConstraintGraphRuleImpl<FPBDJointConstraints>;
-	template class TPBDConstraintGraphRuleImpl<FPBDNullConstraints>;
-	template class TPBDConstraintGraphRuleImpl<FPBDPositionConstraints>;
-	template class TPBDConstraintGraphRuleImpl<FPBDSuspensionConstraints>;
-	template class TPBDConstraintGraphRuleImpl<FPBDRigidDynamicSpringConstraints>;
-	template class TPBDConstraintGraphRuleImpl<FPBDRigidSpringConstraints>;
-
-	template class TPBDConstraintColorRule<FPBDCollisionConstraints>;
-	template class TPBDConstraintIslandRule<FPBDJointConstraints>;
-	template class TPBDConstraintIslandRule<FPBDNullConstraints>;
-	template class TPBDConstraintIslandRule<FPBDPositionConstraints>;
-	template class TPBDConstraintIslandRule<FPBDSuspensionConstraints>;
-	template class TPBDConstraintIslandRule<FPBDRigidDynamicSpringConstraints>;
-	template class TPBDConstraintIslandRule<FPBDRigidSpringConstraints>;
-}
\ No newline at end of file
+//
+// TO BE REMOVED
+//
\ No newline at end of file
