Chaos - Constraint Solver Refactor

diff --git a/Engine/Source/Runtime/Experimental/Chaos/Private/Chaos/PBDJointConstraints.cpp b/Engine/Source/Runtime/Experimental/Chaos/Private/Chaos/PBDJointConstraints.cpp
--- a/Engine/Source/Runtime/Experimental/Chaos/Private/Chaos/PBDJointConstraints.cpp
+++ b/Engine/Source/Runtime/Experimental/Chaos/Private/Chaos/PBDJointConstraints.cpp
@@ -2,7 +2,7 @@
 #include "Chaos/PBDJointConstraints.h"
 #include "Chaos/ChaosDebugDraw.h"
 #include "Chaos/DebugDrawQueue.h"
-#include "Chaos/Evolution/SolverDatas.h"
+#include "Chaos/Island/IslandManager.h"
 #include "Chaos/Joint/ChaosJointLog.h"
 #include "Chaos/Joint/ColoringGraph.h"
 #include "Chaos/Joint/JointConstraintsCVars.h"
@@ -174,15 +174,6 @@ namespace Chaos
 		return ConcreteContainer()->SetConstraintEnabled(ConstraintIndex, bInEnabled);
 	}
 
-	void FPBDJointConstraintHandle::PreGatherInput(const FReal Dt, FPBDIslandSolverData& SolverData)
-	{
-		ConcreteContainer()->PreGatherInput(Dt, ConstraintIndex, SolverData);
-	}
-
-	void FPBDJointConstraintHandle::GatherInput(const FReal Dt, const int32 Particle0Level, const int32 Particle1Level, FPBDIslandSolverData& SolverData)
-	{
-		ConcreteContainer()->GatherInput(Dt, ConstraintIndex, Particle0Level, Particle1Level, SolverData);
-	}
 
 	//
 	// Constraint Settings
@@ -355,9 +346,7 @@ namespace Chaos
 
 	
 	FPBDJointSolverSettings::FPBDJointSolverSettings()
-		: ApplyPairIterations(1)
-		, ApplyPushOutPairIterations(1)
-		, SwingTwistAngleTolerance(1.0e-6f)
+		: SwingTwistAngleTolerance(1.0e-6f)
 		, PositionTolerance(0)
 		, AngleTolerance(0)
 		, MinParentMassRatio(0)
@@ -396,11 +385,10 @@ namespace Chaos
 	//
 
 	
-	FPBDJointConstraints::FPBDJointConstraints(const FPBDJointSolverSettings& InSettings)
-		: FPBDIndexedConstraintContainer(FConstraintContainerHandle::StaticType())
-		, Settings(InSettings)
+	FPBDJointConstraints::FPBDJointConstraints()
+		: Base(FConstraintContainerHandle::StaticType())
+		, Settings()
 		, bJointsDirty(false)
-		, SolverType(EConstraintSolverType::QuasiPbd)
 	{
 	}
 
@@ -758,10 +746,6 @@ namespace Chaos
 		return ConstraintStates[ConstraintIndex].ResimType;
 	}
 
-	void FPBDJointConstraints::UpdatePositionBasedState(const FReal Dt)
-	{
-	}
-
 	void FPBDJointConstraints::PrepareTick()
 	{
 		SCOPE_CYCLE_COUNTER(STAT_Joints_PrepareTick);
@@ -817,71 +801,119 @@ namespace Chaos
 		OutR1 = FRotation3(Q1 * XL1.GetRotation()).ToMatrix();
 	}
 
-	//////////////////////////////////////////////////////////////////////////
-	//
-	// Begin Simple API Solver. Iterate over constraints in array order.
-	//
-	//////////////////////////////////////////////////////////////////////////
+	void FPBDJointConstraints::AddConstraintsToGraph(FPBDIslandManager& IslandManager)
+	{ 
+		IslandManager.AddContainerConstraints(*this);
+	}
 
-	void FPBDJointConstraints::PreGatherInput(const FReal Dt, FPBDIslandSolverData& SolverData)
+	void FPBDJointConstraints::AddBodies(FSolverBodyContainer& SolverBodyContainer)
 	{
 		for (int32 ConstraintIndex = 0; ConstraintIndex < NumConstraints(); ++ConstraintIndex)
 		{
-			if (!ConstraintStates[ConstraintIndex].bDisabled)
-			{
-				PreGatherInput(Dt, ConstraintIndex, SolverData);
-			}
+			AddBodies(ConstraintIndex, SolverBodyContainer);
 		}
 	}
 
-	void FPBDJointConstraints::GatherInput(const FReal Dt, FPBDIslandSolverData& SolverData)
+	void FPBDJointConstraints::GatherInput(const FReal Dt)
 	{
 		for (int32 ConstraintIndex = 0; ConstraintIndex < NumConstraints(); ++ConstraintIndex)
 		{
-			if (!ConstraintStates[ConstraintIndex].bDisabled)
-			{
-				GatherInput(Dt, ConstraintIndex, INDEX_NONE, INDEX_NONE, SolverData);
-			}
+			GatherInput(ConstraintIndex, Dt);
 		}
 	}
 
-	bool FPBDJointConstraints::ApplyPhase1(const FReal Dt, const int32 It, const int32 NumIts, FPBDIslandSolverData& SolverData)
+	void FPBDJointConstraints::ScatterOutput(const FReal Dt)
 	{
-		return ApplyPhase1Serial(Dt, It, NumIts, SolverData);
+		for (int32 ConstraintIndex = 0; ConstraintIndex < NumConstraints(); ++ConstraintIndex)
+		{
+			ScatterOutput(ConstraintIndex, Dt);
+		}
 	}
 
-	bool FPBDJointConstraints::ApplyPhase2(const FReal Dt, const int32 It, const int32 NumIts, FPBDIslandSolverData& SolverData)
+	void FPBDJointConstraints::ApplyPositionConstraints(const FReal Dt, const int32 It, const int32 NumIts)
 	{
-		return ApplyPhase2Serial(Dt, It, NumIts, SolverData);
+		for (int32 ConstraintIndex = 0; ConstraintIndex < NumConstraints(); ++ConstraintIndex)
+		{
+			ApplyPositionConstraint(Dt, ConstraintIndex, It, NumIts);
+		}
 	}
 
-	bool FPBDJointConstraints::ApplyPhase3(const FReal Dt, const int32 It, const int32 NumIts, FPBDIslandSolverData& SolverData)
+	void FPBDJointConstraints::ApplyVelocityConstraints(const FReal Dt, const int32 It, const int32 NumIts)
 	{
-		return ApplyPhase3Serial(Dt, It, NumIts, SolverData);
+		for (int32 ConstraintIndex = 0; ConstraintIndex < NumConstraints(); ++ConstraintIndex)
+		{
+			ApplyVelocityConstraint(Dt, ConstraintIndex, It, NumIts);
+		}
 	}
 
-	//////////////////////////////////////////////////////////////////////////
-	//
-	// End Simple API Solver.
-	//
-	//////////////////////////////////////////////////////////////////////////
+	void FPBDJointConstraints::ApplyProjectionConstraints(const FReal Dt, const int32 It, const int32 NumIts)
+	{
+		for (int32 ConstraintIndex = 0; ConstraintIndex < NumConstraints(); ++ConstraintIndex)
+		{
+			ApplyProjectionConstraint(Dt, ConstraintIndex, It, NumIts);
+		}
+	}
 
-	//////////////////////////////////////////////////////////////////////////
-	//
-	// Begin Graph API Solver. Iterate over constraints in connectivity order.
-	//
-	//////////////////////////////////////////////////////////////////////////
+	void FPBDJointConstraints::AddBodies(const TArrayView<int32>& ConstraintIndices, FSolverBodyContainer& SolverBodyContainer)
+	{
+		for (int32 ConstraintIndex : ConstraintIndices)
+		{
+			AddBodies(ConstraintIndex, SolverBodyContainer);
+		}
+	}
 
-	void FPBDJointConstraints::SetNumIslandConstraints(const int32 NumIslandConstraints, FPBDIslandSolverData& SolverData)
+	void FPBDJointConstraints::GatherInput(const TArrayView<int32>& ConstraintIndices, const FReal Dt)
 	{
-		SolverData.GetConstraintIndices(ContainerId).Reset(NumIslandConstraints);
+		for (int32 ConstraintIndex : ConstraintIndices)
+		{
+			GatherInput(ConstraintIndex, Dt);
+		}
 	}
 
-	void FPBDJointConstraints::PreGatherInput(const FReal Dt, const int32 ConstraintIndex, FPBDIslandSolverData& SolverData)
+	void FPBDJointConstraints::ScatterOutput(const TArrayView<int32>& ConstraintIndices, const FReal Dt)
 	{
-		check(!ConstraintStates[ConstraintIndex].bDisabled);
+		for (int32 ConstraintIndex : ConstraintIndices)
+		{
+			ScatterOutput(ConstraintIndex, Dt);
+		}
+	}
+
+	void FPBDJointConstraints::ApplyPositionConstraints(const TArrayView<int32>& ConstraintIndices, const FReal Dt, const int32 It, const int32 NumIts)
+	{
+		for (int32 ConstraintIndex : ConstraintIndices)
+		{
+			ApplyPositionConstraint(Dt, ConstraintIndex, It, NumIts);
+		}
+	}
+
+	void FPBDJointConstraints::ApplyVelocityConstraints(const TArrayView<int32>& ConstraintIndices, const FReal Dt, const int32 It, const int32 NumIts)
+	{
+		for (int32 ConstraintIndex : ConstraintIndices)
+		{
+			ApplyVelocityConstraint(Dt, ConstraintIndex, It, NumIts);
+		}
+	}
+
+	void FPBDJointConstraints::ApplyProjectionConstraints(const TArrayView<int32>& ConstraintIndices, const FReal Dt, const int32 It, const int32 NumIts)
+	{
+		for (int32 ConstraintIndex : ConstraintIndices)
+		{
+			ApplyProjectionConstraint(Dt, ConstraintIndex, It, NumIts);
+		}
+	}
 
-		SolverData.GetConstraintIndices(ContainerId).Add(ConstraintIndex);
+	//////////////////////////////////////////////////////////////////////////
+	//
+	// Begin single-particle solve methods
+	//
+	//////////////////////////////////////////////////////////////////////////
+
+	void FPBDJointConstraints::AddBodies(const int32 ConstraintIndex, FSolverBodyContainer& SolverBodyContainer)
+	{
+		if (ConstraintStates[ConstraintIndex].bDisabled)
+		{
+			return;
+		}
 
 		int32 Index0, Index1;
 		GetConstrainedParticleIndices(ConstraintIndex, Index0, Index1);
@@ -890,14 +922,27 @@ namespace Chaos
 
 		// Find the solver bodies for the particles we constrain. This will add them to the container
 		// if they aren't there already, and ensure that they are populated with the latest data.
-		SolverData.GetBodyContainer().FindOrAdd(Particle0, Dt);
-		SolverData.GetBodyContainer().FindOrAdd(Particle1, Dt);
+		FSolverBody* SolverBody0 = SolverBodyContainer.FindOrAdd(Particle0);
+		FSolverBody* SolverBody1 = SolverBodyContainer.FindOrAdd(Particle1);
+
+		if (Settings.bUseLinearSolver)
+		{
+			FPBDJointCachedSolver& Solver = CachedConstraintSolvers[ConstraintIndex];
+			Solver.SetSolverBodies(SolverBody0, SolverBody1);
+		}
+		else
+		{
+			FPBDJointSolver& Solver = ConstraintSolvers[ConstraintIndex];
+			Solver.SetSolverBodies(SolverBody0, SolverBody1);
+		}
 	}
 
-	void FPBDJointConstraints::GatherInput(const FReal Dt, const int32 ConstraintIndex, const int32 Particle0Level, const int32 Particle1Level, FPBDIslandSolverData& SolverData)
+	void FPBDJointConstraints::GatherInput(const int32 ConstraintIndex, const FReal Dt)
 	{
-		//SCOPE_CYCLE_COUNTER(STAT_Joints_Gather);
-		check(!ConstraintStates[ConstraintIndex].bDisabled);
+		if (ConstraintStates[ConstraintIndex].bDisabled)
+		{
+			return;
+		}
 
 		const FPBDJointSettings& JointSettings = ConstraintSettings[ConstraintIndex];
 		const FTransformPair& JointFrames = JointSettings.ConnectorTransforms;
@@ -907,17 +952,11 @@ namespace Chaos
 		FGenericParticleHandle Particle0 = FGenericParticleHandle(ConstraintParticles[ConstraintIndex][Index0]);
 		FGenericParticleHandle Particle1 = FGenericParticleHandle(ConstraintParticles[ConstraintIndex][Index1]);
 
-		// Find the solver bodies for the particles we constrain. This will add them to the container
-		// if they aren't there already, and ensure that they are populated with the latest data.
-		FSolverBody* Body0 = SolverData.GetBodyContainer().FindOrAdd(Particle0, Dt);
-		FSolverBody* Body1 = SolverData.GetBodyContainer().FindOrAdd(Particle1, Dt);
-
 		if (Settings.bUseLinearSolver)
 		{
 			FPBDJointCachedSolver& Solver = CachedConstraintSolvers[ConstraintIndex];
 			Solver.Init(
 				Dt,
-				{ Body0, Body1 },
 				Settings,
 				JointSettings,
 				FParticleUtilities::ParticleLocalToCoMLocal(Particle0, JointFrames[Index0]),
@@ -928,7 +967,6 @@ namespace Chaos
 			FPBDJointSolver& Solver = ConstraintSolvers[ConstraintIndex];
 			Solver.Init(
 				Dt,
-				{ Body0, Body1 },
 				Settings,
 				JointSettings,
 				FParticleUtilities::ParticleLocalToCoMLocal(Particle0, JointFrames[Index0]),
@@ -956,167 +994,80 @@ namespace Chaos
 		}
 	}
 
-	void FPBDJointConstraints::ScatterOutput(FReal Dt, FPBDIslandSolverData& SolverData)
+	void FPBDJointConstraints::ScatterOutput(const int32 ConstraintIndex, FReal Dt)
 	{
 		SCOPE_CYCLE_COUNTER(STAT_Joints_Scatter);
 
-		for (int32 ConstraintIndex : SolverData.GetConstraintIndices(ContainerId))
+		if (ConstraintStates[ConstraintIndex].bDisabled)
 		{
-			if (!ConstraintStates[ConstraintIndex].bDisabled)
-			{
-				FPBDJointState& JointState = ConstraintStates[ConstraintIndex];
-
-				int32 Index0, Index1;
-				GetConstrainedParticleIndices(ConstraintIndex, Index0, Index1);
+			return;
+		}
 
-				if (Settings.bUseLinearSolver)
-				{
-					FPBDJointCachedSolver& Solver = CachedConstraintSolvers[ConstraintIndex];
-					// NOTE: LinearImpulse/AngularImpulse in the solver are not really impulses - they are mass-weighted position/rotation delta, or (impulse x dt).
-					if (Dt > UE_SMALL_NUMBER)
-					{
-						if (Index0 == 0)
-						{
-							JointState.LinearImpulse = Solver.GetNetLinearImpulse() / Dt;
-							JointState.AngularImpulse = Solver.GetNetAngularImpulse() / Dt;
-						}
-						else
-						{
-							// Particles were flipped in the solver...
-							JointState.LinearImpulse = -Solver.GetNetLinearImpulse() / Dt;
-							JointState.AngularImpulse = -Solver.GetNetAngularImpulse() / Dt;
-						}
-					}
-					else
-					{
-						JointState.LinearImpulse = FVec3(0);
-						JointState.AngularImpulse = FVec3(0);
-					}
+		FPBDJointState& JointState = ConstraintStates[ConstraintIndex];
 
-					ApplyPlasticityLimits(ConstraintIndex);
+		int32 Index0, Index1;
+		GetConstrainedParticleIndices(ConstraintIndex, Index0, Index1);
 
-					// Remove our solver body reference (they are not valid between frames)
-					CachedConstraintSolvers[ConstraintIndex].Deinit();
+		if (Settings.bUseLinearSolver)
+		{
+			FPBDJointCachedSolver& Solver = CachedConstraintSolvers[ConstraintIndex];
+			// NOTE: LinearImpulse/AngularImpulse in the solver are not really impulses - they are mass-weighted position/rotation delta, or (impulse x dt).
+			if (Dt > UE_SMALL_NUMBER)
+			{
+				if (Index0 == 0)
+				{
+					JointState.LinearImpulse = Solver.GetNetLinearImpulse() / Dt;
+					JointState.AngularImpulse = Solver.GetNetAngularImpulse() / Dt;
 				}
 				else
 				{
-					FPBDJointSolver& Solver = ConstraintSolvers[ConstraintIndex];
-					// NOTE: LinearImpulse/AngularImpulse in the solver are not really impulses - they are mass-weighted position/rotation delta, or (impulse x dt).
-					if (Dt > UE_SMALL_NUMBER)
-					{
-						if (Index0 == 0)
-						{
-							JointState.LinearImpulse = Solver.GetNetLinearImpulse() / Dt;
-							JointState.AngularImpulse = Solver.GetNetAngularImpulse() / Dt;
-						}
-						else
-						{
-							// Particles were flipped in the solver...
-							JointState.LinearImpulse = -Solver.GetNetLinearImpulse() / Dt;
-							JointState.AngularImpulse = -Solver.GetNetAngularImpulse() / Dt;
-						}
-					}
-					else
-					{
-						JointState.LinearImpulse = FVec3(0);
-						JointState.AngularImpulse = FVec3(0);
-					}
-
-					ApplyPlasticityLimits(ConstraintIndex);
-
-					// Remove our solver body reference (they are not valid between frames)
-					ConstraintSolvers[ConstraintIndex].Deinit();
+					// Particles were flipped in the solver...
+					JointState.LinearImpulse = -Solver.GetNetLinearImpulse() / Dt;
+					JointState.AngularImpulse = -Solver.GetNetAngularImpulse() / Dt;
 				}
 			}
-		}
-	}
-
-	bool FPBDJointConstraints::ApplyPhase1Serial(const FReal Dt, const int32 It, const int32 NumIts, FPBDIslandSolverData& SolverData)
-	{
-		CSV_SCOPED_TIMING_STAT(Chaos, ApplyJointConstraints);
-		SCOPE_CYCLE_COUNTER(STAT_Joints_Apply);
-
-		int32 NumActive = 0;
-		const int32 NumPairIts = (SolverType == EConstraintSolverType::QuasiPbd) ? 1 : Settings.ApplyPairIterations;
-		if (NumPairIts > 0)
-		{
-			for (int32 ConstraintIndex : SolverData.GetConstraintIndices(ContainerId))
+			else
 			{
-				NumActive += ApplyPhase1Single(Dt, ConstraintIndex, NumPairIts, It, NumIts);
+				JointState.LinearImpulse = FVec3(0);
+				JointState.AngularImpulse = FVec3(0);
 			}
-		}
-		return (NumActive > 0);
-	}
 
-	bool FPBDJointConstraints::ApplyPhase2Serial(const FReal Dt, const int32 It, const int32 NumIts, FPBDIslandSolverData& SolverData)
-	{
-		CSV_SCOPED_TIMING_STAT(Chaos, PushOutJointConstraints);
-		SCOPE_CYCLE_COUNTER(STAT_Joints_ApplyPushOut);
+			ApplyPlasticityLimits(ConstraintIndex);
 
-		int32 NumActive = 0;
-		for (int32 ConstraintIndex : SolverData.GetConstraintIndices(ContainerId))
-		{
-			NumActive += ApplyPhase2Single(Dt, ConstraintIndex, It, NumIts);
+			// Remove our solver body reference (they are not valid between frames)
+			CachedConstraintSolvers[ConstraintIndex].Deinit();
 		}
-		return (NumActive > 0);
-	}
-	
-	void FPBDJointConstraints::PreparePhase3Serial(const FReal Dt, FPBDIslandSolverData& SolverData)
-	{
-		CSV_SCOPED_TIMING_STAT(Chaos, ProjectJointConstraints);
-		SCOPE_CYCLE_COUNTER(STAT_Joints_ApplyProjection);
-
-		if(Settings.bUseLinearSolver)
+		else
 		{
-			for (int32 ConstraintIndex : SolverData.GetConstraintIndices(ContainerId))
+			FPBDJointSolver& Solver = ConstraintSolvers[ConstraintIndex];
+			// NOTE: LinearImpulse/AngularImpulse in the solver are not really impulses - they are mass-weighted position/rotation delta, or (impulse x dt).
+			if (Dt > UE_SMALL_NUMBER)
 			{
-				const FPBDJointSettings& JointSettings = ConstraintSettings[ConstraintIndex];
-				FPBDJointCachedSolver& Solver = CachedConstraintSolvers[ConstraintIndex];
-				
-				if ((FMath::IsNearlyZero(Solver.InvM(0)) && FMath::IsNearlyZero(Solver.InvM(1))))
+				if (Index0 == 0)
+				{
+					JointState.LinearImpulse = Solver.GetNetLinearImpulse() / Dt;
+					JointState.AngularImpulse = Solver.GetNetAngularImpulse() / Dt;
+				}
+				else
 				{
-					return;
+					// Particles were flipped in the solver...
+					JointState.LinearImpulse = -Solver.GetNetLinearImpulse() / Dt;
+					JointState.AngularImpulse = -Solver.GetNetAngularImpulse() / Dt;
 				}
-				Solver.InitProjection(
-					Dt,
-					Settings,
-					JointSettings);
 			}
-		}
-	}
-
-	bool FPBDJointConstraints::ApplyPhase3Serial(const FReal Dt, const int32 It, const int32 NumIts, FPBDIslandSolverData& SolverData)
-	{
-		CSV_SCOPED_TIMING_STAT(Chaos, ProjectJointConstraints);
-		SCOPE_CYCLE_COUNTER(STAT_Joints_ApplyProjection);
+			else
+			{
+				JointState.LinearImpulse = FVec3(0);
+				JointState.AngularImpulse = FVec3(0);
+			}
 
-		// Prepare phase 3 for the linear solver in order to partially re-init the solver
-		if (It == 0)
-		{
-			PreparePhase3Serial(Dt, SolverData);
-		}
+			ApplyPlasticityLimits(ConstraintIndex);
 
-		for (int32 ConstraintIndex : SolverData.GetConstraintIndices(ContainerId))
-		{
-			ApplyPhase3Single(Dt, ConstraintIndex, It, NumIts);
+			// Remove our solver body reference (they are not valid between frames)
+			ConstraintSolvers[ConstraintIndex].Deinit();
 		}
-
-		return true;
 	}
 
-	//////////////////////////////////////////////////////////////////////////
-	//
-	// End Graph API Solver.
-	//
-	//////////////////////////////////////////////////////////////////////////
-
-
-	//////////////////////////////////////////////////////////////////////////
-	//
-	// Begin single-particle solve methods used by APIs
-	//
-	//////////////////////////////////////////////////////////////////////////
-
 	FReal FPBDJointConstraints::CalculateIterationStiffness(int32 It, int32 NumIts) const
 	{
 		// Linearly interpolate betwwen MinStiffness and MaxStiffness over the first few iterations,
@@ -1181,7 +1132,7 @@ namespace Chaos
 	{
 		// Shock propagation is only enabled for the last iteration, and only for the QPBD solver.
 		// The standard PBD solver runs projection in the second solver phase which is mostly the same thing.
-		if (JointSettings.bShockPropagationEnabled && (It >= (NumIts - Settings.NumShockPropagationIterations)) && (SolverType == EConstraintSolverType::QuasiPbd))
+		if (JointSettings.bShockPropagationEnabled && (It >= (NumIts - Settings.NumShockPropagationIterations)))
 		{
 			if (Body0.IsDynamic() && Body1.IsDynamic())
 			{
@@ -1193,7 +1144,7 @@ namespace Chaos
 
 	// This position solver iterates over each of the inner constraints (position, twist, swing) and solves them independently.
 	// This will converge slowly in some cases, particularly where resolving angular constraints violates position constraints and vice versa.
-	bool FPBDJointConstraints::ApplyPhase1Single(const FReal Dt, const int32 ConstraintIndex, const int32 NumPairIts, const int32 It, const int32 NumIts)
+	bool FPBDJointConstraints::ApplyPositionConstraint(const FReal Dt, const int32 ConstraintIndex, const int32 It, const int32 NumIts)
 	{
 		if (!CanEvaluate(ConstraintIndex))
 		{
@@ -1231,12 +1182,7 @@ namespace Chaos
 			Solver.SetShockPropagationScales(ShockPropagationInvMassScale, FReal(1), Dt);
 
 			const FReal IterationStiffness = CalculateIterationStiffness(It, NumIts);
-			for (int32 PairIt = 0; PairIt < NumPairIts; ++PairIt)
-			{
-				UE_LOG(LogChaosJoint, VeryVerbose, TEXT("  Pair Iteration %d / %d"), PairIt, NumPairIts);
-
-				Solver.ApplyConstraints(Dt, IterationStiffness, Settings, JointSettings);
-			}
+			Solver.ApplyConstraints(Dt, IterationStiffness, Settings, JointSettings);
 
 			// @todo(ccaulfield): The break limit should really be applied to the impulse in the solver to prevent 1-frame impulses larger than the threshold
 			if ((JointSettings.LinearBreakForce!=FLT_MAX) || (JointSettings.AngularBreakTorque!=FLT_MAX))
@@ -1258,35 +1204,9 @@ namespace Chaos
 			const FReal ShockPropagationInvMassScale = CalculateShockPropagationInvMassScale(Solver.Body0(), Solver.Body1(), JointSettings, It, NumIts);
 
 			const FReal IterationStiffness = CalculateIterationStiffness(It, NumIts);
-			for (int32 PairIt = 0; PairIt < NumPairIts; ++PairIt)
-			{
-				UE_LOG(LogChaosJoint, VeryVerbose, TEXT("  Pair Iteration %d / %d"), PairIt, NumPairIts);
-
-				if (SolverType == EConstraintSolverType::StandardPbd)
-				{
-					Solver.UpdateMasses(ShockPropagationInvMassScale, FReal(1));
-				}
-				else
-				{
-					Solver.SetShockPropagationScales(ShockPropagationInvMassScale, FReal(1));
-				}
+			Solver.SetShockPropagationScales(ShockPropagationInvMassScale, FReal(1));
 
-				Solver.ApplyConstraints(Dt, IterationStiffness, Settings, JointSettings);
-
-				if (SolverType == EConstraintSolverType::StandardPbd)
-				{
-					if (Solver.Body0().IsDynamic())
-					{
-						Solver.Body0().SolverBody().ApplyCorrections();
-						Solver.Body0().UpdateRotationDependentState();
-					}
-					if (Solver.Body1().IsDynamic())
-					{
-						Solver.Body1().SolverBody().ApplyCorrections();
-						Solver.Body1().UpdateRotationDependentState();
-					}
-				}
-			}
+			Solver.ApplyConstraints(Dt, IterationStiffness, Settings, JointSettings);
 
 			// @todo(ccaulfield): The break limit should really be applied to the impulse in the solver to prevent 1-frame impulses larger than the threshold
 			if ((JointSettings.LinearBreakForce!=FLT_MAX) || (JointSettings.AngularBreakTorque!=FLT_MAX))
@@ -1300,18 +1220,13 @@ namespace Chaos
 
 	// QuasiPBD applies a velocity solve in phase 2
 	// Standard PBD does nothing
-	bool FPBDJointConstraints::ApplyPhase2Single(const FReal Dt, const int32 ConstraintIndex, const int32 It, const int32 NumIts)
+	bool FPBDJointConstraints::ApplyVelocityConstraint(const FReal Dt, const int32 ConstraintIndex, const int32 It, const int32 NumIts)
 	{
 		if (!CanEvaluate(ConstraintIndex))
 		{
 			return false;
 		}
 
-		if (SolverType == EConstraintSolverType::StandardPbd)
-		{
-			return false;
-		}
-
 		const TVector<TGeometryParticleHandle<FReal, 3>*, 2>& Constraint = ConstraintParticles[ConstraintIndex];
 		UE_LOG(LogChaosJoint, VeryVerbose, TEXT("Solve Joint Velocity Constraint %d %s %s (dt = %f; it = %d / %d)"), ConstraintIndex, *Constraint[0]->ToString(), *Constraint[1]->ToString(), Dt, It, NumIts);
 
@@ -1371,7 +1286,7 @@ namespace Chaos
 	}
 
 	// Projection phase
-	bool FPBDJointConstraints::ApplyPhase3Single(const FReal Dt, const int32 ConstraintIndex, const int32 It, const int32 NumIts)
+	bool FPBDJointConstraints::ApplyProjectionConstraint(const FReal Dt, const int32 ConstraintIndex, const int32 It, const int32 NumIts)
 	{
 		if (!CanEvaluate(ConstraintIndex))
 		{
@@ -1409,6 +1324,7 @@ namespace Chaos
 
 			if (It == 0)
 			{
+				Solver.InitProjection(Dt, Settings, JointSettings);
 				Solver.ApplyTeleports(Dt, Settings, JointSettings);
 			}
 
@@ -1424,7 +1340,7 @@ namespace Chaos
 
 			Solver.Update(Dt, Settings, JointSettings);
 
-			if ((SolverType == EConstraintSolverType::StandardPbd) || (It == 0))
+			if (It == 0)
 			{
 				// @todo(chaos): support reverse parent/child
 				Solver.Body1().UpdateRotationDependentState();
@@ -1432,14 +1348,6 @@ namespace Chaos
 			}
 
 			Solver.ApplyProjections(Dt, Settings, JointSettings);
-
-			if (SolverType == EConstraintSolverType::StandardPbd)
-			{
-				if (Solver.Body1().IsDynamic())
-				{
-					Solver.Body1().SolverBody().ApplyCorrections();
-				}
-			}
 		}
 
 		return true;
diff --git a/Engine/Source/Runtime/Experimental/Chaos/Private/Chaos/PBDNullConstraints.cpp b/Engine/Source/Runtime/Experimental/Chaos/Private/Chaos/PBDNullConstraints.cpp
new file mode 100644
--- /dev/null
+++ b/Engine/Source/Runtime/Experimental/Chaos/Private/Chaos/PBDNullConstraints.cpp
@@ -0,0 +1,11 @@
+// Copyright Epic Games, Inc. All Rights Reserved.
+#include "Chaos/PBDNullConstraints.h"
+#include "Chaos/Island/IslandManager.h"
+
+namespace Chaos
+{
+	void FPBDNullConstraints::AddConstraintsToGraph(FPBDIslandManager& IslandManager)
+	{
+		IslandManager.AddContainerConstraints(*this);
+	}
+}
\ No newline at end of file
diff --git a/Engine/Source/Runtime/Experimental/Chaos/Private/Chaos/PBDPositionConstraints.cpp b/Engine/Source/Runtime/Experimental/Chaos/Private/Chaos/PBDPositionConstraints.cpp
--- a/Engine/Source/Runtime/Experimental/Chaos/Private/Chaos/PBDPositionConstraints.cpp
+++ b/Engine/Source/Runtime/Experimental/Chaos/Private/Chaos/PBDPositionConstraints.cpp
@@ -1,7 +1,7 @@
 // Copyright Epic Games, Inc. All Rights Reserved.
 
 #include "Chaos/PBDPositionConstraints.h"
-#include "Chaos/Evolution/SolverDatas.h"
+#include "Chaos/Island/IslandManager.h"
 
 namespace Chaos
 {
@@ -10,48 +10,61 @@ namespace Chaos
 		return ConcreteContainer()->GetConstrainedParticles(ConstraintIndex);
 	}
 
-	void FPBDPositionConstraintHandle::PreGatherInput(const FReal Dt, FPBDIslandSolverData& SolverData)
+	void FPBDPositionConstraints::AddConstraintsToGraph(FPBDIslandManager& IslandManager)
 	{
-		ConcreteContainer()->PreGatherInput(Dt, ConstraintIndex, SolverData);
+		IslandManager.AddContainerConstraints(*this);
 	}
 
-	void FPBDPositionConstraintHandle::GatherInput(const FReal Dt, const int32 Particle0Level, const int32 Particle1Level, FPBDIslandSolverData& SolverData)
+	void FPBDPositionConstraints::AddBodies(FSolverBodyContainer& SolverBodyContainer)
 	{
-		ConcreteContainer()->GatherInput(Dt, ConstraintIndex, Particle0Level, Particle1Level, SolverData);
+		for (int32 ConstraintIndex = 0; ConstraintIndex < NumConstraints(); ++ConstraintIndex)
+		{
+			AddBodies(ConstraintIndex, SolverBodyContainer);
+		}
 	}
-	
-	void FPBDPositionConstraints::SetNumIslandConstraints(const int32 NumIslandConstraints, FPBDIslandSolverData& SolverData)
+
+	void FPBDPositionConstraints::ScatterOutput(const FReal Dt)
 	{
-		SolverData.GetConstraintIndices(ContainerId).Reset(NumIslandConstraints);
+		for (int32 ConstraintIndex = 0; ConstraintIndex < NumConstraints(); ++ConstraintIndex)
+		{
+			ConstraintSolverBodies[ConstraintIndex] = nullptr;
+		}
 	}
 
-	void FPBDPositionConstraints::PreGatherInput(const FReal Dt, const int32 ConstraintIndex, FPBDIslandSolverData& SolverData)
+	void FPBDPositionConstraints::ApplyPositionConstraints(const FReal Dt, const int32 It, const int32 NumIts)
 	{
-		SolverData.GetConstraintIndices(ContainerId).Add(ConstraintIndex);
-
-		ConstraintSolverBodies[ConstraintIndex] = SolverData.GetBodyContainer().FindOrAdd(ConstrainedParticles[ConstraintIndex], Dt);
+		for (int32 ConstraintIndex = 0; ConstraintIndex < NumConstraints(); ++ConstraintIndex)
+		{
+			ApplySingle(Dt, ConstraintIndex);
+		}
 	}
 
-	void FPBDPositionConstraints::GatherInput(const FReal Dt, const int32 ConstraintIndex, const int32 Particle0Level, const int32 Particle1Level, FPBDIslandSolverData& SolverData)
+	void FPBDPositionConstraints::AddBodies(const TArrayView<int32>& ConstraintIndices, FSolverBodyContainer& SolverBodyContainer)
 	{
+		for (int32 ConstraintIndex : ConstraintIndices)
+		{
+			AddBodies(ConstraintIndex, SolverBodyContainer);
+		}
 	}
-	
-	void FPBDPositionConstraints::ScatterOutput(FReal Dt, FPBDIslandSolverData& SolverData)
+
+	void FPBDPositionConstraints::ScatterOutput(const TArrayView<int32>& ConstraintIndices, const FReal Dt)
 	{
-		for (int32 ConstraintIndex : SolverData.GetConstraintIndices(ContainerId))
+		for (int32 ConstraintIndex : ConstraintIndices)
 		{
 			ConstraintSolverBodies[ConstraintIndex] = nullptr;
 		}
 	}
 
-	bool FPBDPositionConstraints::ApplyPhase1Serial(const FReal Dt, const int32 It, const int32 NumIts, FPBDIslandSolverData& SolverData)
+	void FPBDPositionConstraints::ApplyPositionConstraints(const TArrayView<int32>& ConstraintIndices, const FReal Dt, const int32 It, const int32 NumIts)
 	{
-		for (int32 ConstraintIndex : SolverData.GetConstraintIndices(ContainerId))
+		for (int32 ConstraintIndex : ConstraintIndices)
 		{
 			ApplySingle(Dt, ConstraintIndex);
 		}
+	}
 
-		// @todo(chaos): early iteration termination in FPBDPositionConstraints
-		return true;
+	void FPBDPositionConstraints::AddBodies(const int32 ConstraintIndex, FSolverBodyContainer& SolverBodyContainer)
+	{
+		ConstraintSolverBodies[ConstraintIndex] = SolverBodyContainer.FindOrAdd(ConstrainedParticles[ConstraintIndex]);
 	}
 }
diff --git a/Engine/Source/Runtime/Experimental/Chaos/Private/Chaos/PBDRigidDynamicSpringConstraints.cpp b/Engine/Source/Runtime/Experimental/Chaos/Private/Chaos/PBDRigidDynamicSpringConstraints.cpp
--- a/Engine/Source/Runtime/Experimental/Chaos/Private/Chaos/PBDRigidDynamicSpringConstraints.cpp
+++ b/Engine/Source/Runtime/Experimental/Chaos/Private/Chaos/PBDRigidDynamicSpringConstraints.cpp
@@ -1,7 +1,7 @@
 // Copyright Epic Games, Inc. All Rights Reserved.
 
 #include "Chaos/PBDRigidDynamicSpringConstraints.h"
-#include "Chaos/Evolution/SolverDatas.h"
+#include "Chaos/Island/IslandManager.h"
 #include "Chaos/Particle/ParticleUtilities.h"
 #include "Chaos/Utilities.h"
 
@@ -12,16 +12,6 @@ FParticlePair FPBDRigidDynamicSpringConstraintHandle::GetConstrainedParticles()
 	return ConcreteContainer()->GetConstrainedParticles(ConstraintIndex);
 }
 
-void FPBDRigidDynamicSpringConstraintHandle::PreGatherInput(const FReal Dt, FPBDIslandSolverData& SolverData)
-{
-	ConcreteContainer()->PreGatherInput(Dt, ConstraintIndex, SolverData);
-}
-
-void FPBDRigidDynamicSpringConstraintHandle::GatherInput(const FReal Dt, const int32 Particle0Level, const int32 Particle1Level, FPBDIslandSolverData& SolverData)
-{
-	ConcreteContainer()->GatherInput(Dt, ConstraintIndex, Particle0Level, Particle1Level, SolverData);
-}
-
 // Note: this is called outside the constraint solver loop (it generates constraints and can merge 
 // islands - the equivalent of collision detection) and therefore reads directly from Particles, 
 // rather than from SolverBodies.
@@ -107,47 +97,74 @@ void FPBDRigidDynamicSpringConstraints::UpdatePositionBasedState(const FReal Dt)
 	}
 }
 
-void FPBDRigidDynamicSpringConstraints::SetNumIslandConstraints(const int32 NumIslandConstraints, FPBDIslandSolverData& SolverData)
+void FPBDRigidDynamicSpringConstraints::AddConstraintsToGraph(FPBDIslandManager& IslandManager)
 {
-	SolverData.GetConstraintIndices(ContainerId).Reset(NumIslandConstraints);
+	IslandManager.AddContainerConstraints(*this);
 }
 
-void FPBDRigidDynamicSpringConstraints::PreGatherInput(const FReal Dt, const int32 ConstraintIndex, FPBDIslandSolverData& SolverData)
+void FPBDRigidDynamicSpringConstraints::AddBodies(FSolverBodyContainer& SolverBodyContainer)
 {
-	SolverData.GetConstraintIndices(ContainerId).Add(ConstraintIndex);
+	for (int32 ConstraintIndex = 0; ConstraintIndex < NumConstraints(); ++ConstraintIndex)
+	{
+		AddBodies(ConstraintIndex, SolverBodyContainer);
+	}
+}
 
-	ConstraintSolverBodies[ConstraintIndex] =
+void FPBDRigidDynamicSpringConstraints::ScatterOutput(const FReal Dt)
+{
+	for (int32 ConstraintIndex = 0; ConstraintIndex < NumConstraints(); ++ConstraintIndex)
 	{
-		SolverData.GetBodyContainer().FindOrAdd(Constraints[ConstraintIndex][0], Dt),
-		SolverData.GetBodyContainer().FindOrAdd(Constraints[ConstraintIndex][1], Dt)
-	};
+		ConstraintSolverBodies[ConstraintIndex] =
+		{
+			nullptr,
+			nullptr
+		};
+	}
 }
 
-void FPBDRigidDynamicSpringConstraints::GatherInput(const FReal Dt, const int32 ConstraintIndex, const int32 Particle0Level, const int32 Particle1Level, FPBDIslandSolverData& SolverData)
+void FPBDRigidDynamicSpringConstraints::ApplyPositionConstraints(const FReal Dt, const int32 It, const int32 NumIts)
 {
+	for (int32 ConstraintIndex = 0; ConstraintIndex < NumConstraints(); ++ConstraintIndex)
+	{
+		ApplySingle(Dt, ConstraintIndex);
+	}
+}
+
+void FPBDRigidDynamicSpringConstraints::AddBodies(const TArrayView<int32>& ConstraintIndices, FSolverBodyContainer& SolverBodyContainer)
+{
+	for (int32 ConstraintIndex : ConstraintIndices)
+	{
+		AddBodies(ConstraintIndex, SolverBodyContainer);
+	}
 }
 
-void FPBDRigidDynamicSpringConstraints::ScatterOutput(FReal Dt, FPBDIslandSolverData& SolverData)
+void FPBDRigidDynamicSpringConstraints::ScatterOutput(const TArrayView<int32>& ConstraintIndices, const FReal Dt)
 {
-	for (int32 ConstraintIndex : SolverData.GetConstraintIndices(ContainerId))
+	for (int32 ConstraintIndex : ConstraintIndices)
 	{
 		ConstraintSolverBodies[ConstraintIndex] = 
-		{
-			nullptr,
-			nullptr
+		{ 
+			nullptr, 
+			nullptr 
 		};
 	}
 }
 
-bool FPBDRigidDynamicSpringConstraints::ApplyPhase1Serial(const FReal Dt, const int32 It, const int32 NumIts, FPBDIslandSolverData& SolverData)
+void FPBDRigidDynamicSpringConstraints::ApplyPositionConstraints(const TArrayView<int32>& ConstraintIndices, const FReal Dt, const int32 It, const int32 NumIts)
 {
-	for (int32 ConstraintIndex : SolverData.GetConstraintIndices(ContainerId))
+	for (int32 ConstraintIndex : ConstraintIndices)
 	{
 		ApplySingle(Dt, ConstraintIndex);
 	}
+}
 
-	// @todo(chaos): early iteration termination in FPBDRigidDynamicSpringConstraints
-	return true;
+void FPBDRigidDynamicSpringConstraints::AddBodies(const int32 ConstraintIndex, FSolverBodyContainer& SolverBodyContainer)
+{
+	ConstraintSolverBodies[ConstraintIndex] =
+	{
+		SolverBodyContainer.FindOrAdd(Constraints[ConstraintIndex][0]),
+		SolverBodyContainer.FindOrAdd(Constraints[ConstraintIndex][1])
+	};
 }
 
 FVec3 FPBDRigidDynamicSpringConstraints::GetDelta(const FVec3& WorldSpaceX1, const FVec3& WorldSpaceX2, const int32 ConstraintIndex, const int32 SpringIndex) const
diff --git a/Engine/Source/Runtime/Experimental/Chaos/Private/Chaos/PBDRigidSpringConstraints.cpp b/Engine/Source/Runtime/Experimental/Chaos/Private/Chaos/PBDRigidSpringConstraints.cpp
--- a/Engine/Source/Runtime/Experimental/Chaos/Private/Chaos/PBDRigidSpringConstraints.cpp
+++ b/Engine/Source/Runtime/Experimental/Chaos/Private/Chaos/PBDRigidSpringConstraints.cpp
@@ -1,7 +1,7 @@
 // Copyright Epic Games, Inc. All Rights Reserved.
 
 #include "Chaos/PBDRigidSpringConstraints.h"
-#include "Chaos/Evolution/SolverDatas.h"
+#include "Chaos/Island/IslandManager.h"
 #include "Chaos/Particle/ParticleUtilities.h"
 #include "Chaos/Utilities.h"
 
@@ -37,23 +37,12 @@ namespace Chaos
 		ConcreteContainer()->SetRestLength(ConstraintIndex, SpringLength);
 	}
 
-	void FPBDRigidSpringConstraintHandle::PreGatherInput(const FReal Dt, FPBDIslandSolverData& SolverData)
-	{
-		ConcreteContainer()->PreGatherInput(Dt, ConstraintIndex, SolverData);
-	}
-
-	void FPBDRigidSpringConstraintHandle::GatherInput(const FReal Dt, const int32 Particle0Level, const int32 Particle1Level, FPBDIslandSolverData& SolverData)
-	{
-		ConcreteContainer()->GatherInput(Dt, ConstraintIndex, Particle0Level, Particle1Level, SolverData);
-	}
-
-
 	//
 	// Container Impl
 	//
 
 	FPBDRigidSpringConstraints::FPBDRigidSpringConstraints()
-		: FPBDIndexedConstraintContainer(FConstraintContainerHandle::StaticType())
+		: TPBDIndexedConstraintContainer<FPBDRigidSpringConstraints>(FConstraintContainerHandle::StaticType())
 	{}
 
 	FPBDRigidSpringConstraints::~FPBDRigidSpringConstraints()
@@ -135,62 +124,74 @@ namespace Chaos
 		return SpringSettings[ConstraintIndex].Stiffness * Delta / CombinedInvMass;
 	}
 
-	void FPBDRigidSpringConstraints::PreGatherInput(const FReal Dt, FPBDIslandSolverData& SolverData)
+	void FPBDRigidSpringConstraints::AddConstraintsToGraph(FPBDIslandManager& IslandManager)
 	{
-		for (int32 ConstraintIndex = 0; ConstraintIndex < NumConstraints(); ++ConstraintIndex)
-		{
-			PreGatherInput(Dt, ConstraintIndex, SolverData);
-		}
+		IslandManager.AddContainerConstraints(*this);
 	}
 
-	void FPBDRigidSpringConstraints::GatherInput(const FReal Dt, FPBDIslandSolverData& SolverData)
+	void FPBDRigidSpringConstraints::AddBodies(FSolverBodyContainer& SolverBodyContainer)
 	{
 		for (int32 ConstraintIndex = 0; ConstraintIndex < NumConstraints(); ++ConstraintIndex)
 		{
-			GatherInput(Dt, ConstraintIndex, INDEX_NONE, INDEX_NONE, SolverData);
+			AddBodies(ConstraintIndex, SolverBodyContainer);
 		}
 	}
 
-	bool FPBDRigidSpringConstraints::ApplyPhase1(const FReal Dt, const int32 It, const int32 NumIts, FPBDIslandSolverData& SolverData)
+	void FPBDRigidSpringConstraints::ScatterOutput(const FReal Dt)
 	{
-		return ApplyPhase1Serial(Dt, It, NumIts, SolverData);
+		for (int32 ConstraintIndex = 0; ConstraintIndex < NumConstraints(); ++ConstraintIndex)
+		{
+			ConstraintSolverBodies[ConstraintIndex] =
+			{
+				nullptr,
+				nullptr
+			};
+		}
 	}
 
-	void FPBDRigidSpringConstraints::SetNumIslandConstraints(const int32 NumIslandConstraints, FPBDIslandSolverData& SolverData)
+	void FPBDRigidSpringConstraints::ApplyPositionConstraints(const FReal Dt, const int32 It, const int32 NumIts)
 	{
-		SolverData.GetConstraintIndices(ContainerId).Reset(NumIslandConstraints);
+		for (int32 ConstraintIndex = 0; ConstraintIndex < NumConstraints(); ++ConstraintIndex)
+		{
+			ApplyPhase1Single(Dt, ConstraintIndex);
+		}
 	}
 
-	void FPBDRigidSpringConstraints::PreGatherInput(const FReal Dt, const int32 ConstraintIndex, FPBDIslandSolverData& SolverData)
+	void FPBDRigidSpringConstraints::AddBodies(const TArrayView<int32>& ConstraintIndices, FSolverBodyContainer& SolverBodyContainer)
 	{
-		SolverData.GetConstraintIndices(ContainerId).Add(ConstraintIndex);
-	
-		ConstraintSolverBodies[ConstraintIndex] =
+		for (int32 ConstraintIndex : ConstraintIndices)
 		{
-			SolverData.GetBodyContainer().FindOrAdd(Constraints[ConstraintIndex][0], Dt),
-			SolverData.GetBodyContainer().FindOrAdd(Constraints[ConstraintIndex][1], Dt)
-		};
+			AddBodies(ConstraintIndex, SolverBodyContainer);
+		}
 	}
 
-	void FPBDRigidSpringConstraints::GatherInput(const FReal Dt, const int32 ConstraintIndex, const int32 Particle0Level, const int32 Particle1Level, FPBDIslandSolverData& SolverData)
+	void FPBDRigidSpringConstraints::ScatterOutput(const TArrayView<int32>& ConstraintIndices, const FReal Dt)
 	{
+		for (int32 ConstraintIndex : ConstraintIndices)
+		{
+			ConstraintSolverBodies[ConstraintIndex] = 
+			{ 
+				nullptr, 
+				nullptr 
+			};
+		}
 	}
 
-	void FPBDRigidSpringConstraints::ScatterOutput(FReal Dt, FPBDIslandSolverData& SolverData)
+	void FPBDRigidSpringConstraints::ApplyPositionConstraints(const TArrayView<int32>& ConstraintIndices, const FReal Dt, const int32 It, const int32 NumIts)
 	{
-		for (int32 ConstraintIndex : SolverData.GetConstraintIndices(ContainerId))
+		for (int32 ConstraintIndex : ConstraintIndices)
 		{
-			ConstraintSolverBodies[ConstraintIndex] = { nullptr, nullptr };
+			ApplyPhase1Single(Dt, ConstraintIndex);
 		}
 	}
 
-	bool FPBDRigidSpringConstraints::ApplyPhase1Serial(const FReal Dt, const int32 It, const int32 NumIts, FPBDIslandSolverData& SolverData)
+	void FPBDRigidSpringConstraints::AddBodies(const int32 ConstraintIndex, FSolverBodyContainer& SolverBodyContainer)
 	{
-		for (int32 ConstraintIndex : SolverData.GetConstraintIndices(ContainerId))
+		ConstraintSolverBodies[ConstraintIndex] =
 		{
-			ApplyPhase1Single(Dt, ConstraintIndex);
-		}
-		return true;
+			SolverBodyContainer.FindOrAdd(Constraints[ConstraintIndex][0]),
+			SolverBodyContainer.FindOrAdd(Constraints[ConstraintIndex][1])
+		};
 	}
 
 	void FPBDRigidSpringConstraints::ApplyPhase1Single(const FReal Dt, int32 ConstraintIndex) const
diff --git a/Engine/Source/Runtime/Experimental/Chaos/Private/Chaos/PBDRigidsEvolution.cpp b/Engine/Source/Runtime/Experimental/Chaos/Private/Chaos/PBDRigidsEvolution.cpp
--- a/Engine/Source/Runtime/Experimental/Chaos/Private/Chaos/PBDRigidsEvolution.cpp
+++ b/Engine/Source/Runtime/Experimental/Chaos/Private/Chaos/PBDRigidsEvolution.cpp
@@ -187,7 +187,8 @@ namespace Chaos
 	};
 
 	FPBDRigidsEvolutionBase::FPBDRigidsEvolutionBase(FPBDRigidsSOAs& InParticles, THandleArray<FChaosPhysicsMaterial>& InSolverPhysicsMaterials, bool InIsSingleThreaded)
-	    : Particles(InParticles)
+	    : IslandGroupManager(ConstraintGraph)
+		, Particles(InParticles)
 		, SolverPhysicsMaterials(InSolverPhysicsMaterials)
 		, InternalAcceleration(nullptr)
 		, AsyncInternalAcceleration(nullptr)
diff --git a/Engine/Source/Runtime/Experimental/Chaos/Private/Chaos/PBDRigidsEvolutionGBF.cpp b/Engine/Source/Runtime/Experimental/Chaos/Private/Chaos/PBDRigidsEvolutionGBF.cpp
--- a/Engine/Source/Runtime/Experimental/Chaos/Private/Chaos/PBDRigidsEvolutionGBF.cpp
+++ b/Engine/Source/Runtime/Experimental/Chaos/Private/Chaos/PBDRigidsEvolutionGBF.cpp
@@ -113,12 +113,7 @@ namespace Chaos
 		DECLARE_CYCLE_STAT(TEXT("FPBDRigidsEvolutionGBF::CollisionModifierCallback"), STAT_Evolution_CollisionModifierCallback, STATGROUP_Chaos);
 		DECLARE_CYCLE_STAT(TEXT("FPBDRigidsEvolutionGBF::CCD"), STAT_Evolution_CCD, STATGROUP_Chaos);
 		DECLARE_CYCLE_STAT(TEXT("FPBDRigidsEvolutionGBF::GraphColor"), STAT_Evolution_GraphColor, STATGROUP_Chaos);
-		DECLARE_CYCLE_STAT(TEXT("FPBDRigidsEvolutionGBF::Gather"), STAT_Evolution_Gather, STATGROUP_Chaos);
-		DECLARE_CYCLE_STAT(TEXT("FPBDRigidsEvolutionGBF::Scatter"), STAT_Evolution_Scatter, STATGROUP_Chaos);
-		DECLARE_CYCLE_STAT(TEXT("FPBDRigidsEvolutionGBF::ApplyConstraintsPhase1"), STAT_Evolution_ApplyConstraintsPhase1, STATGROUP_Chaos);
-		DECLARE_CYCLE_STAT(TEXT("FPBDRigidsEvolutionGBF::UpdateVelocities"), STAT_Evolution_UpdateVelocites, STATGROUP_Chaos);
-		DECLARE_CYCLE_STAT(TEXT("FPBDRigidsEvolutionGBF::ApplyConstraintsPhase2"), STAT_Evolution_ApplyConstraintsPhase2, STATGROUP_Chaos);
-		DECLARE_CYCLE_STAT(TEXT("FPBDRigidsEvolutionGBF::ApplyConstraintsPhase3"), STAT_Evolution_ApplyConstraintsPhase3, STATGROUP_Chaos);
+		DECLARE_CYCLE_STAT(TEXT("FPBDRigidsEvolutionGBF::BuildGroups"), STAT_Evolution_BuildGroups, STATGROUP_Chaos);
 		DECLARE_CYCLE_STAT(TEXT("FPBDRigidsEvolutionGBF::DetectCollisions"), STAT_Evolution_DetectCollisions, STATGROUP_Chaos);
 		DECLARE_CYCLE_STAT(TEXT("FPBDRigidsEvolutionGBF::TransferJointCollisions"), STAT_Evolution_TransferJointCollisions, STATGROUP_Chaos);
 		DECLARE_CYCLE_STAT(TEXT("FPBDRigidsEvolutionGBF::PostDetectCollisionsCallback"), STAT_Evolution_PostDetectCollisionsCallback, STATGROUP_Chaos);
@@ -127,10 +122,10 @@ namespace Chaos
 		DECLARE_CYCLE_STAT(TEXT("FPBDRigidsEvolutionGBF::CreateConstraintGraph"), STAT_Evolution_CreateConstraintGraph, STATGROUP_Chaos);
 		DECLARE_CYCLE_STAT(TEXT("FPBDRigidsEvolutionGBF::CreateIslands"), STAT_Evolution_CreateIslands, STATGROUP_Chaos);
 		DECLARE_CYCLE_STAT(TEXT("FPBDRigidsEvolutionGBF::PruneCollisions"), STAT_Evolution_PruneCollisions, STATGROUP_Chaos);
-		DECLARE_CYCLE_STAT(TEXT("FPBDRigidsEvolutionGBF::UpdateAccelerationStructures"), STAT_Evolution_UpdateAccelerationStructures, STATGROUP_Chaos);
 		DECLARE_CYCLE_STAT(TEXT("FPBDRigidsEvolutionGBF::AddSleepingContacts"), STAT_Evolution_AddSleepingContacts, STATGROUP_Chaos);
 		DECLARE_CYCLE_STAT(TEXT("FPBDRigidsEvolutionGBF::PreApplyCallback"), STAT_Evolution_PreApplyCallback, STATGROUP_Chaos);
 		DECLARE_CYCLE_STAT(TEXT("FPBDRigidsEvolutionGBF::ParallelSolve"), STAT_Evolution_ParallelSolve, STATGROUP_Chaos);
+		DECLARE_CYCLE_STAT(TEXT("FPBDRigidsEvolutionGBF::BuildDisabledParticles"), STAT_Evolution_BuildDisabledParticles, STATGROUP_Chaos);
 		DECLARE_CYCLE_STAT(TEXT("FPBDRigidsEvolutionGBF::SaveParticlePostSolve"), STAT_Evolution_SavePostSolve, STATGROUP_Chaos);
 		DECLARE_CYCLE_STAT(TEXT("FPBDRigidsEvolutionGBF::DeactivateSleep"), STAT_Evolution_DeactivateSleep, STATGROUP_Chaos);
 		DECLARE_CYCLE_STAT(TEXT("FPBDRigidsEvolutionGBF::InertiaConditioning"), STAT_Evolution_InertiaConditioning, STATGROUP_Chaos);
@@ -222,35 +217,41 @@ void FPBDRigidsEvolutionGBF::AdvanceOneTimeStep(const FReal Dt,const FSubStepInf
 	UnprepareTick();
 }
 
-void FPBDRigidsEvolutionGBF::ReloadParticlesCache(const int32 Island)
+void FPBDRigidsEvolutionGBF::ReloadParticlesCache()
 {
-	if(!GetConstraintGraph().GetSolverIsland(Island)->IsSleeping())
+	// @todo(chaos): Parallelize
+	FEvolutionResimCache* ResimCache = GetCurrentStepResimCache();
+	if ((ResimCache != nullptr) && ResimCache->IsResimming())
 	{
-		FEvolutionResimCache* ResimCache = GetCurrentStepResimCache();
-		if (ResimCache && ResimCache->IsResimming() && GetConstraintGraph().IslandNeedsResim(Island) == false)
+		for (int32 IslandIndex = 0; IslandIndex < GetConstraintGraph().NumIslands(); ++IslandIndex)
 		{
-			for (auto Particle : GetConstraintGraph().GetIslandParticles(Island))
+			FPBDIsland* Island = GetConstraintGraph().GetIsland(IslandIndex);
+			bool bIsUsingCache = false;
+			if (!Island->IsSleeping() && !GetConstraintGraph().IslandNeedsResim(IslandIndex))
 			{
-				if (auto Rigid = Particle->CastToRigidParticle())
+				for (FPBDIslandParticle& IslandParticle : Island->GetParticles())
 				{
-					ResimCache->ReloadParticlePostSolve(*Rigid);
+					if (auto Rigid = IslandParticle.GetParticle()->CastToRigidParticle())
+					{
+						ResimCache->ReloadParticlePostSolve(*Rigid);
+					}
 				}
+				bIsUsingCache = true;
 			}
-			GetConstraintGraph().GetSolverIsland(Island)->SetIsUsingCache(true);
+			Island->SetIsUsingCache(bIsUsingCache);
 		}
-		GetConstraintGraph().GetSolverIsland(Island)->SetIsUsingCache(false);
 	}
 }
 
 void FPBDRigidsEvolutionGBF::BuildDisabledParticles(const int32 Island, TArray<TArray<FPBDRigidParticleHandle*>>& DisabledParticles, TArray<bool>& SleepedIslands)
 {
-	for (auto Particle : GetConstraintGraph().GetIslandParticles(Island))
+	for (FPBDIslandParticle& IslandParticle : GetConstraintGraph().GetIsland(Island)->GetParticles())
 	{
 		// If a dynamic particle is moving slowly enough for long enough, disable it.
 		// @todo(mlentine): Find a good way of not doing this when we aren't using this functionality
 
 		// increment the disable count for the particle
-		auto PBDRigid = Particle->CastToRigidParticle();
+		auto PBDRigid = IslandParticle.GetParticle()->CastToRigidParticle();
 		if (PBDRigid && PBDRigid->ObjectState() == EObjectStateType::Dynamic)
 		{
 			if (PBDRigid->AuxilaryValue(PhysicsMaterials) && PBDRigid->V().SizeSquared() < PBDRigid->AuxilaryValue(PhysicsMaterials)->DisabledLinearThreshold &&
@@ -280,68 +281,6 @@ void FPBDRigidsEvolutionGBF::BuildDisabledParticles(const int32 Island, TArray<T
 int32 DrawAwake = 0;
 FAutoConsoleVariableRef CVarDrawAwake(TEXT("p.chaos.DebugDrawAwake"),DrawAwake,TEXT("Draw particles that are awake"));
 
-#define CSV_CUSTOM_STAT_PER_ISLAND_HELPER(Stat) CSV_CUSTOM_STAT(PhysicsVerbose, Stat, Stats[Stat] * 1000.0, ECsvCustomStatOp::Set);
-
-//Using aligned to avoid false sharing during perf capture
-struct alignas(64) FPerIslandStats
-{
-#if CSV_PROFILER
-
-	enum EPerIslandStat
-	{
-		/** Any stat added to this enum must be added to ReportStats. Can use inl if it's a problem, but seems fine to just copy paste here*/
-		PerIslandSolve_ReloadCacheTotalSerialized = 0,
-		PerIslandSolve_GatherTotalSerialized,
-		PerIslandSolve_ApplyTotalSerialized,
-		PerIslandSolve_UpdateVelocitiesTotalSerialized,
-		PerIslandSolve_ApplyPushOutTotalSerialized,
-		PerIslandSolve_ApplyProjectionTotalSerialized,
-		PerIslandSolve_ScatterTotalSerialized,
-		NumStats
-	};
-
-	void ReportStats() const
-	{
-		CSV_CUSTOM_STAT_PER_ISLAND_HELPER(PerIslandSolve_ReloadCacheTotalSerialized);
-		CSV_CUSTOM_STAT_PER_ISLAND_HELPER(PerIslandSolve_GatherTotalSerialized);
-		CSV_CUSTOM_STAT_PER_ISLAND_HELPER(PerIslandSolve_ApplyTotalSerialized);
-		CSV_CUSTOM_STAT_PER_ISLAND_HELPER(PerIslandSolve_UpdateVelocitiesTotalSerialized);
-		CSV_CUSTOM_STAT_PER_ISLAND_HELPER(PerIslandSolve_ApplyPushOutTotalSerialized);
-		CSV_CUSTOM_STAT_PER_ISLAND_HELPER(PerIslandSolve_ApplyProjectionTotalSerialized);
-		CSV_CUSTOM_STAT_PER_ISLAND_HELPER(PerIslandSolve_ScatterTotalSerialized);
-	}
-
-	FPerIslandStats()
-	{
-		for (double& Stat : Stats) { Stat = 0.0; }
-	}
-
-	static FPerIslandStats Flatten(const TArray<FPerIslandStats>& PerIslandStats)
-	{
-		FPerIslandStats Flat;
-		for(const FPerIslandStats& IslandStats : PerIslandStats)
-		{
-			for(int32 Stat = 0; Stat < NumStats; ++Stat)
-			{
-				Flat.Stats[Stat] += IslandStats.Stats[Stat];
-			}
-		}
-
-		return Flat;
-	}
-
-	double Stats[EPerIslandStat::NumStats];
-#else
-	void ReportStats() const{}
-#endif
-
-};
-
-#if CSV_PROFILER
-#define CSV_SCOPED_PER_ISLAND_TIMING_STAT(Stat) FScopedDurationTimer Timer_##Stat(PerIslandStats[GroupIndex].Stats[FPerIslandStats::Stat])
-#else
-#define CSV_SCOPED_PER_ISLAND_TIMING_STAT(Stat)
-#endif
 void FPBDRigidsEvolutionGBF::AdvanceOneTimeStepImpl(const FReal Dt, const FSubStepInfo& SubStepInfo)
 {
 	SCOPE_CYCLE_COUNTER(STAT_Evolution_AdvanceOneTimeStep);
@@ -449,13 +388,6 @@ void FPBDRigidsEvolutionGBF::AdvanceOneTimeStepImpl(const FReal Dt, const FSubSt
 		CSV_SCOPED_TIMING_STAT(PhysicsVerbose, StepSolver_CreateIslands);
 		CreateIslands();
 	}
-	{
-		// We keep the graph color stat name for now to compare with previous implementation
-		// @todo : change the name to sort constraints
-		SCOPE_CYCLE_COUNTER(STAT_Evolution_GraphColor);
-		CSV_SCOPED_TIMING_STAT(PhysicsVerbose, StepSolver_GraphColor);
-		SortConstraints();
-	}
 	{
 		// Once the graph is built (constraints are removed/added) we can destroy any collision that are not required
 		SCOPE_CYCLE_COUNTER(STAT_Evolution_PruneCollisions);
@@ -471,117 +403,53 @@ void FPBDRigidsEvolutionGBF::AdvanceOneTimeStepImpl(const FReal Dt, const FSubSt
 
 	CollisionConstraints.SetGravity(GetGravityForces().GetAcceleration());
 
+	// Use the cache to update particles in islands that do not require resimming (not desynched)
+	{
+		CSV_SCOPED_TIMING_STAT(PhysicsVerbose, StepSolver_ReloadCacheTotalSerialized);
+		ReloadParticlesCache();
+	}
+
+	// Assign all islands to a set of groups. Each group is solved in parallel with the others.
+	int32 NumGroups = 0;
+	{
+		SCOPE_CYCLE_COUNTER(STAT_Evolution_BuildGroups);
+		CSV_SCOPED_TIMING_STAT(PhysicsVerbose, StepSolver_BuildGroups);
+		NumGroups = IslandGroupManager.BuildGroups();
+	}
+
+
 	TArray<bool> SleepedIslands;
 	SleepedIslands.SetNum(GetConstraintGraph().NumIslands());
 	TArray<TArray<FPBDRigidParticleHandle*>> DisabledParticles;
 	DisabledParticles.SetNum(GetConstraintGraph().NumIslands());
 	if(Dt > 0)
 	{
-		SCOPE_CYCLE_COUNTER(STAT_Evolution_ParallelSolve);
-		CSV_SCOPED_TIMING_STAT(PhysicsVerbose, StepSolver_PerIslandSolve);
-		const int32 NumGroups = GetConstraintGraph().NumGroups();
-		CSV_CUSTOM_STAT(PhysicsCounters, NumIslandGroups, NumGroups, ECsvCustomStatOp::Set);
-		TArray<FPerIslandStats> PerIslandStats;
-#if CSV_PROFILER
-		PerIslandStats.AddDefaulted(NumGroups);
-#endif
-
-		PhysicsParallelFor(NumGroups, [&](int32 GroupIndex) 
+		// Solve all the constraints
 		{
-			const bool bValidIsland = GetConstraintGraph().GetIslandGroups().IsValidIndex(GroupIndex) && GetConstraintGraph().GetIslandGroup(GroupIndex)->IsValid();
-		  	if (!bValidIsland)
-			{
-				return;
-			}
-			
-			// Reload cache if necessary
-			{
-				CSV_SCOPED_PER_ISLAND_TIMING_STAT(PerIslandSolve_ReloadCacheTotalSerialized);
-				for(auto& IslandSolver : GetConstraintGraph().GetGroupIslands(GroupIndex))
-				{
-					ReloadParticlesCache(IslandSolver->GetIslandIndex());
-				}
-			}
+			SCOPE_CYCLE_COUNTER(STAT_Evolution_ParallelSolve);
+			CSV_SCOPED_TIMING_STAT(PhysicsVerbose, StepSolver_PerIslandSolve);
 
-			// Collect all the data that the constraint solvers operate on
-			{
-				SCOPE_CYCLE_COUNTER(STAT_Evolution_Gather);
-				CSV_SCOPED_PER_ISLAND_TIMING_STAT(PerIslandSolve_GatherTotalSerialized);
-				GatherSolverInput(Dt, GroupIndex);
-			}
-
-			// Run the first phase of the constraint solvers
-			// For GBF this is the hybrid velocity solving step (which also moves the bodies to make the implicit velocity be what it should be)
-			// For PBD/QPBD this is the position solve step
-			{
-				SCOPE_CYCLE_COUNTER(STAT_Evolution_ApplyConstraintsPhase1);
-				CSV_SCOPED_TIMING_STAT(Chaos, ApplyConstraints);
-				CSV_SCOPED_PER_ISLAND_TIMING_STAT(PerIslandSolve_ApplyTotalSerialized);
-				ApplyConstraintsPhase1(Dt, GroupIndex);
-			}
-
-			for(auto& IslandSolver : GetConstraintGraph().GetGroupIslands(GroupIndex))
-			{
-				if(!IslandSolver->IsSleeping() && !IslandSolver->IsUsingCache())
-				{
-					if (PostApplyCallback != nullptr)
-					{
-						PostApplyCallback(IslandSolver->GetIslandIndex());
-					}
-				}
-			}
-
-			// Update implicit velocities from results of constraint solver phase 1
-			{
-				SCOPE_CYCLE_COUNTER(STAT_Evolution_UpdateVelocites);
-				CSV_SCOPED_PER_ISLAND_TIMING_STAT(PerIslandSolve_UpdateVelocitiesTotalSerialized);
-				SetImplicitVelocities(Dt, GroupIndex);
-			}
-
-			// Run the second phase of the constraint solvers
-			// For GBF this is the pushout step
-			// For PBD this does nothing
-			// For QPBD this is the velocity solve step
-			{
-				SCOPE_CYCLE_COUNTER(STAT_Evolution_ApplyConstraintsPhase2);
-				CSV_SCOPED_PER_ISLAND_TIMING_STAT(PerIslandSolve_ApplyPushOutTotalSerialized);
-				ApplyConstraintsPhase2(Dt, GroupIndex);
-			}
-
-			// Apply projections. This is a phase only implemented by joints that applies a semi-physical
-			// correction and/or teleport to fix any errors remaining after the position and velocity solve steps
-			{
-				SCOPE_CYCLE_COUNTER(STAT_Evolution_ApplyConstraintsPhase3);
-				CSV_SCOPED_PER_ISLAND_TIMING_STAT(PerIslandSolve_ApplyProjectionTotalSerialized);
-				ApplyConstraintsPhase3(Dt, GroupIndex);
-			}
+			IslandGroupManager.SetNumIterations(NumPositionIterations, NumVelocityIterations, NumProjectionIterations);
+			IslandGroupManager.Solve(Dt);
+		}
 
-			// Update the particles with the results of the constraint solvers, and also update constraint data
-			// that is accessed externally (net impusles, break info, etc)
-			{
-				SCOPE_CYCLE_COUNTER(STAT_Evolution_Scatter);
-				CSV_SCOPED_PER_ISLAND_TIMING_STAT(PerIslandSolve_ScatterTotalSerialized);
-				ScatterSolverOutput(Dt, GroupIndex);
-			}
+		// Determine which particles can be disabled and which islands are sleeping
+		{
+			SCOPE_CYCLE_COUNTER(STAT_Evolution_BuildDisabledParticles);
 
-			for(auto& IslandSolver : GetConstraintGraph().GetGroupIslands(GroupIndex))
+			// @todo(chaos): improve this - batching will be poor when some island are much larger than others. 
+			// We can't just loop over groups because there are some islands with no constraints (and usually a single particle) 
+			// that do not get added to a group but still need otbe checked for sleeping
+			const int32 NumIslands = GetConstraintGraph().NumIslands();
+			PhysicsParallelFor(NumIslands, [&](const int32 IslandIndex)
 			{
-				if(!IslandSolver->IsSleeping() && !IslandSolver->IsUsingCache())
+				FPBDIsland* Island = GetConstraintGraph().GetIsland(IslandIndex);
+				if (!Island->IsSleeping() && !Island->IsUsingCache())
 				{
-					const int32 Island = IslandSolver->GetIslandIndex();
-					if (PostApplyPushOutCallback != nullptr)
-					{
-						PostApplyPushOutCallback(Island);
-					}
-					BuildDisabledParticles(Island, DisabledParticles, SleepedIslands);
+					BuildDisabledParticles(IslandIndex, DisabledParticles, SleepedIslands);
 				}
-			}
-		});
-
-#if CSV_PROFILER
-		FPerIslandStats FlattenedStats = FPerIslandStats::Flatten(PerIslandStats);
-		FlattenedStats.ReportStats();
-#endif
+			});
+		}
 	}
 
 	{
@@ -632,7 +500,7 @@ void FPBDRigidsEvolutionGBF::AdvanceOneTimeStepImpl(const FReal Dt, const FSubSt
 
 	ParticleUpdatePosition(Particles.GetDirtyParticlesView(), Dt);
 
-	// Clean up the transient data from the constraint graph (e.g., IslandSolvers get cleared here)
+	// Clean up the transient data from the constraint graph (e.g., Islands get cleared here)
 	GetConstraintGraph().EndTick();
 
 	if (bChaos_Solver_TestMode)
@@ -683,132 +551,6 @@ void FPBDRigidsEvolutionGBF::AdvanceOneTimeStepImpl(const FReal Dt, const FSubSt
 #endif
 }
 
-void FPBDRigidsEvolutionGBF::GatherSolverInput(FReal Dt, int32 GroupIndex)
-{
-	// We must initialize the solver body container to be large enough to hold all particles in the
-	// island so that the pointers remain valid (the array should not grow and relocate)
-	if(FPBDIslandGroup* IslandGroup = ConstraintGraph.GetIslandGroups()[GroupIndex].Get())
-	{
-		IslandGroup->GetBodyContainer().Reset(IslandGroup->GetNumParticles());
-
-		for(auto& SolverIsland : IslandGroup->GetIslands())
-		{
-			// We pre-add all the dynamic particles solver bodies involved in any constraint to speed up the gather 
-			const TArray<FGeometryParticleHandle*>& IslandParticles = SolverIsland->GetParticles();
-
-			if(!SolverIsland->IsSleeping() && !SolverIsland->IsUsingCache())
-			{
-				for(auto& Particle : IslandParticles)
-				{
-					FGenericParticleHandle GenericHandle = Particle->Handle();
-					GenericHandle->SetSolverBodyIndex(INDEX_NONE);
-					
-					if(GenericHandle->ObjectState() == EObjectStateType::Dynamic)
-					{
-						if(FPBDRigidParticleHandle* PBDRigid = Particle->Handle()->CastToRigidParticle())
-						{
-							if(ConstraintGraph.GetGraphNodes().IsValidIndex(PBDRigid->ConstraintGraphIndex()))
-							{
-								if(ConstraintGraph.GetGraphNodes()[PBDRigid->ConstraintGraphIndex()].NodeEdges.Num() > 0)
-								{
-									// Only the non sleeping dynamic rigid particles involved in any constraints will be added
-									// by default to the solver bodies
-									const int32 BodyIndex = IslandGroup->GetBodyContainer().GetBodies().Emplace(GenericHandle, Dt);
-									GenericHandle->SetSolverBodyIndex(BodyIndex);
-								}
-							}
-						}
-					}
-				}
-			}
-		}
-	}
-
-	// NOTE: SolverBodies are gathered as part of the constraint gather, in the order that they are first seen 
-	for (FPBDConstraintGraphRule* ConstraintRule : PrioritizedConstraintRules)
-	{
-		ConstraintRule->GatherSolverInput(Dt, GroupIndex);
-	}
-}
-
-void FPBDRigidsEvolutionGBF::ScatterSolverOutput(FReal Dt, int32 GroupIndex)
-{
-	// Scatter solver results for constraints (impulses, break events, etc)
-	for (FPBDConstraintGraphRule* ConstraintRule : PrioritizedConstraintRules)
-	{
-		ConstraintRule->ScatterSolverOutput(Dt, GroupIndex);
-	}
-
-	// Scatter body results back to particles (position, rotation, etc)
-	if(FPBDIslandGroup* IslandGroup = ConstraintGraph.GetIslandGroups()[GroupIndex].Get())
-	{
-		IslandGroup->GetBodyContainer().ScatterOutput();
-	}
-}
-
-void FPBDRigidsEvolutionGBF::ApplyConstraintsPhase1(const FReal Dt, int32 GroupIndex)
-{
-	// @todo(ccaulfield): track whether we are sufficiently solved and can early-out
-	for (int i = 0; i < NumPositionIterations; ++i)
-	{
-		bool bNeedsAnotherIteration = false;
-		for (FPBDConstraintGraphRule* ConstraintRule : PrioritizedConstraintRules)
-		{
-			bNeedsAnotherIteration |= ConstraintRule->ApplyConstraints(Dt, GroupIndex, i, NumPositionIterations);
-		}
-
-		if (ChaosRigidsEvolutionApplyAllowEarlyOutCVar && !bNeedsAnotherIteration)
-		{
-			break;
-		}
-	}
-}
-
-void FPBDRigidsEvolutionGBF::SetImplicitVelocities(const FReal Dt, int32 GroupIndex)
-{
-	if(FPBDIslandGroup* IslandGroup = GetConstraintGraph().GetIslandGroups()[GroupIndex].Get())
-	{
-		IslandGroup->GetBodyContainer().SetImplicitVelocities(Dt);
-	}
-}
-
-void FPBDRigidsEvolutionGBF::ApplyConstraintsPhase2(const FReal Dt, int32 GroupIndex)
-{
-	bool bNeedsAnotherIteration = true;
-	for (int32 It = 0; It < NumVelocityIterations; ++It)
-	{
-		bNeedsAnotherIteration = false;
-		for (FPBDConstraintGraphRule* ConstraintRule : PrioritizedConstraintRules)
-		{
-			bNeedsAnotherIteration |= ConstraintRule->ApplyPushOut(Dt, GroupIndex, It, NumVelocityIterations);
-		}
-
-		if (ChaosRigidsEvolutionApplyPushoutAllowEarlyOutCVar && !bNeedsAnotherIteration)
-		{
-			break;
-		}
-	}
-}
-
-void FPBDRigidsEvolutionGBF::ApplyConstraintsPhase3(const FReal Dt, int32 GroupIndex)
-{
-	// Update the body transforms from the deltas calculated in the constraint solve phases 1 and 2
-	// NOTE: deliberately not updating the world-space inertia as it is not currently needed by joint projection
-	// and no other constraints implement projection
-	if (FPBDIslandGroup* IslandGroup = ConstraintGraph.GetIslandGroups()[GroupIndex].Get())
-	{
-		IslandGroup->GetBodyContainer().ApplyCorrections();
-	}
-
-	for (int32 It = 0; It < NumProjectionIterations; ++It)
-	{
-		for (FPBDConstraintGraphRule* ConstraintRule : PrioritizedConstraintRules)
-		{
-			ConstraintRule->ApplyProjection(Dt, GroupIndex, It, NumProjectionIterations);
-		}
-	}
-}
-
 void FPBDRigidsEvolutionGBF::SetIsDeterministic(const bool bInIsDeterministic)
 {
 	// We detect collisions in parallel, so order is non-deterministic without additional processing
@@ -856,17 +598,12 @@ void FPBDRigidsEvolutionGBF::TestModeResetCollisions()
 FPBDRigidsEvolutionGBF::FPBDRigidsEvolutionGBF(FPBDRigidsSOAs& InParticles,THandleArray<FChaosPhysicsMaterial>& SolverPhysicsMaterials, const TArray<ISimCallbackObject*>* InCollisionModifiers, bool InIsSingleThreaded)
 	: Base(InParticles, SolverPhysicsMaterials, InIsSingleThreaded)
 	, Clustering(*this, Particles.GetClusteredParticles())
-	, JointConstraintRule(JointConstraints, ChaosSolverJointPriority)
-	, SuspensionConstraintRule(SuspensionConstraints, ChaosSolverSuspensionPriority)
-	, CollisionConstraints(InParticles, Collided, PhysicsMaterials, PerParticlePhysicsMaterials, &SolverPhysicsMaterials, 0, 0, DefaultRestitutionThreshold)
-	, CollisionRule(CollisionConstraints, ChaosSolverCollisionPriority)
+	, CollisionConstraints(InParticles, Collided, PhysicsMaterials, PerParticlePhysicsMaterials, &SolverPhysicsMaterials, DefaultNumCollisionsPerBlock, DefaultRestitutionThreshold)
 	, BroadPhase(InParticles)
 	, NarrowPhase(DefaultCollisionCullDistance, BoundsThicknessVelocityMultiplier, CollisionConstraints.GetConstraintAllocator())
 	, CollisionDetector(BroadPhase, NarrowPhase, CollisionConstraints)
 	, PostIntegrateCallback(nullptr)
 	, PreApplyCallback(nullptr)
-	, PostApplyCallback(nullptr)
-	, PostApplyPushOutCallback(nullptr)
 	, CurrentStepResimCacheImp(nullptr)
 	, CollisionModifiers(InCollisionModifiers)
 	, CCDManager()
@@ -904,9 +641,9 @@ FPBDRigidsEvolutionGBF::FPBDRigidsEvolutionGBF(FPBDRigidsSOAs& InParticles,THand
 		GravityForces.Apply(HandleIn, Dt);
 	});
 
-	AddConstraintRule(&SuspensionConstraintRule);
-	AddConstraintRule(&CollisionRule);
-	AddConstraintRule(&JointConstraintRule);
+	AddConstraintContainer(SuspensionConstraints, ChaosSolverSuspensionPriority);
+	AddConstraintContainer(CollisionConstraints, ChaosSolverCollisionPriority);
+	AddConstraintContainer(JointConstraints, ChaosSolverJointPriority);
 
 	SetInternalParticleInitilizationFunction([](const FGeometryParticleHandle*, const FGeometryParticleHandle*) {});
 }
diff --git a/Engine/Source/Runtime/Experimental/Chaos/Private/Chaos/PBDSuspensionConstraints.cpp b/Engine/Source/Runtime/Experimental/Chaos/Private/Chaos/PBDSuspensionConstraints.cpp
--- a/Engine/Source/Runtime/Experimental/Chaos/Private/Chaos/PBDSuspensionConstraints.cpp
+++ b/Engine/Source/Runtime/Experimental/Chaos/Private/Chaos/PBDSuspensionConstraints.cpp
@@ -1,7 +1,7 @@
 // Copyright Epic Games, Inc. All Rights Reserved.
 
 #include "Chaos/PBDSuspensionConstraints.h"
-#include "Chaos/Evolution/SolverDatas.h"
+#include "Chaos/Island/IslandManager.h"
 #include "Chaos/DebugDrawQueue.h"
 
 // Private includes
@@ -63,16 +63,6 @@ namespace Chaos
 		return ConcreteContainer()->GetConstrainedParticles(ConstraintIndex);
 	}
 
-	void FPBDSuspensionConstraintHandle::PreGatherInput(const FReal Dt, FPBDIslandSolverData& SolverData)
-	{
-		ConcreteContainer()->PreGatherInput(Dt, ConstraintIndex, SolverData);
-	}
-
-	void FPBDSuspensionConstraintHandle::GatherInput(const FReal Dt, const int32 Particle0Level, const int32 Particle1Level, FPBDIslandSolverData& SolverData)
-	{
-		ConcreteContainer()->GatherInput(Dt, ConstraintIndex, Particle0Level, Particle1Level, SolverData);
-	}
-
 	Chaos::FPBDSuspensionConstraints::FConstraintContainerHandle* FPBDSuspensionConstraints::AddConstraint(TGeometryParticleHandle<FReal, 3>* Particle, const FVec3& InSuspensionLocalOffset, const FPBDSuspensionSettings& InConstraintSettings)
 	{
 		int32 NewIndex = ConstrainedParticles.Num();
@@ -125,19 +115,99 @@ namespace Chaos
 			SetConstraintIndex(Handles[ConstraintIndex], ConstraintIndex);
 		}
 	}
-	
-	void FPBDSuspensionConstraints::SetNumIslandConstraints(const int32 NumIslandConstraints, FPBDIslandSolverData& SolverData)
+
+	void FPBDSuspensionConstraints::AddConstraintsToGraph(FPBDIslandManager& IslandManager)
+	{
+		IslandManager.AddContainerConstraints(*this);
+	}
+
+	void FPBDSuspensionConstraints::AddBodies(FSolverBodyContainer& SolverBodyContainer)
+	{
+		for (int32 ConstraintIndex = 0; ConstraintIndex < NumConstraints(); ++ConstraintIndex)
+		{
+			AddBodies(ConstraintIndex, SolverBodyContainer);
+		}
+	}
+
+	void FPBDSuspensionConstraints::GatherInput(const FReal Dt)
+	{
+		for (int32 ConstraintIndex = 0; ConstraintIndex < NumConstraints(); ++ConstraintIndex)
+		{
+			GatherInput(ConstraintIndex, Dt);
+		}
+	}
+
+	void FPBDSuspensionConstraints::ScatterOutput(const FReal Dt)
+	{
+		for (int32 ConstraintIndex = 0; ConstraintIndex < NumConstraints(); ++ConstraintIndex)
+		{
+			ScatterOutput(ConstraintIndex, Dt);
+		}
+	}
+
+	void FPBDSuspensionConstraints::ApplyPositionConstraints(const FReal Dt, const int32 It, const int32 NumIts)
+	{
+		for (int32 ConstraintIndex = 0; ConstraintIndex < NumConstraints(); ++ConstraintIndex)
+		{
+			ApplyPositionConstraint(ConstraintIndex, Dt, It, NumIts);
+		}
+	}
+
+	void FPBDSuspensionConstraints::ApplyVelocityConstraints(const FReal Dt, const int32 It, const int32 NumIts)
+	{
+		for (int32 ConstraintIndex = 0; ConstraintIndex < NumConstraints(); ++ConstraintIndex)
+		{
+			ApplyVelocityConstraint(ConstraintIndex, Dt, It, NumIts);
+		}
+	}
+
+	void FPBDSuspensionConstraints::AddBodies(const TArrayView<int32>& ConstraintIndices, FSolverBodyContainer& SolverBodyContainer)
+	{
+		for (int32 ConstraintIndex : ConstraintIndices)
+		{
+			AddBodies(ConstraintIndex, SolverBodyContainer);
+		}
+	}
+
+	void FPBDSuspensionConstraints::GatherInput(const TArrayView<int32>& ConstraintIndices, const FReal Dt)
 	{
-		SolverData.GetConstraintIndices(ContainerId).Reset(NumIslandConstraints);
+		for (int32 ConstraintIndex : ConstraintIndices)
+		{
+			GatherInput(ConstraintIndex, Dt);
+		}
 	}
 
-	// @todo: This does per body work so can't be safely parallelized
-	void FPBDSuspensionConstraints::PreGatherInput(const FReal Dt, const int32 ConstraintIndex, FPBDIslandSolverData& SolverData)
+	void FPBDSuspensionConstraints::ScatterOutput(const TArrayView<int32>& ConstraintIndices, const FReal Dt)
 	{
-		SolverData.GetConstraintIndices(ContainerId).Add(ConstraintIndex);
+		for (int32 ConstraintIndex : ConstraintIndices)
+		{
+			ScatterOutput(ConstraintIndex, Dt);
+		}
+	}
 
-		ConstraintSolverBodies[ConstraintIndex] = SolverData.GetBodyContainer().FindOrAdd(ConstrainedParticles[ConstraintIndex], Dt);
+	void FPBDSuspensionConstraints::ApplyPositionConstraints(const TArrayView<int32>& ConstraintIndices, const FReal Dt, const int32 It, const int32 NumIts)
+	{
+		for (int32 ConstraintIndex : ConstraintIndices)
+		{
+			ApplyPositionConstraint(ConstraintIndex, Dt, It, NumIts);
+		}
+	}
+
+	void FPBDSuspensionConstraints::ApplyVelocityConstraints(const TArrayView<int32>& ConstraintIndices, const FReal Dt, const int32 It, const int32 NumIts)
+	{
+		for (int32 ConstraintIndex : ConstraintIndices)
+		{
+			ApplyVelocityConstraint(ConstraintIndex, Dt, It, NumIts);
+		}
+	}
+
+	void FPBDSuspensionConstraints::AddBodies(const int32 ConstraintIndex, FSolverBodyContainer& SolverBodyContainer)
+	{
+		ConstraintSolverBodies[ConstraintIndex] = SolverBodyContainer.FindOrAdd(ConstrainedParticles[ConstraintIndex]);
+	}
 
+	void FPBDSuspensionConstraints::GatherInput(const int32 ConstraintIndex, const FReal Dt)
+	{
 		ConstraintResults[ConstraintIndex].Reset();
 
 		// Hard-stop Collision Manifold Generation
@@ -222,43 +292,33 @@ namespace Chaos
 
 	}
 
-	void FPBDSuspensionConstraints::GatherInput(const FReal Dt, const int32 ConstraintIndex, const int32 Particle0Level, const int32 Particle1Level, FPBDIslandSolverData& SolverData)
-	{
-	}
-
-	void FPBDSuspensionConstraints::ScatterOutput(FReal Dt, FPBDIslandSolverData& SolverData)
+	void FPBDSuspensionConstraints::ScatterOutput(const int32 ConstraintIndex, FReal Dt)
 	{
-		for (int32 ConstraintIndex : SolverData.GetConstraintIndices(ContainerId))
+		if ((CollisionSolvers.Num() > 0) && (CollisionSolvers[ConstraintIndex] != nullptr) && (CollisionSolvers[ConstraintIndex]->NumManifoldPoints() > 0))
 		{
-			if ((CollisionSolvers.Num() > 0) && (CollisionSolvers[ConstraintIndex] != nullptr) && (CollisionSolvers[ConstraintIndex]->NumManifoldPoints() > 0))
-			{
-				const FPBDCollisionSolver* CollisionSolver = CollisionSolvers[ConstraintIndex];
-				const FPBDCollisionSolverManifoldPoint& ManifoldPoint = CollisionSolver->GetManifoldPoint(0);
-				ConstraintResults[ConstraintIndex].HardStopNetPushOut = ManifoldPoint.NetPushOutNormal * ManifoldPoint.WorldContactNormal;
-				ConstraintResults[ConstraintIndex].HardStopNetImpulse = ManifoldPoint.NetImpulseNormal * ManifoldPoint.WorldContactNormal;
-			}
-
-			ConstraintSolverBodies[ConstraintIndex] = nullptr;
-			CollisionSolvers[ConstraintIndex]->ResetSolverBodies();
+			const FPBDCollisionSolver* CollisionSolver = CollisionSolvers[ConstraintIndex];
+			const FPBDCollisionSolverManifoldPoint& ManifoldPoint = CollisionSolver->GetManifoldPoint(0);
+			ConstraintResults[ConstraintIndex].HardStopNetPushOut = ManifoldPoint.NetPushOutNormal * ManifoldPoint.WorldContactNormal;
+			ConstraintResults[ConstraintIndex].HardStopNetImpulse = ManifoldPoint.NetImpulseNormal * ManifoldPoint.WorldContactNormal;
 		}
+
+		ConstraintSolverBodies[ConstraintIndex] = nullptr;
+		CollisionSolvers[ConstraintIndex]->ResetSolverBodies();
 	}
 
-	bool FPBDSuspensionConstraints::ApplyPhase1Serial(const FReal Dt, const int32 It, const int32 NumIts, FPBDIslandSolverData& SolverData)
+	void FPBDSuspensionConstraints::ApplyPositionConstraint(const int32 ConstraintIndex, const FReal Dt, const int32 It, const int32 NumIts)
 	{
 		if (bChaos_Suspension_Hardstop_Enabled)
 		{
 			// Suspension Hardstop
-			for (int32 ConstraintIndex : SolverData.GetConstraintIndices(ContainerId))
+			const FPBDSuspensionSettings& Setting = ConstraintSettings[ConstraintIndex];
+			if (Setting.Enabled)
 			{
-				const FPBDSuspensionSettings& Setting = ConstraintSettings[ConstraintIndex];
-				if (Setting.Enabled)
+				FPBDCollisionSolver* CollisionSolver = CollisionSolvers[ConstraintIndex];
+				if ((CollisionSolver != nullptr) && (CollisionSolver->NumManifoldPoints() > 0))
 				{
-					FPBDCollisionSolver* CollisionSolver = CollisionSolvers[ConstraintIndex];
-					if ((CollisionSolver != nullptr) && (CollisionSolver->NumManifoldPoints() > 0))
-					{
-						FReal MaxPushoutValue = FMath::Min(Chaos_Suspension_MaxPushout, Chaos_Suspension_MaxPushoutVelocity * Dt);
-						CollisionSolver->SolvePositionNoFriction(FSolverReal(Dt), FSolverReal(MaxPushoutValue));
-					}
+					FReal MaxPushoutValue = FMath::Min(Chaos_Suspension_MaxPushout, Chaos_Suspension_MaxPushoutVelocity * Dt);
+					CollisionSolver->SolvePositionNoFriction(FSolverReal(Dt), FSolverReal(MaxPushoutValue));
 				}
 			}
 		}
@@ -266,39 +326,28 @@ namespace Chaos
 		if (bChaos_Suspension_Spring_Enabled)
 		{
 			// Suspension Spring
-			for (int32 ConstraintIndex : SolverData.GetConstraintIndices(ContainerId))
-			{
-				ApplySingle(Dt, ConstraintIndex);
-			}
+			ApplySingle(ConstraintIndex, Dt);
 		}
-
-		// @todo(chaos): early iteration termination in FPBDSuspensionConstraints
-		return true;
 	}
 
-	bool FPBDSuspensionConstraints::ApplyPhase2Serial(const FReal Dt, const int32 It, const int32 NumIts, FPBDIslandSolverData& SolverData)
+	void FPBDSuspensionConstraints::ApplyVelocityConstraint(const int32 ConstraintIndex, const FReal Dt, const int32 It, const int32 NumIts)
 	{
 		if (bChaos_Suspension_Hardstop_Enabled && bChaos_Suspension_VelocitySolve)
 		{
 			// Suspension Hardstop
-			for (int32 ConstraintIndex : SolverData.GetConstraintIndices(ContainerId))
+			const FPBDSuspensionSettings& Setting = ConstraintSettings[ConstraintIndex];
+			if (Setting.Enabled)
 			{
-				const FPBDSuspensionSettings& Setting = ConstraintSettings[ConstraintIndex];
-				if (Setting.Enabled)
+				FPBDCollisionSolver* CollisionSolver = CollisionSolvers[ConstraintIndex];
+				if ((CollisionSolver != nullptr) && (CollisionSolver->NumManifoldPoints() > 0))
 				{
-					FPBDCollisionSolver* CollisionSolver = CollisionSolvers[ConstraintIndex];
-					if ((CollisionSolver != nullptr) && (CollisionSolver->NumManifoldPoints() > 0))
-					{
-						CollisionSolver->SolveVelocity(FSolverReal(Dt), false);
-					}
+					CollisionSolver->SolveVelocity(FSolverReal(Dt), false);
 				}
 			}
 		}
-
-		return true;
 	}
 
-	void FPBDSuspensionConstraints::ApplySingle(const FReal Dt, int32 ConstraintIndex)
+	void FPBDSuspensionConstraints::ApplySingle(int32 ConstraintIndex, const FReal Dt)
 	{
 		check(ConstraintSolverBodies[ConstraintIndex] != nullptr);
 		FSolverBody& Body = *ConstraintSolverBodies[ConstraintIndex];
diff --git a/Engine/Source/Runtime/Experimental/Chaos/Private/ChaosSolverConfiguration.cpp b/Engine/Source/Runtime/Experimental/Chaos/Private/ChaosSolverConfiguration.cpp
--- a/Engine/Source/Runtime/Experimental/Chaos/Private/ChaosSolverConfiguration.cpp
+++ b/Engine/Source/Runtime/Experimental/Chaos/Private/ChaosSolverConfiguration.cpp
@@ -16,9 +16,9 @@ FChaosSolverConfiguration::FChaosSolverConfiguration()
 	, bGenerateCollisionData(false)
 	, bGenerateBreakData(false)
 	, bGenerateTrailingData(false)
-	, bGenerateContactGraph(true)
 	, Iterations_DEPRECATED(Chaos::FPBDRigidsEvolutionGBF::DefaultNumPositionIterations)
 	, PushOutIterations_DEPRECATED(Chaos::FPBDRigidsEvolutionGBF::DefaultNumVelocityIterations)
+	, bGenerateContactGraph_DEPRECATED(true)
 {
 }
 
diff --git a/Engine/Source/Runtime/Experimental/Chaos/Private/PBDRigidsSolver.cpp b/Engine/Source/Runtime/Experimental/Chaos/Private/PBDRigidsSolver.cpp
--- a/Engine/Source/Runtime/Experimental/Chaos/Private/PBDRigidsSolver.cpp
+++ b/Engine/Source/Runtime/Experimental/Chaos/Private/PBDRigidsSolver.cpp
@@ -41,6 +41,7 @@ DECLARE_DWORD_COUNTER_STAT(TEXT("NumKinematicBodies"), STAT_ChaosCounter_NumKine
 DECLARE_DWORD_COUNTER_STAT(TEXT("NumStaticBodies"), STAT_ChaosCounter_NumStaticBodies, STATGROUP_ChaosCounters);
 DECLARE_DWORD_COUNTER_STAT(TEXT("NumGeomCollBodies"), STAT_ChaosCounter_NumGeometryCollectionBodies, STATGROUP_ChaosCounters);
 DECLARE_DWORD_COUNTER_STAT(TEXT("NumIslands"), STAT_ChaosCounter_NumIslands, STATGROUP_ChaosCounters);
+DECLARE_DWORD_COUNTER_STAT(TEXT("NumIslandGroups"), STAT_ChaosCounter_NumIslandGroups, STATGROUP_ChaosCounters);
 DECLARE_DWORD_COUNTER_STAT(TEXT("NumContacts"), STAT_ChaosCounter_NumContacts, STATGROUP_ChaosCounters);
 DECLARE_DWORD_COUNTER_STAT(TEXT("NumValidConstraints"), STAT_ChaosCounter_NumValidConstraints, STATGROUP_ChaosCounters);
 DECLARE_DWORD_COUNTER_STAT(TEXT("NumActiveConstraints"), STAT_ChaosCounter_NumActiveConstraints, STATGROUP_ChaosCounters);
@@ -179,10 +180,6 @@ namespace Chaos
 		int32 ChaosSolverParticlePoolNumFrameUntilShrink = 30;
 		FAutoConsoleVariableRef CVarChaosSolverParticlePoolNumFrameUntilShrink(TEXT("p.Chaos.Solver.ParticlePoolNumFrameUntilShrink"), ChaosSolverParticlePoolNumFrameUntilShrink, TEXT("Num Frame until we can potentially shrink the pool"));
 
-		// Select the solver technique to use until we settle on the final one...
-		int32 ChaosSolver_SolverType = (int32)Chaos::EConstraintSolverType::QuasiPbd;
-		FAutoConsoleVariableRef CVarChaosSolverSolverType(TEXT("p.Chaos.Solver.SolverType"), ChaosSolver_SolverType, TEXT("0 = None; 1 = GbfPbd; 2 = Pbd; 3 = QuasiPbd"));
-
 		// Joint solver mode (linear vs non-linear)
 		bool bChaosSolverJointUseLinearSolver = true;
 		FAutoConsoleVariableRef CVarChaosSolverJointUseCachedSolver(TEXT("p.Chaos.Solver.Joint.UseLinearSolver"), bChaosSolverJointUseLinearSolver, TEXT("Use linear version of joint solver. (default is true"));
@@ -846,10 +843,6 @@ namespace Chaos
 
 		// Apply CVAR overrides if set
 		{
-			const EConstraintSolverType SolverType = (EConstraintSolverType)FMath::Clamp(ChaosSolver_SolverType, int32(EConstraintSolverType::None), int32(EConstraintSolverType::QuasiPbd));
-			GetEvolution()->GetCollisionConstraints().SetSolverType(SolverType);
-			GetJointConstraints().SetSolverType(SolverType);
-
 			if (ChaosSolverCollisionPositionFrictionIterations >= 0)
 			{
 				MEvolution->GetCollisionConstraints().SetPositionFrictionIterations(ChaosSolverCollisionPositionFrictionIterations);
@@ -1454,6 +1447,7 @@ CSV_CUSTOM_STAT(PhysicsCounters, Name, Value, ECsvCustomStatOp::Set);
 
 		// Constraint counts
 		CHAOS_COUNTER_STAT(NumIslands, GetEvolution()->GetConstraintGraph().NumIslands());
+		CHAOS_COUNTER_STAT(NumIslandGroups, GetEvolution()->GetIslandGroupManager().GetNumActiveGroups());
 		CHAOS_COUNTER_STAT(NumContacts, NumCollisionConstraints());
 		CHAOS_COUNTER_STAT(NumJoints, NumJointConstraints());
 
@@ -1729,7 +1723,6 @@ CSV_CUSTOM_STAT(PhysicsCounters, Name, Value, ECsvCustomStatOp::Set);
 		SetCollisionFilterSettings(InConfig.CollisionFilterSettings);
 		SetBreakingFilterSettings(InConfig.BreakingFilterSettings);
 		SetTrailingFilterSettings(InConfig.TrailingFilterSettings);
-		SetUseContactGraph(InConfig.bGenerateContactGraph);
 	}
 
 	FPBDRigidsSolver::~FPBDRigidsSolver()
diff --git a/Engine/Source/Runtime/Experimental/Chaos/Private/PhysicsProxy/SkeletalMeshPhysicsProxy.cpp b/Engine/Source/Runtime/Experimental/Chaos/Private/PhysicsProxy/SkeletalMeshPhysicsProxy.cpp
--- a/Engine/Source/Runtime/Experimental/Chaos/Private/PhysicsProxy/SkeletalMeshPhysicsProxy.cpp
+++ b/Engine/Source/Runtime/Experimental/Chaos/Private/PhysicsProxy/SkeletalMeshPhysicsProxy.cpp
@@ -11,8 +11,7 @@
 
 FSkeletalMeshPhysicsProxy::FSkeletalMeshPhysicsProxy(UObject* InOwner, const FInitFunc& InInitFunc)
 	: Base(InOwner)
-	, JointConstraints(Chaos::FPBDJointSolverSettings())
-	, JointConstraintsRule(JointConstraints)
+	, JointConstraints()
 	, NextInputProducerBuffer(nullptr)
 	, CurrentOutputConsumerBuffer(nullptr)
 	, bInitializedState(false)
diff --git a/Engine/Source/Runtime/Experimental/Chaos/Public/Chaos/ChaosDebugDrawDeclares.h b/Engine/Source/Runtime/Experimental/Chaos/Public/Chaos/ChaosDebugDrawDeclares.h
--- a/Engine/Source/Runtime/Experimental/Chaos/Public/Chaos/ChaosDebugDrawDeclares.h
+++ b/Engine/Source/Runtime/Experimental/Chaos/Public/Chaos/ChaosDebugDrawDeclares.h
@@ -16,11 +16,14 @@ namespace Chaos
 
 	class FPBDConstraintColor;
 
-	class FPBDConstraintGraph;
-
 	class FPBDJointConstraintHandle;
 
 	class FPBDJointConstraints;
 
+	class FPBDIslandManager;
+
 	class FSimulationSpace;
+
+
+	using FPBDConstraintGraph = FPBDIslandManager;
 }
diff --git a/Engine/Source/Runtime/Experimental/Chaos/Public/Chaos/Collision/CollisionApplyType.h b/Engine/Source/Runtime/Experimental/Chaos/Public/Chaos/Collision/CollisionApplyType.h
--- a/Engine/Source/Runtime/Experimental/Chaos/Public/Chaos/Collision/CollisionApplyType.h
+++ b/Engine/Source/Runtime/Experimental/Chaos/Public/Chaos/Collision/CollisionApplyType.h
@@ -1,17 +1,6 @@
 // Copyright Epic Games, Inc. All Rights Reserved.
 #pragma once
 
-namespace Chaos
-{
-	// @todo(chaos): move this enum to a constraint type independent location
-
-	/**
-	 * What solver method to use for Collision Resolution
-	 */
-	enum class EConstraintSolverType
-	{
-		None,
-		StandardPbd,				// Solve for position in the Apply step with standard PBD algorithm. Calculate Implicit Velocities. Do nothing in the ApplyPushOut step (Original Chaos RBAN Solver)
-		QuasiPbd,					// Solve for position in the Apply step with standard PBD algorithm. Calculate Implicit Velocities. Solve for normal velocity in the ApplyPushOut step (New Chaos Solver)
-	};
-}
\ No newline at end of file
+//
+// TO BE REMOBED
+//
diff --git a/Engine/Source/Runtime/Experimental/Chaos/Public/Chaos/Collision/CollisionConstraintAllocator.h b/Engine/Source/Runtime/Experimental/Chaos/Public/Chaos/Collision/CollisionConstraintAllocator.h
--- a/Engine/Source/Runtime/Experimental/Chaos/Public/Chaos/Collision/CollisionConstraintAllocator.h
+++ b/Engine/Source/Runtime/Experimental/Chaos/Public/Chaos/Collision/CollisionConstraintAllocator.h
@@ -8,6 +8,7 @@
 #include "Chaos/Collision/CollisionKeys.h"
 #include "Chaos/Collision/PBDCollisionConstraint.h"
 #include "Chaos/Collision/ParticlePairMidPhase.h"
+#include "Chaos/ObjectPool.h"
 #include "Chaos/ParticleHandle.h"
 #include "ChaosStats.h"
 
@@ -37,21 +38,66 @@ namespace Chaos
 	 * The Midphase list is pruned at the end of each tick so if particles are destroyed or a particle pair is no longer 
 	 * overlapping.
 	 * 
+	 * NOTE: To reduce RBAN memory use, we do not create any collision blocks until the first call to CreateCollisionConstraint
 	*/
 	class CHAOS_API FCollisionConstraintAllocator
 	{
 	public:
-		FCollisionConstraintAllocator()
+		FCollisionConstraintAllocator(const int32 NumCollisionsPerBlock = 1000)
 			: ParticlePairMidPhases()
 			, ActiveConstraints()
 			, ActiveSweptConstraints()
 			, CurrentEpoch(0)
 			, bInCollisionDetectionPhase(false)
+			, ConstraintPool(NumCollisionsPerBlock, 0)
 		{
 		}
 
 		~FCollisionConstraintAllocator()
 		{
+			// Explicit clear so we don't have to worry about destruction order of pool and midphases
+			Reset();
+			ensure(ConstraintPool.GetNumAllocated() == 0);
+		}
+
+		FPBDCollisionConstraintPtr CreateCollisionConstraint()
+		{
+			// NOTE: Called from the CollisionDetection parellel-for loop.
+			FPBDCollisionConstraint* Constraint;
+			{
+				FScopeLock Lock(&ConstraintPoolCS);
+				Constraint = ConstraintPool.Alloc();
+			}
+
+			// NOTE: Collision constraints are never destroyed in parallel, so the deleter does not require a lock
+			// (See PruneExpiredItems which destroys midphases that own collisions - it is called synchronously).
+			return FPBDCollisionConstraintPtr(Constraint, FPBDCollisionConstraintDeleter(ConstraintPool));
+		}
+
+		FPBDCollisionConstraintPtr CreateCollisionConstraint(
+			FGeometryParticleHandle* Particle0,
+			const FImplicitObject* Implicit0,
+			const FPerShapeData* Shape0,
+			const FBVHParticles* Simplicial0,
+			const FRigidTransform3& ShapeRelativeTransform0,
+			FGeometryParticleHandle* Particle1,
+			const FImplicitObject* Implicit1,
+			const FPerShapeData* Shape1,
+			const FBVHParticles* Simplicial1,
+			const FRigidTransform3& ShapeRelativeTransform1,
+			const FReal CullDistance,
+			const bool bUseManifold,
+			const EContactShapesType ShapePairType)
+		{
+			// NOTE: Called from the CollisionDetection parellel-for loop.
+			FPBDCollisionConstraintPtr Constraint = CreateCollisionConstraint();
+
+			if (Constraint.IsValid())
+			{
+				FPBDCollisionConstraint::Make(Particle0, Implicit0, Shape0, Simplicial0, ShapeRelativeTransform0, Particle1, Implicit1, Shape1, Simplicial1, ShapeRelativeTransform1, CullDistance, bUseManifold, ShapePairType, *Constraint);
+			}
+
+			return Constraint;
 		}
 		
 		/**
@@ -452,6 +498,9 @@ namespace Chaos
 		
 		// The set of mid phases created this tick (i.e., for particle pairs that were not in the map)
 		mutable TLockFreePointerListLIFO<FParticlePairMidPhase> NewParticlePairMidPhases;
+
+		FPBDCollisionConstraintPool ConstraintPool;
+		mutable FCriticalSection ConstraintPoolCS;
 	};
 
 
@@ -469,10 +518,10 @@ namespace Chaos
 	{
 		for (auto& KVP : Constraints)
 		{
-			const TUniquePtr<FPBDCollisionConstraint>& Constraint = KVP.Value;
+			const FPBDCollisionConstraintPtr& Constraint = KVP.Value;
 
 			// If we only want active constraints, check the timestamp
-			if ((Constraint != nullptr) && (!bOnlyActive || (Constraint->GetContainerCookie().LastUsedEpoch >= LastEpoch)))
+			if (Constraint.IsValid() && (!bOnlyActive || (Constraint->GetContainerCookie().LastUsedEpoch >= LastEpoch)))
 			{
 				if (Visitor(*Constraint) == ECollisionVisitorResult::Stop)
 				{
@@ -488,10 +537,10 @@ namespace Chaos
 	{
 		for (auto& KVP : Constraints)
 		{
-			const TUniquePtr<FPBDCollisionConstraint>& Constraint = KVP.Value;
+			const FPBDCollisionConstraintPtr& Constraint = KVP.Value;
 
 			// If we only want active constraints, check the timestamp
-			if ((Constraint != nullptr) && (!bOnlyActive || (Constraint->GetContainerCookie().LastUsedEpoch >= LastEpoch)))
+			if (Constraint.IsValid() && (!bOnlyActive || (Constraint->GetContainerCookie().LastUsedEpoch >= LastEpoch)))
 			{
 				if (!bOnlyActive || (Visitor(*Constraint) == ECollisionVisitorResult::Stop))
 				{
diff --git a/Engine/Source/Runtime/Experimental/Chaos/Public/Chaos/Collision/PBDCollisionConstraint.h b/Engine/Source/Runtime/Experimental/Chaos/Public/Chaos/Collision/PBDCollisionConstraint.h
--- a/Engine/Source/Runtime/Experimental/Chaos/Public/Chaos/Collision/PBDCollisionConstraint.h
+++ b/Engine/Source/Runtime/Experimental/Chaos/Public/Chaos/Collision/PBDCollisionConstraint.h
@@ -27,7 +27,7 @@ namespace Chaos
 	class FSingleShapePairCollisionDetector;
 	class FSolverBody;
 	class FSolverBodyContainer;
-	class FPBDCollisionSolverContainer;
+	class FPBDCollisionContainerSolver;
 	class FPerShapeData;
 
 	UE_DEPRECATED(4.27, "Use FPBDCollisionConstraint instead")
@@ -221,8 +221,8 @@ namespace Chaos
 	 * Each contact constraint contains a Manifold, which is a set of contact points that approximate the
 	 * contact patch between the two shapes.
 	 * 
-	 * Contact constraints are allocated on the heap and have permanent addresses. They use intrusive handles
-	 * to reduce unnecessary indirection.
+	 * Contact constraints are allocated on the heap (Asee FCollisionConstraintAllocator) and have permanent addresses. 
+	 * They use intrusive handles to reduce unnecessary indirection.
 	 * 
 	*/
 	class CHAOS_API FPBDCollisionConstraint final : public FPBDCollisionConstraintHandle
@@ -242,10 +242,10 @@ namespace Chaos
 
 		/**
 		 * @brief Create a contact constraint
-		 * Initializes a constraint stored inline in an object. Only intended to be called once.
+		 * Initializes a constraint stored inline in an object. Only intended to be called once right after construction.
 		 * Does not reinitialize all data so not intended to reset a constraint for reuse with different particles etc.
 		*/
-		static void MakeInline(
+		static void Make(
 			FGeometryParticleHandle* Particle0,
 			const FImplicitObject* Implicit0,
 			const FPerShapeData* Shape0,
@@ -261,26 +261,6 @@ namespace Chaos
 			const EContactShapesType ShapesType,
 			FPBDCollisionConstraint& OutConstraint);
 
-		/**
-		 * @brief Create a contact constraint
-		 * Allocates a constraint on the heap, with a permanent address.
-		 * May return null if we hit the contact limit for the scene.
-		*/
-		static TUniquePtr<FPBDCollisionConstraint> Make(
-			FGeometryParticleHandle* Particle0,
-			const FImplicitObject* Implicit0,
-			const FPerShapeData* Shape0,
-			const FBVHParticles* Simplicial0,
-			const FRigidTransform3& ImplicitLocalTransform0,
-			FGeometryParticleHandle* Particle1,
-			const FImplicitObject* Implicit1,
-			const FPerShapeData* Shape1,
-			const FBVHParticles* Simplicial1,
-			const FRigidTransform3& ImplicitLocalTransform1,
-			const FReal InCullDistance,
-			const bool bInUseManifold,
-			const EContactShapesType ShapesType);
-
 		/**
 		 * @brief For use by the tri mesh and heighfield collision detection as a temporary measure
 		 * @see FHeightField::ContactManifoldImp, FTriangleMeshImplicitObject::ContactManifoldImp
@@ -583,12 +563,6 @@ namespace Chaos
 			SolverBodies[1] = InSolverBody1;
 		}
 
-		int32 GetSolverIndex() const { return SolverIndex; }
-		void SetSolverIndex(int32 InSolverIndex)
-		{
-			SolverIndex = InSolverIndex;
-		}
-
 		/**
 		 * @brief Whether this constraint was fully restored from a previous tick, and the manifold should be reused as-is
 		*/
@@ -814,9 +788,6 @@ namespace Chaos
 		// These are only needed here while we still have the legacy solvers (not QuasiPBD)
 		FSolverBody* SolverBodies[2];
 
-		// Stores the index into the solver container for this constraint
-		int32 SolverIndex;
-
 		// Simplex data from the last call to GJK, used to warm-start GJK
 		FGJKSimplexData GJKWarmStartData;
 
diff --git a/Engine/Source/Runtime/Experimental/Chaos/Public/Chaos/Collision/PBDCollisionConstraintHandle.h b/Engine/Source/Runtime/Experimental/Chaos/Public/Chaos/Collision/PBDCollisionConstraintHandle.h
--- a/Engine/Source/Runtime/Experimental/Chaos/Public/Chaos/Collision/PBDCollisionConstraintHandle.h
+++ b/Engine/Source/Runtime/Experimental/Chaos/Public/Chaos/Collision/PBDCollisionConstraintHandle.h
@@ -11,8 +11,7 @@ namespace Chaos
 {
 	class FPBDCollisionConstraints;
 	class FPBDCollisionConstraint;
-	class FPBDIslandSolverData;
-	class FPBDCollisionSolverContainer;
+	class FPBDCollisionContainerSolver;
 
 	/**
 	 * @brief Whether we should run CCD (swept collision) or not
@@ -76,22 +75,17 @@ namespace Chaos
 
 		ECollisionCCDType GetCCDType() const;
 
-		virtual void SetEnabled(bool InEnabled) override;
+		virtual void SetEnabled(bool InEnabled) override final;
 
-		virtual bool IsEnabled() const override;
+		virtual bool IsEnabled() const override final;
 
-		virtual bool IsProbe() const override;
-
-		//FVec3 GetContactLocation() const;
+		virtual bool IsProbe() const override final;
 
 		FVec3 GetAccumulatedImpulse() const;
 
 		// Declared final so that TPBDConstraintGraphRuleImpl::AddToGraph() does not need to hit vtable
 		virtual FParticlePair GetConstrainedParticles() const override final;
 
-		void PreGatherInput(const FReal Dt, FPBDIslandSolverData& SolverData);
-		void GatherInput(FReal Dt, const int32 Particle0Level, const int32 Particle1Level, FPBDIslandSolverData& SolverData);
-
 		FSolverBody* GetSolverBody0();
 		FSolverBody* GetSolverBody1();
 
diff --git a/Engine/Source/Runtime/Experimental/Chaos/Public/Chaos/Collision/ParticlePairMidPhase.h b/Engine/Source/Runtime/Experimental/Chaos/Public/Chaos/Collision/ParticlePairMidPhase.h
--- a/Engine/Source/Runtime/Experimental/Chaos/Public/Chaos/Collision/ParticlePairMidPhase.h
+++ b/Engine/Source/Runtime/Experimental/Chaos/Public/Chaos/Collision/ParticlePairMidPhase.h
@@ -7,6 +7,7 @@
 #include "Chaos/CollisionResolutionTypes.h"
 #include "Chaos/Collision/CollisionKeys.h"
 #include "Chaos/Collision/CollisionVisitor.h"
+#include "Chaos/ObjectPool.h"
 #include "Chaos/ParticleHandleFwd.h"
 
 #include "ProfilingDebugging/CsvProfiler.h"
@@ -21,6 +22,11 @@ namespace Chaos
 	class FPerShapeData;
 	class FSingleShapePairCollisionDetector;
 
+	// Collision constraints are stored in an ObjectPool. @see FCollisionConstraintAllocator
+	using FPBDCollisionConstraintPool = TObjectPool<FPBDCollisionConstraint>;
+	using FPBDCollisionConstraintDeleter = TObjectPoolDeleter<FPBDCollisionConstraint>;
+	using FPBDCollisionConstraintPtr = TUniquePtr<FPBDCollisionConstraint, FPBDCollisionConstraintDeleter>;
+
 	/**
 	 * @brief Handles collision detection for a pair of simple shapes (i.e., not compound shapes)
 	 * 
@@ -58,7 +64,7 @@ namespace Chaos
 		*/
 		inline bool IsUsedSince(const int32 Epoch) const
 		{
-			return (Constraint != nullptr) && (LastUsedEpoch >= Epoch);
+			return Constraint.IsValid() && (LastUsedEpoch >= Epoch);
 		}
 
 		/**
@@ -120,7 +126,7 @@ namespace Chaos
 		void CreateConstraint(const FReal CullDistance, FCollisionContext& Context);
 
 		FParticlePairMidPhase& MidPhase;
-		TUniquePtr<FPBDCollisionConstraint> Constraint;
+		FPBDCollisionConstraintPtr Constraint;
 		FGeometryParticleHandle* Particle0;
 		FGeometryParticleHandle* Particle1;
 		const FPerShapeData* Shape0;
@@ -251,7 +257,7 @@ namespace Chaos
 		void PruneConstraints();
 
 		FParticlePairMidPhase& MidPhase;
-		TMap<uint32, TUniquePtr<FPBDCollisionConstraint>> Constraints;
+		TMap<uint32, FPBDCollisionConstraintPtr> Constraints;
 		TArray<FPBDCollisionConstraint*> NewConstraints;
 		FGeometryParticleHandle* Particle0;
 		FGeometryParticleHandle* Particle1;
diff --git a/Engine/Source/Runtime/Experimental/Chaos/Public/Chaos/Collision/SolverCollisionContainer.h b/Engine/Source/Runtime/Experimental/Chaos/Public/Chaos/Collision/SolverCollisionContainer.h
--- a/Engine/Source/Runtime/Experimental/Chaos/Public/Chaos/Collision/SolverCollisionContainer.h
+++ b/Engine/Source/Runtime/Experimental/Chaos/Public/Chaos/Collision/SolverCollisionContainer.h
@@ -8,16 +8,14 @@
 
 namespace Chaos
 {
-	namespace Collisions
-	{
-		struct FContactParticleParameters;
-		struct FContactIterationParameters;
-	}
-
+	class FPBDIslandConstraint;
 	class FPBDCollisionConstraint;
+	class FPBDCollisionConstraints;
 	class FPBDCollisionSolver;
 	class FPBDCollisionSolverAdapter;
 	class FPBDCollisionSolverManifoldPoint;
+	class FPBDIsland;
+	class FSolverBody;
 	class FSolverBodyContainer;
 
 	/**
@@ -48,49 +46,52 @@ namespace Chaos
 	};
 
 	/**
-	 * @brief A container of low-level data used to solve collision constraints
+	 * The solver for a set of collision constraints. This collects all the data required to solve a set of collision
+	 * constraints into a contiguous, ordered buffer.
 	*/
-	class FPBDCollisionSolverContainer : public FConstraintSolverContainer
+	class FPBDCollisionContainerSolver : public FConstraintContainerSolver
 	{
 	public:
-		FPBDCollisionSolverContainer();
-		~FPBDCollisionSolverContainer();
+		FPBDCollisionContainerSolver(const FPBDCollisionConstraints& InConstraintContainer, const int32 InPriority);
+		~FPBDCollisionContainerSolver();
 
 		int32 NumSolvers() const { return CollisionSolvers.Num(); }
 
-		virtual void Reset(const int32 InMaxCollisions) override;
+		virtual void Reset(const int32 InMaxCollisions) override final;
 
-		/** 
-		 * Set the number of collision solvers in the container
-		 * @param MaxCollisions Max number of collision solvers
-		 */
-		void SetNum(const int32 MaxCollisions);
+		virtual int32 GetNumConstraints() const override final
+		{
+			return NumSolvers();
+		}
 
 		// Set whether we are using deferred collision detection
 		void SetIsDeferredCollisionDetection(const bool bDeferred) { bDeferredCollisionDetection = bDeferred; }
 
-		// Add a solver constraint for the specified constraint and gather the required solver data
-		void PreAddConstraintSolver(const FReal Dt, FPBDCollisionConstraint& Constraint, FSolverBodyContainer& SolverBodyContainer, int32& SolverIndex);
-		void AddConstraintSolver(const FReal Dt, FPBDCollisionConstraint& Constraint, const int32 Particle0Level, const int32 Particle1Level, FSolverBodyContainer& SolverBodyContainer, const FPBDCollisionSolverSettings& SolverSettings);
-
-		bool SolvePositionSerial(const FReal Dt, const int32 It, const int32 NumIts, const int32 BeginIndex, const int32 EndIndex, const FPBDCollisionSolverSettings& SolverSettings);
-		bool SolvePositionParallel(const FReal Dt, const int32 It, const int32 NumIts, const int32 BeginIndex, const int32 EndIndex, const FPBDCollisionSolverSettings& SolverSettings);
-
-		bool SolveVelocitySerial(const FReal Dt, const int32 It, const int32 NumIts, const int32 BeginIndex, const int32 EndIndex, const FPBDCollisionSolverSettings& SolverSettings);
-		bool SolveVelocityParallel(const FReal Dt, const int32 It, const int32 NumIts, const int32 BeginIndex, const int32 EndIndex, const FPBDCollisionSolverSettings& SolverSettings);
-
-		void ScatterOutput(const FReal Dt, const int32 BeginIndex, const int32 EndIndex);
+		//
+		// IslandGroup API
+		//
+		virtual void AddConstraints() override final;
+		virtual void AddConstraints(const TArrayView<FPBDIslandConstraint>& ConstraintHandles) override final;
+		virtual void AddBodies(FSolverBodyContainer& SolverBodyContainer) override final;
+		virtual void GatherInput(const FReal Dt) override final;
+		virtual void GatherInput(const FReal Dt, const int32 BeginIndex, const int32 EndIndex) override final;
+		virtual void ScatterOutput(const FReal Dt) override final;
+		virtual void ScatterOutput(const FReal Dt, const int32 BeginIndex, const int32 EndIndex) override final;
+		virtual void ApplyPositionConstraints(const FReal Dt, const int32 It, const int32 NumIts) override final;
+		virtual void ApplyVelocityConstraints(const FReal Dt, const int32 It, const int32 NumIts) override final;
+		virtual void ApplyProjectionConstraints(const FReal Dt, const int32 It, const int32 NumIts) override final;
 
 	private:
+		void AddConstraint(FPBDCollisionConstraint& Constraint);
 		void UpdatePositionShockPropagation(const FReal Dt, const int32 It, const int32 NumIts, const int32 BeginIndex, const int32 EndIndex, const FPBDCollisionSolverSettings& SolverSettings);
 		void UpdateVelocityShockPropagation(const FReal Dt, const int32 It, const int32 NumIts, const int32 BeginIndex, const int32 EndIndex, const FPBDCollisionSolverSettings& SolverSettings);
-		bool SolvePositionImpl(const FReal Dt, const int32 It, const int32 NumIts, const int32 BeginIndex, const int32 EndIndex, const FPBDCollisionSolverSettings& SolverSettings, const bool bParallel);
-		bool SolvePositionWithFrictionImpl(const FReal Dt, const int32 BeginIndex, const int32 EndIndex, const FReal MaxPushOut, const bool bParallel);
-		bool SolvePositionNoFrictionImpl(const FReal Dt, const int32 BeginIndex, const int32 EndIndex, const FReal MaxPushOut, const bool bParallel);
-		bool SolveVelocityImpl(const FReal Dt, const int32 It, const int32 NumIts, const int32 BeginIndex, const int32 EndIndex, const FPBDCollisionSolverSettings& SolverSettings, const bool bParallel);
-		void ScatterOutputImpl(const FReal Dt, const int32 BeginIndex, const int32 EndIndex, const bool bParallel);
+		bool SolvePositionImpl(const FReal Dt, const int32 It, const int32 NumIts, const int32 BeginIndex, const int32 EndIndex, const FPBDCollisionSolverSettings& SolverSettings);
+		bool SolvePositionWithFrictionImpl(const FReal Dt, const int32 BeginIndex, const int32 EndIndex, const FReal MaxPushOut);
+		bool SolvePositionNoFrictionImpl(const FReal Dt, const int32 BeginIndex, const int32 EndIndex, const FReal MaxPushOut);
+		bool SolveVelocityImpl(const FReal Dt, const int32 It, const int32 NumIts, const int32 BeginIndex, const int32 EndIndex, const FPBDCollisionSolverSettings& SolverSettings);
 		void UpdateCollisions(const FReal InDt, const int32 BeginIndex, const int32 EndIndex);
 
+		const FPBDCollisionConstraints& ConstraintContainer;
 		TArray<FPBDCollisionSolverAdapter> CollisionSolvers;
 		bool bPerIterationCollisionDetection;
 		bool bDeferredCollisionDetection;
diff --git a/Engine/Source/Runtime/Experimental/Chaos/Public/Chaos/Collision/SpatialAccelerationBroadPhase.h b/Engine/Source/Runtime/Experimental/Chaos/Public/Chaos/Collision/SpatialAccelerationBroadPhase.h
--- a/Engine/Source/Runtime/Experimental/Chaos/Public/Chaos/Collision/SpatialAccelerationBroadPhase.h
+++ b/Engine/Source/Runtime/Experimental/Chaos/Public/Chaos/Collision/SpatialAccelerationBroadPhase.h
@@ -206,7 +206,7 @@ namespace Chaos
 			{
 				const bool bBody1Bounded = HasBoundingBox(Particle1);
 				{
-					SCOPE_CYCLE_COUNTER(STAT_Collisions_AABBTree);
+					//SCOPE_CYCLE_COUNTER(STAT_Collisions_AABBTree);
 					if (bBody1Bounded)
 					{
 						// @todo(chaos): cache this on the particle?
@@ -234,7 +234,7 @@ namespace Chaos
 					}
 				}
 
-				SCOPE_CYCLE_COUNTER(STAT_Collisions_Filtering);
+				//SCOPE_CYCLE_COUNTER(STAT_Collisions_Filtering);
 				const int32 NumPotentials = PotentialIntersections.Num();
 				int32 NumIntoNarrowPhase = 0;
 				TArray<FParticlePairMidPhase*> MidPhasePairs;
