Chaos - Constraint Solver Refactor

diff --git a/Engine/Source/Runtime/Experimental/Chaos/Public/Chaos/ConstraintHandle.h b/Engine/Source/Runtime/Experimental/Chaos/Public/Chaos/ConstraintHandle.h
--- a/Engine/Source/Runtime/Experimental/Chaos/Public/Chaos/ConstraintHandle.h
+++ b/Engine/Source/Runtime/Experimental/Chaos/Public/Chaos/ConstraintHandle.h
@@ -66,9 +66,9 @@ namespace Chaos
 			{
 				return true;
 			}
-			if (TypeID.BaseType != nullptr)
+			if (BaseType != nullptr)
 			{
-				return IsA(*TypeID.BaseType);
+				return BaseType->IsA(TypeID);
 			}
 			return false;
 		}
@@ -95,7 +95,7 @@ namespace Chaos
 
 		FConstraintHandle() 
 			: ConstraintContainer(nullptr)
-			, GraphIndex(INDEX_NONE) 
+			, GraphIndex(INDEX_NONE)
 		{
 		}
 
@@ -124,19 +124,19 @@ namespace Chaos
 			return ConstraintContainer;
 		}
 
-		int32 ConstraintGraphIndex() const
+		bool IsInConstraintGraph() const
 		{
-			return GraphIndex;
+			return (GraphIndex != INDEX_NONE);
 		}
 
-		void SetConstraintGraphIndex(int32 InIndex)
+		int32 GetConstraintGraphIndex() const
 		{
-			GraphIndex = InIndex;
+			return GraphIndex;
 		}
 
-		bool IsInConstraintGraph() const
+		void SetConstraintGraphIndex(const int32 InIndex)
 		{
-			return (GraphIndex != INDEX_NONE);
+			GraphIndex = InIndex;
 		}
 
 		virtual TVec2<FGeometryParticleHandle*> GetConstrainedParticles() const = 0;
@@ -182,7 +182,6 @@ namespace Chaos
 
 		FPBDConstraintContainer* ConstraintContainer;
 		
-		// @todo(chaos): move constraint graph index to base constraint container
 		int32 GraphIndex;
 	};
 
@@ -249,110 +248,6 @@ namespace Chaos
 		}
 	};
 
-	/**
-	 * Base class for handles to constraints in an index-based container
-	 */
-	class CHAOS_API FIndexedConstraintHandle : public FConstraintHandle
-	{
-	public:
-		using FGeometryParticleHandle = TGeometryParticleHandle<FReal, 3>;
-
-		FIndexedConstraintHandle() 
-			: FConstraintHandle()
-			, ConstraintIndex(INDEX_NONE)
-		{
-		}
-
-		FIndexedConstraintHandle(FPBDConstraintContainer* InContainer, int32 InConstraintIndex)
-			: FConstraintHandle(InContainer)
-			, ConstraintIndex(InConstraintIndex)
-		{
-		}
-
-		virtual ~FIndexedConstraintHandle()
-		{
-		}
-
-		virtual bool IsValid() const override
-		{
-			return (ConstraintIndex != INDEX_NONE) && FConstraintHandle::IsValid();
-		}
-
-		int32 GetConstraintIndex() const
-		{
-			return ConstraintIndex;
-		}
-
-		static const FConstraintHandleTypeID& StaticType()
-		{
-			static FConstraintHandleTypeID STypeID(TEXT("FIndexedConstraintHandle"), &FConstraintHandle::StaticType());
-			return STypeID;
-		}
-
-	protected:
-		friend class FPBDIndexedConstraintContainer;
-
-		int32 ConstraintIndex;
-	};
-
-
-	/**
-	 * Utility base class for ConstraintHandles. Provides basic functionality common to most constraint containers.
-	 */
-	template<typename T_CONTAINER>
-	class CHAOS_API TIndexedContainerConstraintHandle : public FIndexedConstraintHandle
-	{
-	public:
-		using Base = FIndexedConstraintHandle;
-		using FGeometryParticleHandle = typename Base::FGeometryParticleHandle;
-		using FConstraintContainer = T_CONTAINER;
-
-		TIndexedContainerConstraintHandle()
-			: FIndexedConstraintHandle()
-		{
-		}
-		
-		TIndexedContainerConstraintHandle(FConstraintContainer* InConstraintContainer, int32 InConstraintIndex)
-			: FIndexedConstraintHandle(InConstraintContainer, InConstraintIndex)
-		{
-		}
-
-		inline virtual void SetEnabled(bool bInEnabled) override
-		{
-			if (ConcreteContainer() != nullptr)
-			{
-				ConcreteContainer()->SetConstraintEnabled(ConstraintIndex, bInEnabled);
-			}
-		}
-
-		inline virtual bool IsEnabled() const override
-		{
-			if (ConcreteContainer() != nullptr)
-			{
-				return ConcreteContainer()->IsConstraintEnabled(ConstraintIndex);
-			}
-			return false;
-		}
-
-		// @todo(chaos): Make this a virtual on FConstraintContainer and move to base class
-		void RemoveConstraint()
-		{
-			ConcreteContainer()->RemoveConstraint(ConstraintIndex);
-		}
-
-	protected:
-		FConstraintContainer* ConcreteContainer()
-		{
-			return static_cast<FConstraintContainer*>(ConstraintContainer);
-		}
-
-		const FConstraintContainer* ConcreteContainer() const
-		{
-			return static_cast<const FConstraintContainer*>(ConstraintContainer);
-		}
-
-		using Base::ConstraintIndex;
-	};
 
 
 	/**
diff --git a/Engine/Source/Runtime/Experimental/Chaos/Public/Chaos/Evolution/PBDMinEvolution.h b/Engine/Source/Runtime/Experimental/Chaos/Public/Chaos/Evolution/PBDMinEvolution.h
--- a/Engine/Source/Runtime/Experimental/Chaos/Public/Chaos/Evolution/PBDMinEvolution.h
+++ b/Engine/Source/Runtime/Experimental/Chaos/Public/Chaos/Evolution/PBDMinEvolution.h
@@ -4,14 +4,16 @@
 #include "Chaos/Core.h"
 #include "Chaos/ArrayCollectionArray.h"
 #include "Chaos/Collision/CollisionApplyType.h"
+#include "Chaos/Evolution/ConstraintGroupSolver.h"
 #include "Chaos/Evolution/SimulationSpace.h"
-#include "Chaos/Evolution/SolverDatas.h"
+#include "Chaos/Evolution/SolverBodyContainer.h"
 #include "Chaos/ParticleHandleFwd.h"
 
 
 namespace Chaos
 {
 	class FBasicCollisionDetector;
+	class FConstraintContainerSolver;
 	class FPBDCollisionConstraints;
 	class FSimpleConstraintRule;
 	class FPBDRigidsSOAs;
@@ -27,24 +29,19 @@ namespace Chaos
 	class CHAOS_API FPBDMinEvolution
 	{
 	public:
-		// @todo(ccaulfield): make it so that CollisionDetection is plugged in with a constraint rule...
-
 		using FCollisionDetector = FBasicCollisionDetector;
 		using FEvolutionCallback = TFunction<void()>;
 		using FRigidParticleSOAs = FPBDRigidsSOAs;
 
 		FPBDMinEvolution(FRigidParticleSOAs& InParticles, TArrayCollectionArray<FVec3>& InPrevX, TArrayCollectionArray<FRotation3>& InPrevR, FCollisionDetector& InCollisionDetector, const FReal InBoundsExtension);
+		~FPBDMinEvolution();
 
-		void AddConstraintRule(FSimpleConstraintRule* Rule);
+		void AddConstraintContainer(FPBDConstraintContainer& InContainer, const int32 Priority = 0);
+		void SetConstraintContainerPriority(const int32 ContainerId, const int32 Priority);
 
 		void Advance(const FReal StepDt, const int32 NumSteps, const FReal RewindDt);
 		void AdvanceOneTimeStep(const FReal Dt, const FReal StepFraction);
 
-		void SetSolverType(const EConstraintSolverType InSolverType)
-		{
-			SolverType = InSolverType;
-		}
-
 		void SetNumPositionIterations(const int32 NumIts)
 		{
 			NumPositionIterations = NumIts;
@@ -60,18 +57,6 @@ namespace Chaos
 			NumProjectionIterations = NumIts;
 		}
 
-		// Legacy
-		void SetNumIterations(const int32 NumIts)
-		{
-			NumApplyIterations = NumIts;
-		}
-
-		// Legacy
-		void SetNumPushOutIterations(const int32 NumIts)
-		{
-			NumApplyPushOutIterations = NumIts;
-		}
-
 		void SetGravity(const FVec3& G)
 		{
 			Gravity = G;
@@ -82,26 +67,6 @@ namespace Chaos
 			BoundsExtension = InBoundsExtension;
 		}
 
-		void SetPostIntegrateCallback(const FEvolutionCallback& Cb)
-		{
-			PostIntegrateCallback = Cb;
-		}
-
-		void SetPostDetectCollisionsCallback(const FEvolutionCallback& Cb)
-		{
-			PostDetectCollisionsCallback = Cb;
-		}
-
-		void SetPostApplyCallback(const FEvolutionCallback& Cb)
-		{
-			PostApplyCallback = Cb;
-		}
-
-		void SetPostApplyPushOutCallback(const FEvolutionCallback& Cb)
-		{
-			PostApplyPushOutCallback = Cb;
-		}
-
 		void SetSimulationSpace(const FSimulationSpace& InSimulationSpace)
 		{
 			SimulationSpace = InSimulationSpace;
@@ -132,9 +97,7 @@ namespace Chaos
 		void GatherInput(FReal Dt);
 		void ScatterOutput(FReal Dt);
 		void ApplyConstraintsPhase1(FReal Dt);
-		void UpdateVelocities(FReal Dt);
 		void ApplyConstraintsPhase2(FReal Dt);
-		void ApplyCorrections(FReal Dt);
 		void ApplyConstraintsPhase3(FReal Dt);
 
 		FRigidParticleSOAs& Particles;
@@ -143,13 +106,9 @@ namespace Chaos
 		TArrayCollectionArray<FVec3>& ParticlePrevXs;
 		TArrayCollectionArray<FRotation3>& ParticlePrevRs;
 
-		TArray<FSimpleConstraintRule*> ConstraintRules;
-		TArray<FSimpleConstraintRule*> PrioritizedConstraintRules;
-		FPBDIslandSolverData SolverData;
+		TArray<FPBDConstraintContainer*> ConstraintContainers;
+		FPBDSceneConstraintGroupSolver ConstraintSolver;
 
-		EConstraintSolverType SolverType;
-		int32 NumApplyIterations;
-		int32 NumApplyPushOutIterations;
 		int32 NumPositionIterations;
 		int32 NumVelocityIterations;
 		int32 NumProjectionIterations;
@@ -157,10 +116,5 @@ namespace Chaos
 		FVec3 Gravity;
 		FSimulationSpaceSettings SimulationSpaceSettings;
 		FSimulationSpace SimulationSpace;
-
-		FEvolutionCallback PostIntegrateCallback;
-		FEvolutionCallback PostDetectCollisionsCallback;
-		FEvolutionCallback PostApplyCallback;
-		FEvolutionCallback PostApplyPushOutCallback;
 	};
 }
diff --git a/Engine/Source/Runtime/Experimental/Chaos/Public/Chaos/Evolution/SolverBodyContainer.h b/Engine/Source/Runtime/Experimental/Chaos/Public/Chaos/Evolution/SolverBodyContainer.h
--- a/Engine/Source/Runtime/Experimental/Chaos/Public/Chaos/Evolution/SolverBodyContainer.h
+++ b/Engine/Source/Runtime/Experimental/Chaos/Public/Chaos/Evolution/SolverBodyContainer.h
@@ -3,6 +3,7 @@
 
 #include "Chaos/Evolution/SolverBody.h"
 #include "Chaos/Framework/ContainerItemHandle.h"
+#include "Chaos/ObjectPool.h"
 #include "Chaos/ParticleHandle.h"
 #include "Chaos/Particle/ParticleUtilities.h"
 
@@ -15,23 +16,27 @@ namespace Chaos
 	class FSolverBodyContainer;
 
 	/**
-	 * @brief An FSolverBody wrapper that binds the solver to a particle and provides gather/scatter methods from/to the particle
-	 * Solver Bodies are used by the core constraint solver
+	 * @brief An FSolverBody wrapper that binds the SolverBody to a particle and provides gather/scatter methods from/to the particle
+	 * Solver Bodies are used by the core constraint solvers.
 	 * @see FSolverBody
 	*/
 	class FSolverBodyAdapter
 	{
 	public:
-		FSolverBodyAdapter(const FGenericParticleHandle& InParticle, const FReal Dt)
-			: Particle(InParticle)
+		FSolverBodyAdapter()
+			: Particle(nullptr)
 		{
-			GatherInput(Dt);
 		}
 		
+		FSolverBodyAdapter(const FGenericParticleHandle& InParticle)
+			: Particle(InParticle)
+		{
+		}
+
 		FSolverBody& GetSolverBody() { return SolverBody; }
 		const FGenericParticleHandle& GetParticle() const { return Particle; }
 
-		FORCEINLINE void GatherInput(const FReal Dt)
+		inline void GatherInput(const FReal Dt)
 		{
 			if (Particle.IsValid())
 			{
@@ -69,6 +74,7 @@ namespace Chaos
 				// valid for dynamic particle for which the SetInvILocal is already doing the job
 			}
 		}
+
 		void ScatterOutput();
 
 	private:
@@ -79,31 +85,42 @@ namespace Chaos
 	/**
 	 * The SolverBodies for a set of particles.
 	 * 
-	 * Each Island owns a SolverBodyContainer containing the data required for solving the constraints in the island.
-	 * Constraints in the island will hold pointers to SolverBodies in the Island's SolverBodyContainer.
+	 * Each IslandGroup owns a SolverBodyContainer containing the data required for solving the constraints in the 
+	 * group's islands. Constraints will hold pointers to SolverBodies in the Island's SolverBodyContainer.
 	 * ConstraintSolvers read and write to their SolverBodies and do not access particles directly.
 	 * 
 	 * SolverBodies are created at the start of the constraint solving phase, and destroyed at the end. They are
-	 * stored in a contiguous array in the order that they are accessed, for cache efficiency. The array
-	 * is guaranteed not to resize while accumulating bodies for the solving phase (it will assert if it does)
-	 * so pointers to elements are valid for the duration of the constraint solving phase, but not longer.
+	 * stored in a semi-contiguous (chunked) array in roughly the order that they are accessed for cache efficiency. 
+	 * Pointers to SolverBodies are persistent for the duration of the tick, but no longer.
 	 * 
-	 * \note This container holds all the state of bodies in an island. Dynamics will only appear in one island, 
+	 * @note This container holds all the state of bodies in an island. Dynamics will only appear in one island, 
 	 * but kinematics are in multiple and so their state will be duplicated in each island. This is ok - we do 
-	 * not update the state of any kinematic bodies in the constraint solvers. We also assume that the number
+	 * not write to the state of any kinematic bodies in the constraint solvers. We also assume that the number
 	 * of kinematics is small compared to the number of dynamics. If this is commonly untrue we may want to 
 	 * consider having a separate (global) container of kinematic solver bodies.
 	 */
 	class FSolverBodyContainer
 	{
+	private:
+		// Settings for the Solver Body array.
+		static const int32 NumBodiesPerChunk = 256;
+		static const int32 BodyArrayChunkSize = NumBodiesPerChunk * sizeof(FSolverBodyAdapter);
+
 	public:
+		using FSolverBodyArray = TChunkedArray<FSolverBodyAdapter, BodyArrayChunkSize>;
+
+		FSolverBodyContainer()
+			: bLocked(true)
+		{}
+
 
 		// Clear the bodies array and allocate enough space for at least MaxBodies bodies.
 		// @param MaxBodies The number of bodies that will get added to the container. The container asserts if we attempt to add more than this.
 		inline void Reset(int MaxBodies)
 		{
-			SolverBodies.Reset(MaxBodies);
+			SolverBodies.Empty(MaxBodies);
 			ParticleToIndexMap.Reset();
+			bLocked = false;
 		}
 
 		// The number of bodies in the container
@@ -115,7 +132,7 @@ namespace Chaos
 		// The maximum number of bodies the container can hold (until Reset() is called again)
 		inline int MaxItems() const
 		{
-			return SolverBodies.Max();
+			return TNumericLimits<int32>::Max();
 		}
 
 		// Get a pointer to the item at the specified index, or null for INDEX_NONE
@@ -154,8 +171,8 @@ namespace Chaos
 			return SolverBodies[Index];
 		}
 
-		/** Return all the solver bodies inside the container */
-		inline TArray<FSolverBodyAdapter>& GetBodies()
+		/** Get all the solver bodies inside the container */
+		inline FSolverBodyArray& GetBodies()
 		{
 			return SolverBodies;
 		}
@@ -170,14 +187,19 @@ namespace Chaos
 		// Add a solver body to represent the solver state of the particle
 		// This should ideally be called in the order in which the bodies will be accessed
 		// (or as close as we can get, given most constraints access 2 bodies so there is no perfect order)
-		FSolverBody* FindOrAdd(FGenericParticleHandle InParticle, const FReal Dt);
+		// NOTE: This does not populate the solver body with any data from the particle - we must also call
+		// GatherInput for that to happen (GatherInput is highly parallizable but FindOrAdd is not)
+		FSolverBody* FindOrAdd(FGenericParticleHandle InParticle);
 
 		// Collect all the data we need from the particles represented by our SolverBodies
-		//void GatherInput();
+		void GatherInput(const FReal Dt, const int32 BeginIndex, const int32 EndIndex);
 
 		// Scatter the solver results back to the particles represented by our SolverBodies
 		void ScatterOutput();
 
+		// Scatter the solver results back to the particles represented by our SolverBodies. Range version used for parallization.
+		void ScatterOutput(const int32 BeginIndex, const int32 EndIndex);
+
 		// Recalculate the velocities of all bodies based on their transform deltas
 		void SetImplicitVelocities(FReal Dt);
 
@@ -187,15 +209,26 @@ namespace Chaos
 		// Can be called after ApplyCorrections to update inertia to match the current transform
 		void UpdateRotationDependentState();
 
+		// Indicate that no more bodies should be added until the next tick (when Reset is called).
+		// Used to ensure that no bodies are added in parallel calls to GatherInput.
+		void Lock() 
+		{ 
+			bLocked = true;
+		}
+
 	private:
-		int32 AddParticle(FGenericParticleHandle InParticle, const FReal Dt);
+		int32 AddParticle(FGenericParticleHandle InParticle);
 
-		// Solver bodies, usually collected in the order in which they are accessed
-		TArray<FSolverBodyAdapter> SolverBodies;
+		FSolverBodyArray SolverBodies;
 
 		// Used to determine if a (non-dynamic) body is already present in the container. 
-		// Dynamic bodies are not in this map - the particle stores its index in this case.
+		// Dynamic bodies are not in this map - the particle stores its index in this case because
+		// dynamic bodies can only be in one island whereas non-dynamics may be const-referenced in many.
 		TMap<FGenericParticleHandle, int32> ParticleToIndexMap;
+
+		// To trap incorrect use of the solve bodies. We do not support adding bodies once we have started
+		// to call GatherInput because GatherInput is parallelized but FindOrAdd is not
+		bool bLocked;
 	};
 
 }
diff --git a/Engine/Source/Runtime/Experimental/Chaos/Public/Chaos/Evolution/SolverConstraintContainer.h b/Engine/Source/Runtime/Experimental/Chaos/Public/Chaos/Evolution/SolverConstraintContainer.h
--- a/Engine/Source/Runtime/Experimental/Chaos/Public/Chaos/Evolution/SolverConstraintContainer.h
+++ b/Engine/Source/Runtime/Experimental/Chaos/Public/Chaos/Evolution/SolverConstraintContainer.h
@@ -1,21 +1,181 @@
 // Copyright Epic Games, Inc. All Rights Reserved.
 #pragma once
 
-namespace Chaos
-{
+#include "Chaos/ConstraintHandle.h"
 
-/** Parent base class for all the solver constraint containers */
-class CHAOS_API FConstraintSolverContainer
+namespace Chaos
 {
-public:
+	class FConstraintHandleHolder;
+	class FPBDIslandConstraint;
+	class FSolverBodyContainer;
+	class FPBDIsland;
 
-	/** Virtual destructor */
-	virtual ~FConstraintSolverContainer() {}
-		
-	/** Reset the container to have zero constraints
-	* @param MaxConstraints Maximum number of constraints that could be in the container
+	/** 
+	 * Base class for all the solver for a set of constraints of a specific type.
+	 * A SolverContainer is used to solve a set constraints in sequential order. There will
+	 * be one solver container for each thread on which we solve constraints (see FPBDIslandConstraintGroupSolver). 
+	 * How constraints are assigned to groups depends on the constraint type and settings, but
+	 * usually a group contains all constraints from one or more islands (unless we are coloring).
 	*/
-	virtual void Reset(const int32 MaxConstraints) {}
-};
+	class CHAOS_API FConstraintContainerSolver
+	{
+	public:
+
+		FConstraintContainerSolver(const int32 InPriority)
+			: Priority(InPriority)
+		{}
+
+		virtual ~FConstraintContainerSolver() {}
+	
+		/**
+		 * Set the solver priority.
+		 * Solvers are sorted by priority. Lower values are solved first so solvers with higher
+		 * priorty values with "win" over lower ones.
+		 * @see FPBDConstraintGroupSolver
+		*/
+		void SetPriority(const int32 InPriority)
+		{
+			Priority = InPriority;
+		}
+
+		/**
+		 * Get the solver priority
+		*/
+		int32 GetPriority() const
+		{
+			return Priority;
+		}
+
+		/** 
+		 * Set the maximum number of constraints the solver will have to handle. 
+		 * This will be called only once per tick, so containers resized here will not have to resize again this tick
+		 * so that pointers to elements in the container will remain valid for the tick (but not beyond).
+		*/
+		virtual void Reset(const int32 MaxConstraints) = 0;
+
+		virtual int32 GetNumConstraints() const = 0;
+
+		virtual void AddConstraints() = 0;
+
+		/**
+		 * Add a set of constraints to the solver container. This can be called multiple times: once for each island a group, but
+		 * there will never be more constraints added than specified in Reset().
+		 * NOTE: this should not do any actual data gathering - it should just add to the list of constraints in this group. All data
+		 * gathering is handled in GatherInput.
+		*/
+		virtual void AddConstraints(const TArrayView<FPBDIslandConstraint>& Constraints) = 0;
+
+		/**
+		 * Add all the required bodies to the body container (required for the constraints added with AddConstraints)
+		*/
+		virtual void AddBodies(FSolverBodyContainer& SolverBodyContainer) = 0;
+
+		virtual void GatherInput(const FReal Dt) = 0;
+		virtual void GatherInput(const FReal Dt, const int32 BeginIndex, const int32 EndIndex) = 0;
+		virtual void ScatterOutput(const FReal Dt) = 0;
+		virtual void ScatterOutput(const FReal Dt, const int32 BeginIndex, const int32 EndIndex) = 0;
+
+		/**
+		 * Apply the position solve to all constraints in the container
+		*/
+		virtual void ApplyPositionConstraints(const FReal Dt, const int32 It, const int32 NumIts) = 0;
+
+		/**
+		 * Apply the velocity solve to all constraints in the container
+		*/
+		virtual void ApplyVelocityConstraints(const FReal Dt, const int32 It, const int32 NumIts) = 0;
+
+		/**
+		 * Apply the projection solve to all constraints in the container
+		*/
+		virtual void ApplyProjectionConstraints(const FReal Dt, const int32 It, const int32 NumIts) = 0;
+
+	private:
+		int32 Priority;
+	};
 	
+	/**
+	 * A constraint solver for use with a simple (non-graph-based) evolution (RBAN) and constraint containers
+	 * with builtin-in solvers (Joints, Suspension, but not Collisions).
+	 * @see TIndexedConstraintContainerSolver, FPBDCollisionContainerSolver
+	 * @todo(chaos): really we should split the base class into group and scene versions. See CreateSceneSolver CreateGroupSolver
+	*/
+	template<typename ConstraintContainerType>
+	class CHAOS_API TSimpleConstraintContainerSolver : public FConstraintContainerSolver
+	{
+	public:
+		using FConstraintContainerType = ConstraintContainerType;
+		using FConstraintHandleType = typename FConstraintContainerType::FConstraintContainerHandle;
+
+		TSimpleConstraintContainerSolver(FConstraintContainerType& InConstraintContainer, const int32 InPriority)
+			: FConstraintContainerSolver(InPriority)
+			, ConstraintContainer(InConstraintContainer)
+		{
+		}
+
+		virtual void Reset(const int32 MaxConstraints) override final
+		{
+		}
+
+		virtual int32 GetNumConstraints() const override final
+		{
+			return ConstraintContainer.GetNumConstraints();
+		}
+
+		virtual void AddConstraints() override final
+		{
+			// We solve all constraints in the container in the order it prefers so nothing to do here
+		}
+
+		virtual void AddConstraints(const TArrayView<FPBDIslandConstraint>& Constraints) override final
+		{
+			// This solver container is for use with the a non-graph evolution. It will not call this function.
+			ensure(false);
+		}
+
+		virtual void AddBodies(FSolverBodyContainer& SolverBodyContainer) override final
+		{
+			ConstraintContainer.AddBodies(SolverBodyContainer);
+		}
+
+		virtual void GatherInput(const FReal Dt) override final
+		{
+			ConstraintContainer.GatherInput(Dt);
+		}
+
+		virtual void GatherInput(const FReal Dt, const int32 BeginIndex, const int32 EndIndex) override final
+		{
+			// This solver container is for use with the a non-graph evolution. It will not call this function.
+			ensure(false);
+		}
+
+		virtual void ScatterOutput(const FReal Dt) override final
+		{
+			ConstraintContainer.ScatterOutput(Dt);
+		}
+
+		virtual void ScatterOutput(const FReal Dt, const int32 BeginIndex, const int32 EndIndex) override final
+		{
+			// This solver container is for use with the a non-graph evolution. It will not call this function.
+			ensure(false);
+		}
+
+		virtual void ApplyPositionConstraints(const FReal Dt, const int32 It, const int32 NumIts) override final
+		{
+			ConstraintContainer.ApplyPositionConstraints(Dt, It, NumIts);
+		}
+
+		virtual void ApplyVelocityConstraints(const FReal Dt, const int32 It, const int32 NumIts) override final
+		{
+			ConstraintContainer.ApplyVelocityConstraints(Dt, It, NumIts);
+		}
+
+		virtual void ApplyProjectionConstraints(const FReal Dt, const int32 It, const int32 NumIts) override final
+		{
+			ConstraintContainer.ApplyProjectionConstraints(Dt, It, NumIts);
+		}
+
+	protected:
+		FConstraintContainerType& ConstraintContainer;
+	};
 }
diff --git a/Engine/Source/Runtime/Experimental/Chaos/Public/Chaos/Evolution/SolverDatas.h b/Engine/Source/Runtime/Experimental/Chaos/Public/Chaos/Evolution/SolverDatas.h
--- a/Engine/Source/Runtime/Experimental/Chaos/Public/Chaos/Evolution/SolverDatas.h
+++ b/Engine/Source/Runtime/Experimental/Chaos/Public/Chaos/Evolution/SolverDatas.h
@@ -1,182 +1,6 @@
 // Copyright Epic Games, Inc. All Rights Reserved.
 #pragma once
 
-#include "Containers/SparseArray.h"
-#include "Chaos/Evolution/SolverBodyContainer.h"
-#include "Chaos/Evolution/SolverConstraintContainer.h"
-
-namespace Chaos
-{
-	
-/** Forward Declaration */
-class FPBDIslandManager;
-
-/**
-* List of bodies and constraints that wil be stored per island.
-*/
-class CHAOS_API FPBDIslandSolverData
-{
-public :
-	/**
-	* Init the solver datas 
-	*/
-	FPBDIslandSolverData(const int32 IslandIndex = 0) : BodyContainer(MakeUnique<FSolverBodyContainer>()), ConstraintDatas(), IslandIndex(IslandIndex)
-	{}
-
-	/** Accessors of the island index */
-	FORCEINLINE const int32& GetIslandIndex() const { return IslandIndex; };
-	FORCEINLINE int32& GetIslandIndex() { return IslandIndex; };
-	
-	/** Accessors of the body container */
-	FORCEINLINE const FSolverBodyContainer& GetBodyContainer() const { return *BodyContainer.Get(); };
-	FORCEINLINE FSolverBodyContainer& GetBodyContainer() { return *BodyContainer.Get(); };
-
-	/**
-	 * @brief The number of constraint containers registered
-	*/
-	int32 NumContainerIds() const { return ConstraintCounts.Num(); }
-
-	/**
-	* @brief Resize the constraint count array with a fixed number of containers
-	* @param NumContainers Number of containers extracted from the constraint rules
-	*/
-	void ResizeConstraintsCounts(const int32 NumContainers);
-
-	/** Number of constraints given a container id */
-	const int32& ConstraintCount(const int32 ContainerId) const;
-	int32& ConstraintCount(const int32 ContainerId);
-
-	/** Accessors of the constraint container */
-	template<typename ContainerType> const ContainerType& GetConstraintContainer(const int32 ContainerId) const;
-	template<typename ContainerType> ContainerType& GetConstraintContainer(const int32 ContainerId);
-
-	/** Accessors of the constraint indices */
-	const TArray<int32>& GetConstraintIndices(const int32 ContainerId) const;
-	TArray<int32>& GetConstraintIndices(const int32 ContainerId);
-	
-	/** Accessors of the constraint handles */
-	const TArray<FConstraintHandle*>& GetConstraintHandles(const int32 ContainerId) const;
-	TArray<FConstraintHandle*>& GetConstraintHandles(const int32 ContainerId);
-
-	/** Accessor of one constraint handle given a container id and a constraint index */
-	template<typename ConstraintType> const ConstraintType* GetConstraintHandle(const int32 ContainerId, const int32 ConstraintIndex) const;
-	template<typename ConstraintType> ConstraintType* GetConstraintHandle(const int32 ContainerId, const int32 ConstraintIndex);
-
-	/** Accessors of the current constraint index */
-	const int32& GetConstraintIndex(const int32 ContainerId) const { return ConstraintDatas[ContainerId].ConstraintIndex; };
-	int32& GetConstraintIndex(const int32 ContainerId) { return ConstraintDatas[ContainerId].ConstraintIndex; };
-
-	/** Init the current constraint index */
-	void InitConstraintIndex(const int32 ContainerId) { ConstraintDatas[ContainerId].ConstraintIndex = 0; };
-
-	/** Add a constraint container to the solver island given a container id
-	* @param ContainerId Constraints container id from which the solver constraint datas is being built
-	*/
-	template<typename ConstraintType> void AddConstraintDatas(const int32 ContainerId);
-
-protected:
-	/**
-	* Constraint datas that will be stored per container (joints, collisions...)
-	*/
-	struct FConstraintDatas
-	{
-		/** Constraint indices that will be used for legacy solver */
-		TArray<int32> ConstraintIndices;
-
-		/** Constraint handles that will be used for legacy solver */
-		TArray<FConstraintHandle*> ConstraintHandles;
-
-		/** Current constraint index in the constraint container */
-		int32 ConstraintIndex;
-	};
-
-	/** Solver body container of that datas */
-	TUniquePtr<FSolverBodyContainer> BodyContainer;
-	
-	/** List of constraint containers (collision, joints...) that will be used to solve constraints */
-	TSparseArray<TUniquePtr<FConstraintSolverContainer>> ConstraintContainers;
-	
-	/** List of constraint datas (collision, joints...) that will be used to solve constraints */
-	TSparseArray<FConstraintDatas> ConstraintDatas;
-
-	/** Number of constraints per container id*/
-	TArray<int32> ConstraintCounts;
-
-	/** Island index in case these datas belong to an island */
-	int32 IslandIndex;
-};
-
-template<typename ContainerType>
-FORCEINLINE const ContainerType& FPBDIslandSolverData::GetConstraintContainer(const int32 ContainerId) const
-{
-	return *static_cast<const ContainerType*>(ConstraintContainers[ContainerId].Get());
-};
-
-template<typename ContainerType>
-FORCEINLINE ContainerType& FPBDIslandSolverData::GetConstraintContainer(const int32 ContainerId)
-{
-	return *static_cast<ContainerType*>(ConstraintContainers[ContainerId].Get());
-};
-	
-FORCEINLINE const TArray<int32>& FPBDIslandSolverData::GetConstraintIndices(const int32 ContainerId) const
-{
-	return ConstraintDatas[ContainerId].ConstraintIndices;
-};
-	
-FORCEINLINE TArray<int32>& FPBDIslandSolverData::GetConstraintIndices(const int32 ContainerId)
-{
-	return ConstraintDatas[ContainerId].ConstraintIndices;
-};
-	
-FORCEINLINE const TArray<FConstraintHandle*>& FPBDIslandSolverData::GetConstraintHandles(const int32 ContainerId) const
-{
-	return ConstraintDatas[ContainerId].ConstraintHandles;
-};
-
-FORCEINLINE TArray<FConstraintHandle*>& FPBDIslandSolverData::GetConstraintHandles(const int32 ContainerId)
-{
-	return ConstraintDatas[ContainerId].ConstraintHandles;
-};
-
-template<typename ConstraintType>
-FORCEINLINE const ConstraintType* FPBDIslandSolverData::GetConstraintHandle(const int32 ContainerId, const int32 ConstraintIndex) const
-{
-	return static_cast<const ConstraintType*>(ConstraintDatas[ContainerId].ConstraintHandles[ConstraintIndex]);
-};
-
-template<typename ConstraintType>
-FORCEINLINE ConstraintType* FPBDIslandSolverData::GetConstraintHandle(const int32 ContainerId, const int32 ConstraintIndex)
-{
-	return static_cast<ConstraintType*>(ConstraintDatas[ContainerId].ConstraintHandles[ConstraintIndex]);
-};
-	
-FORCEINLINE const int32& FPBDIslandSolverData::ConstraintCount(const int32 ContainerId) const
-{
-	return ConstraintCounts[ContainerId];
-};
-	
-FORCEINLINE int32& FPBDIslandSolverData::ConstraintCount(const int32 ContainerId)
-{
-	return ConstraintCounts[ContainerId];
-};
-
-FORCEINLINE void FPBDIslandSolverData::ResizeConstraintsCounts(const int32 NumContainers)
-{
-	ConstraintCounts.SetNum(NumContainers, false);
-	FMemory::Memzero(ConstraintCounts.GetData(), ConstraintCounts.Num() * sizeof(int32));
-};
-	
-template<typename ConstraintType>
-inline void FPBDIslandSolverData::AddConstraintDatas(const int32 ContainerId)
-{
-	if((ContainerId >= 0) && !ConstraintDatas.IsValidIndex(ContainerId) && !ConstraintContainers.IsValidIndex(ContainerId) )
-	{
-		ConstraintDatas.Reserve(ContainerId+1);
-		ConstraintDatas.EmplaceAt(ContainerId, FConstraintDatas());
-		
-		ConstraintContainers.Reserve(ContainerId+1);
-		ConstraintContainers.EmplaceAt(ContainerId, MakeUnique<typename ConstraintType::FConstraintSolverContainerType>());
-	}
-}
-	
-}
\ No newline at end of file
+//
+// TO BE REMOVED
+//
diff --git a/Engine/Source/Runtime/Experimental/Chaos/Public/Chaos/Island/IslandGraph.h b/Engine/Source/Runtime/Experimental/Chaos/Public/Chaos/Island/IslandGraph.h
--- a/Engine/Source/Runtime/Experimental/Chaos/Public/Chaos/Island/IslandGraph.h
+++ b/Engine/Source/Runtime/Experimental/Chaos/Public/Chaos/Island/IslandGraph.h
@@ -200,6 +200,11 @@ public:
 	 */
 	void RemoveNode(const NodeType& NodeItem);
 
+	/**
+	 * Get the level of the specified node item (distance from a kinematic node)
+	 */
+	int32 GetNodeItemLevel(const NodeType& NodeItem) const;
+
 	/**
 	 * Reserve a number of edges in memory before adding them later
 	 * @param NumEdges Number of edges to be reserved in memory
@@ -304,34 +309,28 @@ public:
 	void InitSorting();
 
 	/**
-	* Given a constraint containerId compute all the islands/edges/nodes levels
-	* @param ContainerId Constraint container id that will be compared to the item container
+	* Compute all the islands/edges/nodes levels
 	*/
-	void ComputeLevels(const int32 ContainerId);
+	void ComputeLevels();
 	
 	/**
-	* Given a constraint containerId and a node index update all the islands/edges/nodes levels and push the connected nodes
-	* into the node queue
+	* Update all the islands/edges/nodes levels and push the connected nodes into the node queue
 	* @param NodeIndex Node index that we are currently iterating over
-	* @param ContainerId Constraint container id that will be compared to the item container
 	*/
-	void UpdateLevels(const int32 NodeIndex, const int32 ContainerId);
+	void UpdateLevels(const int32 NodeIndex);
 
 	/**
-	* Given a constraint containerId compute all the islands/edges/nodes colors
-	* @param ContainerId Constraint container id that will be compared to the item container
+	* Compute all the islands/edges/nodes colors
 	* @param MinEdges Minimum number of edges to compute coloring
 	*/
-	void ComputeColors(const int32 ContainerId, const int32 MinEdges);
+	void ComputeColors(const int32 MinEdges);
 
 	/**
-	* Given a constraint containerId and a node index update all the islands/edges/nodes colors and push the connected nodes
-	* into the node queue
+	* Update all the islands/edges/nodes colors and push the connected nodes into the node queue
 	* @param NodeIndex Node index that we are currently iterating over
-	* @param ContainerId Constraint container id that will be compared to the item container
 	* @param MinEdges Minimum number of edges to compute coloring
 	*/
-	void UpdateColors(const int32 NodeIndex, const int32 ContainerId,  const int32 MinEdges);
+	void UpdateColors(const int32 NodeIndex, const int32 MinEdges);
 
 	/**
 	* Pick the first available color that is not used yet by the current graph node and the edge opposite one
diff --git a/Engine/Source/Runtime/Experimental/Chaos/Public/Chaos/Island/IslandGroup.h b/Engine/Source/Runtime/Experimental/Chaos/Public/Chaos/Island/IslandGroup.h
--- a/Engine/Source/Runtime/Experimental/Chaos/Public/Chaos/Island/IslandGroup.h
+++ b/Engine/Source/Runtime/Experimental/Chaos/Public/Chaos/Island/IslandGroup.h
@@ -1,91 +1,72 @@
-﻿// Copyright Epic Games, Inc. All Rights Reserved.
+// Copyright Epic Games, Inc. All Rights Reserved.
 #pragma once
 
 #include "Containers/Array.h"
-#include "Chaos/Evolution/SolverDatas.h"
 
-namespace Chaos
-{
-	
-/** Forward Declaration */
-class FPBDIslandSolver;
+#include "Chaos/Evolution/ConstraintGroupSolver.h"
 
-/**
-* Group of islands that will be used by evolution to run the solvers steps in parallel
-*/
-class CHAOS_API FPBDIslandGroup : public FPBDIslandSolverData
+namespace Chaos
 {
-public:
-
-	/**
-	* Init the island group 
-	*/
-	FPBDIslandGroup(const int32 GroupIndexIn);
-
-	/**
-	* Init group members to their default values
-	*/
-	void InitGroup();
-	
-	/**
-	* Reset the islands list and reserve a number of islands in memory 
-	* @param NumIslands number of islands to be reserved
-	*/
-	void ReserveIslands(const int32 NumIslands);
-
-	/**
-	* Add island to the group
-	* @param IslandSolver Island Solver top be added
-	*/
-	void AddIsland(FPBDIslandSolver* IslandSolver);
+	class FPBDConstraintContainer;
+	class FPBDIsland;
+	class FPBDIslandManager;
 
 	/**
-	* Remove all islands from the group
+	* A set of constraints that will ne solved in sequence on a single thread. This will usually be the constraints
+	* from several islands, but may be a sub set of constraints of a single island when coloring is enabled for
+	* that island.
 	*/
-	void ClearIslands();
-
-	/**
-	* Check if the group is valid and contains islands
-	*/
-	FORCEINLINE bool IsValid() const { return (IslandSolvers.Num() > 0); }
-
-	/**
-	* Get the number of islands within the group
-	*/
-	FORCEINLINE int32 NumIslands() const { return IslandSolvers.Num(); }
-
-	/**
-	* Return the islands within the group
-	*/
-	FORCEINLINE const TArray<FPBDIslandSolver*>& GetIslands() const { return IslandSolvers; }
-
-	/**
-	* Return the islands within the group
-	*/
-	FORCEINLINE TArray<FPBDIslandSolver*>& GetIslands() { return IslandSolvers; }
-
-	/**
-	* Accessors for the number of particles
-	*/
-	FORCEINLINE const int32& GetNumParticles() const  { return ParticlesCount; }
-	FORCEINLINE int32& NumParticles() {return ParticlesCount; }
-
-	/**
-	* Accessors for the number of constraints
-	*/
-	FORCEINLINE const int32& GetNumConstraints() const  { return ConstraintsCount; }
-	FORCEINLINE int32& NumConstraints() {return ConstraintsCount; }
-	
-private:
-
-	/** List of all the island constraints handles */
-	TArray<FPBDIslandSolver*> IslandSolvers;
-
-	/** Number of particles in this group */
-	int32 ParticlesCount = 0;
-
-	/** Number of constraints in this group */
-	int32 ConstraintsCount = 0;
-};
-	
+	class CHAOS_API FPBDIslandConstraintGroupSolver : public FPBDConstraintGroupSolver
+	{
+	public:
+		UE_NONCOPYABLE(FPBDIslandConstraintGroupSolver);
+
+		/**
+		* Init the island group 
+		*/
+		FPBDIslandConstraintGroupSolver(FPBDIslandManager& InIslandManager);
+
+		/**
+		* Add island to the group
+		* @param Island Island to be added
+		*/
+		void AddIsland(FPBDIsland* Island);
+
+		/**
+		* Return the islands within the group
+		*/
+		FORCEINLINE const TArray<FPBDIsland*>& GetIslands() { return Islands; }
+
+		/**
+		 * The number of particles (dynamics, kinematic and static) in the island.
+		*/
+		inline int32 GetNumParticles() const { return NumParticles; }
+
+		/**
+		 * The number of constraints of all types in the island.
+		*/
+		inline int32 GetNumConstraints() const { return NumConstraints; }
+
+		/**
+		 * The number of constraints of the specific container type in the island.
+		*/
+		inline int32 GetNumConstraints(const int32 ContainerId) const { return NumContainerConstraints[ContainerId]; }
+
+	protected:
+		// Base class overrides
+		virtual void SetConstraintSolverImpl(const int32 ContainerId) override final;
+		virtual void ResetImpl() override final;
+		virtual void AddConstraintsImpl() override final;
+		virtual void GatherBodiesImpl(const FReal Dt, const int32 BeginBodyIndex, const int32 EndBodyIndex) override final;
+
+	private:
+		FPBDIslandManager& IslandManager;
+
+		// Item counters used to initialize the solver containers
+		int32 NumParticles;
+		int32 NumConstraints;
+		TArray<int32> NumContainerConstraints;	// Per ContainerId
+
+		TArray<FPBDIsland*> Islands;
+	};
 }
\ No newline at end of file
diff --git a/Engine/Source/Runtime/Experimental/Chaos/Public/Chaos/Island/IslandManager.h b/Engine/Source/Runtime/Experimental/Chaos/Public/Chaos/Island/IslandManager.h
--- a/Engine/Source/Runtime/Experimental/Chaos/Public/Chaos/Island/IslandManager.h
+++ b/Engine/Source/Runtime/Experimental/Chaos/Public/Chaos/Island/IslandManager.h
@@ -10,7 +10,6 @@
 #include "Chaos/GeometryParticlesfwd.h"
 #include "Chaos/Island/IslandGraph.h"
 #include "Chaos/Island/SolverIsland.h"
-#include "Chaos/Island/IslandGroup.h"
 
 namespace Chaos
 {
@@ -36,7 +35,7 @@ using FPBDRigidParticles = TPBDRigidParticles<FReal, 3>;
 class CHAOS_API FPBDIslandManager
 {
 public:
-	using GraphType = FIslandGraph<FGeometryParticleHandle*, FConstraintHandleHolder, FPBDIslandSolver*, FPBDIslandManager>;
+	using GraphType = FIslandGraph<FGeometryParticleHandle*, FConstraintHandleHolder, FPBDIsland*, FPBDIslandManager>;
 	using FGraphNode = GraphType::FGraphNode;
 	using FGraphEdge = GraphType::FGraphEdge;
 	using FGraphIsland = GraphType::FGraphIsland;
@@ -51,11 +50,6 @@ public:
 	*/
 	~FPBDIslandManager();
 
-	/**
-	* Reset nodes and edges graph indices 
-	*/
-	//void ResetIndices();
-
 	/**
 	 * Clear all nodes and edges from the graph
 	*/
@@ -104,6 +98,22 @@ public:
 	  */
 	void AddConstraint(const uint32 ContainerId, FConstraintHandle* ConstraintHandle, const TVector<FGeometryParticleHandle*, 2>& ConstrainedParticles);
 
+	/**
+	 * Utility function for use by the constraint containers to add all their constraints
+	 */
+	template<typename ConstraintContainerType>
+	void AddContainerConstraints(ConstraintContainerType& ConstraintContainer)
+	{
+		const int32 ContainerId = ConstraintContainer.GetContainerId();
+		for (FConstraintHandle* ConstraintHandle : ConstraintContainer.GetConstraintHandles())
+		{
+			if (ConstraintHandle->IsEnabled())
+			{
+				AddConstraint(ContainerId, ConstraintHandle, ConstraintHandle->GetConstrainedParticles());
+			}
+		}
+	}
+
 	/**
 	  * Remove a constraint from the graph
 	  * @param ContainerId Container id the constraint is belonging to
@@ -157,17 +167,25 @@ public:
 	  */
 	void DisableParticle(FGeometryParticleHandle* Particle);
 
+	/**
+	 * Get the Level for the specified particle (distance to a kinmetic in the graph)
+	*/
+	int32 GetParticleLevel(FGeometryParticleHandle* Particle) const;
+
 	/**
 	  * Remove all the graph edges
 	  */
 	void ResetIslands(const TParticleView<FPBDRigidParticles>& PBDRigids);
 
 	/**
-	  * Generate the simulation islands of connected particles (AddConstraints must have already been called). There are no constraints connecting particles in different islands.
-	  * @param PBDRigids Rigid particles that will be used to fill the solver islands
-	  * @param Particles SOA used in the evolution
+	 * Process all the particles and constraints in the graph. Assign all t islands.
+	*/
+	void UpdateGraph();
+
+	/**
+	  * Pack all the constraints into their assigned islands in sorted order
 	  */
-	void UpdateIslands(const TParticleView<FPBDRigidParticles>& PBDRigids, FPBDRigidsSOAs& Particles, const int32 NumContainers);
+	void UpdateIslands(FPBDRigidsSOAs& Particles);
 
 	/**
 	  * @brief Put all particles and constraints in an island to sleep
@@ -184,22 +202,6 @@ public:
 	  * @return return true if the island is going to sleep
 	  */
 	bool SleepInactive(const int32 IslandIndex, const TArrayCollectionArray<TSerializablePtr<FChaosPhysicsMaterial>>& PerParticleMaterialAttributes, const THandleArray<FChaosPhysicsMaterial>& SolverPhysicsMaterials);
-	
-	/**
-	  * Get the island particles stored in the solver island
-	  * @param IslandIndex Island index
-	  * @return List of particles handles in the island
-	  * @note This will be empty for sleeping islands, though the particles will still be in the graph
-	  */
-	const TArray<FGeometryParticleHandle*>& GetIslandParticles(const int32 IslandIndex) const;
-
-	/**
-	  * Get the island constraints stored in the solver island
-	  * @param IslandIndex Island index
-	  * @return List of constraints handles in the island
-	  * @note This will be empty for sleeping islands, though the constraints will still be in the graph
-	  */
-	const TArray<FConstraintHandleHolder>& GetIslandConstraints(const int32 IslandIndex) const;
 
 	/**
 	 * Get the set of islands that the particle belongs to. Will be 0-1 island for dynamics and could be any number for non-dynamics.
@@ -207,6 +209,9 @@ public:
 	 */
 	TArray<int32> FindParticleIslands(const FGeometryParticleHandle* ParticleHandle) const;
 
+	/**
+	 * Get the island index of the constraint
+	*/
 	int32 GetConstraintIsland(const FConstraintHandle* ConstraintHandle) const;
 
 	/**
@@ -224,15 +229,6 @@ public:
 	{
 		return IslandIndexing.Num();
 	}
-
-	/**
-	* The number of groups in the graph.
-	* @return the number of groups
-	*/
-	FORCEINLINE int32 NumGroups() const
-	{
-		return IslandGroups.Num();
-	}
 	
 	/**
 	 * The number of islands in the graph.
@@ -256,63 +252,44 @@ public:
 	 * Used to create arrays large enough to use Particle.UniqueIdx().Idx for indexing.
 	 * @return the max index of the particles in the islands
 	 */
-	inline int32 GetMaxParticleUniqueIdx() const { return MaxParticleIndex; }
-
+	inline int32 GetMaxParticleUniqueIdx() const 
+	{ 
+		return MaxParticleIndex;
+	}
 
 	/**
 	 * Get the sparse island graph index given a dense index between 0 and NumIslands
 	 * @param IslandIndex Island index
 	 * @return the max index of the particles in the islands
 	 */
-	int32 GetGraphIndex(const int32 IslandIndex) const {
-		return IslandIndexing.IsValidIndex(IslandIndex) ? IslandIndexing[IslandIndex] : 0;
-	}
-	
-	/** Add a constraint container to all the solver island given a container id
-	* @param ContainerId Constraints container id from which the solver constraint containers are being built
-	*/
-	template<typename ConstraintType>
-	void AddConstraintDatas(const int32 ContainerId)
+	int32 GetGraphIndex(const int32 IslandIndex) const 
 	{
-		for(auto& IslandGroup : IslandGroups)
-		{
-			IslandGroup->AddConstraintDatas<ConstraintType>(ContainerId);
-		}
-		for(auto& IslandSolver : IslandSolvers)
-		{
-			IslandSolver->AddConstraintDatas<ConstraintType>(ContainerId);
-		}
+		return IslandIndexing.IsValidIndex(IslandIndex) ? IslandIndexing[IslandIndex] : 0;
 	}
 
-	/** Accessors of the graph group */
-	const FPBDIslandGroup* GetIslandGroup(const int32 GroupIndex) const {return IslandGroups.IsValidIndex(GroupIndex) ? IslandGroups[GroupIndex].Get() : nullptr; }
-	FPBDIslandGroup* GetIslandGroup(const int32 GroupIndex) {return IslandGroups.IsValidIndex(GroupIndex) ? IslandGroups[GroupIndex].Get()  : nullptr; }
+	/**
+	 * Register a constraint container with the island manager. Currently this is only used to determin
+	 * how many ContainerIds we need in order to presize arrays etc.
+	*/
+	void AddConstraintContainer(FPBDConstraintContainer& ConstraintContainer);
 
 	/** Accessors of the graph island */
-	const FPBDIslandSolver* GetSolverIsland(const int32 IslandIndex) const {return IslandIndexing.IsValidIndex(IslandIndex) ? IslandSolvers[IslandIndexing[IslandIndex]].Get() : nullptr; }
-	FPBDIslandSolver* GetSolverIsland(const int32 IslandIndex) {return IslandIndexing.IsValidIndex(IslandIndex) ? IslandSolvers[IslandIndexing[IslandIndex]].Get() : nullptr; }
-
-	/** Accessors of the group islands */
-	const TArray<FPBDIslandSolver*>& GetGroupIslands(const int32 GroupIndex) const {return IslandGroups[GroupIndex]->GetIslands(); }
-	TArray<FPBDIslandSolver*>& GetGroupIslands(const int32 GroupIndex) {return IslandGroups[GroupIndex]->GetIslands(); }
-
-	/** Accessors of all the graph groups */
-	const TArray<TUniquePtr<FPBDIslandGroup>>& GetIslandGroups() const {return IslandGroups; }
-	TArray<TUniquePtr<FPBDIslandGroup>>& GetIslandGroups() {return IslandGroups; }
+	const FPBDIsland* GetIsland(const int32 IslandIndex) const {return IslandIndexing.IsValidIndex(IslandIndex) ? Islands[IslandIndexing[IslandIndex]].Get() : nullptr; }
+	FPBDIsland* GetIsland(const int32 IslandIndex) {return IslandIndexing.IsValidIndex(IslandIndex) ? Islands[IslandIndexing[IslandIndex]].Get() : nullptr; }
 
 	/** Get the graph nodes */
+	// @todo(chaos): remove this - only used by testing
 	const TSparseArray<GraphType::FGraphNode>& GetGraphNodes() const { return IslandGraph->GraphNodes; }
 	
 	/** Get the graph edges */
+	// @todo(chaos): remove this - only used by testing
 	const TSparseArray<GraphType::FGraphEdge>& GetGraphEdges() const { return IslandGraph->GraphEdges; }
 
 	/** Get the particle nodes */
+	// @todo(chaos): remove this - only used by testing
 	const TMap<FGeometryParticleHandle*, int32>& GetParticleNodes() const { return IslandGraph->ItemNodes; }
-	
-	/** Get the constraints edges */
-	const TMap<FConstraintHandleHolder, int32>& GetConstraintEdges() const { return IslandGraph->ItemEdges; }
 
-	// Graph Owner API Implementation
+	// Graph Owner API Implementation. Allows us to store edge/node indices on constraints and particles
 	void GraphNodeAdded(FGeometryParticleHandle* NodeItem, const int32 NodeIndex);
 	void GraphNodeRemoved(FGeometryParticleHandle* NodeItem);
 	void GraphEdgeAdded(const FConstraintHandleHolder& EdgeItem, const int32 EdgeIndex);
@@ -334,19 +311,9 @@ protected:
 #endif
 
 	/**
-	* Initialize the groups according to the number of threads
-	*/
-	void InitializeGroups();
-
-	/**
-	 * Build all the island groups
-	 */
-	void BuildGroups(const int32 NumContainers);
-
-	/**
-	* Sync the IslandSolvers with the IslandGraph
+	* Sync the Islands with the IslandGraph
 	*/
-	void SyncIslands(FPBDRigidsSOAs& Particles, const int32 NumContainers);
+	void SyncIslands(FPBDRigidsSOAs& Particles);
 
 
 	/**
@@ -357,7 +324,7 @@ protected:
 	void ValidateIslands() const;
 
 	/** List of solver islands in the manager */
-	TSparseArray<TUniquePtr<FPBDIslandSolver>> IslandSolvers;
+	TSparseArray<TUniquePtr<FPBDIsland>>	Islands;
 
 	/** Island map to make the redirect an index between 0...NumIslands to the persistent island index  */
 	TArray<int32>							IslandIndexing;
@@ -365,11 +332,10 @@ protected:
 	/** Global graph to deal with merging and splitting of islands */
 	TUniquePtr<GraphType>					IslandGraph;
 
+	int32									MaxNumConstraintContainers;
+
 	/** Maximum particle index that we have in the graph*/
 	int32									MaxParticleIndex = INDEX_NONE;
-	
-	/** List of island groups in the manager */
-	TArray<TUniquePtr<FPBDIslandGroup>>		IslandGroups;
 
 	/** Sorted list of islands */
 	TArray<int32>							SortedIslands;
@@ -380,4 +346,5 @@ protected:
 	/** Make sure EndTick is called. Sanity check mostlyhelpful for unit tests */
 	bool									bEndTickCalled;
 };
+
 }
\ No newline at end of file
diff --git a/Engine/Source/Runtime/Experimental/Chaos/Public/Chaos/Island/SolverIsland.h b/Engine/Source/Runtime/Experimental/Chaos/Public/Chaos/Island/SolverIsland.h
--- a/Engine/Source/Runtime/Experimental/Chaos/Public/Chaos/Island/SolverIsland.h
+++ b/Engine/Source/Runtime/Experimental/Chaos/Public/Chaos/Island/SolverIsland.h
@@ -4,189 +4,206 @@
 #include "Containers/Array.h"
 #include "Chaos/ParticleHandle.h"
 #include "Chaos/ConstraintHandle.h"
-#include "Chaos/Evolution/SolverDatas.h"
 
 namespace Chaos
 {
 	
-/** Forward Declaration */
-class FPBDIslandManager;
-
-/**
-* List of per island datas that will created bu the island manager
-*/
-class CHAOS_API FPBDIslandSolver : public FPBDIslandSolverData
-{
-	public:
+	/** Forward Declaration */
+	class FPBDIslandManager;
+	class FPBDRigidsSOAs;
 
 	/**
-	* Init the solver island with an island manager and an index
-	* @param IslandManager Island manager that is storing the solver islands
-	* @param IslandIndex Index of the solver island within that list
+	 * Data for a particle in an Island
 	*/
-	FPBDIslandSolver(const FPBDIslandManager* IslandManager, const int32 IslandIndex);
-
-	/**
-	* Clear all the particles from the island
-	*/
-	void ClearParticles();
-
-	/**
-	* Add particle to the island
-	* @param ParticleHandle ParticleHandle to be added
-	*/
-	void AddParticle(FGenericParticleHandle ParticleHandle);
-
-	/**
-	* Reset the particles list and reserve a number of particles in memory
-	* @param NumParticles Number of particles to be reserved
-	*/
-	void ReserveParticles(const int32 NumParticles);
+	class CHAOS_API FPBDIslandParticle
+	{
+	public:
+		FPBDIslandParticle(FGeometryParticleHandle* InParticleHandle)
+			: ParticleHandle(InParticleHandle)
+		{
+		}
 
-	/**
-	* Remove particle from the island
-	* @param ParticleHandle ParticleHandle to be removed
-	*/
-	void RemoveParticle(FGenericParticleHandle ParticleHandle);
-	
-	/**
-	* Update the particles island index to match the graph index
-	*/
-	void UpdateParticles();
-	
-	/**
-	* Reset the constraints list and reserve a number of constraints in memory 
-	* @param NumConstraints number of constraints to be reserved
-	*/
-	void ReserveConstraints(const int32 NumConstraints);
+		const FGeometryParticleHandle* GetParticle() const { return ParticleHandle; }
+		FGeometryParticleHandle* GetParticle() { return ParticleHandle; }
 
-	/**
-	* Remove all the constraints from the solver island
-	*/
-	void ClearConstraints();
+	private:
+		FGeometryParticleHandle* ParticleHandle;
+	};
 
 	/**
-	* Add constraint to the island
-	* @param ConstraintHandle ConstraintHandle to be added
+	* Data for a constraint in an island. This includes the Level and a sort key (which is build from the level and some 
+	* externally generated key that should be consistent between frames).
 	*/
-	void AddConstraint(FConstraintHandle* ConstraintHandle);
+	class CHAOS_API FPBDIslandConstraint
+	{
+	public:
+		FPBDIslandConstraint(FConstraintHandle* InConstraintHandle, const int32 InLevel, const int32 InColor, const uint32 InSubSortKey)
+			: ConstraintHandle(InConstraintHandle)
+		{
+			// Sort by level, then sort key
+			SortKey = (uint64(InLevel) << 32) | uint64(InSubSortKey);
+		}
+
+		const FConstraintHandle* GetConstraint() const { return ConstraintHandle; }
+		FConstraintHandle* GetConstraint() { return ConstraintHandle; }
+
+		int32 GetLevel() const { return int32((SortKey >> 32) & 0xFFFFF); }
+		
+		int32 GetColor() const { return INDEX_NONE; }
+		
+		uint64 GetSortKey() const { return SortKey; }
 
-	/**
-	* Remove constraint from the island
-	* @param ConstraintHandle ConstraintHandle to be removed
-	*/
-	void RemoveConstraint(FConstraintHandle* ConstraintHandle);
+	private:
+		uint64 SortKey;
+		FConstraintHandle* ConstraintHandle;
+	};
 
 	/**
-	* Sort the islands constraints
+	* A set of constraints and particles that form a connected set in the constraint graph. These constraints
+	* must be solved sequentially (if not colored). FPBDIsland objects are collected into 
 	*/
-	void SortConstraints();
+	class CHAOS_API FPBDIsland
+	{
+	public:
 
-	/**
-	 * Set the island group index
-	 */
-	FORCEINLINE void SetGroupIndex(const int32 GroupIndex) { IslandGroup = GroupIndex; }
+		/**
+		* Init the solver island with an island manager and an index
+		* @param IslandManager Island manager that is storing the solver islands
+		* @param IslandIndex Index of the solver island within that list
+		*/
+		FPBDIsland(const int32 MaxConstraintContainers);
+
+		/**
+		 * The island index as seen by the systems external to the IslandManager.
+		 * The IslandManager uses this to index into the IslandIndexing array to get the
+		 * internal island index.
+		*/
+		int32 GetIslandIndex() const { return IslandIndex; }
+
+		/**
+		 * Set the external island index.
+		 * @see GetIslandIndex.
+		*/
+		void SetIslandIndex(const int32 InIslandIndex) { IslandIndex = InIslandIndex; }
+
+		/**
+		* Clear all the particles from the island
+		*/
+		void ClearParticles();
+
+		/**
+		* Add particle to the island
+		* @param ParticleHandle ParticleHandle to be added
+		*/
+		void AddParticle(FGenericParticleHandle ParticleHandle);
+
+		/**
+		* Reset the particles list and reserve a number of particles in memory
+		* @param NumParticles Number of particles to be reserved
+		*/
+		void ReserveParticles(const int32 NumParticles);
+
+		/**
+		* Update the particles island index to match the graph index
+		*/
+		void UpdateParticles();
 	
-	/**
-	* Get the island group index
-	*/
-	FORCEINLINE const int32& GetGroupIndex() const { return IslandGroup; }
-
-	/**
-	* Return the list of particles within the solver island
-	*/
-	FORCEINLINE const TArray<FGeometryParticleHandle*>& GetParticles() const { return IslandParticles; }
+		/**
+		* Remove all the constraints from the solver island
+		*/
+		void ClearConstraints();
+
+		/**
+		* Add constraint to the island
+		* @param ConstraintHandle ConstraintHandle to be added
+		*/
+		void AddConstraint(FConstraintHandle* ConstraintHandle, const int32 Level, const int32 Color, const uint32 SubSortKey);
+
+		/**
+		* Sort the islands constraints
+		*/
+		void SortConstraints();
+
+		/**
+		* Return the list of particles within the solver island
+		*/
+		FORCEINLINE TArrayView<const FPBDIslandParticle> GetParticles() const { return MakeArrayView(IslandParticles); }
+		FORCEINLINE TArrayView<FPBDIslandParticle> GetParticles() { return MakeArrayView(IslandParticles); }
+
+		/**
+		* Return the list of constraints of the specified type within the island
+		*/
+		TArrayView<const FPBDIslandConstraint> GetConstraints(const int32 ContainerId) const { return MakeArrayView(IslandConstraintsByType[ContainerId]); }
+		TArrayView<FPBDIslandConstraint> GetConstraints(const int32 ContainerId) { return MakeArrayView(IslandConstraintsByType[ContainerId]); }
+
+		const int32 GetNumConstraintContainers() const { return IslandConstraintsByType.Num(); }
+
+		/**
+		* Get the number of particles within the island
+		*/
+		FORCEINLINE int32 GetNumParticles() const { return IslandParticles.Num(); }
+
+		/**
+		* Get the number of constraints (from the specified container) within the island
+		*/
+		FORCEINLINE int32 GetNumConstraints(const int32 ContainerId) const { return IslandConstraintsByType[ContainerId].Num(); }
+
+		/**
+		* Get the number of constraints (of all types) within the island
+		*/
+		FORCEINLINE int32 GetNumConstraints() const { return NumConstraints; }
+
+		/**
+		* Members accessors
+		*/
+		FORCEINLINE bool IsSleeping() const {return bIsSleeping;}
+		FORCEINLINE void SetIsSleeping(const bool bIsSleepingIn ) { bSleepingChanged = bSleepingChanged || (bIsSleeping != bIsSleepingIn); bIsSleeping = bIsSleepingIn; }
+		FORCEINLINE void ResetSleepingChanged() { bSleepingChanged = false; }
+		FORCEINLINE bool IsPersistent() const { return bIsPersistent; }
+		FORCEINLINE void SetIsPersistent(const bool bIsPersistentIn) { bIsPersistent = bIsPersistentIn; }
+		FORCEINLINE bool NeedsResim() const { return bNeedsResim; }
+		FORCEINLINE void SetNeedsResim(const bool bNeedsResimIn) { bNeedsResim = bNeedsResimIn; }
+		FORCEINLINE int32 GetSleepCounter() const { return SleepCounter; }
+		FORCEINLINE void SetSleepCounter(const int32 SleepCounterIn) { SleepCounter = SleepCounterIn; }
+		FORCEINLINE bool SleepingChanged() const { return bSleepingChanged; }
+		FORCEINLINE void SetIsUsingCache(const bool bIsUsingCacheIn ) { bIsUsingCache = bIsUsingCacheIn; }
+		FORCEINLINE bool IsUsingCache() const { return bIsUsingCache; }
+	
+		/**
+		 * Set the sleep state of all particles and constraints based on the island sleep state
+		*/
+		void PropagateSleepState(FPBDRigidsSOAs& Particles);
 
-	/**
-	* Return the list of constraints within the island
-	*/
-	FORCEINLINE const TArray<FConstraintHandleHolder>& GetConstraints() const { return IslandConstraints; }
+		void UpdateSyncState(FPBDRigidsSOAs& Particles);
 
-	/**
-	* Get the number of particles within the island
-	*/
-	FORCEINLINE int32 NumParticles() const {return IslandParticles.Num(); }
-
-	/**
-	* Get the number of constraints within the island
-	*/
-	FORCEINLINE int32 NumConstraints() const {return IslandConstraints.Num(); }
-	
-	/**
-	* Members accessors
-	*/
-	FORCEINLINE bool IsSleeping() const {return bIsSleeping;}
-	FORCEINLINE void SetIsSleeping(const bool bIsSleepingIn ) { bSleepingChanged = bSleepingChanged || (bIsSleeping != bIsSleepingIn); bIsSleeping = bIsSleepingIn; }
-	FORCEINLINE void ResetSleepingChanged() { bSleepingChanged = false; }
-	FORCEINLINE bool IsPersistent() const { return bIsPersistent; }
-	FORCEINLINE void SetIsPersistent(const bool bIsPersistentIn) { bIsPersistent = bIsPersistentIn; }
-	FORCEINLINE bool NeedsResim() const { return bNeedsResim; }
-	FORCEINLINE void SetNeedsResim(const bool bNeedsResimIn) { bNeedsResim = bNeedsResimIn; }
-	FORCEINLINE int32 GetSleepCounter() const { return SleepCounter; }
-	FORCEINLINE void SetSleepCounter(const int32 SleepCounterIn) { SleepCounter = SleepCounterIn; }
-	FORCEINLINE bool SleepingChanged() const { return bSleepingChanged; }
-	FORCEINLINE void SetIsUsingCache(const bool bIsUsingCacheIn ) { bIsUsingCache = bIsUsingCacheIn; }
-	FORCEINLINE bool IsUsingCache() const { return bIsUsingCache; }
-	
-	// template<typename ConstraintType>
-	// void GatherSolverInput(const FReal Dt, const int32 IslandIndex, const int32 ContainerId);
-	
 	private:
 
-	/** Island manager that is storing the IslandSolver */
-	const FPBDIslandManager* IslandManager = nullptr;
+		int32 IslandIndex;
 
-	/** Flag to check if an island is awake or sleeping */
-	bool bIsSleeping = false;
+		/** Flag to check if an island is awake or sleeping */
+		bool bIsSleeping = false;
 
-	/** Flag to check if an island need to be re-simulated or not */
-	bool bNeedsResim = false;
+		/** Flag to check if an island need to be re-simulated or not */
+		bool bNeedsResim = false;
 
-	/** Flag to check if an island is persistent over time */
-	bool bIsPersistent = true;
+		/** Flag to check if an island is persistent over time */
+		bool bIsPersistent = true;
 
-	/** Flag to check if the sleeping state has changed or not */
-	bool bSleepingChanged = false;
+		/** Flag to check if the sleeping state has changed or not */
+		bool bSleepingChanged = false;
 
-	/** Sleep counter to trigger island sleeping */
-	int32 SleepCounter = 0;
+		/** Sleep counter to trigger island sleeping */
+		int32 SleepCounter = 0;
 
-	/** List of all the island particles handles */
-	TArray<FGeometryParticleHandle*> IslandParticles;
+		/** List of all the island particles handles */
+		TArray<FPBDIslandParticle> IslandParticles;
 
-	/** List of all the island constraints handles */
-	TArray<FConstraintHandleHolder> IslandConstraints;
-
-	/** Island Group in which the solver belongs */
-	int32 IslandGroup = 0;
-	
-	/** Check if the island is using the cache or not */
-	bool bIsUsingCache = false;
-	
-};
+		/** List of all the island constraints handles */
+		TArray<TArray<FPBDIslandConstraint>> IslandConstraintsByType;
+		int32 NumConstraints;
 	
-// template<typename ConstraintType>
-// FORCEINLINE void FPBDIslandSolver::GatherSolverInput(const FReal Dt, const int32 IslandIndex, const int32 ContainerId)
-// {
-// 	using FConstraintContainerHandle = typename ConstraintType::FConstraintContainerHandle;
-// 	if(ConstraintContainers.IsValidIndex(ContainerId))
-// 	{
-// 		for (FConstraintHandle* ConstraintHandle : IslandConstraints)
-// 		{
-// 			if (ConstraintHandle->GetContainerId() == ContainerId)
-// 			{
-// 				FConstraintContainerHandle* Constraint = ConstraintHandle->As<FConstraintContainerHandle>();
-//
-// 				// Note we are building the SolverBodies as we go, in the order that we visit them. Each constraint
-// 				// references two bodies, so we won't strictly be accessing only in cache order, but it's about as good as it can be.
-// 				if (Constraint->IsEnabled())
-// 				{
-// 					// @todo(chaos): we should provide Particle Levels in the island rule as well (see TPBDConstraintColorRule)
-// 					Constraint->GatherInput(Dt, GraphIndex, INDEX_NONE, INDEX_NONE, *this);
-// 				}
-// 			}
-// 		}
-// 	}
-// }
+		/** Check if the island is using the cache or not */
+		bool bIsUsingCache = false;
 	
+	};
 }
\ No newline at end of file
diff --git a/Engine/Source/Runtime/Experimental/Chaos/Public/Chaos/Joint/PBDJointCachedSolverGaussSeidel.h b/Engine/Source/Runtime/Experimental/Chaos/Public/Chaos/Joint/PBDJointCachedSolverGaussSeidel.h
--- a/Engine/Source/Runtime/Experimental/Chaos/Public/Chaos/Joint/PBDJointCachedSolverGaussSeidel.h
+++ b/Engine/Source/Runtime/Experimental/Chaos/Public/Chaos/Joint/PBDJointCachedSolverGaussSeidel.h
@@ -193,10 +193,15 @@ struct FAxisConstraintDatas
 		
 		FPBDJointCachedSolver();
 
+		void SetSolverBodies(FSolverBody* SolverBody0, FSolverBody* SolverBody1)
+		{
+			SolverBodies[0] = *SolverBody0;
+			SolverBodies[1] = *SolverBody1;
+		}
+
 		// Called once per frame to initialize the joint solver from the joint settings
 		void Init(
 			const FReal Dt,
-			const FSolverBodyPtrPair& SolverBodyPair,
 			const FPBDJointSolverSettings& SolverSettings,
 			const FPBDJointSettings& JointSettings,
 			const FRigidTransform3& XL0,
diff --git a/Engine/Source/Runtime/Experimental/Chaos/Public/Chaos/Joint/PBDJointSolverGaussSeidel.h b/Engine/Source/Runtime/Experimental/Chaos/Public/Chaos/Joint/PBDJointSolverGaussSeidel.h
--- a/Engine/Source/Runtime/Experimental/Chaos/Public/Chaos/Joint/PBDJointSolverGaussSeidel.h
+++ b/Engine/Source/Runtime/Experimental/Chaos/Public/Chaos/Joint/PBDJointSolverGaussSeidel.h
@@ -129,10 +129,15 @@ namespace Chaos
 
 		FPBDJointSolver();
 
+		void SetSolverBodies(FSolverBody* SolverBody0, FSolverBody* SolverBody1)
+		{
+			SolverBodies[0] = *SolverBody0;
+			SolverBodies[1] = *SolverBody1;
+		}
+
 		// Called once per frame to initialze the joint solver from the joint settings
 		void Init(
 			const FReal Dt,
-			const FSolverBodyPtrPair& SolverBodyPair,
 			const FPBDJointSolverSettings& SolverSettings,
 			const FPBDJointSettings& JointSettings,
 			const FRigidTransform3& XL0,
diff --git a/Engine/Source/Runtime/Experimental/Chaos/Public/Chaos/PBDCollisionConstraints.h b/Engine/Source/Runtime/Experimental/Chaos/Public/Chaos/PBDCollisionConstraints.h
--- a/Engine/Source/Runtime/Experimental/Chaos/Public/Chaos/PBDCollisionConstraints.h
+++ b/Engine/Source/Runtime/Experimental/Chaos/Public/Chaos/PBDCollisionConstraints.h
@@ -23,18 +23,11 @@ class FImplicitObject;
 class FPBDCollisionConstraints;
 class FPBDRigidsSOAs;
 class FPBDCollisionConstraint;
-class FPBDIslandSolverData;
 
 using FRigidBodyContactConstraintsPostComputeCallback = TFunction<void()>;
 using FRigidBodyContactConstraintsPostApplyCallback = TFunction<void(const FReal Dt, const TArray<FPBDCollisionConstraintHandle*>&)>;
 using FRigidBodyContactConstraintsPostApplyPushOutCallback = TFunction<void(const FReal Dt, const TArray<FPBDCollisionConstraintHandle*>&, bool)>;
 
-namespace Collisions
-{
-	struct FContactParticleParameters;
-	struct FContactIterationParameters;
-}
-
 /**
  * A container and solver for collision constraints.
  * 
@@ -53,17 +46,14 @@ public:
 
 	// For use by dependent types
 	using FConstraintContainerHandle = FPBDCollisionConstraintHandle;		// Used by constraint rules
-	using FConstraintSolverContainerType = FPBDCollisionSolverContainer;	// Used by constraint rules
-	using FParticleHandle = TGeometryParticleHandle<FReal, 3>;
 
 	FPBDCollisionConstraints(const FPBDRigidsSOAs& InParticles, 
 		TArrayCollectionArray<bool>& Collided, 
 		const TArrayCollectionArray<TSerializablePtr<FChaosPhysicsMaterial>>& PhysicsMaterials, 
 		const TArrayCollectionArray<TUniquePtr<FChaosPhysicsMaterial>>& PerParticlePhysicsMaterials,
 		const THandleArray<FChaosPhysicsMaterial>* const SimMaterials,
-		const int32 ApplyPairIterations = 1, 
-		const int32 ApplyPushOutPairIterations = 1, 
-		const FReal RestitutionThreshold = 2000.0f);
+		const int32 NumCollisionsPerBlock = 1000,
+		const FReal RestitutionThreshold = FReal(2000));
 
 	virtual ~FPBDCollisionConstraints();
 
@@ -77,14 +67,6 @@ public:
 	 */
 	void DisableHandles();
 
-	/**
-	 * Set the solver method to use
-	 */
-	void SetSolverType(EConstraintSolverType InSolverType)
-	{
-		SolverType = InSolverType;
-	}
-
 	/**
 	 * @brief Enable or disable determinism.
 	 * Support for determinism requires that we sort active constraints each tick, so there is additional cost.
@@ -144,54 +126,26 @@ public:
 	*/
 	void DisableConstraints(const TSet<FGeometryParticleHandle*>& ParticleHandle) {}
 
-	//
-	// General Rule API
-	//
-
-	void PrepareTick() {}
-
-	void UnprepareTick() {}
-
-	/**
-	 * Generate all contact constraints.
-	 */
-	void UpdatePositionBasedState(const FReal Dt);
-
-	//
-	// Simple Rule API
-	//
-
-	void PreGatherInput(const FReal Dt, FPBDIslandSolverData& SolverData);
-	void GatherInput(const FReal Dt, FPBDIslandSolverData& SolverData);
-	void ScatterOutput(const FReal Dt, FPBDIslandSolverData& SolverData);
-
-	bool ApplyPhase1(const FReal Dt, const int32 It, const int32 NumIts, FPBDIslandSolverData& SolverData);
-	bool ApplyPhase2(const FReal Dt, const int32 It, const int32 NumIts, FPBDIslandSolverData& SolverData);
-	bool ApplyPhase3(const FReal Dt, const int32 It, const int32 NumIts, FPBDIslandSolverData& SolverData) { return false; }
 
 	//
-	// Island Rule API
+	// FConstraintContainer Implementation
 	//
+	virtual int32 GetNumConstraints() const override final { return NumConstraints(); }
+	virtual void ResetConstraints() override final { Reset(); }
+	virtual void AddConstraintsToGraph(FPBDIslandManager& IslandManager) override final;
+	virtual void PrepareTick() override final {}
+	virtual void UnprepareTick() override final {}
 
-	void SetNumIslandConstraints(const int32 NumIslandConstraints, FPBDIslandSolverData& SolverData);
-	void PreGatherInput(const FReal Dt, FPBDCollisionConstraint& Constraint, FPBDIslandSolverData& SolverData);
-	void GatherInput(const FReal Dt, FPBDCollisionConstraint& Constraint, const int32 Particle0Level, const int32 Particle1Level, FPBDIslandSolverData& SolverData);
-
-	bool ApplyPhase1Serial(const FReal Dt, const int32 It, const int32 NumIts, FPBDIslandSolverData& SolverData);
-	bool ApplyPhase2Serial(const FReal Dt, const int32 It, const int32 NumIts, FPBDIslandSolverData& SolverData);
-	bool ApplyPhase3Serial(const FReal Dt, const int32 It, const int32 NumIts, FPBDIslandSolverData& SolverData) { return false; }
-
-	//
-	// Color Rule API
-	//
+	virtual TUniquePtr<FConstraintContainerSolver> CreateSceneSolver(const int32 Priority) override final
+	{
+		return MakeUnique<FPBDCollisionContainerSolver>(*this, Priority);
+	}
 
-	void ScatterOutput(const FReal Dt, const int32 BeginIndex, const int32 EndIndex, FPBDIslandSolverData& SolverData);
+	virtual TUniquePtr<FConstraintContainerSolver> CreateGroupSolver(const int32 Priority) override final
+	{
+		return MakeUnique<FPBDCollisionContainerSolver>(*this, Priority);
+	}
 
-	bool ApplyPhase1Serial(const FReal Dt, const int32 It, const int32 NumIts, const int32 BeginIndex, const int32 EndIndex, FPBDIslandSolverData& SolverData);
-	bool ApplyPhase1Parallel(const FReal Dt, const int32 It, const int32 NumIts, const int32 BeginIndex, const int32 EndIndex, FPBDIslandSolverData& SolverData);
-	
-	bool ApplyPhase2Serial(const FReal Dt, const int32 It, const int32 NumIts, const int32 BeginIndex, const int32 EndIndex, FPBDIslandSolverData& SolverData);
-	bool ApplyPhase2Parallel(const FReal Dt, const int32 It, const int32 NumIts, const int32 BeginIndex, const int32 EndIndex, FPBDIslandSolverData& SolverData);
 
 	//
 	// Member Access
@@ -217,26 +171,6 @@ public:
 		return RestitutionThreshold;
 	}
 
-	void SetPairIterations(int32 InPairIterations)
-	{
-		MApplyPairIterations = InPairIterations;
-	}
-
-	int32 GetPairIterations() const
-	{
-		return MApplyPairIterations;
-	}
-
-	void SetPushOutPairIterations(int32 InPairIterations)
-	{
-		MApplyPushOutPairIterations = InPairIterations;
-	}
-
-	int32 GetPushOutPairIterations() const
-	{
-		return MApplyPushOutPairIterations;
-	}
-
 	void SetCollisionsEnabled(bool bInEnableCollisions)
 	{
 		bEnableCollisions = bInEnableCollisions;
@@ -321,23 +255,14 @@ public:
 	const FCollisionDetectorSettings& GetDetectorSettings() const { return DetectorSettings; }
 	void SetDetectorSettings(const FCollisionDetectorSettings& InSettings) { DetectorSettings = InSettings; }
 
+	const FPBDCollisionSolverSettings& GetSolverSettings() const { return SolverSettings; }
+
 protected:
 	FPBDCollisionConstraint& GetConstraint(int32 Index);
-	FPBDCollisionSolverContainer& GetConstraintSolverContainer(FPBDIslandSolverData& SolverData);
-
-	Collisions::FContactParticleParameters GetContactParticleParameters(const FReal Dt);
-	Collisions::FContactIterationParameters GetContactIterationParameters(const FReal Dt, const int32 Iteration, const int32 NumIterations, const int32 NumPairIterations, bool& bNeedsAnotherIteration);
 
 	// Call PruneParticleEdgeCollisions on all particles with ECollisionConstraintFlags::CCF_SmoothEdgeCollisions set in CollisionFlags
 	void PruneEdgeCollisions();
 
-	// The "Legacy" functions handle the older solver types (StandardPbd)
-	// @todo(chaos): remove legacy methods when the new solver is fully operational and used everywhere (RBAN may still optionally use StandardPbd)
-	void LegacyGatherInput(const FReal Dt, FPBDCollisionConstraint& Constraint, const int32 Particle0Level, const int32 Particle1Level, FPBDIslandSolverData& SolverData);
-	void LegacyScatterOutput(const FReal Dt, const int32 BeginIndex, const int32 EndIndex, FPBDIslandSolverData& SolverData);
-	bool LegacyApplyPhase1Serial(const FReal Dt, const int32 Iterations, const int32 NumIterations, const int32 BeginIndex, const int32 EndIndex, FPBDIslandSolverData& SolverData);
-	bool LegacyApplyPhase1Parallel(const FReal Dt, const int32 Iterations, const int32 NumIterations, const int32 BeginIndex, const int32 EndIndex, FPBDIslandSolverData& SolverData);
-
 private:
 
 	const FPBDRigidsSOAs& Particles;
@@ -350,10 +275,6 @@ private:
 	const TArrayCollectionArray<TUniquePtr<FChaosPhysicsMaterial>>& MPerParticlePhysicsMaterials;
 	const THandleArray<FChaosPhysicsMaterial>* const SimMaterials;
 
-	// @todo(chaos): remove pair iteration counts when legacy PBD solver is removed (RBAN-only)
-	int32 MApplyPairIterations;
-	int32 MApplyPushOutPairIterations;
-
 	FReal RestitutionThreshold;
 	bool bEnableCollisions;
 	bool bEnableRestitution;
@@ -375,8 +296,6 @@ private:
 
 	// Settings for the low-level collision solvers
 	FPBDCollisionSolverSettings SolverSettings;
-
-	EConstraintSolverType SolverType;
 };
 
 //
@@ -435,16 +354,6 @@ inline FParticlePair FPBDCollisionConstraintHandle::GetConstrainedParticles() co
 	return { GetContact().GetParticle0(), GetContact().GetParticle1() };
 }
 
-inline void FPBDCollisionConstraintHandle::PreGatherInput(const FReal Dt, FPBDIslandSolverData& SolverData)
-{
-	ConcreteContainer()->PreGatherInput(Dt, GetContact(), SolverData);
-}
-
-inline void FPBDCollisionConstraintHandle::GatherInput(FReal Dt, const int32 Particle0Level, const int32 Particle1Level, FPBDIslandSolverData& SolverData)
-{
-	ConcreteContainer()->GatherInput(Dt, GetContact(), Particle0Level, Particle1Level, SolverData);
-}
-
 inline FSolverBody* FPBDCollisionConstraintHandle::GetSolverBody0()
 {
 	return GetContact().GetSolverBody0();
diff --git a/Engine/Source/Runtime/Experimental/Chaos/Public/Chaos/PBDCollisionConstraintsContact.h b/Engine/Source/Runtime/Experimental/Chaos/Public/Chaos/PBDCollisionConstraintsContact.h
--- a/Engine/Source/Runtime/Experimental/Chaos/Public/Chaos/PBDCollisionConstraintsContact.h
+++ b/Engine/Source/Runtime/Experimental/Chaos/Public/Chaos/PBDCollisionConstraintsContact.h
@@ -1,82 +1,6 @@
 // Copyright Epic Games, Inc. All Rights Reserved.
 #pragma once
 
-#include "Chaos/ExternalCollisionData.h"
-#include "Chaos/PBDCollisionConstraints.h"
-#include "Chaos/CollisionResolutionTypes.h"
-#include "Chaos/Collision/CollisionApplyType.h"
-
-/**
- * @brief This file contains collision solver code for the StandardPbd solver (which will be removed at some point)
- * @see PBDCollisionSolver.h for the new solver
-*/
-
-namespace Chaos
-{
-	class FCollisionContext;
-
-	namespace Collisions
-	{
-		// @todo(chaos): remove this
-		struct FContactParticleParameters 
-		{
-			FContactParticleParameters()
-				: RestitutionVelocityThreshold(0)
-				, bCanDisableContacts(false)
-				, Collided(nullptr)
-			{}
-
-			FContactParticleParameters(
-				FReal InRestitutionVelocityThreshold,
-				bool bInCanDisableContacts,
-				TArrayCollectionArray<bool>* InCollided)
-				: RestitutionVelocityThreshold(InRestitutionVelocityThreshold)
-				, bCanDisableContacts(bInCanDisableContacts)
-				, Collided(InCollided)
-			{}
-
-
-			FReal RestitutionVelocityThreshold;
-			bool bCanDisableContacts;
-			TArrayCollectionArray<bool>* Collided;
-		};
-
-		struct FContactIterationParameters 
-		{
-			FContactIterationParameters()
-				: Dt(0)
-				, Iteration(0)
-				, NumIterations(0)
-				, NumPairIterations(0)
-				, SolverType(EConstraintSolverType::None)
-				, NeedsAnotherIteration(nullptr)
-			{}
-
-			FContactIterationParameters(
-				const FReal InDt,
-				const int32 InIteration,
-				const int32 InNumIterations,
-				const int32 InNumPairIterations,
-				const EConstraintSolverType InSolverType,
-				bool* InNeedsAnotherIteration)
-				: Dt(InDt)
-				, Iteration(InIteration)
-				, NumIterations(InNumIterations)
-				, NumPairIterations(InNumPairIterations)
-				, SolverType(InSolverType)
-				, NeedsAnotherIteration(InNeedsAnotherIteration)
-			{}
-
-			const FReal Dt;
-			const int32 Iteration;
-			const int32 NumIterations;
-			const int32 NumPairIterations;
-			const EConstraintSolverType SolverType;
-			bool* NeedsAnotherIteration;
-		};
-
-		extern void Update(FPBDCollisionConstraint& Constraint, const FReal Dt);
-		extern void Apply(FPBDCollisionConstraint& Constraint, const FContactIterationParameters& IterationParameters, const FContactParticleParameters& ParticleParameters);
-	}
-
-}
+//
+// TO BE DELETED
+//
\ No newline at end of file
diff --git a/Engine/Source/Runtime/Experimental/Chaos/Public/Chaos/PBDConstraintColor.h b/Engine/Source/Runtime/Experimental/Chaos/Public/Chaos/PBDConstraintColor.h
--- a/Engine/Source/Runtime/Experimental/Chaos/Public/Chaos/PBDConstraintColor.h
+++ b/Engine/Source/Runtime/Experimental/Chaos/Public/Chaos/PBDConstraintColor.h
@@ -1,234 +1,6 @@
 // Copyright Epic Games, Inc. All Rights Reserved.
 #pragma once
 
-#include "Chaos/ParticleHandle.h"
-
-namespace Chaos
-{
-	class FPBDConstraintGraph;
-	class FConstraintHandle;
-
-	/**
-	 * Generates color information for a single constraint rule in a connection graph.
-	 * Edges with the same color are non-interacting and can safely be processed in parallel.
-	 */
-	class CHAOS_API FPBDConstraintColor
-	{
-	public:
-		typedef TSet<int32> FColorSet;
-		typedef TArray<FConstraintHandle*> FConstraintList;
-		typedef TMap<int32, FConstraintList> FColorToConstraintListMap;
-		typedef TArray<FColorToConstraintListMap> FLevelToColorToConstraintListMap;
-
-		FPBDConstraintColor()
-			: bUseContactGraph(true)
-		{}
-
-		int32 NumIslands() const
-		{
-			return IslandData.Num();
-		}
-
-		int32 NumIslandEdges(const int32 Island) const
-		{
-			return IslandData[Island].NumEdges;
-		}
-
-		/**
-		 * Initialize the color structures based on the connectivity graph (i.e., reset all color-related node, edge and island data).
-		 */
-		void InitializeColor(const FPBDConstraintGraph& ConstraintGraph);
-
-		/**
-		 * Calculate the color information for the specified island.
-		 */
-		void ComputeColor(const FReal Dt, const int32 Island, const FPBDConstraintGraph& ConstraintGraph, uint32 ContainerId);
-
-		/**
-		 * Get the Level-Color-ConstraintList map for the specified island.
-		 */
-		const FLevelToColorToConstraintListMap& GetIslandLevelToColorToConstraintListMap(int32 Island) const;
-
-		/**
-		 * Get the maximum color index used in the specified island.
-		 */
-		int GetIslandMaxColor(int32 Island) const;
-
-		/**
-		 * Get the maximum level index used in the specified island.
-		 */
-		int GetIslandMaxLevel(int32 Island) const;
-
-		void SetUseContactGraph(bool bInUseContactGraph)
-		{
-			bUseContactGraph = bInUseContactGraph;
-		}
-
-#if !UE_BUILD_SHIPPING
-		// For debugging onlue: run some (slow) validation checks on the graph and log any errors. Return false if there were errors.
-		bool DebugCheckGraph(const FPBDConstraintGraph& ConstraintGraph) const;
-#endif
-
-		/**
-		 * Get the level of a particle after particle levels are computed using ComputeContactGraphGBF 
-		 */
-		int32 GetParticleLevel(const FGeometryParticleHandle* ParticleHandle) const;
-
-	private:
-		void ComputeContactGraph(const int32 Island, const FPBDConstraintGraph& ConstraintGraph, uint32 ContainerId);
-		void ComputeIslandColoring(const int32 Island, const FPBDConstraintGraph& ConstraintGraph, uint32 ContainerId);
-
-		/**
-		 * Compute contact graph using GBF paper Sec 8.1
-		 */
-		void ComputeContactGraphGBF(const FReal Dt, const int32 Island, const FPBDConstraintGraph& ConstraintGraph, const uint32 ContainerId);
-
-		/**
-		 * Compute the mapping between constraint graph node index to island graph node indices. 
-		 * The mapping of dynamic particles are stored in TArray DynamicGlobalToIslandLocalNodeIndices. 
-		 * The mapping of kinematic particles are stored in TMap KinematicGlobalToIslandLocalNodeIndices.
-		 */
-		void ComputeGlobalToIslandLocalNodeMapping(const int32 Island, const FPBDConstraintGraph& ConstraintGraph, TArray<int32>& DynamicGlobalToIslandLocalNodeIndices, TMap<int32, int32>& KinematicGlobalToIslandLocalNodeIndices) const;
-
-		/**
-		 * Compute the graph node index in the island using the mapping data above.
-		 */
-		int32 GetIslandLocalNodeIdx(const int32 GlobalNodeIndex, const TArray<int32>& DynamicGlobalToIslandLocalNodeIndices, const TMap<int32, int32>& KinematicGlobalToIslandLocalNodeIndices, const bool IsNodeDynamic) const;
-
-		struct FDirectedEdge
-		{
-			int32 FirstNode; // Or source node of the directed edge.
-			int32 SecondNode; // Or target node of the directed edge. 
-			FDirectedEdge(){};
-			FDirectedEdge(const int32 FirstInput, const int32 SecondInput) : FirstNode(FirstInput), SecondNode(SecondInput) {};
-		};
-
-		/**
-		 * Collect directed edges from the island.
-		 * All collected edges will point from lower particles to upper particles.
-		 * The set of collected directed edges is a subset of collision constraint edges.
-		 * Horizontal edges will not be collected.
-		 */
-		void CollectIslandDirectedEdges(const FReal Dt, const int32 Island, const FPBDConstraintGraph& ConstraintGraph, const TArray<int32>& DynamicGlobalToIslandLocalNodeIndices, const TMap<int32, int32>& KinematicGlobalToIslandLocalNodeIndices, TArray<FDirectedEdge>& DirectedEdges) const;
-
-		/**
-		 * CSR: Compressed sparse row
-		 * Build CSR representation of the graph from input edges
-		 */
-
-		void BuildGraphCSR(const TArray<FDirectedEdge>& GraphEdges, const int32 NumNodes, TArray<int32>& CSRIndices, TArray<int32>& CSRValues, TArray<int32>& CurrentIndices) const;
-
-		/**
-		 * A node is a root if it has no parents.
-		 * Result is stored in IsRootNode
-		 */
-		void ComputeIsRootNode(const TArray<FDirectedEdge>& GraphEdges, const int32 NumNodes, TArray<bool> &IsRootNode) const;
-
-		/**
-		 * Digraph: directed graph (according to wikipedia: https://en.wikipedia.org/wiki/Directed_graph) that might contain loops
-		 * DAG: directed acyclic graph
-		 * Collapse directed graph to directed acyclic graph. Recursively collapse the loops in the directed graph to single nodes until there are no loops any more.
-		 * For each digraph node, the corresponding DAG node index is stored in DigraphToDAGIndices
-		 * For each DAG node, the corresponding list of digraph nodes is stored in DAGToDigraphIndices
-		 */
-		void CollapseDigraphToDAG(const TArray<int32>& ChildrenCSRIndices, const TArray<int32>& ChildrenCSRValues, const TArray<bool>& IsRootNode, TArray<int32>& DigraphToDAGIndices, TArray<TArray<int32>>& DAGToDigraphIndices, TArray<bool>& Visited, TArray<int32>& TraversalStack, TArray<int32>& NumOnStack) const;
-
-		/**
-		 * Compute DAG Edges from Digraph edges
-		 */
-		void ComputeDAGEdges(const TArray<FDirectedEdge>& DigraphEdges, const TArray<int32>& DigraphToDAGIndices, TArray<FDirectedEdge> &DAGEdges) const;
-
-		/**
-		 * Topological sort DAG nodes
-		 * DAGToDigraphIndices is used to determine whether DAGI is a valid DAG node (point to non-zero Digraph nodes)
-		 */
-		void TopologicalSortDAG(const TArray<int32>& DAGChildrenCSRIndices, const TArray<int32>& DAGChildrenCSRValues, const TArray<TArray<int32>>& DAGToDigraphIndices, TArray<int32>& SortedDAGNodes, TArray<bool>& Visited) const;
-
-		/**
-		 * Assign levels to DAG nodes.
-		 * We want all parents of DAGI have smaller levels than DAGI. The level of DAGI is the max distance to root nodes. 
-		 */
-		void AssignDAGNodeLevels(const TArray<int32>& DAGChildrenCSRIndices, const TArray<int32>& DAGChildrenCSRValues, const TArray<int32>& SortedDAGNodes, TArray<int32> &DAGNodeLevels) const;
-
-		/**
-		 * Assign edge levels based on DAG node levels.
-		 * Edge level is the max of its end node levels.
-		 */
-		void AssignEdgeLevels(const int32 Island, const FPBDConstraintGraph& ConstraintGraph, const TArray<int32>& DynamicGlobalToIslandLocalNodeIndices, const TMap<int32, int32>& KinematicGlobalToIslandLocalNodeIndices, const TArray<int32>& DigraphToDAGIndices, const TArray<int32> DAGNodeLevels);
-
-		/**
-		 * Update particle levels based on DAGNodeLevels.
-		 */
-		void UpdateParticleToLevel(const int32 Island, const FPBDConstraintGraph& ConstraintGraph, const TArray<int32>& DynamicGlobalToIslandLocalNodeIndices, const TMap<int32, int32>& KinematicGlobalToIslandLocalNodeIndices, const TArray<int32>& DigraphToDAGIndices, const TArray<int32>& DAGNodeLevels);
-
-		struct FGraphNodeColor
-		{
-			FGraphNodeColor()
-				: NextColor(0)
-			{
-			}
-			int32 NextColor;
-			FColorSet UsedColors;
-		};
-
-		struct FGraphEdgeColor
-		{
-			FGraphEdgeColor()
-				: Color(INDEX_NONE)
-				, Level(INDEX_NONE)
-			{
-			}
-			int32 Color;
-			int32 Level;
-		};
-
-		struct FIslandColor
-		{
-			FIslandColor()
-				: MaxColor(0)
-				, MaxLevel(0)
-				, NumEdges(0)
-			{
-			}
-			int32 MaxColor;
-			int32 MaxLevel;
-			int32 NumEdges;
-			FLevelToColorToConstraintListMap LevelToColorToConstraintListMap;
-		};
-
-		struct FGBFContactGraphData
-		{
-			TMap<int32, int32> KinematicGlobalToIslandLocalNodeIndices;
-			TArray<FDirectedEdge> DigraphEdges;
-			TArray<int32> ChildrenCSRIndices, ChildrenCSRValues;
-
-			// Buffer array for BuildGraphCSR function. Not used outside BuildGraphCSR.
-			TArray<int32> CSRCurrentIndices; 
-
-			TArray<bool> IsRootNode;
-			TArray<int32> DigraphToDAGIndices;
-			TArray<TArray<int32>> DAGToDigraphIndices;
-
-			// Helper arrays for CollapseDigraphToDAG and TopologicalSortDAG. Not used elsewhere
-			TArray<bool> Visited; 
-			TArray<int32> TraversalStack;
-			TArray<int32> NumOnStack;
-
-			TArray<FDirectedEdge> DAGEdges;
-			TArray<int32> DAGChildrenCSRIndices, DAGChildrenCSRValues;
-			TArray<int32> SortedDAGNodes;
-			TArray<int32> DAGNodeLevels;
-		};
-
-		TArray<FGraphNodeColor> Nodes;
-		TArray<FGraphEdgeColor> Edges;
-		TArray<FIslandColor> IslandData;
-		FLevelToColorToConstraintListMap EmptyLevelToColorToConstraintListMap;
-		TArray<int32> DynamicGlobalToIslandLocalNodeIndicesArray; 
-		TArray<int32> ParticleToLevel; 
-		TArray<const TGeometryParticleHandle<FReal, 3>*> NodeToParticle;
-		TArray<FGBFContactGraphData> IslandToGBFContactGraphData;
-		bool bUseContactGraph;
-	};
-
-}
+//
+// TO BE DELETED
+//
\ No newline at end of file
diff --git a/Engine/Source/Runtime/Experimental/Chaos/Public/Chaos/PBDConstraintContainer.h b/Engine/Source/Runtime/Experimental/Chaos/Public/Chaos/PBDConstraintContainer.h
--- a/Engine/Source/Runtime/Experimental/Chaos/Public/Chaos/PBDConstraintContainer.h
+++ b/Engine/Source/Runtime/Experimental/Chaos/Public/Chaos/PBDConstraintContainer.h
@@ -4,9 +4,16 @@
 #include "Chaos/ConstraintHandle.h"
 #include "Chaos/Transform.h"
 #include "Chaos/Evolution/SolverBody.h"
+#include "Chaos/Evolution/SolverConstraintContainer.h"
+#include "Chaos/Evolution/ConstraintGroupSolver.h"
 
 namespace Chaos
 {
+	class FConstraintContainerSolver;
+	class FPBDIslandManager;
+	using FPBDConstraintGraph = FPBDIslandManager;
+
+
 	/**
 	 * Base class for containers of constraints.
 	 * A Constraint Container holds an array of constraints and provides methods to allocate and deallocate constraints
@@ -19,49 +26,90 @@ namespace Chaos
 
 		virtual ~FPBDConstraintContainer();
 
-		// The ContainerId is used by the Evolution to map constraints graph items back to constraints
+		/**
+		 * The ContainerId is used to map constraint handles back to their constraint container.
+		 * Every container is assigned an ID when it is registered - it will be the array index in the Evolution
+		*/
 		int32 GetContainerId() const
 		{
 			return ContainerId;
 		}
 
+		/**
+		 * @see GetContainerId()
+		*/
 		void SetContainerId(int32 InContainerId)
 		{
 			ContainerId = InContainerId;
 		}
 
+		/**
+		 * The TypeID of the constraints in this container. Used to safely downcast constraint handles.
+		 * @see FConstraintHandle::As()
+		*/
 		const FConstraintHandleTypeID& GetConstraintHandleType() const
 		{
 			return ConstraintHandleType;
 		}
 
-		virtual void DisconnectConstraints(const TSet<TGeometryParticleHandle<FReal, 3>*>&) {}
+		/**
+		 * Get the number of constraints in this container (includes inactive and disabled)
+		*/
+		virtual int32 GetNumConstraints() const = 0;
 
-	protected:
-		FConstraintHandleTypeID ConstraintHandleType;
-		int32 ContainerId;
-	};
+		/**
+		 * Empty the constraints (must be removed from the graph first, if required)
+		*/
+		virtual void ResetConstraints() = 0;
 
-	class CHAOS_API FPBDIndexedConstraintContainer : public FPBDConstraintContainer
-	{
-	public:
-		FPBDIndexedConstraintContainer(FConstraintHandleTypeID InType)
-			: FPBDConstraintContainer(InType)
-		{
-		}
+		/**
+		 * An opportunity to create/destroy constraints based on particle state.
+		*/
+		virtual void UpdatePositionBasedState(const FReal Dt) {}
 
-		virtual void SetConstraintEnabled(int32 ConstraintIndex, bool bEnabled) { }
-		virtual bool IsConstraintEnabled(int32 ConstraintIndex) const { return true; }
+		/**
+		 * Called oncer per tick to initialize buffers required for the rest of the tick
+		*/
+		virtual void PrepareTick() = 0;
+
+		/**
+		 * Should undo any allocations in PrepareTick
+		*/
+		virtual void UnprepareTick() = 0;
+
+		// @todo(chaos): remove the set
+		virtual void DisconnectConstraints(const TSet<TGeometryParticleHandle<FReal, 3>*>&) {}
+
+		/**
+		* Create a constraint solver for an Evolution without Graph support (RBAN evolution).
+		* There will only be one of these per scene (RBAN node) and it is used to solve all constraints
+		* in the container (serially).
+		*/
+		virtual TUniquePtr<FConstraintContainerSolver> CreateSceneSolver(const int32 Priority) = 0;
+
+		/**
+		 * Create a constraint solver for an Evolution with Graph support (World evolution).
+		 * The system will create several of these: usually one per worker thread (Island Group) but possibly
+		 * more in complex scenes where constraint coloring is being used. It will be used to solve constraints
+		 * in groups, with the constraints in each group determined by the graph/islands/islandgroups.
+		*/
+		virtual TUniquePtr<FConstraintContainerSolver> CreateGroupSolver(const int32 Priority) = 0;
+
+		/**
+		 * Add all the constraints in the container to the graph
+		*/
+		virtual void AddConstraintsToGraph(FPBDIslandManager& IslandManager) = 0;
 
 	protected:
-		int32 GetConstraintIndex(const FIndexedConstraintHandle* ConstraintHandle) const;
-		void SetConstraintIndex(FIndexedConstraintHandle* ConstraintHandle, int32 ConstraintIndex) const;
+		FConstraintHandleTypeID ConstraintHandleType;
+		int32 ContainerId;
 	};
 
-
+	//
 	//
 	// From ConstraintHandle.h
 	//
+	//
 
 	inline int32 FConstraintHandle::GetContainerId() const
 	{
@@ -94,4 +142,5 @@ namespace Chaos
 		return ((ConstraintContainer != nullptr) && ConstraintContainer->GetConstraintHandleType().IsA(T::StaticType())) ? static_cast<const T*>(this) : nullptr;
 	}
 
+
 }
\ No newline at end of file
diff --git a/Engine/Source/Runtime/Experimental/Chaos/Public/Chaos/PBDConstraintGraph.h b/Engine/Source/Runtime/Experimental/Chaos/Public/Chaos/PBDConstraintGraph.h
--- a/Engine/Source/Runtime/Experimental/Chaos/Public/Chaos/PBDConstraintGraph.h
+++ b/Engine/Source/Runtime/Experimental/Chaos/Public/Chaos/PBDConstraintGraph.h
@@ -2,10 +2,11 @@
 #pragma once
 
 #include "Chaos/Island/IslandManager.h"
+
 namespace Chaos
 {
-	class CHAOS_API FPBDConstraintGraph : public FPBDIslandManager
-	{};
+	// @todo(chaos): deprecate FPBDConstraintGraph
+	using FPBDConstraintGraph = FPBDIslandManager;
 }
 
 
