Chaos - Constraint Solver Refactor

diff --git a/Engine/Source/Runtime/Experimental/Chaos/Public/Chaos/PBDConstraintRule.h b/Engine/Source/Runtime/Experimental/Chaos/Public/Chaos/PBDConstraintRule.h
--- a/Engine/Source/Runtime/Experimental/Chaos/Public/Chaos/PBDConstraintRule.h
+++ b/Engine/Source/Runtime/Experimental/Chaos/Public/Chaos/PBDConstraintRule.h
@@ -1,363 +1,6 @@
 // Copyright Epic Games, Inc. All Rights Reserved.
 #pragma once
 
-#include "Chaos/Declares.h"
-#include "Chaos/ConstraintHandle.h"
-#include "Chaos/CollisionResolutionTypes.h"
-#include "Chaos/PBDConstraintColor.h"
-#include "Chaos/PBDConstraintGraph.h"
-
-// Only used for Mac/Linux forward decl at the bottom - should be removed as per the comment there (we get a warning if the include is also at the bottom)
-#include "Chaos/PBDCollisionConstraints.h"
-
-
-#define USE_SHOCK_PROPOGATION 1
-
-namespace Chaos
-{
-	/**
-	 * Constraint Rules bind constraint collections to the evolution and provide their update algorithm.
-	 */
-	class CHAOS_API FConstraintRule
-	{
-	public:
-		FConstraintRule(int32 InPriority) : Priority(InPriority) {}
-		virtual ~FConstraintRule() {}
-
-		/** Determines the order in which constraints are resolved. Higher priority constraints override lower priority ones. */
-		int32 GetPriority() const { return Priority; }
-
-		/** Set the constraint resolution priority. Higher priority constraints override lower priority ones. */
-		void SetPriority(const int32 InPriority) { Priority = InPriority; }
-
-		friend bool operator<(const FConstraintRule& L, const FConstraintRule& R)
-		{
-			return L.GetPriority() < R.GetPriority();
-		}
-
-		/** Called once per frame. Can be used to prepare caches etc. */
-		virtual void PrepareTick() {}
-
-		/** Called once per frame. Should undo whatever is done in PrepareTick (can also free any other transient buffers created after) */
-		virtual void UnprepareTick() {}
-
-
-	protected:
-		int32 Priority;
-	};
-
-	/**
-	 * Constraint rule for evolutions that do not use Constraint Graphs or other acceleration schemes.
-	 */
-	class CHAOS_API FSimpleConstraintRule : public FConstraintRule
-	{
-	public:
-		FSimpleConstraintRule(int32 InPriority) : FConstraintRule(InPriority), SolverData(nullptr) {}
-
-		/** Bind the solver datas to the one in the evolution */
-		virtual void BindToDatas(FPBDIslandSolverData& InSolverDatas, const uint32 InContainerId) {}
-
-		virtual void UpdatePositionBasedState(const FReal Dt) {}
-		virtual void GatherSolverInput(const FReal Dt) {}
-		virtual void ScatterSolverOutput(const FReal Dt) {}
-		virtual bool ApplyConstraints(const FReal Dt, const int32 It, const int32 NumIts) { return false; }
-		virtual bool ApplyPushOut(const FReal Dt, const int32 It, const int32 NumIts) { return false; }
-		virtual bool ApplyProjection(const FReal Dt, const int32 It, const int32 NumIts) { return false; }
-	protected:
-		/** Solver datas that are coming from the evolution */
-		FPBDIslandSolverData* SolverData = nullptr;
-	};
-
-	template<typename ConstraintType>
-	class CHAOS_API TSimpleConstraintRule : public FSimpleConstraintRule
-	{
-	public:
-		using FConstraints = ConstraintType;
-
-		TSimpleConstraintRule(int32 InPriority, FConstraints& InConstraints);
-
-		virtual ~TSimpleConstraintRule();
-
-		/** Bind the solver datas to the one in the evolution */
-		virtual void BindToDatas(FPBDIslandSolverData& InSolverDatas, const uint32 InContainerId) override;
-
-		virtual void PrepareTick() override;
-
-		virtual void UnprepareTick() override;
-
-		virtual void UpdatePositionBasedState(const FReal Dt) override;
-
-		virtual void GatherSolverInput(const FReal Dt) override;
-
-		virtual void ScatterSolverOutput(const FReal Dt) override;
-
-		virtual bool ApplyConstraints(const FReal Dt, const int32 It, const int32 NumIts) override;
-
-		virtual bool ApplyPushOut(const FReal Dt, const int32 It, const int32 NumIts) override;
-
-		virtual bool ApplyProjection(const FReal Dt, const int32 It, const int32 NumIts) override;
-
-	private:
-		FConstraints& Constraints;
-	};
-
-
-	/**
-	 * Base class for Constraint Rules that use the Contact Graph (which will be most optimized ones).
-	 * The graph is shared among many/all constraint rules and is held external to the Graph rule itself.
-	 * Each edge in the graph can be mapped back to a constraint controlled by the rule. To support this,
-	 * each rule is assigned an ID which is stored alongside the constraint index in the graph.
-	 * @see TPBDConstraintGraphRuleImpl
-	 */
-	class CHAOS_API FPBDConstraintGraphRule : public FConstraintRule
-	{
-	public:
-		FPBDConstraintGraphRule(int32 InPriority) : FConstraintRule(InPriority) {}
-		virtual ~FPBDConstraintGraphRule() {}
-
-		virtual void BindToGraph(FPBDConstraintGraph& InContactGraph, uint32 InContainerId) {}
-
-		// Collect all the data required to solve the constraints in the specified island. This also fills the SolverBodies
-		virtual void GatherSolverInput(const FReal Dt, int32 GroupIndex) {}
-
-		// Scatter the results of the islands constraint solver(s) out to the appropriate places (e.g., impulses, break flags, etc)
-		virtual void ScatterSolverOutput(const FReal Dt, int32 GroupIndex) {}
-
-		/** Called once per tick to allow constraint containers to create/alter their constraints based on particle position */
-		virtual void UpdatePositionBasedState(const FReal Dt) {}
-
-		/** Apply all corrections for constraints in the specified island. Return true if more iterations are needed. */
-		virtual bool ApplyConstraints(const FReal Dt, int32 GroupIndex, const int32 It, const int32 NumIts) { return false; }
-
-		/** Apply push out for constraints in the specified island. Return true if more iterations are needed. */
-		virtual bool ApplyPushOut(const FReal Dt, int32 GroupIndex, const int32 It, const int32 NumIts) { return false; }
-
-		/** Apply push out for constraints in the specified island. Return true if more iterations are needed. */
-		virtual bool ApplyProjection(const FReal Dt, int32 GroupIndex, const int32 It, const int32 NumIts) { return false; }
-
-		/** Add all constraints to the connectivity graph */
-		virtual void AddToGraph() {}
-
-		/** Initialize and performance-acceleration structures from the contact graph. Called once per evolution update */
-		virtual void InitializeAccelerationStructures() {}
-
-		/** Set up the perf-acceleration structures for the specified island. May be called in parallel for islands */
-		virtual void UpdateAccelerationStructures(const FReal Dt, const int32 GroupIndex) {}
-
-		/** Sort constraints if necessary  */
-		virtual void SortConstraints() {}
-
-		/** Boolean to check if we need to sort the constraints */
-		virtual bool IsSortingEnabled() const {return false;}
-
-		virtual void SetUseContactGraph(const bool InUseContactGraph) {}
-
-		/** Change enabled state on all constraints associated with the specified particles */
-		inline void SetConstraintsEnabled(FGeometryParticleHandle* ParticleHandle, bool bInEnabled)
-		{
-			for (FConstraintHandle* Constraint : ParticleHandle->ParticleConstraints())
-			{
-				if (Constraint->IsEnabled() != bInEnabled)
-				{
-					Constraint->SetEnabled(bInEnabled);
-				}
-			}
-		}
-
-		inline void SetConstraintsEnabled(const TSet<TGeometryParticleHandle<FReal, 3>*>& ParticleHandles, bool bInEnabled)
-		{ 
-			for (TGeometryParticleHandle<FReal, 3>* ParticleHandle : ParticleHandles)
-			{
-				SetConstraintsEnabled(ParticleHandle, bInEnabled);
-			}
-		}
-
-		/** Disconnect all constraints associated with the specified particles */
-		virtual void DisconnectConstraints(const TSet<TGeometryParticleHandle<FReal, 3>*>& RemovedParticles) {  }
-
-		/** Remove all constraints */
-		virtual void ResetConstraints() {}
-
-		/** The number of constraints in the collection */
-		virtual int32 NumConstraints() const { return 0; }
-	};
-
-
-	/**
-	 * ConstraintGraphRule helper base class - templatized on Constraint Container.
-	 */
-	template<typename ConstraintType>
-	class CHAOS_API TPBDConstraintGraphRuleImpl : public FPBDConstraintGraphRule
-	{
-	public:
-		typedef ConstraintType FConstraints;
-
-		TPBDConstraintGraphRuleImpl(FConstraints& InConstraints, int32 InPriority);
-
-		virtual void PrepareTick() override
-		{
-			Constraints.PrepareTick();
-		}
-
-		virtual void UnprepareTick() override
-		{
-			Constraints.UnprepareTick();
-		}
-
-		virtual void BindToGraph(FPBDConstraintGraph& InContactGraph, uint32 InContainerId) override;
-
-		virtual void UpdatePositionBasedState(const FReal Dt) override;
-
-		virtual void AddToGraph() override;
-
-		/** Disconnect all constraints associated with the specified particles */
-		virtual void DisconnectConstraints(const TSet<TGeometryParticleHandle<FReal, 3>*>& RemovedParticles)
-		{
-			Constraints.DisconnectConstraints(RemovedParticles);
-		}
-
-		virtual int32 NumConstraints() const override;
-
-		int32 GetContainerId() const
-		{
-			return Constraints.GetContainerId();
-		}
-
-	protected:
-		FConstraints& Constraints;
-		FPBDConstraintGraph* ConstraintGraph;
-	};
-
-	/**
-	 * Island-based constraint rule. All constraints in an island are updated in single-threaded a loop. Islands may be updated in parallel.
-	 */
-	template<typename ConstraintType>
-	class CHAOS_API TPBDConstraintIslandRule : public TPBDConstraintGraphRuleImpl<ConstraintType>
-	{
-	protected:
-		typedef TPBDConstraintGraphRuleImpl<ConstraintType> Base;
-
-	public:
-		using FConstraints = ConstraintType;
-		using FConstraintContainerHandle = typename FConstraints::FConstraintContainerHandle;
-		using FConstraintList = TArray<FConstraintContainerHandle*>;
-
-		using Base::GetContainerId;
-
-		TPBDConstraintIslandRule(FConstraints& InConstraints, int32 InPriority = 0);
-
-		virtual ~TPBDConstraintIslandRule();
-
-		virtual void GatherSolverInput(const FReal Dt, int32 GroupIndex) override;
-
-		virtual void ScatterSolverOutput(const FReal Dt, int32 GroupIndex) override;
-
-		virtual bool ApplyConstraints(const FReal Dt, int32 GroupIndex, const int32 It, const int32 NumIts) override;
-
-		virtual bool ApplyPushOut(const FReal Dt, int32 GroupIndex, const int32 It, const int32 NumIts) override;
-
-		virtual bool ApplyProjection(const FReal Dt, int32 GroupIndex, const int32 It, const int32 NumIts) override;
-
-		virtual void InitializeAccelerationStructures() override;
-		
-		virtual void UpdateAccelerationStructures(const FReal Dt, const int32 GroupIndex) override;
-
-	protected:
-		using Base::Constraints;
-		using Base::ConstraintGraph;
-	};
-
-	/**
-	 * Level- and Color-based constraint rule. 
-	 * Constraints of the same color are non-interacting and can therefore be processed in parallel. 
-	 * The level is used to implement shock propagation: constraints of lower levels are frozen in 
-	 * place as far as higher-level constraints are concerned.
-	 */
-	template<typename ConstraintType>
-	class CHAOS_API TPBDConstraintColorRule : public TPBDConstraintIslandRule<ConstraintType>
-	{
-		typedef TPBDConstraintIslandRule<ConstraintType> Base;
-
-	public:
-		using FConstraints = ConstraintType;
-		using FConstraintContainerHandle = typename FConstraints::FConstraintContainerHandle;
-
-		TPBDConstraintColorRule(FConstraints& InConstraints, int32 InPriority = 0);
-
-		virtual ~TPBDConstraintColorRule();
-
-		virtual void UpdatePositionBasedState(const FReal Dt) override
-		{
-			Constraints.UpdatePositionBasedState(Dt);
-		}
-
-		virtual void GatherSolverInput(const FReal Dt, int32 GroupIndex) override;
-
-		virtual void ScatterSolverOutput(const FReal Dt, int32 GroupIndex) override;
-
-		virtual bool ApplyConstraints(const FReal Dt, int32 GroupIndex, const int32 It, const int32 NumIts) override;
-
-		virtual void RemoveConstraints(const TSet<TGeometryParticleHandle<FReal, 3>*>& InConstraints)
-		{
-			Constraints.RemoveConstraints(InConstraints);
-		}
-
-		virtual void ResetConstraints()
-		{
-			Constraints.Reset();
-		}
-
-		virtual bool ApplyPushOut(const FReal Dt, int32 GroupIndex, const int32 It, const int32 NumIts) override;
-
-		virtual void InitializeAccelerationStructures() override;
-
-		virtual void UpdateAccelerationStructures(const FReal Dt, const int32 GroupIndex) override;
-
-		/** Sort constraints according to island/level/color*/
-		virtual void SortConstraints() override;
-		
-		/** Boolean to check if we need to sort the constraints*/
-		virtual bool IsSortingEnabled() const override;
-
-		virtual void SetUseContactGraph(const bool bInUseContactGraph) override;
-
-	private:
-		using Base::Base::Constraints;
-		using Base::Base::ConstraintGraph;
-
-		/** Check if sorting is using levels*/
-		bool IsSortingUsingColors() const;
-
-		/** Check if sorting is using colors */
-		bool IsSortingUsingLevels() const;
-
-		/** Compute island levels if necessary */
-		void ComputeLevels();
-
-		/** Compute island colors if necessary */
-		void ComputeColors();
-
-		/** Populate the sorted constraints list based on island/level/color */
-		void PopulateConstraints();
-		
-		/** Loop over all the edges and apply a function */
-		void ForEachEdges(TFunctionRef<void(const int32, const FPBDConstraintGraph::GraphType::FGraphEdge&)> InFunction);
-		
-		// Each array entry contains the [begin, end) index of a set of independent constraints
-		// that can be solved in parallel. The sets are ordered by level and must be solved sequentially.
-		TArray<TArray<TPair<int32, int32>>> ConstraintSets;
-
-		// Sorted constraint handles
-		TArray<FConstraintContainerHandle*> SortedConstraints;
-
-		// Constraint offsets into the sorted handles list for a given tuple island/level/color
-		TArray<int32> ConstraintOffsets;
-
-		// Island offsets into the constraint offsets 
-		TArray<int32> IslandOffsets;
-
-		// Counters to know at which position after the island/level/color constraint offset the handle will be inserted
-		TArray<int32> OffsetCounters;
-	};
-
-}
+//
+// TO BE REMOVED
+//
\ No newline at end of file
diff --git a/Engine/Source/Runtime/Experimental/Chaos/Public/Chaos/PBDJointConstraintTypes.h b/Engine/Source/Runtime/Experimental/Chaos/Public/Chaos/PBDJointConstraintTypes.h
--- a/Engine/Source/Runtime/Experimental/Chaos/Public/Chaos/PBDJointConstraintTypes.h
+++ b/Engine/Source/Runtime/Experimental/Chaos/Public/Chaos/PBDJointConstraintTypes.h
@@ -204,11 +204,6 @@ namespace Chaos
 	public:
 		FPBDJointSolverSettings();
 
-		// Iterations
-		// @todo(chaos): remove pair iteration counts when legacy PBD solver is removed (RBAN-only)
-		int32 ApplyPairIterations;
-		int32 ApplyPushOutPairIterations;
-
 		// Tolerances
 		FReal SwingTwistAngleTolerance;
 		FReal PositionTolerance;
diff --git a/Engine/Source/Runtime/Experimental/Chaos/Public/Chaos/PBDJointConstraints.h b/Engine/Source/Runtime/Experimental/Chaos/Public/Chaos/PBDJointConstraints.h
--- a/Engine/Source/Runtime/Experimental/Chaos/Public/Chaos/PBDJointConstraints.h
+++ b/Engine/Source/Runtime/Experimental/Chaos/Public/Chaos/PBDJointConstraints.h
@@ -10,22 +10,16 @@
 
 #include "Chaos/ConstraintHandle.h"
 #include "Chaos/Collision/CollisionApplyType.h"
-#include "Chaos/Evolution/SolverDatas.h"
-#include "Chaos/Joint/JointSolverConstraints.h"
+#include "Chaos/Evolution/IndexedConstraintContainer.h"
 #include "Chaos/Joint/PBDJointSolverGaussSeidel.h"
 #include "Chaos/Joint/PBDJointCachedSolverGaussSeidel.h"
 #include "Chaos/ParticleHandleFwd.h"
-#include "Chaos/PBDConstraintContainer.h"
 #include "Chaos/PBDJointConstraintTypes.h"
 #include "Chaos/PBDJointConstraintData.h"
 #include "Chaos/GeometryParticles.h"
 
 namespace Chaos
 {
-	class FJointSolverConstraints;
-	class FPBDJointSolver;
-	class FPBDIslandSolverData;
-
 	class CHAOS_API FPBDJointConstraintHandle final : public TIndexedContainerConstraintHandle<FPBDJointConstraints>
 	{
 	public:
@@ -55,10 +49,7 @@ namespace Chaos
 
 		void SetSettings(const FPBDJointSettings& Settings);
 		
-		TVec2<FGeometryParticleHandle*> GetConstrainedParticles() const override;
-
-		void PreGatherInput(const FReal Dt, FPBDIslandSolverData& SolverData);
-		void GatherInput(const FReal Dt, const int32 Particle0Level, const int32 Particle1Level, FPBDIslandSolverData& SolverData);
+		TVec2<FGeometryParticleHandle*> GetConstrainedParticles() const override final;
 
 		static const FConstraintHandleTypeID& StaticType()
 		{
@@ -102,19 +93,18 @@ namespace Chaos
 	/**
 	 * A joint restricting up to 6 degrees of freedom, with linear and angular limits.
 	 */
-	class CHAOS_API FPBDJointConstraints : public FPBDIndexedConstraintContainer
+	class CHAOS_API FPBDJointConstraints : public TPBDIndexedConstraintContainer<FPBDJointConstraints>
 	{
 	public:
-		using Base = FPBDIndexedConstraintContainer;
+		using Base = TPBDIndexedConstraintContainer<FPBDJointConstraints>;
 
 		using FConstraintContainerHandle = FPBDJointConstraintHandle;
 		using FConstraintHandleAllocator = TConstraintHandleAllocator<FPBDJointConstraints>;
 		using FVectorPair = TVector<FVec3, 2>;
 		using FTransformPair = TVector<FRigidTransform3, 2>;
 		using FHandles = TArray<FConstraintContainerHandle*>;
-		using FConstraintSolverContainerType = FConstraintSolverContainer;	// @todo(chaos): Add island solver for this constraint type
 
-		FPBDJointConstraints(const FPBDJointSolverSettings& InSettings = FPBDJointSolverSettings());
+		FPBDJointConstraints();
 
 		virtual ~FPBDJointConstraints();
 
@@ -123,9 +113,6 @@ namespace Chaos
 
 		void SetUseLinearJointSolver(const bool bInEnable) { Settings.bUseLinearSolver = bInEnable; }
 
-		void SetNumPairIterations(const int32 NumPairIterationss) { Settings.ApplyPairIterations = NumPairIterationss; }
-		void SetNumPushOutPairIterations(const int32 NumPairIterationss) { Settings.ApplyPushOutPairIterations = NumPairIterationss; }
-
 		//
 		// Constraint Container API
 		//
@@ -253,45 +240,36 @@ namespace Chaos
 		EResimType GetConstraintResimType(int32 ConstraintIndex) const;
 		
 		//
-		// General Rule API
+		// FConstraintContainer Implementation
 		//
-
-		void PrepareTick();
-
-		void UnprepareTick();
-
-		void UpdatePositionBasedState(const FReal Dt);
+		virtual int32 GetNumConstraints() const override final { return NumConstraints(); }
+		virtual void ResetConstraints() override final {}
+		virtual void AddConstraintsToGraph(FPBDIslandManager& IslandManager) override final;
+		virtual void PrepareTick() override final;
+		virtual void UnprepareTick() override final;
 
 		//
-		// Simple Rule API
+		// TSimpleConstraintContainerSolver API - used by RBAN solvers
 		//
-
-		void PreGatherInput(const FReal Dt, FPBDIslandSolverData& SolverData);
-		void GatherInput(const FReal Dt, FPBDIslandSolverData& SolverData);
-		void ScatterOutput(const FReal Dt, FPBDIslandSolverData& SolverData);
-
-		bool ApplyPhase1(const FReal Dt, const int32 It, const int32 NumIts, FPBDIslandSolverData& SolverData);
-		bool ApplyPhase2(const FReal Dt, const int32 It, const int32 NumIts, FPBDIslandSolverData& SolverData);
-		bool ApplyPhase3(const FReal Dt, const int32 It, const int32 NumIts, FPBDIslandSolverData& SolverData);
+		void AddBodies(FSolverBodyContainer& SolverBodyContainer);
+		void GatherInput(const FReal Dt);
+		void ScatterOutput(const FReal Dt);
+		void ApplyPositionConstraints(const FReal Dt, const int32 It, const int32 NumIts);
+		void ApplyVelocityConstraints(const FReal Dt, const int32 It, const int32 NumIts);
+		void ApplyProjectionConstraints(const FReal Dt, const int32 It, const int32 NumIts);
 
 		//
-		// Island Rule API
+		// TIndexedConstraintContainerSolver API - used by World solvers
 		//
+		void AddBodies(const TArrayView<int32>& ConstraintIndices, FSolverBodyContainer& SolverBodyContainer);
+		void GatherInput(const TArrayView<int32>& ConstraintIndices, const FReal Dt);
+		void ScatterOutput(const TArrayView<int32>& ConstraintIndices, const FReal Dt);
+		void ApplyPositionConstraints(const TArrayView<int32>& ConstraintIndices, const FReal Dt, const int32 It, const int32 NumIts);
+		void ApplyVelocityConstraints(const TArrayView<int32>& ConstraintIndices, const FReal Dt, const int32 It, const int32 NumIts);
+		void ApplyProjectionConstraints(const TArrayView<int32>& ConstraintIndices, const FReal Dt, const int32 It, const int32 NumIts);
 
-		void SetNumIslandConstraints(const int32 NumIslandConstraints, FPBDIslandSolverData& SolverData);
-		void PreGatherInput(const FReal Dt, const int32 ConstraintIndex, FPBDIslandSolverData& SolverData);
-		void GatherInput(const FReal Dt, const int32 ConstraintIndex, const int32 Particle0Level, const int32 Particle1Level, FPBDIslandSolverData& SolverData);
-		bool ApplyPhase1Serial(const FReal Dt, const int32 It, const int32 NumIts, FPBDIslandSolverData& SolverData);
-		bool ApplyPhase2Serial(const FReal Dt, const int32 It, const int32 NumIts, FPBDIslandSolverData& SolverData);
-		bool ApplyPhase3Serial(const FReal Dt, const int32 It, const int32 NumIts, FPBDIslandSolverData& SolverData);
-
-		/**
-		 * Set the solver method to use
-		 */
-		void SetSolverType(EConstraintSolverType InSolverType)
-		{
-			SolverType = InSolverType;
-		}
+		void GatherInput(const int32 ConstraintIndex, const FReal Dt);
+		void ScatterOutput(const int32 ConstraintIndex, const FReal Dt);
 
 	protected:
 		using Base::GetConstraintIndex;
@@ -311,11 +289,10 @@ namespace Chaos
 
 		bool CanEvaluate(const int32 ConstraintIndex) const;
 
-		void PreparePhase3Serial(const FReal Dt, FPBDIslandSolverData& SolverData);
-
-		bool ApplyPhase1Single(const FReal Dt, const int32 ConstraintIndex, const int32 NumPairIts, const int32 It, const int32 NumIts);
-		bool ApplyPhase2Single(const FReal Dt, const int32 ConstraintIndex, const int32 It, const int32 NumIts);
-		bool ApplyPhase3Single(const FReal Dt, const int32 ConstraintIndex, const int32 It, const int32 NumIts);
+		void AddBodies(const int32 ConstraintIndex, FSolverBodyContainer& SolverBodyContainer);
+		bool ApplyPositionConstraint(const FReal Dt, const int32 ConstraintIndex, const int32 It, const int32 NumIts);
+		bool ApplyVelocityConstraint(const FReal Dt, const int32 ConstraintIndex, const int32 It, const int32 NumIts);
+		bool ApplyProjectionConstraint(const FReal Dt, const int32 ConstraintIndex, const int32 It, const int32 NumIts);
 		void ApplyBreakThreshold(const FReal Dt, int32 ConstraintIndex, const FVec3& LinearImpulse, const FVec3& AngularImpulse);
 		void ApplyPlasticityLimits(const int32 ConstraintIndex);
 
@@ -334,10 +311,6 @@ namespace Chaos
 		// @todo(ccaulfield): optimize storage for joint solver
 		TArray<FPBDJointSolver> ConstraintSolvers;
 		TArray<FPBDJointCachedSolver> CachedConstraintSolvers;
-
-		EConstraintSolverType SolverType;
-
-
 	};
 
 }
diff --git a/Engine/Source/Runtime/Experimental/Chaos/Public/Chaos/PBDNullConstraints.h b/Engine/Source/Runtime/Experimental/Chaos/Public/Chaos/PBDNullConstraints.h
--- a/Engine/Source/Runtime/Experimental/Chaos/Public/Chaos/PBDNullConstraints.h
+++ b/Engine/Source/Runtime/Experimental/Chaos/Public/Chaos/PBDNullConstraints.h
@@ -2,7 +2,7 @@
 #pragma once
 
 #include "Chaos/Declares.h"
-#include "Chaos/ConstraintHandle.h"
+#include "Chaos/Evolution/IndexedConstraintContainer.h"
 #include "Chaos/ParticleHandleFwd.h"
 
 
@@ -35,11 +35,10 @@ namespace Chaos
 	/**
 	 * Constraint Container with minimal API required to test the Graph.
 	 */
-	class FPBDNullConstraints : public FPBDIndexedConstraintContainer
+	class FPBDNullConstraints : public TPBDIndexedConstraintContainer<FPBDNullConstraints>
 	{
 	public:
 		using FConstraintContainerHandle = FPBDNullConstraintHandle;
-		using FConstraintSolverContainerType = FConstraintSolverContainer;	// @todo(chaos): Add island solver for this constraint type
 
 		FPBDNullConstraints();
 
@@ -72,27 +71,34 @@ namespace Chaos
 			return Handles;
 		}
 
-		void PrepareTick() {}
-		void UnprepareTick() {}
-		void UpdatePositionBasedState(const FReal Dt) {}
-
-		// Constraint Rule API
-		void SetNumIslandConstraints(const int32 NumIslandConstraints, FPBDIslandSolverData& SolverData) {}
-
-		// Simple Constraint Rule API
-		void PreGatherInput(const FReal Dt, FPBDIslandSolverData& SolverData) {}
-		void GatherInput(const FReal Dt, FPBDIslandSolverData& SolverData) {}
-		void ScatterOutput(const FReal Dt, FPBDIslandSolverData& SolverData) {}
-		bool ApplyPhase1(const FReal Dt, const int32 It, const int32 NumIts, FPBDIslandSolverData& SolverData) { return true; }
-		bool ApplyPhase2(const FReal Dt, const int32 It, const int32 NumIts, FPBDIslandSolverData& SolverData) { return true; }
-		bool ApplyPhase3(const FReal Dt, const int32 It, const int32 NumIts, FPBDIslandSolverData& SolverData) { return true; }
-
-		// Island Constraint Rule API
-		void PreGatherInput(const FReal Dt, const int32 ConstraintIndex, FPBDIslandSolverData& SolverData) {}
-		void GatherInput(const FReal Dt, const int32 ConstraintIndex, const int32 Particle0Level, const int32 Particle1Level, FPBDIslandSolverData& SolverData) {}
-		bool ApplyPhase1Serial(const FReal Dt, const int32 It, const int32 NumIts, FPBDIslandSolverData& SolverData) { return true; }
-		bool ApplyPhase2Serial(const FReal Dt, const int32 It, const int32 NumIts, FPBDIslandSolverData& SolverData) { return true; }
-		bool ApplyPhase3Serial(const FReal Dt, const int32 It, const int32 NumIts, FPBDIslandSolverData& SolverData) { return true; }
+		//
+		// FConstraintContainer Implementation
+		//
+		virtual int32 GetNumConstraints() const override final { return NumConstraints(); }
+		virtual void ResetConstraints() override final {}
+		virtual void AddConstraintsToGraph(FPBDIslandManager& IslandManager) override final;
+		virtual void PrepareTick() override final {}
+		virtual void UnprepareTick() override final {}
+
+		//
+		// TSimpleConstraintContainerSolver API - used by RBAN solvers
+		//
+		void AddBodies(FSolverBodyContainer& SolverBodyContainer) {}
+		void GatherInput(const FReal Dt) {}
+		void ScatterOutput(const FReal Dt) {}
+		void ApplyPositionConstraints(const FReal Dt, const int32 It, const int32 NumIts) {}
+		void ApplyVelocityConstraints(const FReal Dt, const int32 It, const int32 NumIts) {}
+		void ApplyProjectionConstraints(const FReal Dt, const int32 It, const int32 NumIts) {}
+
+		//
+		// TIndexedConstraintContainerSolver API - used by World solvers
+		//
+		void AddBodies(const TArrayView<int32>& ConstraintIndices, FSolverBodyContainer& SolverBodyContainer) {}
+		void GatherInput(const TArrayView<int32>& ConstraintIndices, const FReal Dt) {}
+		void ScatterOutput(const TArrayView<int32>& ConstraintIndices, const FReal Dt) {}
+		void ApplyPositionConstraints(const TArrayView<int32>& ConstraintIndices, const FReal Dt, const int32 It, const int32 NumIts) {}
+		void ApplyVelocityConstraints(const TArrayView<int32>& ConstraintIndices, const FReal Dt, const int32 It, const int32 NumIts) {}
+		void ApplyProjectionConstraints(const TArrayView<int32>& ConstraintIndices, const FReal Dt, const int32 It, const int32 NumIts) {}
 
 		TArray<FPBDNullConstraint> Constraints;
 		TArray<FPBDNullConstraintHandle*> Handles;
@@ -132,16 +138,6 @@ namespace Chaos
 			return ConcreteContainer()->GetConstrainedParticles(GetConstraintIndex());
 		}
 
-		void PreGatherInput(const FReal Dt, FPBDIslandSolverData& SolverData)
-		{
-			ConcreteContainer()->PreGatherInput(Dt, GetConstraintIndex(), SolverData);
-		}
-
-		void GatherInput(const FReal Dt, const int32 Particle0Level, const int32 Particle1Level, FPBDIslandSolverData& SolverData)
-		{
-			ConcreteContainer()->GatherInput(Dt, GetConstraintIndex(), Particle0Level, Particle1Level, SolverData);
-		}
-
 		static const FConstraintHandleTypeID& StaticType()
 		{
 			static FConstraintHandleTypeID STypeID(TEXT("FPBDNullConstraintHandle"), &FIndexedConstraintHandle::StaticType());
@@ -151,7 +147,7 @@ namespace Chaos
 
 
 	inline FPBDNullConstraints::FPBDNullConstraints()
-		: FPBDIndexedConstraintContainer(FPBDNullConstraintHandle::StaticType())
+		: TPBDIndexedConstraintContainer<FPBDNullConstraints>(FPBDNullConstraintHandle::StaticType())
 	{
 	}
 
diff --git a/Engine/Source/Runtime/Experimental/Chaos/Public/Chaos/PBDPositionConstraints.h b/Engine/Source/Runtime/Experimental/Chaos/Public/Chaos/PBDPositionConstraints.h
--- a/Engine/Source/Runtime/Experimental/Chaos/Public/Chaos/PBDPositionConstraints.h
+++ b/Engine/Source/Runtime/Experimental/Chaos/Public/Chaos/PBDPositionConstraints.h
@@ -2,15 +2,12 @@
 #pragma once
 
 #include "Chaos/Array.h"
-#include "Chaos/ConstraintHandle.h"
-#include "Chaos/Evolution/SolverDatas.h"
+#include "Chaos/Evolution/IndexedConstraintContainer.h"
 #include "Chaos/ParticleHandle.h"
-#include "Chaos/PBDConstraintContainer.h"
 
 namespace Chaos
 {
 	class FPBDPositionConstraints;
-	class FPBDIslandSolverData;
 
 	class FPBDPositionConstraintHandle final : public TIndexedContainerConstraintHandle<FPBDPositionConstraints>
 	{
@@ -25,9 +22,6 @@ namespace Chaos
 		
 		virtual FParticlePair GetConstrainedParticles() const override;
 
-		void PreGatherInput(const FReal Dt, FPBDIslandSolverData& SolverData);
-		void GatherInput(const FReal Dt, const int32 Particle0Level, const int32 Particle1Level, FPBDIslandSolverData& SolverData);
-
 		static const FConstraintHandleTypeID& StaticType()
 		{
 			static FConstraintHandleTypeID STypeID(TEXT("FPositionConstraintHandle"), &FIndexedConstraintHandle::StaticType());
@@ -40,24 +34,21 @@ namespace Chaos
 	};
 
 	//! Constraint a single particle to a world-space position
-	class FPBDPositionConstraints : public FPBDIndexedConstraintContainer
+	class FPBDPositionConstraints : public TPBDIndexedConstraintContainer<FPBDPositionConstraints>
 	{
 	public:
-		using Base = FPBDIndexedConstraintContainer;
-		//using FReal = T;
-		//static const int Dimensions = 3;
+		using Base = TPBDIndexedConstraintContainer<FPBDPositionConstraints>;
 		using FConstraintContainerHandle = FPBDPositionConstraintHandle;
 		using FConstraintHandleAllocator = TConstraintHandleAllocator<FPBDPositionConstraints>;
 		using FHandles = TArray<FConstraintContainerHandle*>;
-		using FConstraintSolverContainerType = FConstraintSolverContainer;	// @todo(chaos): Add island solver for this constraint type
 
 		FPBDPositionConstraints(const FReal InStiffness = (FReal)1.)
-			: FPBDIndexedConstraintContainer(FConstraintContainerHandle::StaticType())
+			: TPBDIndexedConstraintContainer<FPBDPositionConstraints>(FConstraintContainerHandle::StaticType())
 			, Stiffness(InStiffness)
 		{}
 
 		FPBDPositionConstraints(TArray<FVec3>&& Locations, TArray<FPBDRigidParticleHandle*>&& InConstrainedParticles, const FReal InStiffness = (FReal)1.)
-			: FPBDIndexedConstraintContainer(FConstraintContainerHandle::StaticType())
+			: TPBDIndexedConstraintContainer<FPBDPositionConstraints>(FConstraintContainerHandle::StaticType())
 			, Targets(MoveTemp(Locations)), ConstrainedParticles(MoveTemp(InConstrainedParticles)), Stiffness(InStiffness)
 		{
 			if (ConstrainedParticles.Num() > 0)
@@ -179,28 +170,44 @@ namespace Chaos
 			Targets[ConstraintIndex] = Position;
 		}
 
+		void UpdatePositionBasedState(const FReal Dt) {}
+
 		//
-		// Island Rule API
+		// FConstraintContainer Implementation
 		//
+		virtual int32 GetNumConstraints() const override final { return NumConstraints(); }
+		virtual void ResetConstraints() override final {}
+		virtual void AddConstraintsToGraph(FPBDIslandManager& IslandManager) override final;
+		virtual void PrepareTick() override final {}
+		virtual void UnprepareTick() override final {}
 
-		void PrepareTick() {}
-		void UnprepareTick() {}
-		void UpdatePositionBasedState(const FReal Dt) {}
-
-		void SetNumIslandConstraints(const int32 NumIslandConstraints, FPBDIslandSolverData& SolverData);
-		void PreGatherInput(const FReal Dt, const int32 ConstraintIndex, FPBDIslandSolverData& SolverData);
-		void GatherInput(const FReal Dt, const int32 ConstraintIndex, const int32 Particle0Level, const int32 Particle1Level, FPBDIslandSolverData& SolverData);
-		void ScatterOutput(FReal Dt, FPBDIslandSolverData& SolverData);
+		//
+		// TSimpleConstraintContainerSolver API - used by RBAN solvers
+		//
+		void AddBodies(FSolverBodyContainer& SolverBodyContainer);
+		void GatherInput(const FReal Dt) {}
+		void ScatterOutput(const FReal Dt);
+		void ApplyPositionConstraints(const FReal Dt, const int32 It, const int32 NumIts);
+		void ApplyVelocityConstraints(const FReal Dt, const int32 It, const int32 NumIts) {}
+		void ApplyProjectionConstraints(const FReal Dt, const int32 It, const int32 NumIts) {}
 
-		bool ApplyPhase1Serial(const FReal Dt, const int32 It, const int32 NumIts, FPBDIslandSolverData& SolverData);
-		bool ApplyPhase2Serial(const FReal Dt, const int32 It, const int32 NumIts, FPBDIslandSolverData& SolverData) { return false; }
-		bool ApplyPhase3Serial(const FReal Dt, const int32 It, const int32 NumIts, FPBDIslandSolverData& SolverData) { return false; }
+		//
+		// TIndexedConstraintContainerSolver API - used by World solvers
+		//
+		void AddBodies(const TArrayView<int32>& ConstraintIndices, FSolverBodyContainer& SolverBodyContainer);
+		void GatherInput(const TArrayView<int32>& ConstraintIndices, const FReal Dt) {}
+		void ScatterOutput(const TArrayView<int32>& ConstraintIndices, const FReal Dt);
+		void ApplyPositionConstraints(const TArrayView<int32>& ConstraintIndices, const FReal Dt, const int32 It, const int32 NumIts);
+		void ApplyVelocityConstraints(const TArrayView<int32>& ConstraintIndices, const FReal Dt, const int32 It, const int32 NumIts) {}
+		void ApplyProjectionConstraints(const TArrayView<int32>& ConstraintIndices, const FReal Dt, const int32 It, const int32 NumIts) {}
 
 	protected:
 		using Base::GetConstraintIndex;
 		using Base::SetConstraintIndex;
 
 	private:
+		void AddBodies(const int32 ConstraintIndex, FSolverBodyContainer& SolverBodyContainer);
+
 		void ApplySingle(const FReal Dt, int32 ConstraintIndex) const
 		{
 			FSolverBody* Body = ConstraintSolverBodies[ConstraintIndex];
diff --git a/Engine/Source/Runtime/Experimental/Chaos/Public/Chaos/PBDRigidDynamicSpringConstraints.h b/Engine/Source/Runtime/Experimental/Chaos/Public/Chaos/PBDRigidDynamicSpringConstraints.h
--- a/Engine/Source/Runtime/Experimental/Chaos/Public/Chaos/PBDRigidDynamicSpringConstraints.h
+++ b/Engine/Source/Runtime/Experimental/Chaos/Public/Chaos/PBDRigidDynamicSpringConstraints.h
@@ -2,15 +2,12 @@
 #pragma once
 
 #include "Chaos/Array.h"
-#include "Chaos/ConstraintHandle.h"
+#include "Chaos/Evolution/IndexedConstraintContainer.h"
 #include "Chaos/ParticleHandle.h"
-#include "Chaos/PBDConstraintContainer.h"
-#include "Chaos/Evolution/SolverConstraintContainer.h"
 
 namespace Chaos
 {
 	class FPBDRigidDynamicSpringConstraints;
-	class FPBDIslandSolverData;
 
 	class CHAOS_API FPBDRigidDynamicSpringConstraintHandle final : public TIndexedContainerConstraintHandle<FPBDRigidDynamicSpringConstraints>
 	{
@@ -27,10 +24,7 @@ namespace Chaos
 		{
 		}
 
-		virtual FParticlePair GetConstrainedParticles() const override;
-
-		void PreGatherInput(const FReal Dt, FPBDIslandSolverData& SolverData);
-		void GatherInput(const FReal Dt, const int32 Particle0Level, const int32 Particle1Level, FPBDIslandSolverData& SolverData);
+		virtual FParticlePair GetConstrainedParticles() const override final;
 
 		static const FConstraintHandleTypeID& StaticType()
 		{
@@ -43,26 +37,24 @@ namespace Chaos
 		using Base::ConcreteContainer;
 	};
 
-	class CHAOS_API FPBDRigidDynamicSpringConstraints : public FPBDIndexedConstraintContainer
+	class CHAOS_API FPBDRigidDynamicSpringConstraints : public TPBDIndexedConstraintContainer<FPBDRigidDynamicSpringConstraints>
 	{
 	public:
-		using Base = FPBDIndexedConstraintContainer;
+		using Base = TPBDIndexedConstraintContainer<FPBDRigidDynamicSpringConstraints>;
 		using FConstrainedParticlePair = TVec2<FGeometryParticleHandle*>;
-		//static const int Dimensions = 3;
 		using FConstraintContainerHandle = FPBDRigidDynamicSpringConstraintHandle;
 		using FConstraintHandleAllocator = TConstraintHandleAllocator<FPBDRigidDynamicSpringConstraints>;
 		using FHandles = TArray<FConstraintContainerHandle*>;
-		using FConstraintSolverContainerType = FConstraintSolverContainer;	// @todo(chaos): Add island solver for this constraint type
 
 		FPBDRigidDynamicSpringConstraints(const FReal InStiffness = (FReal)1.)
-			: FPBDIndexedConstraintContainer(FConstraintContainerHandle::StaticType())
+			: TPBDIndexedConstraintContainer<FPBDRigidDynamicSpringConstraints>(FConstraintContainerHandle::StaticType())
 			, CreationThreshold(1)
 			, MaxSprings(1)
 			, Stiffness(InStiffness) 
 		{}
 
 		FPBDRigidDynamicSpringConstraints(TArray<FConstrainedParticlePair>&& InConstraints, const FReal InCreationThreshold = (FReal)1., const int32 InMaxSprings = 1, const FReal InStiffness = (FReal)1.)
-			: FPBDIndexedConstraintContainer(FConstraintContainerHandle::StaticType())
+			: TPBDIndexedConstraintContainer<FPBDRigidDynamicSpringConstraints>(FConstraintContainerHandle::StaticType())
 			, Constraints(MoveTemp(InConstraints))
 			, CreationThreshold(InCreationThreshold)
 			, MaxSprings(InMaxSprings)
@@ -197,30 +189,43 @@ namespace Chaos
 			return Constraints[ConstraintIndex];
 		}
 
-
+		void UpdatePositionBasedState(const FReal Dt);
 
 		//
-		// Island Rule API
+		// FConstraintContainer Implementation
 		//
+		virtual int32 GetNumConstraints() const override final { return NumConstraints(); }
+		virtual void ResetConstraints() override final {}
+		virtual void AddConstraintsToGraph(FPBDIslandManager& IslandManager) override final;
+		virtual void PrepareTick() override final {}
+		virtual void UnprepareTick() override final {}
 
-		void PrepareTick() {}
-		void UnprepareTick() {}
-		void UpdatePositionBasedState(const FReal Dt);
-
-		void SetNumIslandConstraints(const int32 NumIslandConstraints, FPBDIslandSolverData& SolverData);
-		void PreGatherInput(const FReal Dt, const int32 ConstraintIndex, FPBDIslandSolverData& SolverData);
-		void GatherInput(const FReal Dt, const int32 ConstraintIndex, const int32 Particle0Level, const int32 Particle1Level, FPBDIslandSolverData& SolverData);
-		void ScatterOutput(FReal Dt, FPBDIslandSolverData& SolverData);
+		//
+		// TSimpleConstraintContainerSolver API - used by RBAN solvers
+		//
+		void AddBodies(FSolverBodyContainer& SolverBodyContainer);
+		void GatherInput(const FReal Dt) {}
+		void ScatterOutput(const FReal Dt);
+		void ApplyPositionConstraints(const FReal Dt, const int32 It, const int32 NumIts);
+		void ApplyVelocityConstraints(const FReal Dt, const int32 It, const int32 NumIts) {}
+		void ApplyProjectionConstraints(const FReal Dt, const int32 It, const int32 NumIts) {}
 
-		bool ApplyPhase1Serial(const FReal Dt, const int32 It, const int32 NumIts, FPBDIslandSolverData& SolverData);
-		bool ApplyPhase2Serial(const FReal Dt, const int32 It, const int32 NumIts, FPBDIslandSolverData& SolverData) { return false; }
-		bool ApplyPhase3Serial(const FReal Dt, const int32 It, const int32 NumIts, FPBDIslandSolverData& SolverData) { return false; }
+		//
+		// TIndexedConstraintContainerSolver API - used by World solvers
+		//
+		void AddBodies(const TArrayView<int32>& ConstraintIndices, FSolverBodyContainer& SolverBodyContainer);
+		void GatherInput(const TArrayView<int32>& ConstraintIndices, const FReal Dt) {}
+		void ScatterOutput(const TArrayView<int32>& ConstraintIndices, const FReal Dt);
+		void ApplyPositionConstraints(const TArrayView<int32>& ConstraintIndices, const FReal Dt, const int32 It, const int32 NumIts);
+		void ApplyVelocityConstraints(const TArrayView<int32>& ConstraintIndices, const FReal Dt, const int32 It, const int32 NumIts) {}
+		void ApplyProjectionConstraints(const TArrayView<int32>& ConstraintIndices, const FReal Dt, const int32 It, const int32 NumIts) {}
 
 	protected:
 		using Base::GetConstraintIndex;
 		using Base::SetConstraintIndex;
 
 	private:
+		void AddBodies(const int32 ConstraintIndex, FSolverBodyContainer& SolverBodyContainer);
 		void ApplySingle(const FReal Dt, int32 ConstraintIndex) const;
 
 		FVec3 GetDelta(const FVec3& WorldSpaceX1, const FVec3& WorldSpaceX2, const int32 ConstraintIndex, const int32 SpringIndex) const;
diff --git a/Engine/Source/Runtime/Experimental/Chaos/Public/Chaos/PBDRigidSpringConstraints.h b/Engine/Source/Runtime/Experimental/Chaos/Public/Chaos/PBDRigidSpringConstraints.h
--- a/Engine/Source/Runtime/Experimental/Chaos/Public/Chaos/PBDRigidSpringConstraints.h
+++ b/Engine/Source/Runtime/Experimental/Chaos/Public/Chaos/PBDRigidSpringConstraints.h
@@ -2,15 +2,12 @@
 #pragma once
 
 #include "Chaos/Array.h"
-#include "Chaos/ConstraintHandle.h"
+#include "Chaos/Evolution/IndexedConstraintContainer.h"
 #include "Chaos/ParticleHandle.h"
-#include "Chaos/PBDConstraintContainer.h"
-#include "Chaos/Evolution/SolverConstraintContainer.h"
 
 namespace Chaos
 {
 	class FPBDRigidSpringConstraints;
-	class FPBDIslandSolverData;
 
 	class CHAOS_API FPBDRigidSpringConstraintHandle final : public TIndexedContainerConstraintHandle<FPBDRigidSpringConstraints>
 	{
@@ -31,15 +28,12 @@ namespace Chaos
 		const TVector<FVec3, 2>& GetConstraintPositions() const;
 		void SetConstraintPositions(const TVector<FVec3, 2>& ConstraintPositions);
 		
-		virtual FParticlePair GetConstrainedParticles() const override;
+		virtual FParticlePair GetConstrainedParticles() const override final;
 
 		// Get the rest length of the spring
 		FReal GetRestLength() const;
 		void SetRestLength(const FReal SpringLength);
 
-		void PreGatherInput(const FReal Dt, FPBDIslandSolverData& SolverData);
-		void GatherInput(const FReal Dt, const int32 Particle0Level, const int32 Particle1Level, FPBDIslandSolverData& SolverData);
-
 		static const FConstraintHandleTypeID& StaticType()
 		{
 			static FConstraintHandleTypeID STypeID(TEXT("FRigidSpringConstraintHandle"), &FIndexedConstraintHandle::StaticType());
@@ -48,17 +42,16 @@ namespace Chaos
 	};
 
 
-	class FPBDRigidSpringConstraints : public FPBDIndexedConstraintContainer
+	class FPBDRigidSpringConstraints : public TPBDIndexedConstraintContainer<FPBDRigidSpringConstraints>
 	{
 	public:
 		// @todo(ccaulfield): an alternative AddConstraint which takes the constrain settings rather than assuming everything is in world-space rest pose
 
-		using Base = FPBDIndexedConstraintContainer;
+		using Base = TPBDIndexedConstraintContainer<FPBDRigidSpringConstraints>;
 		using FConstraintContainerHandle = FPBDRigidSpringConstraintHandle;
 		using FConstraintHandleAllocator = TConstraintHandleAllocator<FPBDRigidSpringConstraints>;
 		using FConstrainedParticlePair = TVector<TGeometryParticleHandle<FReal, 3>*, 2>;
 		using FHandles = TArray<FConstraintContainerHandle*>;
-		using FConstraintSolverContainerType = FConstraintSolverContainer;	// @todo(chaos): Add island solver for this constraint type
 
 		FPBDRigidSpringConstraints();
 		virtual ~FPBDRigidSpringConstraints();
@@ -159,39 +152,42 @@ namespace Chaos
 			SpringSettings[ConstraintIndex].RestLength = SpringLength;
 		}
 
+
 		//
-		// Simple Rule API
+		// FConstraintContainer Implementation
 		//
-
-		void PreGatherInput(const FReal Dt, FPBDIslandSolverData& SolverData);
-		void GatherInput(const FReal Dt, FPBDIslandSolverData& SolverData);
-		void ScatterOutput(const FReal Dt, FPBDIslandSolverData& SolverData);
-
-		bool ApplyPhase1(const FReal Dt, const int32 It, const int32 NumIts, FPBDIslandSolverData& SolverData);
-		bool ApplyPhase2(const FReal Dt, const int32 It, const int32 NumIts, FPBDIslandSolverData& SolverData) { return false; }
-		bool ApplyPhase3(const FReal Dt, const int32 It, const int32 NumIts, FPBDIslandSolverData& SolverData) { return false; }
+		virtual int32 GetNumConstraints() const override final { return NumConstraints(); }
+		virtual void ResetConstraints() override final {}
+		virtual void AddConstraintsToGraph(FPBDIslandManager& IslandManager) override final;
+		virtual void PrepareTick() override final {}
+		virtual void UnprepareTick() override final {}
 
 		//
-		// Island Rule API
+		// TSimpleConstraintContainerSolver API - used by RBAN solvers
 		//
+		void AddBodies(FSolverBodyContainer& SolverBodyContainer);
+		void GatherInput(const FReal Dt) {}
+		void ScatterOutput(const FReal Dt);
+		void ApplyPositionConstraints(const FReal Dt, const int32 It, const int32 NumIts);
+		void ApplyVelocityConstraints(const FReal Dt, const int32 It, const int32 NumIts) {}
+		void ApplyProjectionConstraints(const FReal Dt, const int32 It, const int32 NumIts) {}
 
-		void PrepareTick() {}
-		void UnprepareTick() {}
-		void UpdatePositionBasedState(const FReal Dt) {}
-
-		void SetNumIslandConstraints(const int32 NumIslandConstraints, FPBDIslandSolverData& SolverData);
-		void PreGatherInput(const FReal Dt, const int32 ConstraintIndex, FPBDIslandSolverData& SolverData);
-		void GatherInput(const FReal Dt, const int32 ConstraintIndex, const int32 Particle0Level, const int32 Particle1Level, FPBDIslandSolverData& SolverData);
-
-		bool ApplyPhase1Serial(const FReal Dt, const int32 It, const int32 NumIts, FPBDIslandSolverData& SolverData);
-		bool ApplyPhase2Serial(const FReal Dt, const int32 It, const int32 NumIts, FPBDIslandSolverData& SolverData) { return false; }
-		bool ApplyPhase3Serial(const FReal Dt, const int32 It, const int32 NumIts, FPBDIslandSolverData& SolverData) { return false; }
+		//
+		// TIndexedConstraintContainerSolver API - used by World solvers
+		//
+		void AddBodies(const TArrayView<int32>& ConstraintIndices, FSolverBodyContainer& SolverBodyContainer);
+		void GatherInput(const TArrayView<int32>& ConstraintIndices, const FReal Dt) {}
+		void ScatterOutput(const TArrayView<int32>& ConstraintIndices, const FReal Dt);
+		void ApplyPositionConstraints(const TArrayView<int32>& ConstraintIndices, const FReal Dt, const int32 It, const int32 NumIts);
+		void ApplyVelocityConstraints(const TArrayView<int32>& ConstraintIndices, const FReal Dt, const int32 It, const int32 NumIts) {}
+		void ApplyProjectionConstraints(const TArrayView<int32>& ConstraintIndices, const FReal Dt, const int32 It, const int32 NumIts) {}
 
 	protected:
 		using Base::GetConstraintIndex;
 		using Base::SetConstraintIndex;
 
 	private:
+		void AddBodies(const int32 ConstraintIndex, FSolverBodyContainer& SolverBodyContainer);
 		void ApplyPhase1Single(const FReal Dt, int32 ConstraintIndex) const;
 
 		void InitDistance(int32 ConstraintIndex, const FVec3& Location0, const FVec3& Location1);
diff --git a/Engine/Source/Runtime/Experimental/Chaos/Public/Chaos/PBDRigidsEvolution.h b/Engine/Source/Runtime/Experimental/Chaos/Public/Chaos/PBDRigidsEvolution.h
--- a/Engine/Source/Runtime/Experimental/Chaos/Public/Chaos/PBDRigidsEvolution.h
+++ b/Engine/Source/Runtime/Experimental/Chaos/Public/Chaos/PBDRigidsEvolution.h
@@ -6,7 +6,6 @@
 #include "Chaos/PBDConstraintGraph.h"
 #include "Chaos/PBDRigidClustering.h"
 #include "Chaos/PBDRigidParticles.h"
-#include "Chaos/PBDConstraintRule.h"
 #include "Chaos/ParticleHandle.h"
 #include "Chaos/Transform.h"
 #include "Chaos/Framework/DebugSubstep.h"
@@ -15,6 +14,7 @@
 #include "Chaos/PBDRigidsSOAs.h"
 #include "Chaos/SpatialAccelerationCollection.h"
 #include "Chaos/PBDRigidsEvolutionFwd.h"
+#include "Chaos/Island/IslandGroupManager.h"
 #include "Chaos/Defines.h"
 #include "Chaos/PendingSpatialData.h"
 #include "ProfilingDebugging/CsvProfiler.h"
@@ -341,11 +341,14 @@ public:
 	CHAOS_API FPBDRigidsSOAs& GetParticles() { return Particles; }
 	CHAOS_API const FPBDRigidsSOAs& GetParticles() const { return Particles; }
 
-	CHAOS_API void AddConstraintRule(FPBDConstraintGraphRule* ConstraintRule)
+	// NOTE: we do not currently support removing containers. In a few places we assume the ContainerId is equal to the array index.
+	CHAOS_API void AddConstraintContainer(FPBDConstraintContainer& InContainer, const int32 Priority = 0)
 	{
-		uint32 ContainerId = (uint32)ConstraintRules.Num();
-		ConstraintRules.Add(ConstraintRule);
-		ConstraintRule->BindToGraph(ConstraintGraph, ContainerId);
+		const int32 ContainerId = ConstraintContainers.Add(&InContainer);
+		InContainer.SetContainerId(ContainerId);
+
+		GetConstraintGraph().AddConstraintContainer(InContainer);
+		IslandGroupManager.AddConstraintContainer(InContainer, Priority);
 	}
 
 	CHAOS_API void SetNumPositionIterations(int32 InNumIterations)
@@ -568,9 +571,9 @@ public:
 			RemoveConstraintsFromConstraintGraph(ParticleHandle->ParticleConstraints());
 		}
 
-		for (FPBDConstraintGraphRule* ConstraintRule : ConstraintRules)
+		for (FPBDConstraintContainer* Container : ConstraintContainers)
 		{
-			ConstraintRule->DisconnectConstraints(RemovedParticles);
+			Container->DisconnectConstraints(RemovedParticles);
 		}
 	}
 
@@ -583,9 +586,12 @@ public:
 	{
 		RemoveConstraintsFromConstraintGraph(ParticleHandle->ParticleConstraints());
 
-		for (FPBDConstraintGraphRule* ConstraintRule : ConstraintRules)
+		for (FConstraintHandle* Constraint : ParticleHandle->ParticleConstraints())
 		{
-			ConstraintRule->SetConstraintsEnabled(ParticleHandle, false);
+			if (Constraint->IsEnabled())
+			{
+				Constraint->SetEnabled(false);
+			}
 		}
 	}
 
@@ -600,9 +606,12 @@ public:
 	/** Enable constraints from the enabled particles; constraints will only become enabled if their particle end points are valid */
 	CHAOS_API void EnableConstraints(FGeometryParticleHandle* ParticleHandle)
 	{
-		for (FPBDConstraintGraphRule* ConstraintRule : ConstraintRules)
+		for (FConstraintHandle* Constraint : ParticleHandle->ParticleConstraints())
 		{
-			ConstraintRule->SetConstraintsEnabled(ParticleHandle, true);
+			if (!Constraint->IsEnabled())
+			{
+				Constraint->SetEnabled(true);
+			}
 		}
 	}
 
@@ -630,9 +639,9 @@ public:
 		}
 
 		// Remove all constraints from the containers
-		for(FPBDConstraintGraphRule* ConstraintRule : ConstraintRules)
+		for (FPBDConstraintContainer* Container : ConstraintContainers)
 		{
-			ConstraintRule->ResetConstraints();
+			Container->ResetConstraints();
 		}
 	}
 
@@ -661,44 +670,29 @@ public:
 	{
 		ConstraintGraph.InitializeGraph(Particles.GetNonDisabledDynamicView());
 
-		for (FPBDConstraintGraphRule* ConstraintRule : ConstraintRules)
+		for (FPBDConstraintContainer* Container : ConstraintContainers)
 		{
-			ConstraintRule->AddToGraph();
+			Container->AddConstraintsToGraph(ConstraintGraph);
 		}
 
 		ConstraintGraph.ResetIslands(Particles.GetNonDisabledDynamicView());
-
-		for (FPBDConstraintGraphRule* ConstraintRule : ConstraintRules)
-		{
-			ConstraintRule->InitializeAccelerationStructures();
-		}
 	}
 
 	void PrepareTick()
 	{
 		Collisions::ResetChaosCollisionCounters();
 
-		for (FPBDConstraintGraphRule* ConstraintRule : ConstraintRules)
+		for (FPBDConstraintContainer* Container : ConstraintContainers)
 		{
-			ConstraintRule->PrepareTick();
+			Container->PrepareTick();
 		}
 	}
 
 	void UnprepareTick()
 	{
-		for (FPBDConstraintGraphRule* ConstraintRule : ConstraintRules)
-		{
-			ConstraintRule->UnprepareTick();
-		}
-	}
-
-	void UpdateAccelerationStructures(const FReal Dt, const int32 Island)
-	{
-		CSV_SCOPED_TIMING_STAT(Chaos, UpdateAccelerationStructures);
-
-		for (FPBDConstraintGraphRule* ConstraintRule : ConstraintRules)
+		for (FPBDConstraintContainer* Container : ConstraintContainers)
 		{
-			ConstraintRule->UpdateAccelerationStructures(Dt, Island);
+			Container->UnprepareTick();
 		}
 	}
 
@@ -828,6 +822,8 @@ public:
 
 	CHAOS_API const FPBDConstraintGraph& GetConstraintGraph() const { return ConstraintGraph; }
 	CHAOS_API FPBDConstraintGraph& GetConstraintGraph() { return ConstraintGraph; }
+	CHAOS_API const FPBDIslandGroupManager& GetIslandGroupManager() const { return IslandGroupManager; }
+
 
 	void SetResim(bool bInResim) { bIsResim = bInResim; }
 	const bool IsResimming() const { return bIsResim; }
@@ -873,9 +869,9 @@ protected:
 	int32 NumConstraints() const
 	{
 		int32 NumConstraints = 0;
-		for (const FPBDConstraintGraphRule* ConstraintRule : ConstraintRules)
+		for (const FPBDConstraintContainer* Container : ConstraintContainers)
 		{
-			NumConstraints += ConstraintRule->NumConstraints();
+			NumConstraints += Container->GetNumConstraints();
 		}
 		return NumConstraints;
 	}
@@ -905,53 +901,34 @@ protected:
 
 	void UpdateConstraintPositionBasedState(FReal Dt)
 	{
-		for (FPBDConstraintGraphRule* ConstraintRule : ConstraintRules)
+		// If any constraint container state depends on particle state, it gets updated here.
+		// E.g., we can create constraints between close particles etc. Collision detection 
+		// could be called from here, but currently is called explicitly elsewhere
+		for (FPBDConstraintContainer* ConstraintContainer : ConstraintContainers)
 		{
-			ConstraintRule->UpdatePositionBasedState(Dt);
+			ConstraintContainer->UpdatePositionBasedState(Dt);
 		}
 	}
 
 	void CreateConstraintGraph()
 	{
+		// Add new particles to the graph
 		ConstraintGraph.InitializeGraph(Particles.GetNonDisabledDynamicView());
 
-		for (FPBDConstraintGraphRule* ConstraintRule : ConstraintRules)
+		// Add all constraints to graph (all constraints except sleeping collisions are removed every tick)
+		for (FPBDConstraintContainer* ConstraintContainer : ConstraintContainers)
 		{
-			ConstraintRule->AddToGraph();
+			ConstraintContainer->AddConstraintsToGraph(ConstraintGraph);
 		}
 
-		// Apply rules in priority order
-		// @todo(ccaulfield): only really needed when list or priorities change
-		PrioritizedConstraintRules = ConstraintRules;
-		PrioritizedConstraintRules.StableSort();
+		// Process the addidions
+		ConstraintGraph.UpdateGraph();
 	}
 
 	void CreateIslands()
 	{
-		ConstraintGraph.UpdateIslands(Particles.GetNonDisabledDynamicView(), Particles, ConstraintRules.Num());
-
-		for (FPBDConstraintGraphRule* ConstraintRule : ConstraintRules)
-		{
-			ConstraintRule->InitializeAccelerationStructures();
-		}
-	}
-	
-	/** Sort constraints if necessary */
-	void SortConstraints()
-	{
-		bool bNeedsSorting = false;
-		for (FPBDConstraintGraphRule* ConstraintRule : ConstraintRules)
-		{
-			bNeedsSorting |= ConstraintRule->IsSortingEnabled();
-		}
-		if(bNeedsSorting)
-		{
-			ConstraintGraph.GetIslandGraph()->InitSorting();
-			for (FPBDConstraintGraphRule* ConstraintRule : ConstraintRules)
-			{
-				ConstraintRule->SortConstraints();
-			}
-		}
+		// Package the constraints and particles into islands
+		ConstraintGraph.UpdateIslands(Particles);
 	}
 	
 	void FlushInternalAccelerationQueue();
@@ -965,9 +942,9 @@ protected:
 	FUpdatePositionRule ParticleUpdatePosition;
 	FKinematicUpdateRule KinematicUpdate;
 	FCaptureRewindRule CaptureRewindData;
-	TArray<FPBDConstraintGraphRule*> ConstraintRules;
-	TArray<FPBDConstraintGraphRule*> PrioritizedConstraintRules;
-	FPBDConstraintGraph ConstraintGraph;
+	TArray<FPBDConstraintContainer*> ConstraintContainers;
+	FPBDIslandManager ConstraintGraph;
+	FPBDIslandGroupManager IslandGroupManager;
 	TArrayCollectionArray<TSerializablePtr<FChaosPhysicsMaterial>> PhysicsMaterials;
 	TArrayCollectionArray<TUniquePtr<FChaosPhysicsMaterial>> PerParticlePhysicsMaterials;
 	TArrayCollectionArray<int32> ParticleDisableCount;
diff --git a/Engine/Source/Runtime/Experimental/Chaos/Public/Chaos/PBDRigidsEvolutionGBF.h b/Engine/Source/Runtime/Experimental/Chaos/Public/Chaos/PBDRigidsEvolutionGBF.h
--- a/Engine/Source/Runtime/Experimental/Chaos/Public/Chaos/PBDRigidsEvolutionGBF.h
+++ b/Engine/Source/Runtime/Experimental/Chaos/Public/Chaos/PBDRigidsEvolutionGBF.h
@@ -53,13 +53,9 @@ namespace Chaos
 
 		using FGravityForces = FPerParticleGravity;
 		using FCollisionConstraints = FPBDCollisionConstraints;
-		using FCollisionConstraintRule = TPBDConstraintColorRule<FCollisionConstraints>;
 		using FCollisionDetector = FSpatialAccelerationCollisionDetector;
 		using FExternalForces = FPerParticleExternalForces;
-		using FJointConstraintsRule = TPBDConstraintIslandRule<FPBDJointConstraints>;
-		using FSuspensionConstraintsRule = TPBDConstraintIslandRule<FPBDSuspensionConstraints>;
 		using FJointConstraints = FPBDJointConstraints;
-		using FJointConstraintRule = TPBDConstraintIslandRule<FJointConstraints>;
 
 		// Default iteration counts
 		static constexpr int32 DefaultNumPositionIterations = 8;
@@ -70,6 +66,7 @@ namespace Chaos
 		static constexpr FRealSingle DefaultCollisionCullDistance = 3.0f;
 		static constexpr FRealSingle DefaultCollisionMaxPushOutVelocity = 1000.0f;
 		static constexpr int32 DefaultRestitutionThreshold = 1000;
+		static constexpr int32 DefaultNumCollisionsPerBlock = 1000;
 
 		CHAOS_API FPBDRigidsEvolutionGBF(FPBDRigidsSOAs& InParticles, THandleArray<FChaosPhysicsMaterial>& SolverPhysicsMaterials, const TArray<ISimCallbackObject*>* InCollisionModifiers = nullptr, bool InIsSingleThreaded = false);
 		CHAOS_API ~FPBDRigidsEvolutionGBF();
@@ -89,16 +86,6 @@ namespace Chaos
 			PreApplyCallback = Cb;
 		}
 
-		FORCEINLINE void SetPostApplyCallback(const FPBDRigidsEvolutionIslandCallback& Cb)
-		{
-			PostApplyCallback = Cb;
-		}
-
-		FORCEINLINE void SetPostApplyPushOutCallback(const FPBDRigidsEvolutionIslandCallback& Cb)
-		{
-			PostApplyPushOutCallback = Cb;
-		}
-
 		FORCEINLINE void SetInternalParticleInitilizationFunction(const FPBDRigidsEvolutionInternalHandleCallback& Cb)
 		{ 
 			InternalParticleInitilization = Cb;
@@ -120,9 +107,6 @@ namespace Chaos
 		FORCEINLINE FCollisionConstraints& GetCollisionConstraints() { return CollisionConstraints; }
 		FORCEINLINE const FCollisionConstraints& GetCollisionConstraints() const { return CollisionConstraints; }
 
-		FORCEINLINE FCollisionConstraintRule& GetCollisionConstraintsRule() { return CollisionRule; }
-		FORCEINLINE const FCollisionConstraintRule& GetCollisionConstraintsRule() const { return CollisionRule; }
-
 		FORCEINLINE FCollisionDetector& GetCollisionDetector() { return CollisionDetector; }
 		FORCEINLINE const FCollisionDetector& GetCollisionDetector() const { return CollisionDetector; }
 
@@ -139,10 +123,9 @@ namespace Chaos
 		FORCEINLINE const FPBDSuspensionConstraints& GetSuspensionConstraints() const { return SuspensionConstraints; }
 
 		/**
-		 * Reload the particles cache within an island
-		 * @param Island Index of the island in which the cache will be used
+		 * Reload the particles cache for all particles where appropriate
 		 */
-		void ReloadParticlesCache(const int32 Island);
+		void ReloadParticlesCache();
 
 		/**
 		 * Build the list of disables particles and update the sleeping flag on the island
@@ -152,8 +135,6 @@ namespace Chaos
 		 */
 		void BuildDisabledParticles(const int32 Island, TArray<TArray<FPBDRigidParticleHandle*>>& DisabledParticles, TArray<bool>& SleepedIslands);
 
-		void DestroyConstraint(FConstraintHandle* Constraint);
-
 		void DestroyParticleCollisionsInAllocator(FGeometryParticleHandle* Particle);
 
 		virtual void DestroyTransientConstraints(FGeometryParticleHandle* Particle) override;
@@ -304,26 +285,6 @@ namespace Chaos
 			}
 		}
 
-		// First phase of constraint solver
-		// For GBF this is the velocity solve phase
-		// For PBD/QuasiPBD this is the position solve phase
-		void ApplyConstraintsPhase1(const FReal Dt, int32 GroupIndex);
-
-		// Calculate the implicit velocites based on the change in position from ApplyConstraintsPhase1
-		void SetImplicitVelocities(const FReal Dt, int32 GroupIndex);
-		
-		// Second phase of constraint solver (after implicit velocity calculation following results of phase 1)
-		// For StandardPBD this does nothing
-		// For QuasiPBD this is the velocity solve phase
-		void ApplyConstraintsPhase2(const FReal Dt, int32 GroupIndex);
-
-		// Apply the accumulated corrections to the solver bodies
-		void ApplyCorrections(const FReal Dt, int32 GroupIndex);
-
-		// Third phase of constraint solver
-		// This is the projection phase (only joints implement this atm)
-		void ApplyConstraintsPhase3(const FReal Dt, int32 GroupIndex);
-
 		CHAOS_API void Serialize(FChaosArchive& Ar);
 
 		CHAOS_API TUniquePtr<IResimCacheBase> CreateExternalResimCache() const;
@@ -359,9 +320,6 @@ namespace Chaos
 
 		CHAOS_API void AdvanceOneTimeStepImpl(const FReal dt, const FSubStepInfo& SubStepInfo);
 
-		void GatherSolverInput(FReal Dt, int32 GroupIndex);
-		void ScatterSolverOutput(FReal Dt, int32 GroupIndex);
-
 		void UpdateInertiaConditioning();
 
 		FEvolutionResimCache* GetCurrentStepResimCache()
@@ -372,13 +330,10 @@ namespace Chaos
 		FRigidClustering Clustering;
 
 		FPBDJointConstraints JointConstraints;
-		TPBDConstraintIslandRule<FPBDJointConstraints> JointConstraintRule;
 		FPBDSuspensionConstraints SuspensionConstraints;
-		TPBDConstraintIslandRule<FPBDSuspensionConstraints> SuspensionConstraintRule;
 
 		FGravityForces GravityForces;
 		FCollisionConstraints CollisionConstraints;
-		FCollisionConstraintRule CollisionRule;
 		FSpatialAccelerationBroadPhase BroadPhase;
 		FNarrowPhase NarrowPhase;
 		FSpatialAccelerationCollisionDetector CollisionDetector;
@@ -386,8 +341,6 @@ namespace Chaos
 		FPBDRigidsEvolutionCallback PostIntegrateCallback;
 		FPBDRigidsEvolutionCallback PostDetectCollisionsCallback;
 		FPBDRigidsEvolutionCallback PreApplyCallback;
-		FPBDRigidsEvolutionIslandCallback PostApplyCallback;
-		FPBDRigidsEvolutionIslandCallback PostApplyPushOutCallback;
 		FPBDRigidsEvolutionInternalHandleCallback InternalParticleInitilization;
 		FEvolutionResimCache* CurrentStepResimCacheImp;
 		const TArray<ISimCallbackObject*>* CollisionModifiers;
diff --git a/Engine/Source/Runtime/Experimental/Chaos/Public/Chaos/PBDSuspensionConstraints.h b/Engine/Source/Runtime/Experimental/Chaos/Public/Chaos/PBDSuspensionConstraints.h
--- a/Engine/Source/Runtime/Experimental/Chaos/Public/Chaos/PBDSuspensionConstraints.h
+++ b/Engine/Source/Runtime/Experimental/Chaos/Public/Chaos/PBDSuspensionConstraints.h
@@ -2,10 +2,8 @@
 #pragma once
 
 #include "Chaos/Array.h"
-#include "Chaos/ConstraintHandle.h"
-#include "Chaos/Evolution/SolverDatas.h"
+#include "Chaos/Evolution/IndexedConstraintContainer.h"
 #include "Chaos/ParticleHandle.h"
-#include "Chaos/PBDConstraintContainer.h"
 #include "Chaos/PBDSuspensionConstraintTypes.h"
 #include "Chaos/PBDSuspensionConstraintData.h"
 #include "Chaos/Particle/ParticleUtilities.h"
@@ -16,7 +14,6 @@
 namespace Chaos
 {
 	class FPBDSuspensionConstraints;
-	class FPBDIslandSolverData;
 	class FPBDCollisionSolver;
 
 	class CHAOS_API FPBDSuspensionConstraintHandle final : public TIndexedContainerConstraintHandle<FPBDSuspensionConstraints>
@@ -33,10 +30,7 @@ namespace Chaos
 
 		void SetSettings(const FPBDSuspensionSettings& Settings);
 
-		virtual FParticlePair GetConstrainedParticles() const override;
-
-		void PreGatherInput(const FReal Dt, FPBDIslandSolverData& SolverData);
-		void GatherInput(const FReal Dt, const int32 Particle0Level, const int32 Particle1Level, FPBDIslandSolverData& SolverData);
+		virtual FParticlePair GetConstrainedParticles() const override final;
 
 		static const FConstraintHandleTypeID& StaticType()
 		{
@@ -48,22 +42,21 @@ namespace Chaos
 		using Base::ConcreteContainer;
 	};
 
-	class CHAOS_API FPBDSuspensionConstraints : public FPBDIndexedConstraintContainer
+	class CHAOS_API FPBDSuspensionConstraints : public TPBDIndexedConstraintContainer<FPBDSuspensionConstraints>
 	{
 	public:
-		using Base = FPBDIndexedConstraintContainer;
+		using Base = TPBDIndexedConstraintContainer<FPBDSuspensionConstraints>;
 		using FConstraintContainerHandle = FPBDSuspensionConstraintHandle;
 		using FConstraintHandleAllocator = TConstraintHandleAllocator<FPBDSuspensionConstraints>;
 		using FHandles = TArray<FConstraintContainerHandle*>;
-		using FConstraintSolverContainerType = FConstraintSolverContainer;	// @todo(chaos): Add island solver for this constraint type
 
 		FPBDSuspensionConstraints(const FPBDSuspensionSolverSettings& InSolverSettings = FPBDSuspensionSolverSettings())
-			: FPBDIndexedConstraintContainer(FConstraintContainerHandle::StaticType())
+			: TPBDIndexedConstraintContainer<FPBDSuspensionConstraints>(FConstraintContainerHandle::StaticType())
 			, SolverSettings(InSolverSettings)
 		{}
 
 		FPBDSuspensionConstraints(TArray<FVec3>&& Locations, TArray<TGeometryParticleHandle<FReal,3>*>&& InConstrainedParticles, TArray<FVec3>&& InLocalOffset, TArray<FPBDSuspensionSettings>&& InConstraintSettings)
-			: FPBDIndexedConstraintContainer(FConstraintContainerHandle::StaticType())
+			: TPBDIndexedConstraintContainer<FPBDSuspensionConstraints>(FConstraintContainerHandle::StaticType())
 			, ConstrainedParticles(MoveTemp(InConstrainedParticles)), SuspensionLocalOffset(MoveTemp(InLocalOffset)), ConstraintSettings(MoveTemp(InConstraintSettings))
 		{
 			if (ConstrainedParticles.Num() > 0)
@@ -226,31 +219,47 @@ namespace Chaos
 		}
 
 		//
-		// Island Rule API
+		// FConstraintContainer Implementation
 		//
+		virtual int32 GetNumConstraints() const override final { return NumConstraints(); }
+		virtual void ResetConstraints() override final {}
+		virtual void AddConstraintsToGraph(FPBDIslandManager& IslandManager) override final;
+		virtual void PrepareTick() override final {}
+		virtual void UnprepareTick() override final {}
 
-		void PrepareTick() {}
-		void UnprepareTick() {}
-		void UpdatePositionBasedState(const FReal Dt) {}
+		//
+		// TSimpleConstraintContainerSolver API - used by RBAN
+		//
+		void AddBodies(FSolverBodyContainer& SolverBodyContainer);
+		void GatherInput(const FReal Dt);
+		void ScatterOutput(const FReal Dt);
+		void ApplyPositionConstraints(const FReal Dt, const int32 It, const int32 NumIts);
+		void ApplyVelocityConstraints(const FReal Dt, const int32 It, const int32 NumIts);
+		void ApplyProjectionConstraints(const FReal Dt, const int32 It, const int32 NumIts) {}
 
-		void SetNumIslandConstraints(const int32 NumIslandConstraints, FPBDIslandSolverData& SolverData);
-		void PreGatherInput(const FReal Dt, const int32 ConstraintIndex, FPBDIslandSolverData& SolverData);
-		void GatherInput(const FReal Dt, const int32 ConstraintIndex, const int32 Particle0Level, const int32 Particle1Level, FPBDIslandSolverData& SolverData);
-		void ScatterOutput(FReal Dt, FPBDIslandSolverData& SolverData);
+		//
+		// TIndexedConstraintContainerSolver API - used by World solvers
+		//
+		void AddBodies(const TArrayView<int32>& ConstraintIndices, FSolverBodyContainer& SolverBodyContainer);
+		void GatherInput(const TArrayView<int32>& ConstraintIndices, const FReal Dt);
+		void ScatterOutput(const TArrayView<int32>& ConstraintIndices, const FReal Dt);
+		void ApplyPositionConstraints(const TArrayView<int32>& ConstraintIndices, const FReal Dt, const int32 It, const int32 NumIts);
+		void ApplyVelocityConstraints(const TArrayView<int32>& ConstraintIndices, const FReal Dt, const int32 It, const int32 NumIts);
+		void ApplyProjectionConstraints(const TArrayView<int32>& ConstraintIndices, const FReal Dt, const int32 It, const int32 NumIts) {}
 
-		bool ApplyPhase1Serial(const FReal Dt, const int32 It, const int32 NumIts, FPBDIslandSolverData& SolverData);
-		bool ApplyPhase2Serial(const FReal Dt, const int32 It, const int32 NumIts, FPBDIslandSolverData& SolverData);
-		bool ApplyPhase3Serial(const FReal Dt, const int32 It, const int32 NumIts, FPBDIslandSolverData& SolverData) { return false; }
 
 	protected:
 		using Base::GetConstraintIndex;
 		using Base::SetConstraintIndex;
 
 	private:
-
-		void ApplySingle(const FReal Dt, int32 ConstraintIndex);
+		void AddBodies(const int32 ConstraintIndex, FSolverBodyContainer& SolverBodyContainer);
+		void GatherInput(const int32 ConstraintIndex, FReal Dt);
+		void ScatterOutput(const int32 ConstraintIndex, FReal Dt);
+		void ApplyPositionConstraint(const int32 ConstraintIndex, const FReal Dt, const int32 It, const int32 NumIts);
+		void ApplyVelocityConstraint(const int32 ConstraintIndex, const FReal Dt, const int32 It, const int32 NumIts);
+		void ApplySingle(int32 ConstraintIndex, const FReal Dt);
 		
-		void ApplyPositionConstraintSoft(const int ConstraintIndex, const FReal Dt, const bool bAccelerationMode);
 		FPBDSuspensionSolverSettings SolverSettings;
 
 		TArray<FGeometryParticleHandle*> ConstrainedParticles;
diff --git a/Engine/Source/Runtime/Experimental/Chaos/Public/ChaosSolverConfiguration.h b/Engine/Source/Runtime/Experimental/Chaos/Public/ChaosSolverConfiguration.h
--- a/Engine/Source/Runtime/Experimental/Chaos/Public/ChaosSolverConfiguration.h
+++ b/Engine/Source/Runtime/Experimental/Chaos/Public/ChaosSolverConfiguration.h
@@ -90,9 +90,6 @@ struct FChaosSolverConfiguration
 	UPROPERTY(EditAnywhere, Category = "SolverConfiguration|GeneratedData", meta = (EditCondition = bGenerateTrailingData))
 	FSolverTrailingFilterSettings TrailingFilterSettings;
 
-	UPROPERTY(EditAnywhere, Category = "SolverConfiguration|Contacts")
-	bool bGenerateContactGraph;
-
 private:
 
 	// Renamed to PositionIterations
@@ -102,4 +99,8 @@ private:
 	// Renamed to VelocityIterations
 	UPROPERTY()
 	int32 PushOutIterations_DEPRECATED;
+
+	// No longer used
+	UPROPERTY()
+	bool bGenerateContactGraph_DEPRECATED;
 };
diff --git a/Engine/Source/Runtime/Experimental/Chaos/Public/ChaosStats.h b/Engine/Source/Runtime/Experimental/Chaos/Public/ChaosStats.h
--- a/Engine/Source/Runtime/Experimental/Chaos/Public/ChaosStats.h
+++ b/Engine/Source/Runtime/Experimental/Chaos/Public/ChaosStats.h
@@ -16,6 +16,7 @@ DECLARE_STATS_GROUP(TEXT("ChaosMinEvolution"), STATGROUP_ChaosMinEvolution, STAT
 DECLARE_STATS_GROUP(TEXT("ChaosCounters"), STATGROUP_ChaosCounters, STATCAT_Advanced);
 DECLARE_STATS_GROUP(TEXT("ChaosIterations"), STATGROUP_ChaosIterations, STATCAT_Advanced);
 DECLARE_STATS_GROUP(TEXT("ChaosCollisionCounters"), STATGROUP_ChaosCollisionCounters, STATCAT_Advanced);
+DECLARE_STATS_GROUP(TEXT("ChaosConstraintDetails"), STATGROUP_ChaosConstraintDetails, STATCAT_Advanced);
 
 DECLARE_CYCLE_STAT_EXTERN(TEXT("Physics Tick"), STAT_ChaosTick, STATGROUP_Chaos, CHAOS_API);
 DECLARE_CYCLE_STAT_EXTERN(TEXT("Physics Advance"), STAT_PhysicsAdvance, STATGROUP_Chaos, CHAOS_API);
diff --git a/Engine/Source/Runtime/Experimental/Chaos/Public/PBDRigidsSolver.h b/Engine/Source/Runtime/Experimental/Chaos/Public/PBDRigidsSolver.h
--- a/Engine/Source/Runtime/Experimental/Chaos/Public/PBDRigidsSolver.h
+++ b/Engine/Source/Runtime/Experimental/Chaos/Public/PBDRigidsSolver.h
@@ -12,7 +12,6 @@
 #include "Chaos/PBDPositionConstraints.h"
 #include "Chaos/PBDSuspensionConstraints.h"
 #include "Chaos/PBDJointConstraints.h"
-#include "Chaos/PBDConstraintRule.h"
 #include "Chaos/PerParticleGravity.h"
 #include "Chaos/ParticleHandle.h"
 #include "Chaos/Transform.h"
@@ -113,13 +112,7 @@ namespace Chaos
 		typedef FPBDRigidDynamicSpringConstraints FRigidDynamicSpringConstraints;
 		typedef FPBDPositionConstraints FPositionConstraints;
 
-		typedef TPBDConstraintIslandRule<FPBDJointConstraints> FJointConstraintsRule;
-		typedef TPBDConstraintIslandRule<FRigidDynamicSpringConstraints> FRigidDynamicSpringConstraintsRule;
-		typedef TPBDConstraintIslandRule<FPositionConstraints> FPositionConstraintsRule;
-		typedef TPBDConstraintIslandRule<FPBDSuspensionConstraints> FSuspensionConstraintsRule;
-
 		using FJointConstraints = FPBDJointConstraints;
-		using FJointConstraintRule = TPBDConstraintIslandRule<FJointConstraints>;
 		//
 		// Execution API
 		//
@@ -180,7 +173,6 @@ namespace Chaos
 		void SetProjectionIterations(const int32 InNumIterations) { GetEvolution()->SetNumProjectionIterations(InNumIterations); }
 		void SetCollisionCullDistance(const FReal InCullDistance) { GetEvolution()->GetNarrowPhase().SetBoundsExpansion(InCullDistance); }
 		void SetCollisionMaxPushOutVelocity(const FReal InMaxPushOutVelocity) { GetEvolution()->GetCollisionConstraints().SetMaxPushOutVelocity(InMaxPushOutVelocity); }
-		void SetUseContactGraph(const bool bInUseContactGraph) { GetEvolution()->GetCollisionConstraintsRule().SetUseContactGraph(bInUseContactGraph); }
 
 		/**/
 		void SetGenerateCollisionData(bool bDoGenerate) { GetEventFilters()->SetGenerateCollisionEvents(bDoGenerate); }
diff --git a/Engine/Source/Runtime/Experimental/Chaos/Public/PhysicsProxy/SkeletalMeshPhysicsProxy.h b/Engine/Source/Runtime/Experimental/Chaos/Public/PhysicsProxy/SkeletalMeshPhysicsProxy.h
--- a/Engine/Source/Runtime/Experimental/Chaos/Public/PhysicsProxy/SkeletalMeshPhysicsProxy.h
+++ b/Engine/Source/Runtime/Experimental/Chaos/Public/PhysicsProxy/SkeletalMeshPhysicsProxy.h
@@ -6,9 +6,9 @@
 #include "GeometryCollection/GeometryCollectionSimulationTypes.h"
 
 #include "Chaos/PBDJointConstraints.h"
-#include "Chaos/PBDConstraintRule.h"
 #include "Framework/TripleBufferedData.h"
 
+// @todo(chaos): remove this file
 
 struct CHAOS_API FSkeletalMeshPhysicsProxyParams
 {
@@ -168,12 +168,10 @@ public:
 
 private:
 	using FJointConstraints = Chaos::FPBDJointConstraints;
-	using FJointConstraintsRule = Chaos::TPBDConstraintIslandRule<FJointConstraints>;
 
 	FSkeletalMeshPhysicsProxyParams Parameters;
 	TArray<int32> RigidBodyIds;
 	FJointConstraints JointConstraints;
-	FJointConstraintsRule JointConstraintsRule;
 	// @todo(ccaulfield): sort out the IO buffer stuff
 	Chaos::TTripleBufferedData<FSkeletalMeshPhysicsProxyInputs> InputBuffers;
 	Chaos::TBufferedData<FSkeletalMeshPhysicsProxyOutputs> OutputBuffers;
diff --git a/Engine/Source/Runtime/Experimental/ChaosCore/Public/Chaos/ObjectPool.h b/Engine/Source/Runtime/Experimental/ChaosCore/Public/Chaos/ObjectPool.h
--- a/Engine/Source/Runtime/Experimental/ChaosCore/Public/Chaos/ObjectPool.h
+++ b/Engine/Source/Runtime/Experimental/ChaosCore/Public/Chaos/ObjectPool.h
@@ -146,6 +146,14 @@ namespace Chaos
 			return Blocks.Num();
 		}
 
+		/**
+		 * Get the max number of items per block
+		*/
+		int32 GetNumPerBlock() const
+		{
+			return NumPerBlock;
+		}
+
 		/**
 		 * Get the allocated (total) size of all blocks in the pool
 		 */
@@ -223,6 +231,14 @@ namespace Chaos
 			return static_cast<float>(sizeof(ObjectType)) / PaddedItemSize;
 		}
 
+		/**
+		 * Get the number of allocated items.
+		*/
+		int32 GetNumAllocated() const
+		{
+			return GetCapacity() - GetNumFree();
+		}
+
 		/**
 		 * Gets the number of free items the pool has remaining
 		 * @return Number of free items
@@ -267,7 +283,7 @@ namespace Chaos
 		// will need to destruct any non-free objects). This is specialized so that if the type
 		// is trivially destructible we can avoid adding the extra member.
 		template<typename T_>
-		struct alignas(ItemAlign) TItem<T_, typename TRequiresDestructor<T_>>
+		struct alignas(ItemAlign) TItem<T_, TRequiresDestructor<T_>>
 		{
 			TItem()
 				: NextFree(INDEX_NONE)
@@ -295,7 +311,7 @@ namespace Chaos
 
 		// Size of the item, plus padding up to the correct alignment so we can allocate the whole
 		// block with the correct size.
-		constexpr static int32 PaddedItemSize = Align(sizeof(FItem), ItemAlign);
+		constexpr static int32 PaddedItemSize = Align(int32(sizeof(FItem)), ItemAlign);
 
 		// A block is an area of memory large enough to hold NumInBlock items, correctly aligned and
 		// provide items on demand to the pool
@@ -465,7 +481,7 @@ namespace Chaos
 		 */
 		template<
 			typename ObjectType_ = ObjectType, 
-			typename TTrivialDestruct<ObjectType_>* = nullptr, 
+			TTrivialDestruct<ObjectType_>* = nullptr, 
 			typename... TArgs>
 		static ObjectType* Construct(ObjectType* At, TArgs&&... Args)
 		{
@@ -484,7 +500,7 @@ namespace Chaos
 		 */
 		template<
 			typename ObjectType_ = ObjectType, 
-			typename TRequiresDestructor<ObjectType_>* = nullptr, 
+			TRequiresDestructor<ObjectType_>* = nullptr, 
 			typename... TArgs>
 		static ObjectType* Construct(ObjectType* At, TArgs&&... Args)
 		{
@@ -505,7 +521,7 @@ namespace Chaos
 		 */
 		template<
 			typename ObjectType_ = ObjectType, 
-			typename TRequiresDestructor<ObjectType_>* = nullptr>
+			TRequiresDestructor<ObjectType_>* = nullptr>
 		static void ConditionalDestruct(ObjectType* At)
 		{
 			checkSlow(At);
@@ -527,7 +543,7 @@ namespace Chaos
 		 */
 		template<
 			typename ObjectType_ = ObjectType, 
-			typename TTrivialDestruct<ObjectType_>* = nullptr>
+			TTrivialDestruct<ObjectType_>* = nullptr>
 		static void ConditionalDestruct(ObjectType*)
 		{}
 
@@ -537,4 +553,35 @@ namespace Chaos
 		int32 Capacity = 0;
 		int32 FreeCount = 0;
 	};
+
+
+	/**
+	 * A deleter for use with TUniquePtr and a TObjectPool item
+	 */
+	template<typename ObjectType>
+	class TObjectPoolDeleter
+	{
+	public:
+		TObjectPoolDeleter()
+			: Pool(nullptr)
+		{
+		}
+
+		TObjectPoolDeleter(TObjectPool<ObjectType>& InPool)
+			: Pool(&InPool)
+		{
+		}
+
+		void operator()(ObjectType* Object)
+		{
+			if (Object != nullptr)
+			{
+				check(Pool != nullptr);
+				Pool->Free(Object);
+			}
+		}
+
+	private:
+		TObjectPool<ObjectType>* Pool;
+	};
 }
\ No newline at end of file
diff --git a/Engine/Source/Runtime/Experimental/ChaosSolverEngine/Private/Chaos/ChaosSolverActor.cpp b/Engine/Source/Runtime/Experimental/ChaosSolverEngine/Private/Chaos/ChaosSolverActor.cpp
--- a/Engine/Source/Runtime/Experimental/ChaosSolverEngine/Private/Chaos/ChaosSolverActor.cpp
+++ b/Engine/Source/Runtime/Experimental/ChaosSolverEngine/Private/Chaos/ChaosSolverActor.cpp
@@ -186,7 +186,6 @@ AChaosSolverActor::AChaosSolverActor(const FObjectInitializer& ObjectInitializer
 	, DoGenerateBreakingData_DEPRECATED(true)
 	, DoGenerateTrailingData_DEPRECATED(true)
 	, MassScale_DEPRECATED(1.f)
-	, bGenerateContactGraph_DEPRECATED(true)
 	, bHasFloor(true)
 	, FloorHeight(0.f)
 	, ChaosDebugSubstepControl()
@@ -366,7 +365,6 @@ void AChaosSolverActor::PostLoad()
 		Properties.BreakingFilterSettings = BreakingFilterSettings_DEPRECATED;
 		Properties.CollisionFilterSettings = CollisionFilterSettings_DEPRECATED;
 		Properties.TrailingFilterSettings = TrailingFilterSettings_DEPRECATED;
-		Properties.bGenerateContactGraph = bGenerateContactGraph_DEPRECATED;
 	}
 }
 
diff --git a/Engine/Source/Runtime/Experimental/ChaosSolverEngine/Public/Chaos/ChaosSolverActor.h b/Engine/Source/Runtime/Experimental/ChaosSolverEngine/Public/Chaos/ChaosSolverActor.h
--- a/Engine/Source/Runtime/Experimental/ChaosSolverEngine/Public/Chaos/ChaosSolverActor.h
+++ b/Engine/Source/Runtime/Experimental/ChaosSolverEngine/Public/Chaos/ChaosSolverActor.h
@@ -108,8 +108,6 @@ public:
 	FSolverTrailingFilterSettings TrailingFilterSettings_DEPRECATED;
 	UPROPERTY()
 	float MassScale_DEPRECATED;
-	UPROPERTY()
-	bool bGenerateContactGraph_DEPRECATED;
 	/** End deprecated properties */
 	
 	UPROPERTY(EditAnywhere, Category = Settings)
