Chaos : fix cook warning related to selection material in geometry collection
- Make sure to remove reference to the selection material in the material array when cooking and saving

diff --git a/Engine/Source/Runtime/Experimental/GeometryCollectionEngine/Private/GeometryCollection/GeometryCollectionObject.cpp b/Engine/Source/Runtime/Experimental/GeometryCollectionEngine/Private/GeometryCollection/GeometryCollectionObject.cpp
--- a/Engine/Source/Runtime/Experimental/GeometryCollectionEngine/Private/GeometryCollection/GeometryCollectionObject.cpp
+++ b/Engine/Source/Runtime/Experimental/GeometryCollectionEngine/Private/GeometryCollection/GeometryCollectionObject.cpp
@@ -666,6 +666,14 @@ void UGeometryCollection::Serialize(FArchive& Ar)
 	if (bIsCookedOrCooking && Ar.IsSaving())
 	{
 #if WITH_EDITOR
+		// if we have a valid selection material, let's make sure we replace it with one that will be cooked
+		// this avoid getting warning about the selected material being reference but not cooked
+		const int32 SelectedMaterialIndex = GetBoneSelectedMaterialIndex();
+		if (!Materials.IsEmpty() && Materials.IsValidIndex(SelectedMaterialIndex))
+		{
+			Materials[SelectedMaterialIndex] = Materials[0];
+		}
+
 		if (bStripOnCook && EnableNanite && NaniteData)
 		{
 			// If this is a cooked archive, we strip unnecessary data from the Geometry Collection to keep the memory footprint as small as possible.
