Chaos - Re-enable collisions when joint with IgnoreCollisions is destroyed

diff --git a/Engine/Source/Runtime/Experimental/Chaos/Private/PhysicsProxy/JointConstraintProxy.cpp b/Engine/Source/Runtime/Experimental/Chaos/Private/PhysicsProxy/JointConstraintProxy.cpp
--- a/Engine/Source/Runtime/Experimental/Chaos/Private/PhysicsProxy/JointConstraintProxy.cpp
+++ b/Engine/Source/Runtime/Experimental/Chaos/Private/PhysicsProxy/JointConstraintProxy.cpp
@@ -109,6 +109,80 @@ void FJointConstraintPhysicsProxy::DestroyOnPhysicsThread(FPBDRigidsSolver* InSo
 	{
 		TVec2<FGeometryParticleHandle*> Particles = Constraint_PT->GetConstrainedParticles();
 		
+		// Note: the current implementation of FJointConstraintPhysicsProxy::PushStateOnPhysicsThread means that if there are multiple constraints for the same body pair,
+		// then they will stomp over each other's collision disable state in FIgnoreCollisionManager. 
+		// @todo(chaos): Fix this ^^ UE-162071
+		// @todo(chaos): Move all Joint creation, regsitration and update code to Evolution
+		if (!Constraint_PT->GetSettings().bCollisionEnabled)
+		{
+			FGeometryParticleHandle* Handle0 = Particles[0];
+			FGeometryParticleHandle* Handle1 = Particles[1];
+
+			if (Handle0 && Handle1)
+			{
+				FPBDRigidParticleHandle* RigidHandle0 = Handle0->CastToRigidParticle();
+				FPBDRigidParticleHandle* RigidHandle1 = Handle1->CastToRigidParticle();
+
+				// As long as one particle is a rigid we can add the ignore entry, one particle can be a static
+				if (RigidHandle0 || RigidHandle1)
+				{
+					// Make sure that there is no other active constraint on the same particle pair that is disabling collision 
+					const bool bDisabledCollisionByAnotherConstraint = [&]()
+					{
+						const FConstraintHandleArray& Constraints0 = Handle0->ParticleConstraints();
+
+						for (const FConstraintHandle* Constraint0 : Constraints0)
+						{
+							if (Constraint_PT == Constraint0)
+							{
+								continue;
+							}
+
+							const FPBDJointConstraintHandle* OtherConstraint = Constraint0->As<FPBDJointConstraintHandle>();
+							if (!OtherConstraint)
+							{
+								continue;
+							}
+
+							const TVector<FGeometryParticleHandle*, 2>& OtherParticles = OtherConstraint->GetConstrainedParticles();
+							TArray<FGeometryParticleHandle*, TInlineAllocator<2>> OtherBasePairsArray = { OtherParticles[0], OtherParticles[1] };
+							if (!OtherBasePairsArray.Contains(Handle0) || !OtherBasePairsArray.Contains(Handle1))
+							{
+								continue;
+							}
+
+							if (OtherConstraint->GetSettings().bCollisionEnabled == false)
+							{
+								return true;
+							}
+						}
+
+						return false;
+					}();
+
+					if (bDisabledCollisionByAnotherConstraint == false)
+					{
+						const FUniqueIdx ID0 = Handle0->UniqueIdx();
+						const FUniqueIdx ID1 = Handle1->UniqueIdx();
+						FIgnoreCollisionManager& IgnoreCollisionManager = InSolver->GetEvolution()->GetBroadPhase().GetIgnoreCollisionManager();
+
+						// For rigid/dynamic particles, add the broadphase flag and the IDs to check for disabled collisions
+						if (RigidHandle0)
+						{
+							//can't remove collision constraint flag because we may still be ignoring collision for others
+							IgnoreCollisionManager.RemoveIgnoreCollisionsFor(ID0, ID1);
+						}
+
+						if (RigidHandle1)
+						{
+							//can't remove collision constraint flag because we may still be ignoring collision for others
+							IgnoreCollisionManager.RemoveIgnoreCollisionsFor(ID1, ID0);
+						}
+					}
+				}
+			}
+		}
+
 		// Ensure that our connected particles are aware that this constraint no longer exists
 		if(Particles[0])
 		{
