Chaos : Fracture : Fix break propagation logicnot working properly

diff --git a/Engine/Source/Runtime/Experimental/Chaos/Private/Chaos/PBDRigidClustering.cpp b/Engine/Source/Runtime/Experimental/Chaos/Private/Chaos/PBDRigidClustering.cpp
--- a/Engine/Source/Runtime/Experimental/Chaos/Private/Chaos/PBDRigidClustering.cpp
+++ b/Engine/Source/Runtime/Experimental/Chaos/Private/Chaos/PBDRigidClustering.cpp
@@ -643,11 +643,7 @@ namespace Chaos
 		TArray<FPBDRigidParticleHandle*>& Children = MChildren[ClusteredParticle];
 
 		// only used for propagation
-		TArray<FReal> AppliedStrains;
-		if (bUseDamagePropagation)
-		{
-			AppliedStrains.SetNum(Children.Num());
-		}
+		TMap<FPBDRigidParticleHandle*, FReal> AppliedStrains;
 
 		for (int32 ChildIdx = Children.Num() - 1; ChildIdx >= 0; --ChildIdx)
 		{
@@ -686,26 +682,27 @@ namespace Chaos
 			}
 			if (bUseDamagePropagation)
 			{
-				AppliedStrains[ChildIdx] = MaxAppliedStrain;
+				AppliedStrains.Add(Child, MaxAppliedStrain);
 			}
 			Child->ClearExternalStrain();
 		}
 
 		// if necessary propagate strain through the graph
+		// IMPORTANT: this assumes that the connectivity graph has not yet been updated from pieces that broke off
 		if (bUseDamagePropagation)
 		{
-			for (int32 ChildIdx = 0; ChildIdx < Children.Num(); ChildIdx++)
+			for (const auto& AppliedStrain: AppliedStrains)
 			{
-				FPBDRigidClusteredParticleHandle* ClusteredChild = Children[ChildIdx]->CastToClustered();
+				FPBDRigidClusteredParticleHandle* ClusteredChild = AppliedStrain.Key->CastToClustered();
 
-				const FReal AppliedStrain = AppliedStrains[ChildIdx];
+				const FReal AppliedStrainValue = AppliedStrain.Value;
 				FReal PropagatedStrainPerConnection = 0.0f;
 
 				// @todo(chaos) : may not be optimal, but good enough for now
-				if (ActivatedChildren.Contains(Children[ChildIdx]))
+				if (BreakDamagePropagationFactor > 0 && ActivatedChildren.Contains(AppliedStrain.Key))
 				{
 					// break damage propagation case: we only look at the broken pieces and propagate the strain remainder 
-					const FReal RemainingStrain = (AppliedStrain - ClusteredChild->Strain());
+					const FReal RemainingStrain = (AppliedStrainValue - ClusteredChild->Strain());
 					if (RemainingStrain > 0)
 					{
 						const FReal AdjustedRemainingStrain = (FReal)BreakDamagePropagationFactor * RemainingStrain;
@@ -713,10 +710,10 @@ namespace Chaos
 						PropagatedStrainPerConnection = AdjustedRemainingStrain / ClusteredChild->ConnectivityEdges().Num();
 					}
 				}
-				else
+				else if (ShockDamagePropagationFactor > 0)
 				{
 					// shock damage propagation case : for all the non broken pieces, proapagate the actual applied strain 
-					PropagatedStrainPerConnection = (FReal)ShockDamagePropagationFactor * AppliedStrain;
+					PropagatedStrainPerConnection = (FReal)ShockDamagePropagationFactor * AppliedStrainValue;
 				}
 
 				if (PropagatedStrainPerConnection > 0)
