Chaos : Fix capsule raycast sometimes returning false negatve
- reverted changes from 19279987
- made a more general way to check if the ray and the capsule do intersect before computing the actual intersection details

diff --git a/Engine/Source/Runtime/Experimental/Chaos/Public/Chaos/Capsule.h b/Engine/Source/Runtime/Experimental/Chaos/Public/Chaos/Capsule.h
--- a/Engine/Source/Runtime/Experimental/Chaos/Public/Chaos/Capsule.h
+++ b/Engine/Source/Runtime/Experimental/Chaos/Public/Chaos/Capsule.h
@@ -148,6 +148,16 @@ namespace Chaos
 			const FReal R2 = R * R;
 			OutFaceIndex = INDEX_NONE;
 
+			// first test, segment to segment distance
+			const FVec3 RayEndPoint = StartPoint + Dir * Length;
+			FVec3 OutP1, OutP2;
+			FMath::SegmentDistToSegmentSafe(X1, X2, StartPoint, RayEndPoint, OutP1, OutP2);
+			const FReal DistanceSquared = (OutP2 - OutP1).SizeSquared();
+			if (DistanceSquared > R2 + DBL_EPSILON)
+			{
+				return false;
+			}
+
 			// Raycast against capsule bounds.
 			// We will use the intersection point as our ray start, this prevents precision issues if start is far away.
 			// All calculations below should use LocalStart/LocalLength, and convert output time using input length if intersecting.
@@ -285,15 +295,8 @@ namespace Chaos
 					const FReal PositionLengthOnCoreCylinder = FVec3::DotProduct(CylinderToSpherePosition, MVector);
 					if (PositionLengthOnCoreCylinder >= 0 && PositionLengthOnCoreCylinder < MHeight)
 					{
-						const FVec3 SegmentToPosition = (CylinderToSpherePosition - MVector * PositionLengthOnCoreCylinder);
-						if (SegmentToPosition.SquaredLength() > ( R2 + UE_KINDA_SMALL_NUMBER))
-						{
-							// the contact point is actually outside of the cylinder
-							// this can happen if the ray starts within the cylinder bounds but do not intersect with the cylinder
-							return false;
-						}
 						OutTime = Time + RemovedLength; // Account for ray clipped against bounds
-						OutNormal = SegmentToPosition / R;
+						OutNormal = (CylinderToSpherePosition - MVector * PositionLengthOnCoreCylinder) / R;
 						OutPosition = SpherePosition - OutNormal * Thickness;
 						return true;
 					}
