Chaos - fix RBAN debug draw when not near the origin

diff --git a/Engine/Source/Runtime/Experimental/Chaos/Private/Chaos/ChaosDebugDraw.cpp b/Engine/Source/Runtime/Experimental/Chaos/Private/Chaos/ChaosDebugDraw.cpp
--- a/Engine/Source/Runtime/Experimental/Chaos/Private/Chaos/ChaosDebugDraw.cpp
+++ b/Engine/Source/Runtime/Experimental/Chaos/Private/Chaos/ChaosDebugDraw.cpp
@@ -497,7 +497,7 @@ namespace Chaos
 			const EImplicitObjectType InnerType = GetInnerType(Implicit->GetType());
 
 			// Are we within the region of interest?
-			if (Particle->HasBounds() && !FDebugDrawQueue::GetInstance().IsInRegionOfInterest(Particle->WorldSpaceInflatedBounds()))
+			if (Implicit->HasBoundingBox() && !FDebugDrawQueue::GetInstance().IsInRegionOfInterest(Implicit->BoundingBox().TransformedAABB(ShapeTransform)))
 			{
 				return;
 			}
