Clang Static Analysis GJK

diff --git a/Engine/Source/Runtime/Experimental/Chaos/Public/Chaos/GJK.h b/Engine/Source/Runtime/Experimental/Chaos/Public/Chaos/GJK.h
--- a/Engine/Source/Runtime/Experimental/Chaos/Public/Chaos/GJK.h
+++ b/Engine/Source/Runtime/Experimental/Chaos/Public/Chaos/GJK.h
@@ -1805,7 +1805,7 @@ namespace Chaos
 		const VectorRegister4Float InitialDirSimd = MakeVectorRegisterFloatFromDouble(MakeVectorRegister(InitialDir[0], InitialDir[1], InitialDir[2], 0.0));
 		const VectorRegister4Float RayDirSimd = MakeVectorRegisterFloatFromDouble(MakeVectorRegister(RayDir[0], RayDir[1], RayDir[2], 0.0));
 
-		FRealSingle OutTimeFloat;
+		FRealSingle OutTimeFloat = 0.0f;
 		VectorRegister4Float OutPositionSimd, OutNormalSimd;
 		bool result = GJKRaycast2ImplSimd(A, B, Rotation, Translation, RayDirSimd, static_cast<FRealSingle>(RayLength), OutTimeFloat, OutPositionSimd, OutNormalSimd, bComputeMTD, InitialDirSimd);
 
