Chaos : fix body instance warning for geomtry collection when calling SetSsimulatePhysics
- make sure we do not set the OwnerComponent before calling the body instance matching function

diff --git a/Engine/Source/Runtime/Experimental/GeometryCollectionEngine/Private/GeometryCollection/GeometryCollectionComponent.cpp b/Engine/Source/Runtime/Experimental/GeometryCollectionEngine/Private/GeometryCollection/GeometryCollectionComponent.cpp
--- a/Engine/Source/Runtime/Experimental/GeometryCollectionEngine/Private/GeometryCollection/GeometryCollectionComponent.cpp
+++ b/Engine/Source/Runtime/Experimental/GeometryCollectionEngine/Private/GeometryCollection/GeometryCollectionComponent.cpp
@@ -744,7 +744,15 @@ bool UGeometryCollectionComponent::CanEditSimulatePhysics()
 
 void UGeometryCollectionComponent::SetSimulatePhysics(bool bEnabled)
 {
-	Super::SetSimulatePhysics(bEnabled);
+	// make sure owner component is set to null before calling Super::SetSimulatePhysics
+	// this will prevent unwanted log warning to trigger in BodyInstance::SetInstanceSimulatePhysics() because
+	// in geometry collection , body instance never holds a valid physics handle
+	const TWeakObjectPtr<UPrimitiveComponent> PreviousOwnerCOmponent = BodyInstance.OwnerComponent;
+	{
+		BodyInstance.OwnerComponent = nullptr;
+		Super::SetSimulatePhysics(bEnabled);
+		BodyInstance.OwnerComponent = PreviousOwnerCOmponent;
+	}
 
 	if (bEnabled && !PhysicsProxy)
 	{
