Fracture editor : fix crash when collecting debug data after the number of transform has changed

diff --git a/Engine/Source/Runtime/Experimental/Chaos/Public/PhysicsProxy/GeometryCollectionPhysicsProxy.h b/Engine/Source/Runtime/Experimental/Chaos/Public/PhysicsProxy/GeometryCollectionPhysicsProxy.h
--- a/Engine/Source/Runtime/Experimental/Chaos/Public/PhysicsProxy/GeometryCollectionPhysicsProxy.h
+++ b/Engine/Source/Runtime/Experimental/Chaos/Public/PhysicsProxy/GeometryCollectionPhysicsProxy.h
@@ -560,7 +560,13 @@ public:
 
 	void Reset(int32 NumTransforms);
 	int32 Num() const { return DamageData.Num(); }
-	const FDamageData& operator[](int32 TransformIndex) const { return DamageData[TransformIndex]; }
+	const FDamageData& operator[](int32 TransformIndex) const 
+	{ 
+		static FDamageData DefaultDamageData;
+		if (DamageData.IsValidIndex(TransformIndex))
+			return DamageData[TransformIndex]; 
+		return DefaultDamageData;
+	}
 	void SampleDamage(int32 TransformIndex, float Damage, float DamageThreshold);
 	
 private:
