Geometry collection : make sure max cluster level is taken in account when using size specific damage threshold

diff --git a/Engine/Source/Runtime/Experimental/Chaos/Private/PhysicsProxy/GeometryCollectionPhysicsProxy.cpp b/Engine/Source/Runtime/Experimental/Chaos/Private/PhysicsProxy/GeometryCollectionPhysicsProxy.cpp
--- a/Engine/Source/Runtime/Experimental/Chaos/Private/PhysicsProxy/GeometryCollectionPhysicsProxy.cpp
+++ b/Engine/Source/Runtime/Experimental/Chaos/Private/PhysicsProxy/GeometryCollectionPhysicsProxy.cpp
@@ -1175,34 +1175,38 @@ DECLARE_CYCLE_STAT(TEXT("FGeometryCollectionPhysicsProxy::BuildClusters:GlobalMa
 float FGeometryCollectionPhysicsProxy::ComputeDamageThreshold(const FGeometryDynamicCollection& DynamicCollection, int32 TransformIndex) const 
 {
 	float DamageThreshold = TNumericLimits<float>::Max();
-	if (Parameters.bUseSizeSpecificDamageThresholds)
+
+	const int32 LevelOffset = Parameters.bUsePerClusterOnlyDamageThreshold ? 0 : -1;
+	const int32 Level = FMath::Clamp(CalculateHierarchyLevel(DynamicCollection, TransformIndex) + LevelOffset, 0, INT_MAX);
+	if (Level >= Parameters.MaxClusterLevel)
 	{
-		// bounding box volume is used as a fallback to find specific size if the relative size if not available
-		// ( May happen with older GC )
-		FBox LocalBoundingBox;
-		const FGeometryDynamicCollection::FSharedImplicit& Implicit = DynamicCollection.Implicits[TransformIndex];
-		if (Implicit && Implicit->HasBoundingBox())
-		{
-			const Chaos::FAABB3& ImplicitBoundingBox = Implicit->BoundingBox();
-			LocalBoundingBox = FBox(ImplicitBoundingBox.Min(), ImplicitBoundingBox.Max());
-		}
-		
-		const FSharedSimulationSizeSpecificData& SizeSpecificData = GetSizeSpecificData(Parameters.Shared.SizeSpecificData, *Parameters.RestCollection, TransformIndex, LocalBoundingBox);
-		DamageThreshold = SizeSpecificData.DamageThreshold;
+		DamageThreshold = TNumericLimits<float>::Max();
 	}
 	else
 	{
-		const int32 LevelOffset =  Parameters.bUsePerClusterOnlyDamageThreshold? 0: -1;
-		const int32 NumThresholds = Parameters.DamageThreshold.Num();
-		const int32 Level = FMath::Clamp(CalculateHierarchyLevel(DynamicCollection, TransformIndex) + LevelOffset, 0, INT_MAX);
-		const float DefaultDamage = NumThresholds > 0 ? Parameters.DamageThreshold[NumThresholds - 1] : 0.f;
-		DamageThreshold = Level < NumThresholds ? Parameters.DamageThreshold[Level] : DefaultDamage;
+		if (Parameters.bUseSizeSpecificDamageThresholds)
+		{
+			// bounding box volume is used as a fallback to find specific size if the relative size if not available
+			// ( May happen with older GC )
+			FBox LocalBoundingBox;
+			const FGeometryDynamicCollection::FSharedImplicit& Implicit = DynamicCollection.Implicits[TransformIndex];
+			if (Implicit && Implicit->HasBoundingBox())
+			{
+				const Chaos::FAABB3& ImplicitBoundingBox = Implicit->BoundingBox();
+				LocalBoundingBox = FBox(ImplicitBoundingBox.Min(), ImplicitBoundingBox.Max());
+			}
 
-		if (Level >= Parameters.MaxClusterLevel)
+			const FSharedSimulationSizeSpecificData& SizeSpecificData = GetSizeSpecificData(Parameters.Shared.SizeSpecificData, *Parameters.RestCollection, TransformIndex, LocalBoundingBox);
+			DamageThreshold = SizeSpecificData.DamageThreshold;
+		}
+		else
 		{
-			DamageThreshold = TNumericLimits<float>::Max();
+			const int32 NumThresholds = Parameters.DamageThreshold.Num();
+			const float DefaultDamage = NumThresholds > 0 ? Parameters.DamageThreshold[NumThresholds - 1] : 0.f;
+			DamageThreshold = Level < NumThresholds ? Parameters.DamageThreshold[Level] : DefaultDamage;
 		}
 	}
+
 	return DamageThreshold;
 }
 
