Iris
- ReplicationStateDescriptorBuilder now automatically uses custom NetSerializer from Super if a derived struct does not implement one while the Super does. If this is the case a warning will be issued which can be silenced by declaring a forwarding serializer or a custom netserializser specific to the derived class if needed.
- Added SkipCheckForCustomNetSerializerForStruct flag to FReplicationStateDescriptorBuilder::FParameters,
   if set the builder will skip the check if there is adeclared NetSerializer for the struct and instead build a descriptor based on
   the struct properties. This is useful for custom NetSerializers that can reuse part of the generated descriptor.

diff --git a/Engine/Plugins/Runtime/GameplayAbilities/Source/GameplayAbilities/Private/Serialization/GameplayAbilityRepAnimMontageNetSerializer.cpp b/Engine/Plugins/Runtime/GameplayAbilities/Source/GameplayAbilities/Private/Serialization/GameplayAbilityRepAnimMontageNetSerializer.cpp
--- a/Engine/Plugins/Runtime/GameplayAbilities/Source/GameplayAbilities/Private/Serialization/GameplayAbilityRepAnimMontageNetSerializer.cpp
+++ b/Engine/Plugins/Runtime/GameplayAbilities/Source/GameplayAbilities/Private/Serialization/GameplayAbilityRepAnimMontageNetSerializer.cpp
@@ -1,7 +1,6 @@
 // Copyright Epic Games, Inc. All Rights Reserved.
 
 #include "Serialization/GameplayAbilityRepAnimMontageNetSerializer.h"
-#include "Serialization/InternalGameplayAbilityRepAnimMontageNetSerializer.h"
 
 #if UE_WITH_IRIS
 
@@ -439,8 +438,10 @@ void FGameplayAbilityRepAnimMontageNetSerializer::FNetSerializerRegistryDelegate
 	static_assert(sizeof(FGameplayAbilityRepAnimMontage) == ExpectedSizeOfFGameplayAbilityRepAnimMontage && alignof(FGameplayAbilityRepAnimMontage) == ExpectedAlignOfFGameplayAbilityRepAnimMontage, "FGameplayAbilityRepAnimMontage layout has changed. Might need to update FGameplayAbilityRepAnimMontageNetSerializer to include new data or update the size.");
 
 	// Use helper to avoid getting hold of the FGameplayAbilityRepAnimMontageNetSerializer that we are setting up and validating.
-	const UStruct* BaseStruct = FGameplayAbilityRepAnimMontageNetSerializerSerializationHelper::StaticStruct();
-	StructNetSerializerConfigForBase.StateDescriptor = FReplicationStateDescriptorBuilder::CreateDescriptorForStruct(BaseStruct);
+	const UStruct* BaseStruct = FGameplayAbilityRepAnimMontage::StaticStruct();
+	FReplicationStateDescriptorBuilder::FParameters Params;
+	Params.SkipCheckForCustomNetSerializerForStruct = 1U;
+	StructNetSerializerConfigForBase.StateDescriptor = FReplicationStateDescriptorBuilder::CreateDescriptorForStruct(BaseStruct, Params);
 
 	const FReplicationStateDescriptor* Descriptor = StructNetSerializerConfigForBase.StateDescriptor.GetReference();
 	check(Descriptor != nullptr);
diff --git a/Engine/Plugins/Runtime/GameplayAbilities/Source/GameplayAbilities/Private/Serialization/GameplayEffectContextNetSerializer.cpp b/Engine/Plugins/Runtime/GameplayAbilities/Source/GameplayAbilities/Private/Serialization/GameplayEffectContextNetSerializer.cpp
--- a/Engine/Plugins/Runtime/GameplayAbilities/Source/GameplayAbilities/Private/Serialization/GameplayEffectContextNetSerializer.cpp
+++ b/Engine/Plugins/Runtime/GameplayAbilities/Source/GameplayAbilities/Private/Serialization/GameplayEffectContextNetSerializer.cpp
@@ -489,7 +489,11 @@ void FGameplayEffectContextNetSerializer::FNetSerializerRegistryDelegates::OnPos
 			LowLevelFatalError(TEXT("%s Size: %d Alignment: %d"), TEXT("FGameplayEffectContext layout has changed. Need to update FGameplayEffectContextNetSerializer."), GEStruct->GetStructureSize(), GEStruct->GetMinAlignment());
 		}
 
-		StructNetSerializerConfigForGE.StateDescriptor = FReplicationStateDescriptorBuilder::CreateDescriptorForStruct(GEStruct);
+		// In this case we want to build a descriptor based on the struct members rather than the serializer we try to register
+		FReplicationStateDescriptorBuilder::FParameters Params;
+		Params.SkipCheckForCustomNetSerializerForStruct = true;
+
+		StructNetSerializerConfigForGE.StateDescriptor = FReplicationStateDescriptorBuilder::CreateDescriptorForStruct(GEStruct, Params);
 		const FReplicationStateDescriptor* Descriptor = StructNetSerializerConfigForGE.StateDescriptor.GetReference();
 		check(Descriptor != nullptr);
 		GEStateTraits = Descriptor->Traits;
diff --git a/Engine/Plugins/Runtime/GameplayAbilities/Source/GameplayAbilities/Private/Serialization/InternalGameplayAbilityRepAnimMontageNetSerializer.h b/Engine/Plugins/Runtime/GameplayAbilities/Source/GameplayAbilities/Private/Serialization/InternalGameplayAbilityRepAnimMontageNetSerializer.h
deleted file mode 100644
diff --git a/Engine/Plugins/Runtime/GameplayAbilities/Source/GameplayAbilities/Private/Serialization/InternalGameplayEffectContextNetSerializer.h b/Engine/Plugins/Runtime/GameplayAbilities/Source/GameplayAbilities/Private/Serialization/InternalGameplayEffectContextNetSerializer.h
--- a/Engine/Plugins/Runtime/GameplayAbilities/Source/GameplayAbilities/Private/Serialization/InternalGameplayEffectContextNetSerializer.h
+++ b/Engine/Plugins/Runtime/GameplayAbilities/Source/GameplayAbilities/Private/Serialization/InternalGameplayEffectContextNetSerializer.h
@@ -5,13 +5,10 @@
 #include "GameplayEffectTypes.h"
 #include "Engine/EngineTypes.h"
 #include "Templates/SharedPointer.h"
-#include "InternalGameplayEffectContextNetSerializer.generated.h"
 
 // Implementation is located in GameplayEffectContextNetSerializer.cpp
-USTRUCT()
 struct FGameplayEffectContextAccessorForNetSerializer : public FGameplayEffectContext
 {
-	GENERATED_BODY()
 
 public:
 	void CopyReplicatedFieldsFrom(const FGameplayEffectContextAccessorForNetSerializer& GE);	
diff --git a/Engine/Plugins/Runtime/GameplayAbilities/Source/GameplayAbilities/Private/Serialization/MinimalReplicationTagCountNetSerializer.cpp b/Engine/Plugins/Runtime/GameplayAbilities/Source/GameplayAbilities/Private/Serialization/MinimalReplicationTagCountNetSerializer.cpp
new file mode 100644
--- /dev/null
+++ b/Engine/Plugins/Runtime/GameplayAbilities/Source/GameplayAbilities/Private/Serialization/MinimalReplicationTagCountNetSerializer.cpp
@@ -0,0 +1,35 @@
+// Copyright Epic Games, Inc. All Rights Reserved.
+#if UE_WITH_IRIS
+
+#include "GameplayEffectTypes.h"
+#include "Iris/ReplicationState/PropertyNetSerializerInfoRegistry.h"
+#include "Iris/Serialization/NetSerializerDelegates.h"
+
+namespace UE::Net
+{
+
+// Until we implement a proper NetSerializer JIRA:UE-161669 we fall back on the LastResortNetSerializer
+static const FName PropertyNetSerializerRegistry_NAME_MinimalReplicationTagCountMap("MinimalReplicationTagCountMap");
+UE_NET_IMPLEMENT_NAMED_STRUCT_LASTRESORT_NETSERIALIZER_INFO(PropertyNetSerializerRegistry_NAME_MinimalReplicationTagCountMap);
+
+struct FMinimalReplicationTagCountMapNetSerializerRegistryDelegates : protected FNetSerializerRegistryDelegates
+{
+	~FMinimalReplicationTagCountMapNetSerializerRegistryDelegates();
+	virtual void OnPreFreezeNetSerializerRegistry();
+};
+static FMinimalReplicationTagCountMapNetSerializerRegistryDelegates MinimalReplicationTagCountMapNetSerializerRegistryDelegates;
+
+FMinimalReplicationTagCountMapNetSerializerRegistryDelegates::~FMinimalReplicationTagCountMapNetSerializerRegistryDelegates()
+{
+	UE_NET_UNREGISTER_NETSERIALIZER_INFO(PropertyNetSerializerRegistry_NAME_MinimalReplicationTagCountMap);
+}
+
+void FMinimalReplicationTagCountMapNetSerializerRegistryDelegates::OnPreFreezeNetSerializerRegistry()
+{
+	UE_NET_REGISTER_NETSERIALIZER_INFO(PropertyNetSerializerRegistry_NAME_MinimalReplicationTagCountMap);
+}
+
+}
+
+#endif
+
diff --git a/Engine/Source/Runtime/Experimental/Iris/Core/Private/Iris/ReplicationState/ReplicationStateDescriptorBuilder.cpp b/Engine/Source/Runtime/Experimental/Iris/Core/Private/Iris/ReplicationState/ReplicationStateDescriptorBuilder.cpp
--- a/Engine/Source/Runtime/Experimental/Iris/Core/Private/Iris/ReplicationState/ReplicationStateDescriptorBuilder.cpp
+++ b/Engine/Source/Runtime/Experimental/Iris/Core/Private/Iris/ReplicationState/ReplicationStateDescriptorBuilder.cpp
@@ -1743,6 +1743,17 @@ bool FPropertyReplicationStateDescriptorBuilder::IsStructWithCustomSerializer(FM
 	const FPropertyNetSerializerInfo* NetSerializerInfo = FPropertyNetSerializerInfoRegistry::FindStructSerializerInfo(InStruct->GetFName());
 	if (!NetSerializerInfo)
 	{
+		if (InStruct->GetSuperStruct())
+		{
+			const bool bSuperHasCustomNetSerializer = IsStructWithCustomSerializer(OutMemberProperty, InStruct->GetSuperStruct());
+
+			// Warn if one of our supers have a CustomNetSerializer, if that is the case user must be made aware and either declare an explicit forwarding NetSerializer
+			// or implement a custom NetSerializer
+			const bool bIsAllowedToWarn = ShouldUseIrisReplication() || FApp::IsGame() || IsRunningDedicatedServer() || IsRunningClientOnly();
+			UE_CLOG(bSuperHasCustomNetSerializer && bIsAllowedToWarn && bWarnAboutStructsWithCustomSerialization, LogIris, Warning, TEXT("Detected derived struct %s which do not have a custom NetSerializer while Super %s has a custom NetSerializer. Silence warning by declaring an explicit forwarding serializer or implement a custom NetSerialzier for the derived struct"), ToCStr(InStruct->GetName()), ToCStr(InStruct->GetSuperStruct()->GetName()));
+
+			return bSuperHasCustomNetSerializer;
+		}
 		return false;
 	}
 
@@ -1956,6 +1967,7 @@ FReplicationStateDescriptorBuilder::FParameters::FParameters()
   IncludeSuper(1U),
   GetLifeTimeProperties(1U),
   EnableFastArrayHandling(1U),
+  SkipCheckForCustomNetSerializerForStruct(0U),
   SinglePropertyIndex(-1)
 {
 }
@@ -2006,14 +2018,14 @@ TRefCountPtr<const FReplicationStateDescriptor> FReplicationStateDescriptorBuild
 	bool bAllMembersAreReplicated = true;
 
 	// We have a special case for structs with custom serializers
-	if (FPropertyReplicationStateDescriptorBuilder::IsStructWithCustomSerializer(MemberProperty, InStruct))
+	if (!Parameters.SkipCheckForCustomNetSerializerForStruct && FPropertyReplicationStateDescriptorBuilder::IsStructWithCustomSerializer(MemberProperty, InStruct))
 	{
 		Builder.AddMemberProperty(MemberProperty);
 	}
 	else
 	{
-		const bool bIsAllowedToWarn = FApp::IsGame() || IsRunningDedicatedServer() || IsRunningClientOnly();
-		UE_CLOG(bIsStructWithCustomSerialization && bIsAllowedToWarn && bWarnAboutStructsWithCustomSerialization, LogIris, Warning, TEXT("Generating descriptor for struct %s that has custom serialization."), ToCStr(InStruct->GetName()));
+		const bool bIsAllowedToWarn = ShouldUseIrisReplication() || FApp::IsGame() || IsRunningDedicatedServer() || IsRunningClientOnly();
+		UE_CLOG(!Parameters.SkipCheckForCustomNetSerializerForStruct && bIsStructWithCustomSerialization && bIsAllowedToWarn && bWarnAboutStructsWithCustomSerialization, LogIris, Warning, TEXT("Generating descriptor for struct %s that has custom serialization."), ToCStr(InStruct->GetName()));
 
 		for (TFieldIterator<FProperty> It(InStruct, EFieldIteratorFlags::IncludeSuper); It; ++It)
 		{
diff --git a/Engine/Source/Runtime/Experimental/Iris/Core/Public/Iris/ReplicationState/PropertyNetSerializerInfoRegistry.h b/Engine/Source/Runtime/Experimental/Iris/Core/Public/Iris/ReplicationState/PropertyNetSerializerInfoRegistry.h
--- a/Engine/Source/Runtime/Experimental/Iris/Core/Public/Iris/ReplicationState/PropertyNetSerializerInfoRegistry.h
+++ b/Engine/Source/Runtime/Experimental/Iris/Core/Public/Iris/ReplicationState/PropertyNetSerializerInfoRegistry.h
@@ -4,6 +4,7 @@
 
 #include "Containers/Array.h"
 #include "Iris/Serialization/NetSerializer.h"
+#include "Iris/Serialization/NetSerializerDelegates.h"
 #include "Templates/Casts.h"
 #include "Templates/Tuple.h"
 #include "UObject/NameTypes.h"
@@ -121,7 +122,7 @@ private:
 
 }
 
-// Only needed if we want to export PropertyNetSerializerInfo, this goes in the header if we need to export ti
+// Only needed if we want to export PropertyNetSerializerInfo, this goes in the header if we need to export it
 #define UE_NET_DECLARE_NETSERIALIZER_INFO(NetSerializerInfo) \
 const UE::Net::FPropertyNetSerializerInfo& GetPropertyNetSerializerInfo_##NetSerializerInfo();
 
@@ -152,3 +153,18 @@ UE::Net::FPropertyNetSerializerInfoRegistry::Unregister(&GetPropertyNetSerialize
 #else
 #define UE_NET_UNREGISTER_NETSERIALIZER_INFO(...) 
 #endif
+
+// Implement minimal required delegates for a NetSerializer
+#define UE_NET_IMPLEMENT_NETSERIALZIER_REGISTRY_DELEGATES(Name) \
+struct F##Name##NetSerializerRegistryDelegates : protected FNetSerializerRegistryDelegates \
+{ \
+	~F##Name##NetSerializerRegistryDelegates() { UE_NET_UNREGISTER_NETSERIALIZER_INFO(PropertyNetSerializerRegistry_NAME_##Name); } \
+	virtual void OnPreFreezeNetSerializerRegistry() override { UE_NET_REGISTER_NETSERIALIZER_INFO(PropertyNetSerializerRegistry_NAME_##Name); } \
+}; \
+static F##Name##NetSerializerRegistryDelegates Name##NetSerializerRegistryDelegates;
+
+// Utility that can be used to forward serialization of a Struct to a specific NetSerializer
+#define UE_NET_IMPLEMENT_FORWARDING_NETSERIALIZER_AND_REGISTRY_DELEGATES(Name, SerializerName) \
+	static const FName PropertyNetSerializerRegistry_NAME_##Name( PREPROCESSOR_TO_STRING(Name) ); \
+	UE_NET_IMPLEMENT_NAMED_STRUCT_NETSERIALIZER_INFO(PropertyNetSerializerRegistry_NAME_##Name, FGameplayEffectContextNetSerializer); \
+	UE_NET_IMPLEMENT_NETSERIALZIER_REGISTRY_DELEGATES(Name)
diff --git a/Engine/Source/Runtime/Experimental/Iris/Core/Public/Iris/ReplicationState/ReplicationStateDescriptorBuilder.h b/Engine/Source/Runtime/Experimental/Iris/Core/Public/Iris/ReplicationState/ReplicationStateDescriptorBuilder.h
--- a/Engine/Source/Runtime/Experimental/Iris/Core/Public/Iris/ReplicationState/ReplicationStateDescriptorBuilder.h
+++ b/Engine/Source/Runtime/Experimental/Iris/Core/Public/Iris/ReplicationState/ReplicationStateDescriptorBuilder.h
@@ -31,6 +31,8 @@ public:
 		uint32 GetLifeTimeProperties : 1;
 		// If EnableFastArrayHandling is true and the struct inherits from FFastArraySerializer then special logic will be applied which allows it to be bound to a FastArrayReplicationFragment.
 		uint32 EnableFastArrayHandling : 1;
+		// In SkipCheckForCustomNetSerializer is true the descriptor will be built using the underlying representation even if a CustomNetSerializer is registered for the struct
+		uint32 SkipCheckForCustomNetSerializerForStruct : 1;
 		 // If SinglePropertyIndex != -1 we will create a state that only includes the specified property.
 		int32 SinglePropertyIndex;
 	};
