Chaos - fix joint projection (knock-on from big solver refactor)

diff --git a/Engine/Source/Runtime/Experimental/Chaos/Private/Chaos/PBDJointConstraints.cpp b/Engine/Source/Runtime/Experimental/Chaos/Private/Chaos/PBDJointConstraints.cpp
--- a/Engine/Source/Runtime/Experimental/Chaos/Private/Chaos/PBDJointConstraints.cpp
+++ b/Engine/Source/Runtime/Experimental/Chaos/Private/Chaos/PBDJointConstraints.cpp
@@ -15,10 +15,6 @@
 
 #include "HAL/IConsoleManager.h"
 
-#if INTEL_ISPC
-#include "PBDJointSolverGaussSeidel.ispc.generated.h"
-#endif
-
 namespace Chaos
 {
 	DECLARE_CYCLE_STAT(TEXT("Joints::Sort"), STAT_Joints_Sort, STATGROUP_ChaosJoint);
@@ -848,6 +844,14 @@ namespace Chaos
 
 	void FPBDJointConstraints::ApplyProjectionConstraints(const FReal Dt, const int32 It, const int32 NumIts)
 	{
+		if (It == 0)
+		{
+			for (int32 ConstraintIndex = 0; ConstraintIndex < NumConstraints(); ++ConstraintIndex)
+			{
+				PrepareProjectionConstraint(Dt, ConstraintIndex, It, NumIts);
+			}
+		}
+
 		for (int32 ConstraintIndex = 0; ConstraintIndex < NumConstraints(); ++ConstraintIndex)
 		{
 			ApplyProjectionConstraint(Dt, ConstraintIndex, It, NumIts);
@@ -896,6 +900,14 @@ namespace Chaos
 
 	void FPBDJointConstraints::ApplyProjectionConstraints(const TArrayView<int32>& ConstraintIndices, const FReal Dt, const int32 It, const int32 NumIts)
 	{
+		if (It == 0)
+		{
+			for (int32 ConstraintIndex : ConstraintIndices)
+			{
+				PrepareProjectionConstraint(Dt, ConstraintIndex, It, NumIts);
+			}
+		}
+
 		for (int32 ConstraintIndex : ConstraintIndices)
 		{
 			ApplyProjectionConstraint(Dt, ConstraintIndex, It, NumIts);
@@ -1285,6 +1297,25 @@ namespace Chaos
 		return true;
 	}
 
+	void FPBDJointConstraints::PrepareProjectionConstraint(const FReal Dt, const int32 ConstraintIndex, const int32 It, const int32 NumIts)
+	{
+		// Collect all the data for projection prior to the first iteration. 
+		// This must happen for all joints before we project any joints so the the initial state for each joint is not polluted by any earlier projections.
+		// @todo(chaos): if we ever support projection on other constraint types, we will need a PrepareProjection phase so that all constraint types
+		// can initialize correctly before any constraints apply their projection. For now we can just check the iteration count is zero.
+		check(It == 0);
+
+		if (Settings.bUseLinearSolver)
+		{
+			FPBDJointCachedSolver& Solver = CachedConstraintSolvers[ConstraintIndex];
+			const FPBDJointSettings& JointSettings = ConstraintSettings[ConstraintIndex];
+			if (JointSettings.bProjectionEnabled && (Solver.IsDynamic(0) || Solver.IsDynamic(1)))
+			{
+				Solver.InitProjection(Dt, Settings, JointSettings);
+			}
+		}
+	}
+
 	// Projection phase
 	bool FPBDJointConstraints::ApplyProjectionConstraint(const FReal Dt, const int32 ConstraintIndex, const int32 It, const int32 NumIts)
 	{
@@ -1324,7 +1355,6 @@ namespace Chaos
 
 			if (It == 0)
 			{
-				Solver.InitProjection(Dt, Settings, JointSettings);
 				Solver.ApplyTeleports(Dt, Settings, JointSettings);
 			}
 
diff --git a/Engine/Source/Runtime/Experimental/Chaos/Public/Chaos/PBDJointConstraints.h b/Engine/Source/Runtime/Experimental/Chaos/Public/Chaos/PBDJointConstraints.h
--- a/Engine/Source/Runtime/Experimental/Chaos/Public/Chaos/PBDJointConstraints.h
+++ b/Engine/Source/Runtime/Experimental/Chaos/Public/Chaos/PBDJointConstraints.h
@@ -292,6 +292,7 @@ namespace Chaos
 		void AddBodies(const int32 ConstraintIndex, FSolverBodyContainer& SolverBodyContainer);
 		bool ApplyPositionConstraint(const FReal Dt, const int32 ConstraintIndex, const int32 It, const int32 NumIts);
 		bool ApplyVelocityConstraint(const FReal Dt, const int32 ConstraintIndex, const int32 It, const int32 NumIts);
+		void PrepareProjectionConstraint(const FReal Dt, const int32 ConstraintIndex, const int32 It, const int32 NumIts);
 		bool ApplyProjectionConstraint(const FReal Dt, const int32 ConstraintIndex, const int32 It, const int32 NumIts);
 		void ApplyBreakThreshold(const FReal Dt, int32 ConstraintIndex, const FVec3& LinearImpulse, const FVec3& AngularImpulse);
 		void ApplyPlasticityLimits(const int32 ConstraintIndex);
