Chaos - add cvar to disable collisions in the main scene for debugging

diff --git a/Engine/Source/Runtime/Experimental/Chaos/Private/Chaos/PBDCollisionConstraints.cpp b/Engine/Source/Runtime/Experimental/Chaos/Private/Chaos/PBDCollisionConstraints.cpp
--- a/Engine/Source/Runtime/Experimental/Chaos/Private/Chaos/PBDCollisionConstraints.cpp
+++ b/Engine/Source/Runtime/Experimental/Chaos/Private/Chaos/PBDCollisionConstraints.cpp
@@ -24,10 +24,6 @@
 // Private includes
 #include "Collision/PBDCollisionSolver.h"
 
-#if INTEL_ISPC
-#include "PBDCollisionConstraints.ispc.generated.h"
-#endif
-
 //PRAGMA_DISABLE_OPTIMIZATION
 
 namespace Chaos
@@ -362,11 +358,14 @@ namespace Chaos
 
 	void FPBDCollisionConstraints::AddConstraintsToGraph(FPBDIslandManager& IslandManager)
 	{
-		for (FPBDCollisionConstraintHandle* ConstraintHandle : GetConstraints())
+		if (GetCollisionsEnabled())
 		{
-			if (ConstraintHandle->IsEnabled())
+			for (FPBDCollisionConstraintHandle* ConstraintHandle : GetConstraints())
 			{
-				IslandManager.AddConstraint(GetContainerId(), ConstraintHandle, ConstraintHandle->GetConstrainedParticles());
+				if (ConstraintHandle->IsEnabled())
+				{
+					IslandManager.AddConstraint(GetContainerId(), ConstraintHandle, ConstraintHandle->GetConstrainedParticles());
+				}
 			}
 		}
 	}
diff --git a/Engine/Source/Runtime/Experimental/Chaos/Private/PBDRigidsSolver.cpp b/Engine/Source/Runtime/Experimental/Chaos/Private/PBDRigidsSolver.cpp
--- a/Engine/Source/Runtime/Experimental/Chaos/Private/PBDRigidsSolver.cpp
+++ b/Engine/Source/Runtime/Experimental/Chaos/Private/PBDRigidsSolver.cpp
@@ -184,6 +184,11 @@ namespace Chaos
 		bool bChaosSolverJointUseLinearSolver = true;
 		FAutoConsoleVariableRef CVarChaosSolverJointUseCachedSolver(TEXT("p.Chaos.Solver.Joint.UseLinearSolver"), bChaosSolverJointUseLinearSolver, TEXT("Use linear version of joint solver. (default is true"));
 
+		// Enable/Disable collisions
+		bool bChaosSolverCollisionEnabled = true;
+		FAutoConsoleVariableRef CVarChaosSolverCollisionDisable(TEXT("p.Chaos.Solver.Collision.Enabled"), bChaosSolverCollisionEnabled, TEXT("Enable/Disable collisions in the main scene."));
+
+
 		// Iteration count cvars
 		// These override the engine config if >= 0
 
@@ -819,6 +824,8 @@ namespace Chaos
 
 	void FPBDRigidsSolver::PrepareAdvanceBy(const FReal DeltaTime)
 	{
+		MEvolution->GetCollisionConstraints().SetCollisionsEnabled(bChaosSolverCollisionEnabled);
+
 		FCollisionDetectorSettings CollisionDetectorSettings = MEvolution->GetCollisionConstraints().GetDetectorSettings();
 		CollisionDetectorSettings.bAllowManifoldReuse = (ChaosSolverCollisionAllowManifoldUpdate != 0);
 		CollisionDetectorSettings.bDeferNarrowPhase = (ChaosSolverCollisionDeferNarrowPhase != 0);
