Chaos - fix island sleep state when resued as a new island for a solo particle

#okforgtihub public

diff --git a/Engine/Source/Runtime/Experimental/Chaos/Private/Chaos/Island/IslandGraph.cpp b/Engine/Source/Runtime/Experimental/Chaos/Private/Chaos/Island/IslandGraph.cpp
--- a/Engine/Source/Runtime/Experimental/Chaos/Private/Chaos/Island/IslandGraph.cpp
+++ b/Engine/Source/Runtime/Experimental/Chaos/Private/Chaos/Island/IslandGraph.cpp
@@ -304,7 +304,7 @@ void FIslandGraph<NodeType, EdgeType, IslandType, OwnerType>::AttachIslands(cons
 			else if (!bFirstValidIsland && !bSecondValidIsland)
 			{
 				FGraphIsland GraphIsland = { 1, 2 };
-				GraphEdge.IslandIndex = GraphIslands.Emplace(GraphIsland);
+				GraphEdge.IslandIndex = AddIsland(MoveTemp(GraphIsland));
 
 				// We set both island indices to be the equal to the edge one
 				if (FirstNode.bValidNode)
@@ -332,7 +332,7 @@ void FIslandGraph<NodeType, EdgeType, IslandType, OwnerType>::AttachIslands(cons
 			if(!GraphIslands.IsValidIndex(GraphNodes[GraphEdge.FirstNode].IslandIndex))
 			{
 				FGraphIsland GraphIsland = { 1, 1 };
-				GraphNodes[GraphEdge.FirstNode].IslandIndex =  GraphIslands.Emplace(GraphIsland);
+				GraphNodes[GraphEdge.FirstNode].IslandIndex = AddIsland(MoveTemp(GraphIsland));
 			}
 			GraphEdge.IslandIndex =  GraphNodes[GraphEdge.FirstNode].IslandIndex;
 		}
@@ -343,7 +343,7 @@ void FIslandGraph<NodeType, EdgeType, IslandType, OwnerType>::AttachIslands(cons
 			if(!GraphIslands.IsValidIndex(GraphNodes[GraphEdge.SecondNode].IslandIndex))
 			{
 				FGraphIsland GraphIsland = { 1, 1 };
-				GraphNodes[GraphEdge.SecondNode].IslandIndex =  GraphIslands.Emplace(GraphIsland);
+				GraphNodes[GraphEdge.SecondNode].IslandIndex = AddIsland(MoveTemp(GraphIsland));
 			}
 			GraphEdge.IslandIndex =  GraphNodes[GraphEdge.SecondNode].IslandIndex;
 		}
@@ -779,7 +779,7 @@ void FIslandGraph<NodeType, EdgeType, IslandType, OwnerType>::SplitIslands()
 						if(GraphIslands[CurrentIsland].IslandCounter == GraphCounter)
 						{
 							FGraphIsland GraphIsland = { 0, 1, 0, false, false };
-							CurrentIsland = GraphIslands.Emplace(GraphIsland);
+							CurrentIsland = AddIsland(MoveTemp(GraphIsland));
 						}
 						
 						GraphIslands[CurrentIsland].IslandCounter = GraphCounter;
@@ -876,7 +876,7 @@ void FIslandGraph<NodeType, EdgeType, IslandType, OwnerType>::UpdateGraph()
 			if (GraphNode.bValidNode && !GraphIslands.IsValidIndex(GraphNode.IslandIndex))
 			{
 				FGraphIsland GraphIsland = { 0, 1 }; // {NumEdges, NumNodes}
-				GraphNode.IslandIndex = GraphIslands.Emplace(GraphIsland);
+				GraphNode.IslandIndex = AddIsland(MoveTemp(GraphIsland));
 
 #if CHAOS_CONSTRAINTHANDLE_DEBUG_ENABLED
 				// Should never have any NodeIslands if we get here
diff --git a/Engine/Source/Runtime/Experimental/Chaos/Private/Chaos/Island/IslandManager.cpp b/Engine/Source/Runtime/Experimental/Chaos/Private/Chaos/Island/IslandManager.cpp
--- a/Engine/Source/Runtime/Experimental/Chaos/Private/Chaos/Island/IslandManager.cpp
+++ b/Engine/Source/Runtime/Experimental/Chaos/Private/Chaos/Island/IslandManager.cpp
@@ -613,6 +613,24 @@ inline bool SolverIslandSortPredicate(const TUniquePtr<FPBDIsland>& SolverIsland
 	return SolverIslandL->GetNumConstraints() < SolverIslandR->GetNumConstraints();
 }
 	
+void FPBDIslandManager::GraphIslandAdded(const int32 IslandIndex)
+{
+	FGraphIsland& GraphIsland = IslandGraph->GraphIslands[IslandIndex];
+	check(GraphIsland.IslandItem == nullptr);
+
+	// An graph island was created, so we should assign an island manager island to it
+	Islands.Reserve(IslandGraph->MaxNumIslands());
+	if (!Islands.IsValidIndex(IslandIndex))
+	{
+		Islands.EmplaceAt(IslandIndex, MakeUnique<FPBDIsland>(MaxNumConstraintContainers));
+	}
+
+	// NOTE: The index may be reused and we may already have an island so we need to reset the island
+	Islands[IslandIndex]->Reuse();
+	GraphIsland.IslandItem = Islands[IslandIndex].Get();
+
+}
+
 void FPBDIslandManager::SyncIslands(FPBDRigidsSOAs& Particles)
 {
 	Islands.Reserve(IslandGraph->MaxNumIslands());
@@ -630,12 +648,8 @@ void FPBDIslandManager::SyncIslands(FPBDRigidsSOAs& Particles)
 
 			// If the island item is not set on the graph we check if the solver island
 			// at the right index is already there. If not we create a new solver island.
-			if(!Islands.IsValidIndex(IslandIndex))
-			{
-				Islands.EmplaceAt(IslandIndex, MakeUnique<FPBDIsland>(MaxNumConstraintContainers));
-			}
-			FPBDIsland*& Island = IslandGraph->GraphIslands[IslandIndex].IslandItem;
-			Island = Islands[IslandIndex].Get();
+			check(Islands.IsValidIndex(IslandIndex));
+			FPBDIsland* Island = Islands[IslandIndex].Get();
 			
 			// We then transfer the persistent flag and the graph dense index to the solver island
 			Island->SetIsPersistent(IslandGraph->GraphIslands[IslandIndex].bIsPersistent);
diff --git a/Engine/Source/Runtime/Experimental/Chaos/Private/Chaos/Island/SolverIsland.cpp b/Engine/Source/Runtime/Experimental/Chaos/Private/Chaos/Island/SolverIsland.cpp
--- a/Engine/Source/Runtime/Experimental/Chaos/Private/Chaos/Island/SolverIsland.cpp
+++ b/Engine/Source/Runtime/Experimental/Chaos/Private/Chaos/Island/SolverIsland.cpp
@@ -19,6 +19,15 @@ namespace Chaos
 	IslandConstraintsByType.SetNum(MaxConstraintContainers);
 }
 
+void FPBDIsland::Reuse()
+{
+	check(GetNumParticles() == 0);
+	check(GetNumConstraints() == 0);
+
+	SetSleepCounter(0);
+	ResetSleepingChanged();
+}
+
 void FPBDIsland::UpdateParticles()
 {
 	for (FPBDIslandParticle& IslandParticle : IslandParticles)
diff --git a/Engine/Source/Runtime/Experimental/Chaos/Public/Chaos/Island/IslandGraph.h b/Engine/Source/Runtime/Experimental/Chaos/Public/Chaos/Island/IslandGraph.h
--- a/Engine/Source/Runtime/Experimental/Chaos/Public/Chaos/Island/IslandGraph.h
+++ b/Engine/Source/Runtime/Experimental/Chaos/Public/Chaos/Island/IslandGraph.h
@@ -357,6 +357,16 @@ public:
 		return false;
 	}
 
+	int32 AddIsland(FGraphIsland&& Island)
+	{
+		const int32 IslandIndex = GraphIslands.Emplace(MoveTemp(Island));
+		if (Owner != nullptr)
+		{
+			Owner->GraphIslandAdded(IslandIndex);
+		}
+		return IslandIndex;
+	}
+
 	FGraphOwner* Owner = nullptr;
 
 	/** List of graph nodes */
@@ -385,6 +395,8 @@ public:
 template<typename NodeType, typename EdgeType>
 struct TNullIslandGraphOwner
 {
+	void GraphIslandAdded(const int32 IslandIndex) {}
+
 	void GraphNodeAdded(const NodeType& NodeItem, const int32 NodeIndex) {}
 	void GraphNodeRemoved(const NodeType& NodeItem) {}
 
diff --git a/Engine/Source/Runtime/Experimental/Chaos/Public/Chaos/Island/IslandManager.h b/Engine/Source/Runtime/Experimental/Chaos/Public/Chaos/Island/IslandManager.h
--- a/Engine/Source/Runtime/Experimental/Chaos/Public/Chaos/Island/IslandManager.h
+++ b/Engine/Source/Runtime/Experimental/Chaos/Public/Chaos/Island/IslandManager.h
@@ -290,6 +290,7 @@ public:
 	const TMap<FGeometryParticleHandle*, int32>& GetParticleNodes() const { return IslandGraph->ItemNodes; }
 
 	// Graph Owner API Implementation. Allows us to store edge/node indices on constraints and particles
+	void GraphIslandAdded(const int32 IslandIndex);
 	void GraphNodeAdded(FGeometryParticleHandle* NodeItem, const int32 NodeIndex);
 	void GraphNodeRemoved(FGeometryParticleHandle* NodeItem);
 	void GraphEdgeAdded(const FConstraintHandleHolder& EdgeItem, const int32 EdgeIndex);
diff --git a/Engine/Source/Runtime/Experimental/Chaos/Public/Chaos/Island/SolverIsland.h b/Engine/Source/Runtime/Experimental/Chaos/Public/Chaos/Island/SolverIsland.h
--- a/Engine/Source/Runtime/Experimental/Chaos/Public/Chaos/Island/SolverIsland.h
+++ b/Engine/Source/Runtime/Experimental/Chaos/Public/Chaos/Island/SolverIsland.h
@@ -73,6 +73,12 @@ namespace Chaos
 		*/
 		FPBDIsland(const int32 MaxConstraintContainers);
 
+		/**
+		 * This island is about to be reused as a new island. It should be empty, but we need to reset
+		 * sleep state and other persistent state.
+		 */
+		void Reuse();
+
 		/**
 		 * The island index as seen by the systems external to the IslandManager.
 		 * The IslandManager uses this to index into the IslandIndexing array to get the
