Chaos destruction : fix internal clusters remove on break  behavior
- Unbroken children of internal clusters do not get removed on break anymore

diff --git a/Engine/Source/Runtime/Experimental/GeometryCollectionEngine/Private/GeometryCollection/GeometryCollectionComponent.cpp b/Engine/Source/Runtime/Experimental/GeometryCollectionEngine/Private/GeometryCollection/GeometryCollectionComponent.cpp
--- a/Engine/Source/Runtime/Experimental/GeometryCollectionEngine/Private/GeometryCollection/GeometryCollectionComponent.cpp
+++ b/Engine/Source/Runtime/Experimental/GeometryCollectionEngine/Private/GeometryCollection/GeometryCollectionComponent.cpp
@@ -3778,6 +3778,11 @@ struct FGeometryCollectionDecayContext
 
 void UGeometryCollectionComponent::UpdateDecay(int32 TransformIdx, float UpdatedDecay, bool bUseClusterCrumbling, bool bHasDynamicInternalClusterParent, FGeometryCollectionDecayContext& ContextInOut)
 {
+	if (bHasDynamicInternalClusterParent && !bUseClusterCrumbling)
+	{
+		return;
+	}
+
 	TManagedArray<float>& Decay = ContextInOut.DecayFacade.DecayAttribute.Modify();
 	if (UpdatedDecay > Decay[TransformIdx])
 	{
@@ -3869,7 +3874,7 @@ void UGeometryCollectionComponent::IncrementBreakTimer(float DeltaTime)
 		return;
 	}
 	
-	if (DynamicCollection && PhysicsProxy)
+	if (RestCollection && DynamicCollection && PhysicsProxy)
 	{
 		FGeometryCollectionRemoveOnBreakDynamicFacade RemoveOnBreakFacade(*DynamicCollection);
 		FGeometryCollectionDecayDynamicFacade DecayFacade(*DynamicCollection);
@@ -3880,7 +3885,8 @@ void UGeometryCollectionComponent::IncrementBreakTimer(float DeltaTime)
 			&& DynamicStateFacade.IsValid())
 		{
 			FGeometryCollectionDecayContext DecayContext(*PhysicsProxy, DecayFacade);
-			
+			const TManagedArray<int32>& OriginalParents = RestCollection->GetGeometryCollection()->Parent;
+
 			TManagedArray<float>& Decay = DecayFacade.DecayAttribute.Modify();
 			for (int32 TransformIdx = 0; TransformIdx < Decay.Num(); ++TransformIdx)
 			{
@@ -3889,10 +3895,10 @@ void UGeometryCollectionComponent::IncrementBreakTimer(float DeltaTime)
 					&& (DynamicStateFacade.HasBrokenOff(TransformIdx) || HasDynamicInternalClusterParent))
 				{
 					bool UseClusterCrumbling = RemoveOnBreakFacade.UseClusterCrumbling(TransformIdx);
-					if (HasDynamicInternalClusterParent && DynamicCollection->Parent[TransformIdx] > INDEX_NONE)
+					if (HasDynamicInternalClusterParent && OriginalParents[TransformIdx] > INDEX_NONE)
 					{
 						// look up the state of the original parent to decide if the internal cluster parent should cumble 
-						UseClusterCrumbling = RemoveOnBreakFacade.UseClusterCrumbling(DynamicCollection->Parent[TransformIdx]);
+						UseClusterCrumbling = RemoveOnBreakFacade.UseClusterCrumbling(OriginalParents[TransformIdx]);
 					}
 					
 					const float UpdatedBreakDecay = RemoveOnBreakFacade.UpdateBreakTimerAndComputeDecay(TransformIdx, DeltaTime);
