Chaos - fix static analysis warning

diff --git a/Engine/Source/Runtime/Experimental/Chaos/Private/Chaos/HeightField.cpp b/Engine/Source/Runtime/Experimental/Chaos/Private/Chaos/HeightField.cpp
--- a/Engine/Source/Runtime/Experimental/Chaos/Private/Chaos/HeightField.cpp
+++ b/Engine/Source/Runtime/Experimental/Chaos/Private/Chaos/HeightField.cpp
@@ -449,6 +449,8 @@ namespace Chaos
 		THeightfieldSweepVisitorCCD(const typename FHeightField::FDataType* InData, const GeomQueryType& InQueryGeom, const FRigidTransform3& InStartTM, const FVec3& InDir, const FReal Length, const FReal InIgnorePenetration, const FReal InTargetPenetration)
 			: OutDistance(TNumericLimits<FReal>::Max())
 			, OutPhi(TNumericLimits<FReal>::Max())
+			, OutPosition(0)
+			, OutNormal(0)
 			, OutFaceIndex(INDEX_NONE)
 			, HfData(InData)
 			, StartTM(InStartTM)
@@ -2127,14 +2129,14 @@ namespace Chaos
 		FReal TOI = (Length > 0) ? SQVisitor.OutDistance / Length : FReal(0);
 		FReal Phi = SQVisitor.OutPhi;
 
-		// @todo(chaos): Legacy path to be removed when fully tested. See ComputeSweptContactTOIAndPhiAtTargetPenetration
-		if (!CVars::bCCDNewTargetDepthMode)
-		{
-			LegacyComputeSweptContactTOIAndPhiAtTargetPenetration(FVec3::DotProduct(Dir, SQVisitor.OutNormal), Length, InIgnorePenetration, InTargetPenetration, TOI, Phi);
-		}
-
 		if (TOI <= 1)
 		{
+			// @todo(chaos): Legacy path to be removed when fully tested. See ComputeSweptContactTOIAndPhiAtTargetPenetration
+			if (!CVars::bCCDNewTargetDepthMode)
+			{
+				LegacyComputeSweptContactTOIAndPhiAtTargetPenetration(FVec3::DotProduct(Dir, SQVisitor.OutNormal), Length, InIgnorePenetration, InTargetPenetration, TOI, Phi);
+			}
+
 			OutTOI = TOI;
 			OutPhi = Phi;
 			OutPosition = SQVisitor.OutPosition;
diff --git a/Engine/Source/Runtime/Experimental/Chaos/Public/Chaos/Plane.h b/Engine/Source/Runtime/Experimental/Chaos/Public/Chaos/Plane.h
--- a/Engine/Source/Runtime/Experimental/Chaos/Public/Chaos/Plane.h
+++ b/Engine/Source/Runtime/Experimental/Chaos/Public/Chaos/Plane.h
@@ -155,6 +155,11 @@ public:
 		ensure(FMath::IsNearlyEqual(Dir.SizeSquared(), (FReal)1, (FReal)UE_KINDA_SMALL_NUMBER));
 		CHAOS_ENSURE(Length > 0);
 		OutFaceIndex = INDEX_NONE;
+		OutTime = 0;
+
+		// This is mainly to fix static analysis warnings
+		OutPosition = FVec3(0);
+		OutNormal = FVec3(0);
 
 		const FReal SignedDist = FVec3::DotProduct(StartPoint - (FVec3)MX, (FVec3)MNormal);
 		if (FMath::Abs(SignedDist) < Thickness)
@@ -163,7 +168,7 @@ public:
 			//const FReal DirDotNormal = FVec3::DotProduct(Dir, (FVec3)MNormal);
 			//OutPosition = StartPoint;
 			//OutNormal = DirDotNormal < 0 ? MNormal : -MNormal;
-			OutTime = 0;
+			//OutTime = 0;
 			return true;
 		}
 
