Make the re-import workflow able to re-import a single asset if the same source file is use for multiple assets.
The factory node "reference object" member is now an attribute so it can be serialize in the asset import data.

Avoid duplicating the whole NodeContainer, we just need the factory node and the factory node dependencies of the node

#rnx

diff --git a/Engine/Plugins/Experimental/Interchange/Runtime/Source/Import/Private/Animation/InterchangeAnimSequenceFactory.cpp b/Engine/Plugins/Experimental/Interchange/Runtime/Source/Import/Private/Animation/InterchangeAnimSequenceFactory.cpp
--- a/Engine/Plugins/Experimental/Interchange/Runtime/Source/Import/Private/Animation/InterchangeAnimSequenceFactory.cpp
+++ b/Engine/Plugins/Experimental/Interchange/Runtime/Source/Import/Private/Animation/InterchangeAnimSequenceFactory.cpp
@@ -805,6 +805,8 @@ UObject* UInterchangeAnimSequenceFactory::CreateAsset(const FCreateAssetParams&
 		UE_LOG(LogInterchangeImport, Warning, TEXT("Invalid skeleton factory node, the skeleton factory node is obligatory to import this animsequence [%s]!"), *Arguments.AssetName);
 		return nullptr;
 	}
+	FSoftObjectPath SkeletonFactoryNodeReferenceObject;
+	SkeletonFactoryNode->GetCustomReferenceObject(SkeletonFactoryNodeReferenceObject);
 
 	USkeleton* Skeleton = nullptr;
 
@@ -818,9 +820,9 @@ UObject* UInterchangeAnimSequenceFactory::CreateAsset(const FCreateAssetParams&
 		{
 			SkeletonObject = SpecifiedSkeleton.TryLoad();
 		}
-		else if (SkeletonFactoryNode->ReferenceObject.IsValid())
+		else if (SkeletonFactoryNodeReferenceObject.IsValid())
 		{
-			SkeletonObject = SkeletonFactoryNode->ReferenceObject.TryLoad();
+			SkeletonObject = SkeletonFactoryNodeReferenceObject.TryLoad();
 		}
 
 		if (SkeletonObject)
diff --git a/Engine/Plugins/Experimental/Interchange/Runtime/Source/Import/Private/Animation/InterchangeAnimationTrackSetFactory.cpp b/Engine/Plugins/Experimental/Interchange/Runtime/Source/Import/Private/Animation/InterchangeAnimationTrackSetFactory.cpp
--- a/Engine/Plugins/Experimental/Interchange/Runtime/Source/Import/Private/Animation/InterchangeAnimationTrackSetFactory.cpp
+++ b/Engine/Plugins/Experimental/Interchange/Runtime/Source/Import/Private/Animation/InterchangeAnimationTrackSetFactory.cpp
@@ -271,13 +271,26 @@ namespace UE::Interchange::Private
 
 		const FString TrackSetFactoryNodeUid = TEXT("Factory_") + TrackSetNodeUid;
 		const UInterchangeAnimationTrackSetFactoryNode* InstanceFactoryNode = Cast<UInterchangeAnimationTrackSetFactoryNode>(NodeContainer.GetNode(TrackSetFactoryNodeUid));
-		if (!InstanceFactoryNode || !InstanceFactoryNode->ReferenceObject.IsValid())
+		
+		FString InstanceNodeDisplayLabel = InstanceNode.GetDisplayLabel();
+		auto LogMissingTrackError = [&InstanceNodeDisplayLabel]()
 		{
-			UE_LOG(LogInterchangeImport, Error, TEXT("Cannot find factory of animation track set referenced by animation track %s."), *InstanceNode.GetDisplayLabel());
+			UE_LOG(LogInterchangeImport, Error, TEXT("Cannot find factory of animation track set referenced by animation track %s."), *InstanceNodeDisplayLabel);
+		};
+
+		if (!InstanceFactoryNode)
+		{
+			LogMissingTrackError();
 			return;
 		}
-
-		UMovieSceneSequence* TargetMovieSceneSequence = CastChecked<UMovieSceneSequence>(InstanceFactoryNode->ReferenceObject.TryLoad());
+		FSoftObjectPath ReferenceObject;
+		InstanceFactoryNode->GetCustomReferenceObject(ReferenceObject);
+		if (!ReferenceObject.IsValid())
+		{
+			LogMissingTrackError();
+			return;
+		}
+		UMovieSceneSequence* TargetMovieSceneSequence = CastChecked<UMovieSceneSequence>(ReferenceObject.TryLoad());
 
 		// Create SubTrack
 		UMovieSceneSubTrack* SubTrack = MovieScene->FindMasterTrack<UMovieSceneSubTrack>();
@@ -442,9 +455,14 @@ namespace UE::Interchange::Private
 			const FString ActorFactoryNodeUid = TEXT("Factory_") + ActorNodeUid;
 			const UInterchangeFactoryBaseNode* ActorFactoryNode = Cast<UInterchangeFactoryBaseNode>(NodeContainer.GetNode(ActorFactoryNodeUid));
 
-			if (ActorFactoryNode && ActorFactoryNode->ReferenceObject.IsValid())
+			if (ActorFactoryNode)
 			{
-				Actor = Cast<AActor>(ActorFactoryNode->ReferenceObject.TryLoad());
+				FSoftObjectPath ReferenceObject;
+				ActorFactoryNode->GetCustomReferenceObject(ReferenceObject);
+				if (ReferenceObject.IsValid())
+				{
+					Actor = Cast<AActor>(ReferenceObject.TryLoad());
+				}
 			}
 		}
 
diff --git a/Engine/Plugins/Experimental/Interchange/Runtime/Source/Import/Private/Material/InterchangeMaterialFactory.cpp b/Engine/Plugins/Experimental/Interchange/Runtime/Source/Import/Private/Material/InterchangeMaterialFactory.cpp
--- a/Engine/Plugins/Experimental/Interchange/Runtime/Source/Import/Private/Material/InterchangeMaterialFactory.cpp
+++ b/Engine/Plugins/Experimental/Interchange/Runtime/Source/Import/Private/Material/InterchangeMaterialFactory.cpp
@@ -121,7 +121,9 @@ namespace UE
 
 						if (const UInterchangeMaterialFunctionFactoryNode* MaterialFactoryNode = Cast<UInterchangeMaterialFunctionFactoryNode>(Arguments.NodeContainer->GetNode(MaterialFactoryNodeUid)))
 						{
-							if (UMaterialFunctionInterface* MaterialFunction = Cast<UMaterialFunctionInterface>(MaterialFactoryNode->ReferenceObject.TryLoad()))
+							FSoftObjectPath ReferenceObject;
+							MaterialFactoryNode->GetCustomReferenceObject(ReferenceObject);
+							if (UMaterialFunctionInterface* MaterialFunction = Cast<UMaterialFunctionInterface>(ReferenceObject.TryLoad()))
 							{
 								MaterialFunction->UpdateFromFunctionResource();
 								FunctionCallExpression->SetMaterialFunction(MaterialFunction);
@@ -141,7 +143,9 @@ namespace UE
 
 					if (const UInterchangeTextureFactoryNode* TextureFactoryNode = Cast<UInterchangeTextureFactoryNode>(Arguments.NodeContainer->GetNode(TextureFactoryNodeUid)))
 					{
-						if (UTexture* Texture = Cast<UTexture>(TextureFactoryNode->ReferenceObject.TryLoad()))
+						FSoftObjectPath ReferenceObject;
+						TextureFactoryNode->GetCustomReferenceObject(ReferenceObject);
+						if (UTexture* Texture = Cast<UTexture>(ReferenceObject.TryLoad()))
 						{
 							TextureExpression->Texture = Texture;
 						}
@@ -1041,7 +1045,9 @@ void UInterchangeMaterialFactory::SetupMaterialInstance(UMaterialInstance* Mater
 					{
 						if (const UInterchangeTextureFactoryNode* TextureFactoryNode = Cast<UInterchangeTextureFactoryNode>(NodeContainer->GetNode(InputValue)))
 						{
-							if (UTexture* InputTexture = Cast<UTexture>(TextureFactoryNode->ReferenceObject.TryLoad()))
+							FSoftObjectPath ReferenceObject;
+							TextureFactoryNode->GetCustomReferenceObject(ReferenceObject);
+							if (UTexture* InputTexture = Cast<UTexture>(ReferenceObject.TryLoad()))
 							{
 								if ( InputTexture != InstanceValue )
 								{
diff --git a/Engine/Plugins/Experimental/Interchange/Runtime/Source/Import/Private/Mesh/InterchangePhysicsAssetFactory.cpp b/Engine/Plugins/Experimental/Interchange/Runtime/Source/Import/Private/Mesh/InterchangePhysicsAssetFactory.cpp
--- a/Engine/Plugins/Experimental/Interchange/Runtime/Source/Import/Private/Mesh/InterchangePhysicsAssetFactory.cpp
+++ b/Engine/Plugins/Experimental/Interchange/Runtime/Source/Import/Private/Mesh/InterchangePhysicsAssetFactory.cpp
@@ -24,7 +24,7 @@ UObject* UInterchangePhysicsAssetFactory::CreateEmptyAsset(const FCreateAssetPar
 		return nullptr;
 	}
 
-	const UInterchangePhysicsAssetFactoryNode* PhysicsAssetNode = Cast<UInterchangePhysicsAssetFactoryNode>(Arguments.AssetNode);
+	UInterchangePhysicsAssetFactoryNode* PhysicsAssetNode = Cast<UInterchangePhysicsAssetFactoryNode>(Arguments.AssetNode);
 	if (PhysicsAssetNode == nullptr)
 	{
 		return nullptr;
@@ -49,7 +49,7 @@ UObject* UInterchangePhysicsAssetFactory::CreateEmptyAsset(const FCreateAssetPar
 		UE_LOG(LogInterchangeImport, Warning, TEXT("Could not create PhysicsAsset asset %s"), *Arguments.AssetName);
 		return nullptr;
 	}
-	PhysicsAssetNode->ReferenceObject = PhysicsAsset;
+	PhysicsAssetNode->SetCustomReferenceObject(FSoftObjectPath(PhysicsAsset));
 
 	PhysicsAsset->PreEditChange(nullptr);
 #endif //WITH_EDITORONLY_DATA
@@ -70,7 +70,7 @@ UObject* UInterchangePhysicsAssetFactory::CreateAsset(const FCreateAssetParams&
 		return nullptr;
 	}
 
-	const UInterchangePhysicsAssetFactoryNode* PhysicsAssetNode = Cast<UInterchangePhysicsAssetFactoryNode>(Arguments.AssetNode);
+	UInterchangePhysicsAssetFactoryNode* PhysicsAssetNode = Cast<UInterchangePhysicsAssetFactoryNode>(Arguments.AssetNode);
 	if (PhysicsAssetNode == nullptr)
 	{
 		return nullptr;
@@ -116,7 +116,7 @@ UObject* UInterchangePhysicsAssetFactory::CreateAsset(const FCreateAssetParams&
 				return nullptr;
 			}
 			
-			PhysicsAssetNode->ReferenceObject = PhysicsAsset;
+			PhysicsAssetNode->SetCustomReferenceObject(FSoftObjectPath(PhysicsAsset));
 		}
 		
 		//Getting the file Hash will cache it into the source data
diff --git a/Engine/Plugins/Experimental/Interchange/Runtime/Source/Import/Private/Mesh/InterchangeSkeletalMeshFactory.cpp b/Engine/Plugins/Experimental/Interchange/Runtime/Source/Import/Private/Mesh/InterchangeSkeletalMeshFactory.cpp
--- a/Engine/Plugins/Experimental/Interchange/Runtime/Source/Import/Private/Mesh/InterchangeSkeletalMeshFactory.cpp
+++ b/Engine/Plugins/Experimental/Interchange/Runtime/Source/Import/Private/Mesh/InterchangeSkeletalMeshFactory.cpp
@@ -1064,13 +1064,21 @@ UObject* UInterchangeSkeletalMeshFactory::CreateAsset(const FCreateAssetParams&
 		FName MaterialSlotName = *SlotMaterialDependency.Key;
 
 		const UInterchangeBaseMaterialFactoryNode* MaterialFactoryNode = Cast<UInterchangeBaseMaterialFactoryNode>(Arguments.NodeContainer->GetNode(SlotMaterialDependency.Value));
-		if (!MaterialFactoryNode || !MaterialFactoryNode->ReferenceObject.IsValid() || !MaterialFactoryNode->IsEnabled())
+		if (!MaterialFactoryNode || !MaterialFactoryNode->IsEnabled())
 		{
 			UpdateOrAddSkeletalMaterial(MaterialSlotName, UMaterial::GetDefaultMaterial(MD_Surface));
 			continue;
 		}
 
-		UMaterialInterface* MaterialInterface = Cast<UMaterialInterface>(MaterialFactoryNode->ReferenceObject.ResolveObject());
+		FSoftObjectPath MaterialFactoryNodeReferenceObject;
+		MaterialFactoryNode->GetCustomReferenceObject(MaterialFactoryNodeReferenceObject);
+		if (!MaterialFactoryNodeReferenceObject.IsValid())
+		{
+			UpdateOrAddSkeletalMaterial(MaterialSlotName, UMaterial::GetDefaultMaterial(MD_Surface));
+			continue;
+		}
+
+		UMaterialInterface* MaterialInterface = Cast<UMaterialInterface>(MaterialFactoryNodeReferenceObject.ResolveObject());
 		UpdateOrAddSkeletalMaterial(MaterialSlotName, MaterialInterface ? MaterialInterface : UMaterial::GetDefaultMaterial(MD_Surface));
 	}
 
@@ -1104,6 +1112,8 @@ UObject* UInterchangeSkeletalMeshFactory::CreateAsset(const FCreateAssetParams&
 			UE_LOG(LogInterchangeImport, Warning, TEXT("Invalid Skeleton LOD when importing SkeletalMesh asset %s"), *Arguments.AssetName);
 			continue;
 		}
+		FSoftObjectPath SkeletonNodeReferenceObject;
+		SkeletonNode->GetCustomReferenceObject(SkeletonNodeReferenceObject);
 
 		FSoftObjectPath SpecifiedSkeleton;
 		SkeletalMeshFactoryNode->GetCustomSkeletonSoftObjectPath(SpecifiedSkeleton);
@@ -1116,9 +1126,9 @@ UObject* UInterchangeSkeletalMeshFactory::CreateAsset(const FCreateAssetParams&
 			{
 				SkeletonObject = SpecifiedSkeleton.TryLoad();
 			}
-			else if (SkeletonNode->ReferenceObject.IsValid())
+			else if (SkeletonNodeReferenceObject.IsValid())
 			{
-				SkeletonObject = SkeletonNode->ReferenceObject.TryLoad();
+				SkeletonObject = SkeletonNodeReferenceObject.TryLoad();
 			}
 
 			if (SkeletonObject)
diff --git a/Engine/Plugins/Experimental/Interchange/Runtime/Source/Import/Private/Mesh/InterchangeSkeletonFactory.cpp b/Engine/Plugins/Experimental/Interchange/Runtime/Source/Import/Private/Mesh/InterchangeSkeletonFactory.cpp
--- a/Engine/Plugins/Experimental/Interchange/Runtime/Source/Import/Private/Mesh/InterchangeSkeletonFactory.cpp
+++ b/Engine/Plugins/Experimental/Interchange/Runtime/Source/Import/Private/Mesh/InterchangeSkeletonFactory.cpp
@@ -25,7 +25,7 @@ UObject* UInterchangeSkeletonFactory::CreateEmptyAsset(const FCreateAssetParams&
 		return nullptr;
 	}
 
-	const UInterchangeSkeletonFactoryNode* SkeletonNode = Cast<UInterchangeSkeletonFactoryNode>(Arguments.AssetNode);
+	UInterchangeSkeletonFactoryNode* SkeletonNode = Cast<UInterchangeSkeletonFactoryNode>(Arguments.AssetNode);
 	if (SkeletonNode == nullptr)
 	{
 		return nullptr;
@@ -50,7 +50,7 @@ UObject* UInterchangeSkeletonFactory::CreateEmptyAsset(const FCreateAssetParams&
 		UE_LOG(LogInterchangeImport, Warning, TEXT("Could not create Skeleton asset %s"), *Arguments.AssetName);
 		return nullptr;
 	}
-	SkeletonNode->ReferenceObject = Skeleton;
+	SkeletonNode->SetCustomReferenceObject(FSoftObjectPath(Skeleton));
 
 	Skeleton->PreEditChange(nullptr);
 #endif //WITH_EDITORONLY_DATA
@@ -71,7 +71,7 @@ UObject* UInterchangeSkeletonFactory::CreateAsset(const FCreateAssetParams& Argu
 		return nullptr;
 	}
 
-	const UInterchangeSkeletonFactoryNode* SkeletonNode = Cast<UInterchangeSkeletonFactoryNode>(Arguments.AssetNode);
+	UInterchangeSkeletonFactoryNode* SkeletonNode = Cast<UInterchangeSkeletonFactoryNode>(Arguments.AssetNode);
 	if (SkeletonNode == nullptr)
 	{
 		return nullptr;
@@ -134,7 +134,7 @@ UObject* UInterchangeSkeletonFactory::CreateAsset(const FCreateAssetParams& Argu
 			}
 		}
 		//The joint will be added by the skeletalmesh factory since we need a valid skeletalmesh to add joint to a skeleton
-		SkeletonNode->ReferenceObject = Skeleton;
+		SkeletonNode->SetCustomReferenceObject(FSoftObjectPath(Skeleton));
 		
 		//Getting the file Hash will cache it into the source data
 		Arguments.SourceData->GetFileContentHash();
diff --git a/Engine/Plugins/Experimental/Interchange/Runtime/Source/Import/Private/Mesh/InterchangeStaticMeshFactory.cpp b/Engine/Plugins/Experimental/Interchange/Runtime/Source/Import/Private/Mesh/InterchangeStaticMeshFactory.cpp
--- a/Engine/Plugins/Experimental/Interchange/Runtime/Source/Import/Private/Mesh/InterchangeStaticMeshFactory.cpp
+++ b/Engine/Plugins/Experimental/Interchange/Runtime/Source/Import/Private/Mesh/InterchangeStaticMeshFactory.cpp
@@ -192,13 +192,20 @@ UObject* UInterchangeStaticMeshFactory::CreateAsset(const FCreateAssetParams& Ar
 		FName MaterialSlotName = *SlotMaterialDependency.Key;
 
 		const UInterchangeBaseMaterialFactoryNode* MaterialFactoryNode = Cast<UInterchangeBaseMaterialFactoryNode>(Arguments.NodeContainer->GetNode(SlotMaterialDependency.Value));
-		if (!MaterialFactoryNode || !MaterialFactoryNode->ReferenceObject.IsValid() || !MaterialFactoryNode->IsEnabled())
+		if (!MaterialFactoryNode || !MaterialFactoryNode->IsEnabled())
+		{
+			UpdateOrAddStaticMaterial(MaterialSlotName, UMaterial::GetDefaultMaterial(MD_Surface));
+			continue;
+		}
+		FSoftObjectPath MaterialFactoryNodeReferenceObject;
+		MaterialFactoryNode->GetCustomReferenceObject(MaterialFactoryNodeReferenceObject);
+		if (!MaterialFactoryNodeReferenceObject.IsValid())
 		{
 			UpdateOrAddStaticMaterial(MaterialSlotName, UMaterial::GetDefaultMaterial(MD_Surface));
 			continue;
 		}
 
-		UMaterialInterface* MaterialInterface = Cast<UMaterialInterface>(MaterialFactoryNode->ReferenceObject.ResolveObject());
+		UMaterialInterface* MaterialInterface = Cast<UMaterialInterface>(MaterialFactoryNodeReferenceObject.TryLoad());
 		UpdateOrAddStaticMaterial(MaterialSlotName, MaterialInterface ? MaterialInterface : UMaterial::GetDefaultMaterial(MD_Surface));
 	}
 
diff --git a/Engine/Plugins/Experimental/Interchange/Runtime/Source/Import/Private/Scene/InterchangeActorHelper.cpp b/Engine/Plugins/Experimental/Interchange/Runtime/Source/Import/Private/Scene/InterchangeActorHelper.cpp
--- a/Engine/Plugins/Experimental/Interchange/Runtime/Source/Import/Private/Scene/InterchangeActorHelper.cpp
+++ b/Engine/Plugins/Experimental/Interchange/Runtime/Source/Import/Private/Scene/InterchangeActorHelper.cpp
@@ -24,7 +24,9 @@ AActor* UE::Interchange::ActorHelper::GetSpawnedParentActor(const UInterchangeBa
 
 	if (const UInterchangeFactoryBaseNode* ParentFactoryNode = NodeContainer->GetFactoryNode(FactoryNode->GetParentUid()))
 	{
-		ParentActor = Cast<AActor>(ParentFactoryNode->ReferenceObject.TryLoad());
+		FSoftObjectPath ReferenceObject;
+		ParentFactoryNode->GetCustomReferenceObject(ReferenceObject);
+		ParentActor = Cast<AActor>(ReferenceObject.TryLoad());
 	}
 
 	return  ParentActor;
@@ -158,13 +160,19 @@ void UE::Interchange::ActorHelper::ApplySlotMaterialDependencies(const UIntercha
 	for (TPair<FString, FString>& SlotMaterialDependency : SlotMaterialDependencies)
 	{
 		const UInterchangeBaseMaterialFactoryNode* MaterialFactoryNode = Cast<UInterchangeBaseMaterialFactoryNode>(NodeContainer.GetNode(SlotMaterialDependency.Value));
-		if (!MaterialFactoryNode || !MaterialFactoryNode->ReferenceObject.IsValid() || !MaterialFactoryNode->IsEnabled())
+		if (!MaterialFactoryNode || !MaterialFactoryNode->IsEnabled())
+		{
+			continue;
+		}
+		FSoftObjectPath ReferenceObject;
+		MaterialFactoryNode->GetCustomReferenceObject(ReferenceObject);
+		if (!ReferenceObject.IsValid())
 		{
 			continue;
 		}
 
 		FName MaterialSlotName = *SlotMaterialDependency.Key;
-		UMaterialInterface* MaterialInterface = Cast<UMaterialInterface>(MaterialFactoryNode->ReferenceObject.ResolveObject());
+		UMaterialInterface* MaterialInterface = Cast<UMaterialInterface>(ReferenceObject.TryLoad());
 		if (!MaterialInterface)
 		{
 			MaterialInterface = UMaterial::GetDefaultMaterial(MD_Surface);
diff --git a/Engine/Plugins/Experimental/Interchange/Runtime/Source/Import/Private/Scene/InterchangeLightActorFactory.cpp b/Engine/Plugins/Experimental/Interchange/Runtime/Source/Import/Private/Scene/InterchangeLightActorFactory.cpp
--- a/Engine/Plugins/Experimental/Interchange/Runtime/Source/Import/Private/Scene/InterchangeLightActorFactory.cpp
+++ b/Engine/Plugins/Experimental/Interchange/Runtime/Source/Import/Private/Scene/InterchangeLightActorFactory.cpp
@@ -37,7 +37,9 @@ UObject* UInterchangeLightActorFactory::CreateSceneObject(const UInterchangeFact
 				IESTexture = UInterchangeFactoryBaseNode::BuildFactoryNodeUid(IESTexture);
 				if(const UInterchangeTextureLightProfileFactoryNode* TextureFactoryNode = Cast<UInterchangeTextureLightProfileFactoryNode>(CreateSceneObjectsParams.NodeContainer->GetNode(IESTexture)))
 				{
-					if(UTextureLightProfile* Texture = Cast<UTextureLightProfile>(TextureFactoryNode->ReferenceObject.TryLoad()))
+					FSoftObjectPath ReferenceObject;
+					TextureFactoryNode->GetCustomReferenceObject(ReferenceObject);
+					if(UTextureLightProfile* Texture = Cast<UTextureLightProfile>(ReferenceObject.TryLoad()))
 					{
 						LightComponent->SetIESTexture(Texture);
 					}
diff --git a/Engine/Plugins/Experimental/Interchange/Runtime/Source/Import/Private/Scene/InterchangeSceneVariantSetsFactory.cpp b/Engine/Plugins/Experimental/Interchange/Runtime/Source/Import/Private/Scene/InterchangeSceneVariantSetsFactory.cpp
--- a/Engine/Plugins/Experimental/Interchange/Runtime/Source/Import/Private/Scene/InterchangeSceneVariantSetsFactory.cpp
+++ b/Engine/Plugins/Experimental/Interchange/Runtime/Source/Import/Private/Scene/InterchangeSceneVariantSetsFactory.cpp
@@ -216,7 +216,9 @@ namespace UE::Interchange::Private
 			{
 				if (const UInterchangeBaseMaterialFactoryNode* TargetFactoryNode = Cast<UInterchangeBaseMaterialFactoryNode>(NodeContainer.GetNode(TargetFactoryNodeUid)))
 				{
-					if (const UMaterialInterface* TargetMaterial = Cast<UMaterialInterface>(TargetFactoryNode->ReferenceObject.TryLoad()))
+					FSoftObjectPath ReferenceObject;
+					TargetFactoryNode->GetCustomReferenceObject(ReferenceObject);
+					if (const UMaterialInterface* TargetMaterial = Cast<UMaterialInterface>(ReferenceObject.TryLoad()))
 					{
 						// FVariantManager::CreateMaterialPropertyCaptures returns an array of the newly found properties
 						// Therefore, if a mesh actor has more than 1 material and a binding has more than 1 material variant
@@ -275,9 +277,11 @@ namespace UE::Interchange::Private
 		const FString ActorFactoryNodeUid = UInterchangeFactoryBaseNode::BuildFactoryNodeUid(ActorNodeUid);
 		const UInterchangeFactoryBaseNode* ActorFactoryNode = Cast<UInterchangeFactoryBaseNode>(NodeContainer.GetNode(ActorFactoryNodeUid));
 
-		if (ActorFactoryNode && ActorFactoryNode->ReferenceObject.IsValid())
+		if (ActorFactoryNode )
 		{
-			return Cast<AActor>(ActorFactoryNode->ReferenceObject.TryLoad());
+			FSoftObjectPath ReferenceObject;
+			ActorFactoryNode->GetCustomReferenceObject(ReferenceObject);
+			return ReferenceObject.IsValid() ? Cast<AActor>(ReferenceObject.TryLoad()) : nullptr;
 		}
 
 		return nullptr;
diff --git a/Engine/Plugins/Experimental/Interchange/Runtime/Source/Import/Private/Scene/InterchangeSkeletalMeshActorFactory.cpp b/Engine/Plugins/Experimental/Interchange/Runtime/Source/Import/Private/Scene/InterchangeSkeletalMeshActorFactory.cpp
--- a/Engine/Plugins/Experimental/Interchange/Runtime/Source/Import/Private/Scene/InterchangeSkeletalMeshActorFactory.cpp
+++ b/Engine/Plugins/Experimental/Interchange/Runtime/Source/Import/Private/Scene/InterchangeSkeletalMeshActorFactory.cpp
@@ -52,7 +52,9 @@ void UInterchangeSkeletalMeshActorFactory::PostImportPreCompletedCallback(const
 		{
 			if (const UInterchangeFactoryBaseNode* MeshNode = UE::Interchange::ActorHelper::FindAssetInstanceFactoryNode(Arguments.NodeContainer, Arguments.FactoryNode))
 			{
-				if (USkeletalMesh* SkeletalMesh = Cast<USkeletalMesh>(MeshNode->ReferenceObject.TryLoad()))
+				FSoftObjectPath ReferenceObject;
+				MeshNode->GetCustomReferenceObject(ReferenceObject);
+				if (USkeletalMesh* SkeletalMesh = Cast<USkeletalMesh>(ReferenceObject.TryLoad()))
 				{
 					SkeletalMeshComponent->SetSkeletalMesh(SkeletalMesh);
 
diff --git a/Engine/Plugins/Experimental/Interchange/Runtime/Source/Import/Private/Scene/InterchangeStaticMeshActorFactory.cpp b/Engine/Plugins/Experimental/Interchange/Runtime/Source/Import/Private/Scene/InterchangeStaticMeshActorFactory.cpp
--- a/Engine/Plugins/Experimental/Interchange/Runtime/Source/Import/Private/Scene/InterchangeStaticMeshActorFactory.cpp
+++ b/Engine/Plugins/Experimental/Interchange/Runtime/Source/Import/Private/Scene/InterchangeStaticMeshActorFactory.cpp
@@ -48,7 +48,9 @@ void UInterchangeStaticMeshActorFactory::SetupStaticMeshActor(const UInterchange
 
 		if (const UInterchangeFactoryBaseNode* MeshNode = UE::Interchange::ActorHelper::FindAssetInstanceFactoryNode(NodeContainer, ActorFactoryNode))
 		{
-			if (UStaticMesh* StaticMesh = Cast<UStaticMesh>(MeshNode->ReferenceObject.TryLoad()))
+			FSoftObjectPath ReferenceObject;
+			MeshNode->GetCustomReferenceObject(ReferenceObject);
+			if (UStaticMesh* StaticMesh = Cast<UStaticMesh>(ReferenceObject.TryLoad()))
 			{
 				StaticMeshComponent->SetStaticMesh(StaticMesh);
 
diff --git a/Engine/Plugins/Experimental/Interchange/Runtime/Source/Pipelines/Private/InterchangeGenericSkeletalMeshPipeline.cpp b/Engine/Plugins/Experimental/Interchange/Runtime/Source/Pipelines/Private/InterchangeGenericSkeletalMeshPipeline.cpp
--- a/Engine/Plugins/Experimental/Interchange/Runtime/Source/Pipelines/Private/InterchangeGenericSkeletalMeshPipeline.cpp
+++ b/Engine/Plugins/Experimental/Interchange/Runtime/Source/Pipelines/Private/InterchangeGenericSkeletalMeshPipeline.cpp
@@ -309,7 +309,7 @@ UInterchangeSkeletonFactoryNode* UInterchangeGenericMeshPipeline::CreateSkeleton
 	if (CommonSkeletalMeshesAndAnimationsProperties->Skeleton)
 	{
 		SkeletonFactoryNode->SetEnabled(false);
-		SkeletonFactoryNode->ReferenceObject = CommonSkeletalMeshesAndAnimationsProperties->Skeleton;
+		SkeletonFactoryNode->SetCustomReferenceObject(FSoftObjectPath(CommonSkeletalMeshesAndAnimationsProperties->Skeleton));
 	}
 #if WITH_EDITOR
 	//Iterate all joints to set the meta data value in the skeleton node
@@ -648,9 +648,11 @@ void UInterchangeGenericMeshPipeline::PostImportPhysicsAssetImport(UObject* Crea
 		{
 			if (const UInterchangeSkeletalMeshFactoryNode* SkeletalMeshFactoryNode = Cast<const UInterchangeSkeletalMeshFactoryNode>(BaseNodeContainer->GetFactoryNode(SkeletalMeshFactoryNodeUid)))
 			{
-				if (SkeletalMeshFactoryNode->ReferenceObject.IsValid())
+				FSoftObjectPath ReferenceObject;
+				SkeletalMeshFactoryNode->GetCustomReferenceObject(ReferenceObject);
+				if (ReferenceObject.IsValid())
 				{
-					if (USkeletalMesh* SkeletalMesh = Cast<USkeletalMesh>(SkeletalMeshFactoryNode->ReferenceObject.TryLoad()))
+					if (USkeletalMesh* SkeletalMesh = Cast<USkeletalMesh>(ReferenceObject.TryLoad()))
 					{
 						auto CreateFromSkeletalMeshLambda = [CreatedPhysicsAsset, SkeletalMesh]()
 						{
diff --git a/Engine/Source/Runtime/Experimental/Interchange/Core/Private/Nodes/InterchangeFactoryBaseNode.cpp b/Engine/Source/Runtime/Experimental/Interchange/Core/Private/Nodes/InterchangeFactoryBaseNode.cpp
--- a/Engine/Source/Runtime/Experimental/Interchange/Core/Private/Nodes/InterchangeFactoryBaseNode.cpp
+++ b/Engine/Source/Runtime/Experimental/Interchange/Core/Private/Nodes/InterchangeFactoryBaseNode.cpp
@@ -143,6 +143,16 @@ bool UInterchangeFactoryBaseNode::RemoveFactoryDependencyUid(const FString& Depe
 	return FactoryDependencies.RemoveItem(DependencyUid);
 }
 
+bool UInterchangeFactoryBaseNode::GetCustomReferenceObject(FSoftObjectPath& AttributeValue) const
+{
+	IMPLEMENT_NODE_ATTRIBUTE_GETTER(ReferenceObject, FSoftObjectPath)
+}
+
+bool UInterchangeFactoryBaseNode::SetCustomReferenceObject(const FSoftObjectPath& AttributeValue)
+{
+	IMPLEMENT_NODE_ATTRIBUTE_SETTER_NODELEGATE(ReferenceObject, FSoftObjectPath)
+}
+
 FString UInterchangeFactoryBaseNode::BuildFactoryNodeUid(const FString& TranslatedNodeUid)
 {
 	return TEXT("Factory_") + TranslatedNodeUid;
diff --git a/Engine/Source/Runtime/Experimental/Interchange/Core/Public/Nodes/InterchangeFactoryBaseNode.h b/Engine/Source/Runtime/Experimental/Interchange/Core/Public/Nodes/InterchangeFactoryBaseNode.h
--- a/Engine/Source/Runtime/Experimental/Interchange/Core/Public/Nodes/InterchangeFactoryBaseNode.h
+++ b/Engine/Source/Runtime/Experimental/Interchange/Core/Public/Nodes/InterchangeFactoryBaseNode.h
@@ -193,9 +193,17 @@ public:
 	UFUNCTION(BlueprintCallable, Category = "Interchange | Node")
 	bool RemoveFactoryDependencyUid(const FString& DependencyUid);
 
+	/**
+	 * Return the custom ReferenceObject. The UObject this factory node has created.
+	 */
+	UFUNCTION(BlueprintCallable, Category = "Interchange | Node")
+	bool GetCustomReferenceObject(FSoftObjectPath& AttributeValue) const;
 
-	UPROPERTY()
-	mutable FSoftObjectPath ReferenceObject;
+	/**
+	 * Set the custom ReferenceObject. The UObject this factory node has created.
+	 */
+	UFUNCTION(BlueprintCallable, Category = "Interchange | Node")
+	bool SetCustomReferenceObject(const FSoftObjectPath& AttributeValue);
 
 	static FString BuildFactoryNodeUid(const FString& TranslatedNodeUid);
 
@@ -253,7 +261,7 @@ protected:
 private:
 
 	const UE::Interchange::FAttributeKey Macro_CustomSubPathKey = UE::Interchange::FAttributeKey(TEXT("SubPath"));
-
+	const UE::Interchange::FAttributeKey Macro_CustomReferenceObjectKey = UE::Interchange::FAttributeKey(TEXT("ReferenceObject"));
 };
 
 template<typename AttributeType>
diff --git a/Engine/Source/Runtime/Experimental/Interchange/Engine/Private/InterchangeImportCommon.cpp b/Engine/Source/Runtime/Experimental/Interchange/Engine/Private/InterchangeImportCommon.cpp
--- a/Engine/Source/Runtime/Experimental/Interchange/Engine/Private/InterchangeImportCommon.cpp
+++ b/Engine/Source/Runtime/Experimental/Interchange/Engine/Private/InterchangeImportCommon.cpp
@@ -21,6 +21,7 @@ namespace UE
 		{
 			UInterchangeAssetImportData* BeginSetupAssetData(FFactoryCommon::FUpdateImportAssetDataParameters& Parameters)
 			{
+				TRACE_CPUPROFILER_EVENT_SCOPE("UE::Interchange::Private::ImportCommon::BeginSetupAssetData")
 				if (!ensure(IsInGameThread()))
 				{
 					return nullptr;
@@ -42,13 +43,35 @@ namespace UE
 				return AssetImportData;
 			}
 
+			void RecursivelyDuplicateFactoryNodeDependencies(UInterchangeBaseNodeContainer* SourceContainer, UInterchangeBaseNodeContainer* DestinationContainer, TArray<FString>& FactoryDependencies)
+			{
+				for (const FString& DependencyUid : FactoryDependencies)
+				{
+					if (UInterchangeFactoryBaseNode* FactoryNode = Cast<UInterchangeFactoryBaseNode>(SourceContainer->GetFactoryNode(DependencyUid)))
+					{
+						FObjectDuplicationParameters DupParam(FactoryNode, DestinationContainer);
+						DestinationContainer->AddNode(CastChecked<UInterchangeBaseNode>(StaticDuplicateObjectEx(DupParam)));
+
+						TArray<FString> ChildFactoryDependencies;
+						FactoryNode->GetFactoryDependencies(ChildFactoryDependencies);
+						RecursivelyDuplicateFactoryNodeDependencies(SourceContainer, DestinationContainer, ChildFactoryDependencies);
+					}
+				}
+			}
 
 			void EndSetupAssetData(FFactoryCommon::FUpdateImportAssetDataParameters& Parameters, UInterchangeAssetImportData* AssetImportData)
 			{
-				//Set the interchange node graph data
-				AssetImportData->NodeUniqueID = Parameters.NodeUniqueID;
-				FObjectDuplicationParameters DupParam(Parameters.NodeContainer, AssetImportData);
-				AssetImportData->NodeContainer = CastChecked<UInterchangeBaseNodeContainer>(StaticDuplicateObjectEx(DupParam));
+				TRACE_CPUPROFILER_EVENT_SCOPE("UE::Interchange::Private::ImportCommon::EndSetupAssetData")
+				AssetImportData->NodeContainer = NewObject<UInterchangeBaseNodeContainer>(AssetImportData);
+				//We copy only the factory node dependencies, we use this only 
+				if (UInterchangeFactoryBaseNode* FactoryNode = Cast<UInterchangeFactoryBaseNode>(Parameters.NodeContainer->GetFactoryNode(Parameters.NodeUniqueID)))
+				{
+					//Set the interchange node graph data
+					AssetImportData->NodeUniqueID = Parameters.NodeUniqueID;
+					TArray<FString> FactoryDependencies;
+					FactoryDependencies.Add(AssetImportData->NodeUniqueID);
+					RecursivelyDuplicateFactoryNodeDependencies(Parameters.NodeContainer, AssetImportData->NodeContainer, FactoryDependencies);
+				}
 				AssetImportData->Pipelines.Reset();
 				for (const UObject* Pipeline : Parameters.Pipelines)
 				{
@@ -82,6 +105,7 @@ namespace UE
 
 		UAssetImportData* FFactoryCommon::UpdateImportAssetData(FUpdateImportAssetDataParameters& Parameters)
 		{
+			TRACE_CPUPROFILER_EVENT_SCOPE("UE::Interchange::FFactoryCommon::UpdateImportAssetData")
 			return UpdateImportAssetData(Parameters, [&Parameters](UInterchangeAssetImportData* AssetImportData)
 				{
 #if WITH_EDITORONLY_DATA
@@ -97,6 +121,7 @@ namespace UE
 
 		UAssetImportData* FFactoryCommon::UpdateImportAssetData(FUpdateImportAssetDataParameters& Parameters, TFunctionRef<void(UInterchangeAssetImportData*)> CustomFileSourceUpdate)
 		{
+			TRACE_CPUPROFILER_EVENT_SCOPE("UE::Interchange::FFactoryCommon::UpdateImportAssetData2")
 #if WITH_EDITORONLY_DATA
 			if (!ensure(IsInGameThread()))
 			{
@@ -150,6 +175,7 @@ namespace UE
 
 		UAssetImportData* FFactoryCommon::SetImportAssetData(FSetImportAssetDataParameters& Parameters)
 		{
+			TRACE_CPUPROFILER_EVENT_SCOPE("UE::Interchange::FFactoryCommon::SetImportAssetData")
 			UInterchangeAssetImportData* AssetImportData = Private::ImportCommon::BeginSetupAssetData(Parameters);
 
 			// Update the source files
@@ -185,6 +211,7 @@ namespace UE
 
 		bool FFactoryCommon::GetSourceFilenames(const UAssetImportData* AssetImportData, TArray<FString>& OutSourceFilenames)
 		{
+			TRACE_CPUPROFILER_EVENT_SCOPE("UE::Interchange::FFactoryCommon::GetSourceFilenames")
 			if (Cast<UInterchangeAssetImportData>(AssetImportData) != nullptr)
 			{
 				AssetImportData->ExtractFilenames(OutSourceFilenames);
@@ -195,6 +222,7 @@ namespace UE
 
 		bool FFactoryCommon::SetSourceFilename(UAssetImportData* AssetImportData, const FString& SourceFilename, int32 SourceIndex, const FString& SourceLabel)
 		{
+			TRACE_CPUPROFILER_EVENT_SCOPE("UE::Interchange::FFactoryCommon::SetSourceFilename")
 			if (AssetImportData)
 			{
 				const int32 SafeSourceIndex = SourceIndex == INDEX_NONE ? 0 : SourceIndex;				
@@ -215,6 +243,7 @@ namespace UE
 		
 		bool FFactoryCommon::SetReimportSourceIndex(const UObject* Object, UAssetImportData* AssetImportData, int32 SourceIndex)
 		{
+			TRACE_CPUPROFILER_EVENT_SCOPE("UE::Interchange::FFactoryCommon::SetReimportSourceIndex")
 			UInterchangeAssetImportData* InterchangeAssetImportData = Cast<UInterchangeAssetImportData>(AssetImportData);
 			if (!InterchangeAssetImportData)
 			{
@@ -247,6 +276,7 @@ namespace UE
 										  , UInterchangeFactoryBaseNode* CurrentAssetNode
 										  , UInterchangeFactoryBaseNode* PipelineAssetNode)
 		{
+			TRACE_CPUPROFILER_EVENT_SCOPE("UE::Interchange::FFactoryCommon::ApplyReimportStrategyToAsset")
 			if (!ensure(PreviousAssetNode) || !ensure(PipelineAssetNode) || !ensure(CurrentAssetNode))
 			{
 				return;
diff --git a/Engine/Source/Runtime/Experimental/Interchange/Engine/Private/Tasks/InterchangeTaskCreateAsset.cpp b/Engine/Source/Runtime/Experimental/Interchange/Engine/Private/Tasks/InterchangeTaskCreateAsset.cpp
--- a/Engine/Source/Runtime/Experimental/Interchange/Engine/Private/Tasks/InterchangeTaskCreateAsset.cpp
+++ b/Engine/Source/Runtime/Experimental/Interchange/Engine/Private/Tasks/InterchangeTaskCreateAsset.cpp
@@ -48,12 +48,12 @@ namespace UE
 				OutPackageName = FPaths::Combine(*SanitizedPackageBasePath, *SubPath, *OutAssetName);
 			}
 
-			UObject* GetExistingObjectFromAssetImportData(TSharedPtr<UE::Interchange::FImportAsyncHelper, ESPMode::ThreadSafe> AsyncHelper, UInterchangeFactoryBaseNode* FactoryNode)
+			UObject* GetExistingObjectFromAssetImportData(UObject* ReimportObject, UInterchangeFactoryBaseNode* FactoryNode, FString& OutPackageName, FString& OutAssetName)
 			{
 				TRACE_CPUPROFILER_EVENT_SCOPE("UE::Interchange::Private::GetExistingObjectFromAssetImportData")
 				UInterchangeAssetImportData* OriginalAssetImportData = nullptr;
 				TArray<UObject*> SubObjects;
-				GetObjectsWithOuter(AsyncHelper->TaskData.ReimportObject, SubObjects);
+				GetObjectsWithOuter(ReimportObject, SubObjects);
 				for (UObject* SubObject : SubObjects)
 				{
 					OriginalAssetImportData = Cast<UInterchangeAssetImportData>(SubObject);
@@ -65,28 +65,26 @@ namespace UE
 
 				if (OriginalAssetImportData)
 				{
-					if (UInterchangeBaseNodeContainer* NodeContainer = OriginalAssetImportData->NodeContainer)
+					if (UInterchangeBaseNodeContainer* OriginalNodeContainer = OriginalAssetImportData->NodeContainer)
 					{
-						UClass* FactoryNodeClass = FactoryNode->GetClass();
-						UInterchangeBaseNode* SelectedOriginalNode = nullptr;
-						NodeContainer->BreakableIterateNodes([FactoryNodeClass, &SelectedOriginalNode](const FString&, UInterchangeBaseNode* OriginalNode)
-							{
-								if (OriginalNode->GetClass() == FactoryNodeClass && OriginalNode->GetParentUid() == UInterchangeBaseNode::InvalidNodeUid())
-								{
-									SelectedOriginalNode = OriginalNode;
-									return true;
-								}
-								return false;
-							});
-
-						// Hack for the reimport with a new file. (to revisited for MVP as this not a future proof solution. Should there be some sort of adapter that tell us how to do the mapping? Or should the pipeline do the mapping on a reimport? After all it is the pipeline that chose the name of the asset.
-						if (SelectedOriginalNode)
+						if (UInterchangeFactoryBaseNode* OriginalFactoryNode = OriginalNodeContainer->GetFactoryNode(OriginalAssetImportData->NodeUniqueID))
 						{
-							return AsyncHelper->TaskData.ReimportObject;
+							FSoftObjectPath ReferenceObject;
+							OriginalFactoryNode->GetCustomReferenceObject(ReferenceObject);
+							if (ReferenceObject.TryLoad() == ReimportObject)
+							{
+								OutPackageName = ReimportObject->GetPackage()->GetPathName();
+								OutAssetName = ReimportObject->GetName();
+
+								FactoryNode->SetDisplayLabel(OutAssetName);
+								FactoryNode->SetAssetName(OutAssetName);
+								// Hack for the texture reimport with a new file. (to revisited for MVP as this not a future proof solution. Should there be some sort of adapter that tell us
+								// how to do the mapping? Or should the pipeline do the mapping on a reimport? After all it is the pipeline that chose the name of the asset.
+								return ReimportObject;
+							}
 						}
 					}
 				}
-
 				return nullptr;
 			}
 
@@ -124,8 +122,9 @@ void UE::Interchange::FTaskCreatePackage::DoTask(ENamedThreads::Type CurrentThre
 	UPackage* Pkg = nullptr;
 	FString PackageName;
 	FString AssetName;
+	UObject* ReimportObject = AsyncHelper->TaskData.ReimportObject;
 	//If we do a reimport no need to create a package
-	if (AsyncHelper->TaskData.ReimportObject)
+	if (ReimportObject)
 	{
 		Private::InternalGetPackageName(*AsyncHelper, SourceIndex, PackageBasePath, FactoryNode, PackageName, AssetName);
 		UPackage* ResultFindPackage = FindPackage(nullptr, *PackageName);
@@ -139,10 +138,17 @@ void UE::Interchange::FTaskCreatePackage::DoTask(ENamedThreads::Type CurrentThre
 			ExistingObject = FindFirstObject<UObject>(*AssetName, EFindFirstObjectOptions::NativeFirst | EFindFirstObjectOptions::EnsureIfAmbiguous);
 		}
 
-		if (ExistingObject != AsyncHelper->TaskData.ReimportObject)
+		if (!ExistingObject)
 		{
-			// Try to see if we can do a mapping from the source data (we should revisit this for the MVP)
-			ExistingObject = UE::Interchange::Private::GetExistingObjectFromAssetImportData(AsyncHelper, FactoryNode);
+			ExistingObject = Private::GetExistingObjectFromAssetImportData(ReimportObject, FactoryNode, PackageName, AssetName);
+		}
+
+		if (ExistingObject != ReimportObject)
+		{
+			//Set the existing object so other factory can link UObject correctly (i.e. mesh link to existing material)
+			FactoryNode->SetCustomReferenceObject(FSoftObjectPath(ExistingObject));
+			//Skip this asset, the re-import is not for this asset
+			return;
 		}
 		
 
@@ -161,7 +167,8 @@ void UE::Interchange::FTaskCreatePackage::DoTask(ENamedThreads::Type CurrentThre
 			{
 				CreateAssetParams.NodeContainer = AsyncHelper->BaseNodeContainers[SourceIndex].Get();
 			}
-			CreateAssetParams.ReimportObject = AsyncHelper->TaskData.ReimportObject;
+			CreateAssetParams.ReimportObject = ReimportObject;
+			FactoryNode->SetCustomReferenceObject(FSoftObjectPath(ReimportObject));
 			//We call CreateEmptyAsset to ensure any resource use by an existing UObject is release on the game thread
 			Factory->CreateEmptyAsset(CreateAssetParams);
 		}
@@ -217,7 +224,7 @@ void UE::Interchange::FTaskCreatePackage::DoTask(ENamedThreads::Type CurrentThre
 		{
 			CreateAssetParams.NodeContainer = AsyncHelper->BaseNodeContainers[SourceIndex].Get();
 		}
-		CreateAssetParams.ReimportObject = AsyncHelper->TaskData.ReimportObject;
+		CreateAssetParams.ReimportObject = ReimportObject;
 		//Make sure the asset UObject is created with the correct type on the main thread
 		UObject* NodeAsset = Factory->CreateEmptyAsset(CreateAssetParams);
 		if (NodeAsset)
@@ -234,8 +241,8 @@ void UE::Interchange::FTaskCreatePackage::DoTask(ENamedThreads::Type CurrentThre
 			AssetInfo.ImportedObject = NodeAsset;
 			AssetInfo.Factory = Factory;
 			AssetInfo.FactoryNode = FactoryNode;
-			AssetInfo.bIsReimport = bool(AsyncHelper->TaskData.ReimportObject);
-			FactoryNode->ReferenceObject = FSoftObjectPath(NodeAsset);
+			AssetInfo.bIsReimport = bool(ReimportObject != nullptr);
+			FactoryNode->SetCustomReferenceObject(FSoftObjectPath(NodeAsset));
 		}
 	}
 
@@ -274,7 +281,8 @@ void UE::Interchange::FTaskCreateAsset::DoTask(ENamedThreads::Type CurrentThread
 	Private::InternalGetPackageName(*AsyncHelper, SourceIndex, PackageBasePath, FactoryNode, PackageName, AssetName);
 	bool bSkipAsset = false;
 	UObject* ExistingObject = nullptr;
-	if (AsyncHelper->TaskData.ReimportObject)
+	UObject* ReimportObject = AsyncHelper->TaskData.ReimportObject;
+	if (ReimportObject)
 	{
 		//When we re-import one particular asset, if the source file contain other assets, we want to set the node= reference UObject for those asset to the existing asset
 		//The way to discover this case is to compare the re-import asset with the node asset.
@@ -287,17 +295,16 @@ void UE::Interchange::FTaskCreateAsset::DoTask(ENamedThreads::Type CurrentThread
 		{
 			ExistingObject = FindFirstObject<UObject>(*AssetName, EFindFirstObjectOptions::NativeFirst | EFindFirstObjectOptions::EnsureIfAmbiguous);
 		}
-
-		if(ExistingObject != AsyncHelper->TaskData.ReimportObject)
+		
+		if (!ExistingObject)
 		{
-			// Try to see if we can do a mapping from the source data (we should revisit this for the MVP)
-			ExistingObject = UE::Interchange::Private::GetExistingObjectFromAssetImportData(AsyncHelper, FactoryNode);
+			ExistingObject = UE::Interchange::Private::GetExistingObjectFromAssetImportData(ReimportObject, FactoryNode, PackageName, AssetName);
 		}
 
-		bSkipAsset = !ExistingObject || ExistingObject != AsyncHelper->TaskData.ReimportObject;
+		bSkipAsset = !ExistingObject || ExistingObject != ReimportObject;
 		if (!bSkipAsset)
 		{
-			Pkg = AsyncHelper->TaskData.ReimportObject->GetPackage();
+			Pkg = ReimportObject->GetPackage();
 			PackageName = Pkg->GetPathName();
 			AssetName = ExistingObject->GetName();
 		}
@@ -355,7 +362,7 @@ void UE::Interchange::FTaskCreateAsset::DoTask(ENamedThreads::Type CurrentThread
 		{
 			CreateAssetParams.NodeContainer = AsyncHelper->BaseNodeContainers[SourceIndex].Get();
 		}
-		CreateAssetParams.ReimportObject = AsyncHelper->TaskData.ReimportObject;
+		CreateAssetParams.ReimportObject = ReimportObject;
 
 		NodeAsset = Factory->CreateAsset(CreateAssetParams);
 	}
@@ -376,7 +383,7 @@ void UE::Interchange::FTaskCreateAsset::DoTask(ENamedThreads::Type CurrentThread
 				AssetInfo.ImportedObject = NodeAsset;
 				AssetInfo.Factory = Factory;
 				AssetInfo.FactoryNode = FactoryNode;
-				AssetInfo.bIsReimport = bool(AsyncHelper->TaskData.ReimportObject);
+				AssetInfo.bIsReimport = bool(ReimportObject != nullptr);
 			}
 
 			// Fill in destination asset and type in any results which have been added previously by a translator or pipeline, now that we have a corresponding factory.
@@ -396,6 +403,6 @@ void UE::Interchange::FTaskCreateAsset::DoTask(ENamedThreads::Type CurrentThread
 			}
 		}
 
-		FactoryNode->ReferenceObject = FSoftObjectPath(NodeAsset);
+		FactoryNode->SetCustomReferenceObject(FSoftObjectPath(NodeAsset));
 	}
 }
\ No newline at end of file
diff --git a/Engine/Source/Runtime/Experimental/Interchange/Engine/Private/Tasks/InterchangeTaskCreateSceneObjects.cpp b/Engine/Source/Runtime/Experimental/Interchange/Engine/Private/Tasks/InterchangeTaskCreateSceneObjects.cpp
--- a/Engine/Source/Runtime/Experimental/Interchange/Engine/Private/Tasks/InterchangeTaskCreateSceneObjects.cpp
+++ b/Engine/Source/Runtime/Experimental/Interchange/Engine/Private/Tasks/InterchangeTaskCreateSceneObjects.cpp
@@ -93,7 +93,7 @@ void UE::Interchange::FTaskCreateSceneObjects::DoTask(ENamedThreads::Type Curren
 				ObjectInfo.FactoryNode = FactoryNode;
 			}
 
-			FactoryNode->ReferenceObject = FSoftObjectPath(SceneObject);
+			FactoryNode->SetCustomReferenceObject(FSoftObjectPath(SceneObject));
 		}
 	}
 }
\ No newline at end of file
