Interchange: Make sure interchange generate unique asset name.

#rnx

diff --git a/Engine/Source/Runtime/Experimental/Interchange/Engine/Private/Tasks/InterchangeTaskParsing.cpp b/Engine/Source/Runtime/Experimental/Interchange/Engine/Private/Tasks/InterchangeTaskParsing.cpp
--- a/Engine/Source/Runtime/Experimental/Interchange/Engine/Private/Tasks/InterchangeTaskParsing.cpp
+++ b/Engine/Source/Runtime/Experimental/Interchange/Engine/Private/Tasks/InterchangeTaskParsing.cpp
@@ -228,14 +228,23 @@ void UE::Interchange::FTaskParsing::DoTask(ENamedThreads::Type CurrentThread, co
 		return GraphEvents;
 	};
 
+	struct FTaskParsingRenameInfo
+	{
+		UInterchangeFactoryBaseNode* FactoryNode = nullptr;
+		int32 SourceIndex = INDEX_NONE;
+		FString OriginalName = FString();
+		FString NewName = FString();
+	};
+
+	TMap<FString, FTaskParsingRenameInfo> RenameAssets;
 	TSet<FString> CreatedTasksAssetNames; // Tracks for which asset name we have created a task so that we don't have 2 tasks for the same asset name
-	TFunction<FGraphEventRef(FTaskData&)> CreateTasksFromData = [this, &AsyncHelper, &CreatedTasksAssetNames](FTaskData& TaskData)
+	TFunction<FGraphEventRef(FTaskData&)> CreateTasksFromData = [this, &AsyncHelper, &RenameAssets, &CreatedTasksAssetNames](FTaskData& TaskData)
 	{
 		check(TaskData.Nodes.Num() == 1); //We expect 1 node per asset task
 
 		const int32 SourceIndex = TaskData.SourceIndex;
 		const UClass* const FactoryClass = TaskData.FactoryClass;
-		UInterchangeFactoryBaseNode* const FactoryNode = TaskData.Nodes[0];
+		UInterchangeFactoryBaseNode* FactoryNode = TaskData.Nodes[0];
 		const bool bFactoryCanRunOnAnyThread = FactoryClass->GetDefaultObject<UInterchangeFactoryBase>()->CanExecuteOnAnyThread();
 
 		if (TaskData.bIsSceneNode)
@@ -249,7 +258,32 @@ void UE::Interchange::FTaskParsing::DoTask(ENamedThreads::Type CurrentThread, co
 			FString PackageSubPath;
 			FactoryNode->GetCustomSubPath(PackageSubPath);
 
-			const FString AssetFullPath = FPaths::Combine(PackageBasePath, PackageSubPath, FactoryNode->GetAssetName());
+			FString AssetFullPath = FPaths::Combine(PackageBasePath, PackageSubPath, FactoryNode->GetAssetName());
+
+			//Make sure there is no duplicate name full path
+			uint32 NameIndex = 1;
+			FString NewName = AssetFullPath;
+			bool NameClash = true;
+			
+			while (NameClash)
+			{
+				if (CreatedTasksAssetNames.Contains(NewName))
+				{
+					const FString NameIndexString = FString::FromInt(NameIndex++);
+					NewName = AssetFullPath + NameIndexString;
+					FTaskParsingRenameInfo& RenameInfo = RenameAssets.FindOrAdd(AssetFullPath);
+					RenameInfo.FactoryNode = FactoryNode;
+					RenameInfo.OriginalName = AssetFullPath;
+					RenameInfo.NewName = NewName;
+					RenameInfo.SourceIndex = SourceIndex;
+					FactoryNode->SetDisplayLabel(FactoryNode->GetDisplayLabel() + NameIndexString);
+				}
+				else
+				{
+					AssetFullPath = NewName;
+					NameClash = false;
+				}
+			}
 
 			if (ensureMsgf(!CreatedTasksAssetNames.Contains(AssetFullPath),
 				TEXT("Found multiple task data with the same asset name (%s). Only one will be executed."), *AssetFullPath))
@@ -300,6 +334,27 @@ void UE::Interchange::FTaskParsing::DoTask(ENamedThreads::Type CurrentThread, co
 		CompletionPrerequistes.Add(TaskData.GraphEventRef);
 	}
 
+	if (!RenameAssets.IsEmpty())
+	{
+		TMap<TWeakObjectPtr<UInterchangeTranslatorBase>, FString> TranslatorMessageMap;
+		for (const TPair<FString, FTaskParsingRenameInfo>& RenameAssetKvp : RenameAssets)
+		{
+			const FTaskParsingRenameInfo& RenameInfo = RenameAssetKvp.Value;
+			FString& Message = TranslatorMessageMap.FindOrAdd(AsyncHelper->Translators[RenameInfo.SourceIndex]);
+			Message += FText::Format(NSLOCTEXT("InterchangeTaskParsingDoTask", "RenamedAssetMessagePerAsset", "\n OriginalName:[{0}] NewName:[{1}]")
+				, FText::FromString(RenameInfo.OriginalName)
+				, FText::FromString(RenameInfo.NewName)).ToString();
+		}
+		for (const TPair<TWeakObjectPtr<UInterchangeTranslatorBase>, FString>& MessagePerTranslator : TranslatorMessageMap)
+		{
+			UInterchangeResultWarning_Generic* WarningResult = NewObject<UInterchangeResultWarning_Generic>(GetTransientPackage(), UInterchangeResultWarning_Generic::StaticClass());
+			FString Message = NSLOCTEXT("InterchangeTaskParsingDoTask", "RenamedAssetsMessageHeader", "Renamed Assets:").ToString();
+			Message += MessagePerTranslator.Value;
+			WarningResult->Text = FText::FromString(Message);
+			MessagePerTranslator.Key->AddMessage(WarningResult);
+		}
+	}
+
 	//Add an async task for pre completion
 	
 	FGraphEventArray PreCompletionPrerequistes;
