Warning: potentially uninitialized local variable used in GJK.h

diff --git a/Engine/Source/Runtime/Experimental/Chaos/Public/Chaos/GJK.h b/Engine/Source/Runtime/Experimental/Chaos/Public/Chaos/GJK.h
--- a/Engine/Source/Runtime/Experimental/Chaos/Public/Chaos/GJK.h
+++ b/Engine/Source/Runtime/Experimental/Chaos/Public/Chaos/GJK.h
@@ -1458,7 +1458,7 @@ namespace Chaos
 		VectorRegister4Float As[4] = { VectorZeroFloat(), VectorZeroFloat(), VectorZeroFloat(), VectorZeroFloat() };
 		VectorRegister4Float Bs[4] = { VectorZeroFloat(), VectorZeroFloat(), VectorZeroFloat(), VectorZeroFloat() };
 
-		VectorRegister4Float Barycentric;
+		VectorRegister4Float Barycentric = VectorZeroFloat();
 
 		VectorRegister4Float Inflation = VectorAdd(MarginASimd, MarginBSimd);
 		constexpr VectorRegister4Float Eps2Simd = MakeVectorRegisterFloatConstant(1e-6f, 1e-6f, 1e-6f, 1e-6f);
