Dataflow Graph:
- added NodeColorsFactory to register node title/node body tint color based on category
- exposed category color settings in Project Settings/Dataflow Plugin for override

diff --git a/Engine/Plugins/Experimental/Dataflow/Source/DataflowEditor/DataflowEditor.Build.cs b/Engine/Plugins/Experimental/Dataflow/Source/DataflowEditor/DataflowEditor.Build.cs
--- a/Engine/Plugins/Experimental/Dataflow/Source/DataflowEditor/DataflowEditor.Build.cs
+++ b/Engine/Plugins/Experimental/Dataflow/Source/DataflowEditor/DataflowEditor.Build.cs
@@ -36,7 +36,6 @@ namespace UnrealBuildTool.Rules
 					"DataflowCore",
 					"DataflowEngine",
 					"Slate",
-					"DeveloperSettings",
 				}
 			);
 		}
diff --git a/Engine/Plugins/Experimental/Dataflow/Source/DataflowEditor/Private/Dataflow/DataflowSettings.cpp b/Engine/Plugins/Experimental/Dataflow/Source/DataflowEditor/Private/Dataflow/DataflowSettings.cpp
deleted file mode 100644
diff --git a/Engine/Plugins/Experimental/Dataflow/Source/DataflowEditor/Public/Dataflow/DataflowSettings.h b/Engine/Plugins/Experimental/Dataflow/Source/DataflowEditor/Public/Dataflow/DataflowSettings.h
deleted file mode 100644
diff --git a/Engine/Plugins/Experimental/GeometryCollectionPlugin/Source/GeometryCollectionNodes/Private/Dataflow/GeometryCollectionNodes.cpp b/Engine/Plugins/Experimental/GeometryCollectionPlugin/Source/GeometryCollectionNodes/Private/Dataflow/GeometryCollectionNodes.cpp
--- a/Engine/Plugins/Experimental/GeometryCollectionPlugin/Source/GeometryCollectionNodes/Private/Dataflow/GeometryCollectionNodes.cpp
+++ b/Engine/Plugins/Experimental/GeometryCollectionPlugin/Source/GeometryCollectionNodes/Private/Dataflow/GeometryCollectionNodes.cpp
@@ -23,6 +23,8 @@
 
 namespace Dataflow
 {
+	static const FLinearColor CDefaultNodeBodyTintColor = FLinearColor(0.f, 0.f, 0.f, 0.5f);
+
 	void GeometryCollectionEngineAssetNodes()
 	{
 		DATAFLOW_NODE_REGISTER_CREATION_FACTORY(FGetCollectionAssetDataflowNode);
@@ -64,6 +66,22 @@ namespace Dataflow
 		DATAFLOW_NODE_REGISTER_CREATION_FACTORY(FHashVectorDataflowNode);
 		DATAFLOW_NODE_REGISTER_CREATION_FACTORY(FFloatToIntDataflowNode);
 		DATAFLOW_NODE_REGISTER_CREATION_FACTORY(FMathConstantsDataflowNode);
+
+		// GeometryCollection
+		DATAFLOW_NODE_REGISTER_CREATION_FACTORY_NODE_COLORS_BY_CATEGORY("GeometryCollection", FLinearColor(0.55f, 0.45f, 1.0f), CDefaultNodeBodyTintColor);
+		// Development
+		DATAFLOW_NODE_REGISTER_CREATION_FACTORY_NODE_COLORS_BY_CATEGORY("Development", FLinearColor(1.f, 0.f, 0.f), CDefaultNodeBodyTintColor);
+		// Utilities|String
+		DATAFLOW_NODE_REGISTER_CREATION_FACTORY_NODE_COLORS_BY_CATEGORY("Utilities|String", FLinearColor(0.5f, 0.f, 0.5f), CDefaultNodeBodyTintColor);
+		// Fracture
+		DATAFLOW_NODE_REGISTER_CREATION_FACTORY_NODE_COLORS_BY_CATEGORY("Fracture", FLinearColor(1.f, 1.f, 0.8f), CDefaultNodeBodyTintColor);
+		// Utilities
+		DATAFLOW_NODE_REGISTER_CREATION_FACTORY_NODE_COLORS_BY_CATEGORY("Utilities", FLinearColor(1.f, 1.f, 0.f), CDefaultNodeBodyTintColor);
+		// Math
+		DATAFLOW_NODE_REGISTER_CREATION_FACTORY_NODE_COLORS_BY_CATEGORY("Math", FLinearColor(0.f, 0.4f, 0.8f), CDefaultNodeBodyTintColor);
+		// Generators
+		DATAFLOW_NODE_REGISTER_CREATION_FACTORY_NODE_COLORS_BY_CATEGORY("Generators", FLinearColor(.7f, 0.7f, 0.7f), CDefaultNodeBodyTintColor);
+
 	}
 }
 
diff --git a/Engine/Source/Runtime/Experimental/Dataflow/Core/DataflowCore.Build.cs b/Engine/Source/Runtime/Experimental/Dataflow/Core/DataflowCore.Build.cs
--- a/Engine/Source/Runtime/Experimental/Dataflow/Core/DataflowCore.Build.cs
+++ b/Engine/Source/Runtime/Experimental/Dataflow/Core/DataflowCore.Build.cs
@@ -12,8 +12,10 @@ namespace UnrealBuildTool.Rules
 					"Core",
 					"CoreUObject",
 					"Chaos",
+					"DeveloperSettings",
+					"InputCore",
 				}
-				);
+			);
 		}
 	}
 }
diff --git a/Engine/Source/Runtime/Experimental/Dataflow/Core/Private/Dataflow/DataflowNodeColorsFactory.cpp b/Engine/Source/Runtime/Experimental/Dataflow/Core/Private/Dataflow/DataflowNodeColorsFactory.cpp
new file mode 100644
--- /dev/null
+++ b/Engine/Source/Runtime/Experimental/Dataflow/Core/Private/Dataflow/DataflowNodeColorsFactory.cpp
@@ -0,0 +1,130 @@
+// Copyright Epic Games, Inc. All Rights Reserved.
+
+#include "Dataflow/DataflowNodeColorsFactory.h"
+
+#include "Dataflow/DataflowNode.h"
+//#include "Misc/MessageDialog.h"
+#include "Misc/LazySingleton.h"
+
+namespace Dataflow
+{
+	FNodeColorsFactory::FNodeColorsFactory()
+	{
+		DataflowSettings = GetMutableDefault<UDataflowSettings>();
+		DataflowSettingsChangedDelegateHandle = DataflowSettings->GetOnDataflowSettingsChangedDelegate().AddRaw(this, &FNodeColorsFactory::NodeColorsChangedInSettings);
+
+		const FNodeColorsMap NodeColorsMap = DataflowSettings->GetNodeColorsMap();
+
+		for (auto& Elem : NodeColorsMap)
+		{
+			if (ColorsMap.Contains(Elem.Key))
+			{
+				ColorsMap[Elem.Key] = Elem.Value;
+			}
+			else
+			{
+				ColorsMap.Add(Elem.Key, Elem.Value);
+			}
+		}
+	}
+
+	FNodeColorsFactory::~FNodeColorsFactory()
+	{
+		DataflowSettings->GetOnDataflowSettingsChangedDelegate().Remove(DataflowSettingsChangedDelegateHandle);
+	}
+
+	FNodeColorsFactory& FNodeColorsFactory::Get()
+	{
+		return TLazySingleton<FNodeColorsFactory>::Get();
+	}
+
+	void FNodeColorsFactory::TearDown()
+	{
+		return TLazySingleton<FNodeColorsFactory>::TearDown();
+	}
+
+	void FNodeColorsFactory::RegisterNodeColors(const FName& Category, const FNodeColors& NodeColors)
+	{
+		if (!ColorsMap.Contains(Category))
+		{
+			ColorsMap.Add(Category, NodeColors);
+		}
+
+		// Register colors in DataflowSettings
+		GetMutableDefault<UDataflowSettings>()->RegisterColors(Category, NodeColors);
+	}
+
+	FLinearColor FNodeColorsFactory::GetNodeTitleColor(const FName& Category)
+	{
+		if (ColorsMap.Contains(Category))
+		{
+			return ColorsMap[Category].NodeTitleColor;
+		}
+		else
+		{
+			// Check if any of the parent category registered
+			FString CategoryString = Category.ToString();
+			if (CategoryString.Contains(TEXT("|")))
+			{
+				do {
+					FString LfString, RtString;
+					if (CategoryString.Split(TEXT("|"), &LfString, &RtString, ESearchCase::IgnoreCase, ESearchDir::FromEnd))
+					{
+						if (ColorsMap.Contains(FName(*LfString)))
+						{
+							return ColorsMap[FName(*LfString)].NodeTitleColor;
+						}
+
+						CategoryString = LfString;
+					}
+				} while (CategoryString.Contains(TEXT("|")));
+			}
+		}
+		return FNodeColors().NodeTitleColor;
+	}
+
+	FLinearColor FNodeColorsFactory::GetNodeBodyTintColor(const FName& Category)
+	{
+		if (ColorsMap.Contains(Category))
+		{
+			return ColorsMap[Category].NodeBodyTintColor;
+		}
+		else
+		{
+			// Check if any of the parent category registered
+			FString CategoryString = Category.ToString();
+			if (CategoryString.Contains(TEXT("|")))
+			{
+				do {
+					FString LfString, RtString;
+					if (CategoryString.Split(TEXT("|"), &LfString, &RtString, ESearchCase::IgnoreCase, ESearchDir::FromEnd))
+					{
+						if (ColorsMap.Contains(FName(*LfString)))
+						{
+							return ColorsMap[FName(*LfString)].NodeBodyTintColor;
+						}
+
+						CategoryString = LfString;
+					}
+				} while (CategoryString.Contains(TEXT("|")));
+			}
+		}
+		return FNodeColors().NodeBodyTintColor;
+	}
+
+	void FNodeColorsFactory::NodeColorsChangedInSettings(const FNodeColorsMap& NodeColorsMap)
+	{
+		for (auto& Elem : NodeColorsMap)
+		{
+			if (ColorsMap.Contains(Elem.Key))
+			{
+				ColorsMap[Elem.Key] = Elem.Value;
+			}
+			else
+			{
+				ColorsMap.Add(Elem.Key, Elem.Value);
+			}
+		}
+	}
+
+}
\ No newline at end of file
diff --git a/Engine/Source/Runtime/Experimental/Dataflow/Core/Private/Dataflow/DataflowSettings.cpp b/Engine/Source/Runtime/Experimental/Dataflow/Core/Private/Dataflow/DataflowSettings.cpp
new file mode 100644
--- /dev/null
+++ b/Engine/Source/Runtime/Experimental/Dataflow/Core/Private/Dataflow/DataflowSettings.cpp
@@ -0,0 +1,49 @@
+// Copyright Epic Games, Inc. All Rights Reserved.
+
+#include "Dataflow/DataflowSettings.h"
+#include "Dataflow/DataflowNodeColorsFactory.h"
+
+UDataflowSettings::UDataflowSettings(const FObjectInitializer& ObjectInitlaizer)
+	: Super(ObjectInitlaizer)
+{
+	ManagedArrayCollectionPinTypeColor = CManagedArrayCollectionPinTypeColor;
+	ArrayPinTypeColor = CArrayPinTypeColor;
+	BoxPinTypeColor = CBoxPinTypeColor;
+}
+
+FName UDataflowSettings::GetCategoryName() const
+{
+	return TEXT("Plugins");
+}
+
+#if WITH_EDITOR
+
+FText UDataflowSettings::GetSectionText() const
+{
+	return NSLOCTEXT("DataflowPlugin", "DataflowSettingsSection", "Dataflow");
+}
+
+void UDataflowSettings::PostEditChangeProperty(FPropertyChangedEvent& PropertyChangedEvent)
+{
+	if (PropertyChangedEvent.Property != nullptr)
+	{
+		OnDataflowSettingsChangedDelegate.Broadcast(NodeColorsMap);
+	}
+
+	Super::PostEditChangeProperty(PropertyChangedEvent);
+}
+
+#endif
+
+FNodeColors UDataflowSettings::RegisterColors(const FName& Category, const FNodeColors& Colors)
+{
+	if (!NodeColorsMap.Contains(Category))
+	{
+		NodeColorsMap.Add(Category, Colors);
+	}
+	return NodeColorsMap[Category];
+}
+
+
+
+
diff --git a/Engine/Source/Runtime/Experimental/Dataflow/Core/Public/Dataflow/DataflowCore.h b/Engine/Source/Runtime/Experimental/Dataflow/Core/Public/Dataflow/DataflowCore.h
--- a/Engine/Source/Runtime/Experimental/Dataflow/Core/Public/Dataflow/DataflowCore.h
+++ b/Engine/Source/Runtime/Experimental/Dataflow/Core/Public/Dataflow/DataflowCore.h
@@ -7,4 +7,5 @@
 #include "Dataflow/DataflowInputOutput.h"
 #include "Dataflow/DataflowNode.h"
 #include "Dataflow/DataflowNodeFactory.h"
+#include "Dataflow/DataflowNodeColorsFactory.h"
 #include "Dataflow/DataflowNodeParameters.h"
diff --git a/Engine/Source/Runtime/Experimental/Dataflow/Core/Public/Dataflow/DataflowNode.h b/Engine/Source/Runtime/Experimental/Dataflow/Core/Public/Dataflow/DataflowNode.h
--- a/Engine/Source/Runtime/Experimental/Dataflow/Core/Public/Dataflow/DataflowNode.h
+++ b/Engine/Source/Runtime/Experimental/Dataflow/Core/Public/Dataflow/DataflowNode.h
@@ -7,6 +7,7 @@
 #include "Dataflow/DataflowConnection.h"
 #include "Dataflow/DataflowNodeParameters.h"
 #include "UObject/StructOnScope.h"
+#include "Dataflow/DataflowSettings.h"
 
 #include "DataflowNode.generated.h"
 
@@ -218,6 +219,10 @@ public:																				\
 	TYPE() {}																		\
 private:
 
+#define DATAFLOW_NODE_REGISTER_CREATION_FACTORY_NODE_COLORS_BY_CATEGORY(A, C1, C2)	\
+{																					\
+	FNodeColorsFactory::Get().RegisterNodeColors(A, {C1, C2});						\
+}																					\
 
 }
 
diff --git a/Engine/Source/Runtime/Experimental/Dataflow/Core/Public/Dataflow/DataflowNodeColorsFactory.h b/Engine/Source/Runtime/Experimental/Dataflow/Core/Public/Dataflow/DataflowNodeColorsFactory.h
new file mode 100644
--- /dev/null
+++ b/Engine/Source/Runtime/Experimental/Dataflow/Core/Public/Dataflow/DataflowNodeColorsFactory.h
@@ -0,0 +1,42 @@
+// Copyright Epic Games, Inc. All Rights Reserved.
+
+#pragma once
+
+#include "CoreMinimal.h"
+#include "Dataflow/DataflowGraph.h"
+#include "Dataflow/DataflowSettings.h"
+#include "ChaosLog.h"
+
+struct FDataflowNode;
+struct FDataflowConnection;
+	
+class FLazySingleton;
+
+namespace Dataflow
+{
+	//
+	//
+	//
+	class DATAFLOWCORE_API FNodeColorsFactory
+	{
+	public:
+		static FNodeColorsFactory& Get();
+		static void TearDown();
+
+		void RegisterNodeColors(const FName& Category, const FNodeColors& NodeColors);
+		FLinearColor GetNodeTitleColor(const FName& Category);
+		FLinearColor GetNodeBodyTintColor(const FName& Category);
+		void NodeColorsChangedInSettings(const FNodeColorsMap& NodeColorsMap);
+
+	private:
+		FNodeColorsFactory();
+		~FNodeColorsFactory();
+
+		TMap<FName, FNodeColors > ColorsMap;					// [Category] -> Colors
+		FDelegateHandle DataflowSettingsChangedDelegateHandle;
+		UDataflowSettings* DataflowSettings;
+
+		friend FLazySingleton;
+	};
+
+}
diff --git a/Engine/Source/Runtime/Experimental/Dataflow/Core/Public/Dataflow/DataflowSettings.h b/Engine/Source/Runtime/Experimental/Dataflow/Core/Public/Dataflow/DataflowSettings.h
new file mode 100644
--- /dev/null
+++ b/Engine/Source/Runtime/Experimental/Dataflow/Core/Public/Dataflow/DataflowSettings.h
@@ -0,0 +1,74 @@
+// Copyright Epic Games, Inc. All Rights Reserved.
+
+#pragma once
+
+#include "Engine/DeveloperSettings.h"
+#include "InputCoreTypes.h"
+
+#include "DataflowSettings.generated.h"
+
+static const FLinearColor CManagedArrayCollectionPinTypeColor = FLinearColor(0.353393f, 0.454175f, 1.0f, 1.0f);
+static const FLinearColor CArrayPinTypeColor = FLinearColor(1.0f, 0.172585f, 0.0f, 1.0f);
+static const FLinearColor CBoxPinTypeColor = FLinearColor(0.013575f, 0.770000f, 0.429609f, 1.0f);
+
+typedef TMap<FName, FNodeColors> FNodeColorsMap;
+
+USTRUCT()
+struct FNodeColors
+{
+	GENERATED_USTRUCT_BODY()
+
+	UPROPERTY(EditAnywhere, Category = Colors)
+	FLinearColor NodeTitleColor = FLinearColor(0.f, 0.f, 0.f);
+
+	UPROPERTY(EditAnywhere, Category = Colors)
+	FLinearColor NodeBodyTintColor = FLinearColor(0.f, 0.f, 0.f);
+};
+
+UCLASS(config = Editor, defaultconfig, meta = (DisplayName = "Dataflow"))
+class DATAFLOWCORE_API UDataflowSettings : public UDeveloperSettings
+{
+	GENERATED_UCLASS_BODY()
+
+	/** TArray<> pin type color. The other pin colors are defined in the general editor settings. */
+	UPROPERTY(config, EditAnywhere, Category = PinColors)
+	FLinearColor ArrayPinTypeColor;
+
+	/** FManagedArrayCollection pin type color. The other pin colors are defined in the general editor settings. */
+	UPROPERTY(config, EditAnywhere, Category = PinColors)
+	FLinearColor ManagedArrayCollectionPinTypeColor;
+
+	/** FBox pin type color. The other pin colors are defined in the general editor settings. */
+	UPROPERTY(config, EditAnywhere, Category = PinColors)
+	FLinearColor BoxPinTypeColor;
+
+	UPROPERTY(config, EditAnywhere, Category = NodeColors)
+	TMap<FName, FNodeColors> NodeColorsMap;
+
+	// Begin UDeveloperSettings Interface
+	virtual FName GetCategoryName() const override;
+
+	FNodeColors RegisterColors(const FName& Category, const FNodeColors& Colors);
+
+	const TMap<FName, FNodeColors>& GetNodeColorsMap() { return NodeColorsMap; }
+
+#if WITH_EDITOR
+	virtual FText GetSectionText() const override;
+#endif
+	// END UDeveloperSettings Interface
+
+#if WITH_EDITOR
+	virtual void PostEditChangeProperty(FPropertyChangedEvent& PropertyChangedEvent) override;
+#endif
+
+public:
+	DECLARE_MULTICAST_DELEGATE_OneParam(FOnDataflowSettingsChanged, const FNodeColorsMap&);
+
+	/** Gets a multicast delegate which is called whenever one of the parameters in this settings object changes. */
+	FOnDataflowSettingsChanged& GetOnDataflowSettingsChangedDelegate() { return OnDataflowSettingsChangedDelegate; }
+
+protected:
+	FOnDataflowSettingsChanged OnDataflowSettingsChangedDelegate;
+
+};
+
diff --git a/Engine/Source/Runtime/Experimental/Dataflow/Engine/Private/Dataflow/DataflowEdNode.cpp b/Engine/Source/Runtime/Experimental/Dataflow/Engine/Private/Dataflow/DataflowEdNode.cpp
--- a/Engine/Source/Runtime/Experimental/Dataflow/Engine/Private/Dataflow/DataflowEdNode.cpp
+++ b/Engine/Source/Runtime/Experimental/Dataflow/Engine/Private/Dataflow/DataflowEdNode.cpp
@@ -19,6 +19,7 @@ void UDataflowEdNode::AllocateDefaultPins()
 {
 	UE_LOG(DATAFLOWNODE_LOG, Verbose, TEXT("UDataflowEdNode::AllocateDefaultPins()"));
 	// called on node creation from UI. 
+
 #if WITH_EDITOR && !(UE_BUILD_SHIPPING || UE_BUILD_TEST)
 	if (DataflowGraph)
 	{
@@ -107,6 +108,7 @@ void UDataflowEdNode::PinConnectionListChanged(UEdGraphPin* Pin)
 
 	Super::PinConnectionListChanged(Pin);
 }
+
 #endif // WITH_EDITOR && !UE_BUILD_SHIPPING
 
 void UDataflowEdNode::Serialize(FArchive& Ar)
@@ -115,6 +117,39 @@ void UDataflowEdNode::Serialize(FArchive& Ar)
 	Ar << DataflowNodeGuid;
 }
 
+#if WITH_EDITOR
+
+FLinearColor UDataflowEdNode::GetNodeTitleColor() const
+{
+	if (DataflowGraph)
+	{
+		if (DataflowNodeGuid.IsValid())
+		{
+			if (TSharedPtr<FDataflowNode> DataflowNode = DataflowGraph->FindBaseNode(DataflowNodeGuid))
+			{
+				return Dataflow::FNodeColorsFactory::Get().GetNodeTitleColor(DataflowNode->GetCategory());
+			}
+		}
+	}
+	return FLinearColor(0.f, 0.f, 0.f);
+}
+
+FLinearColor UDataflowEdNode::GetNodeBodyTintColor() const
+{
+	if (DataflowGraph)
+	{
+		if (DataflowNodeGuid.IsValid())
+		{
+			if (TSharedPtr<FDataflowNode> DataflowNode = DataflowGraph->FindBaseNode(DataflowNodeGuid))
+			{
+				return Dataflow::FNodeColorsFactory::Get().GetNodeBodyTintColor(DataflowNode->GetCategory());
+			}
+		}
+	}
+	return FLinearColor(0.f, 0.f, 0.f);
+}
+
+#endif
 
 #undef LOCTEXT_NAMESPACE
 
diff --git a/Engine/Source/Runtime/Experimental/Dataflow/Engine/Public/Dataflow/DataflowEdNode.h b/Engine/Source/Runtime/Experimental/Dataflow/Engine/Public/Dataflow/DataflowEdNode.h
--- a/Engine/Source/Runtime/Experimental/Dataflow/Engine/Public/Dataflow/DataflowEdNode.h
+++ b/Engine/Source/Runtime/Experimental/Dataflow/Engine/Public/Dataflow/DataflowEdNode.h
@@ -21,7 +21,6 @@ class UObject;
 namespace Dataflow { class FGraph; }
 
 
-
 UCLASS()
 class DATAFLOWENGINE_API UDataflowEdNode : public UEdGraphNode
 {
@@ -30,7 +29,6 @@ class DATAFLOWENGINE_API UDataflowEdNode : public UEdGraphNode
 	FGuid DataflowNodeGuid;
 	TSharedPtr<Dataflow::FGraph> DataflowGraph;
 public:
-
 	// UEdGraphNode interface
 	virtual void AllocateDefaultPins();
 	virtual FText GetNodeTitle(ENodeTitleType::Type TitleType) const;
@@ -47,6 +45,11 @@ public:
 	FGuid GetDataflowNodeGuid() const { return DataflowNodeGuid; }
 	void SetDataflowNodeGuid(FGuid InGuid) { DataflowNodeGuid = InGuid; }
 
+#if WITH_EDITOR
+	virtual FLinearColor GetNodeTitleColor() const override;
+	virtual FLinearColor GetNodeBodyTintColor() const override;
+#endif
+
 	// UObject interface
 	void Serialize(FArchive& Ar);
 	// End UObject interface
