use the controller when removing constraints to ensure it's propagated

diff --git a/Engine/Source/Runtime/Experimental/Animation/Constraints/Private/ConstraintsManager.cpp b/Engine/Source/Runtime/Experimental/Animation/Constraints/Private/ConstraintsManager.cpp
--- a/Engine/Source/Runtime/Experimental/Animation/Constraints/Private/ConstraintsManager.cpp
+++ b/Engine/Source/Runtime/Experimental/Animation/Constraints/Private/ConstraintsManager.cpp
@@ -138,15 +138,23 @@ void UConstraintsManager::OnActorDestroyed(AActor* InActor)
 {
 	if (USceneComponent* SceneComponent = InActor->GetRootComponent())
 	{
-		Constraints.RemoveAll([this, SceneComponent](const TObjectPtr<UTickableConstraint>& Constraint)
+		TArray< int32 > IndicesToRemove;
+		for (int32 Index = 0; Index < Constraints.Num(); ++Index)
 		{
-			const bool bIsRemoved = IsValid(Constraint) && Constraint->ReferencesObject(SceneComponent);
-			if (bIsRemoved)
+			if (IsValid(Constraints[Index]) && Constraints[Index]->ReferencesObject(SceneComponent))
 			{
-				OnConstraintRemoved_BP.Broadcast(this, Constraint);
+				IndicesToRemove.Add(Index);
 			}
-			return bIsRemoved;
-		} );
+		}
+
+		if (!IndicesToRemove.IsEmpty())
+		{
+			FConstraintsManagerController& Controller = FConstraintsManagerController::Get(InActor->GetWorld());
+			for (int32 Index = IndicesToRemove.Num()-1; Index >= 0; --Index)
+			{
+				Controller.RemoveConstraint(Index);
+			}
+		}
 	}
 }
 
