fixed log message when trying to create a cycle

diff --git a/Engine/Source/Runtime/Experimental/Animation/Constraints/Private/TransformConstraint.cpp b/Engine/Source/Runtime/Experimental/Animation/Constraints/Private/TransformConstraint.cpp
--- a/Engine/Source/Runtime/Experimental/Animation/Constraints/Private/TransformConstraint.cpp
+++ b/Engine/Source/Runtime/Experimental/Animation/Constraints/Private/TransformConstraint.cpp
@@ -962,7 +962,8 @@ UTickableTransformConstraint* FTransformConstraintUtils::CreateAndAddFromActors(
 	if (ParentHandle->HasDirectDependencyWith(*ChildHandle))
 	{
 #if WITH_EDITOR
-		UE_LOG(LogTemp, Error, TEXT("%s: %s has a direct dependecy to %s."), *ParentHandle->GetLabel(), *ChildHandle->GetLabel());
+		UE_LOG(LogTemp, Error, TEXT("%s: %s has a direct dependecy to %s."),
+			ErrorPrefix, *ParentHandle->GetLabel(), *ChildHandle->GetLabel());
 #endif
 		bDeleteHandles = true;
 	}
