Iris - NetBlobAssembler error detection.

#rnx

diff --git a/Engine/Plugins/Runtime/ReplicationSystemTestPlugin/Source/Private/Tests/ReplicationSystem/NetBlob/MockNetBlob.cpp b/Engine/Plugins/Runtime/ReplicationSystemTestPlugin/Source/Private/Tests/ReplicationSystem/NetBlob/MockNetBlob.cpp
--- a/Engine/Plugins/Runtime/ReplicationSystemTestPlugin/Source/Private/Tests/ReplicationSystem/NetBlob/MockNetBlob.cpp
+++ b/Engine/Plugins/Runtime/ReplicationSystemTestPlugin/Source/Private/Tests/ReplicationSystem/NetBlob/MockNetBlob.cpp
@@ -148,6 +148,9 @@ UMockSequentialPartialNetBlobHandler::~UMockSequentialPartialNetBlobHandler()
 void UMockSequentialPartialNetBlobHandler::Init(const FSequentialPartialNetBlobHandlerInitParams& InitParams)
 {
 	Super::Init(InitParams);
+	UE::Net::FNetBlobAssemblerInitParams AssemblerInitParams;
+	AssemblerInitParams.PartialNetBlobHandlerConfig = GetConfig();
+	Assembler.Init(AssemblerInitParams);
 }
 
 TRefCountPtr<UE::Net::FNetBlob> UMockSequentialPartialNetBlobHandler::CreateNetBlob(const FNetBlobCreationInfo& CreationInfo) const
diff --git a/Engine/Plugins/Runtime/ReplicationSystemTestPlugin/Source/Private/Tests/ReplicationSystem/NetBlob/MockNetBlob.h b/Engine/Plugins/Runtime/ReplicationSystemTestPlugin/Source/Private/Tests/ReplicationSystem/NetBlob/MockNetBlob.h
--- a/Engine/Plugins/Runtime/ReplicationSystemTestPlugin/Source/Private/Tests/ReplicationSystem/NetBlob/MockNetBlob.h
+++ b/Engine/Plugins/Runtime/ReplicationSystemTestPlugin/Source/Private/Tests/ReplicationSystem/NetBlob/MockNetBlob.h
@@ -81,6 +81,8 @@ public:
 
 	void Init(const FSequentialPartialNetBlobHandlerInitParams& InitParams);
 
+	const USequentialPartialNetBlobHandlerConfig* GetConfig() const;
+
 	const FCallCounts& GetFunctionCallCounts() const { return CallCounts; }
 	void ResetFunctionCallCounts() { CallCounts = FCallCounts({}); }
 
@@ -91,3 +93,8 @@ private:
 	UE::Net::FNetBlobAssembler Assembler;
 	mutable FCallCounts CallCounts;
 };
+
+inline const USequentialPartialNetBlobHandlerConfig* UMockSequentialPartialNetBlobHandler::GetConfig() const
+{
+	return CastChecked<const USequentialPartialNetBlobHandlerConfig>(Super::GetConfig(), ECastCheckedType::NullAllowed);
+}
\ No newline at end of file
diff --git a/Engine/Plugins/Runtime/ReplicationSystemTestPlugin/Source/Private/Tests/ReplicationSystem/NetBlob/TestPartialNetBlob.cpp b/Engine/Plugins/Runtime/ReplicationSystemTestPlugin/Source/Private/Tests/ReplicationSystem/NetBlob/TestPartialNetBlob.cpp
--- a/Engine/Plugins/Runtime/ReplicationSystemTestPlugin/Source/Private/Tests/ReplicationSystem/NetBlob/TestPartialNetBlob.cpp
+++ b/Engine/Plugins/Runtime/ReplicationSystemTestPlugin/Source/Private/Tests/ReplicationSystem/NetBlob/TestPartialNetBlob.cpp
@@ -138,4 +138,40 @@ UE_NET_TEST_FIXTURE(FPartialNetBlobTestFixture, TestSplitBlobIsReceivedDespitePa
 	Server->DestroyObject(ServerObject);
 }
 
+UE_NET_TEST_FIXTURE(FPartialNetBlobTestFixture, PartOfUnexpectedSizeCausesError)
+{
+	FTestContext ServerContext;
+	SetupTestContext(ServerContext, Server);
+
+	TArray<TRefCountPtr<FNetBlob>> PartialNetBlobs;
+	TRefCountPtr<FNetBlob> WrongSizeSecondBlob;
+	{
+		const USequentialPartialNetBlobHandlerConfig* Config = MockSequentialPartialNetBlobHandler->GetConfig();
+
+		// Guesstimate on blob overhead. The important thing is we get exactly as many pars as we wish.
+		constexpr uint32 BlobOverheadBitCount = 128U;
+
+		constexpr int32 PartCount = 3U;
+		uint32 TotalPayloadBitCountForOriginalSplit = PartCount*Config->GetMaxPartBitCount() - BlobOverheadBitCount - 1U;
+		const TRefCountPtr<FNetBlob>& Blob = CreateUnreliableMockNetBlob(TotalPayloadBitCountForOriginalSplit);
+		MockSequentialPartialNetBlobHandler->SplitNetBlob(Blob, PartialNetBlobs);
+		UE_NET_ASSERT_EQ(PartialNetBlobs.Num(), PartCount);
+
+		// Create a second blob with two parts, where the second part will be sent to cause an error.
+		constexpr int32 OtherPartCount = 2U;
+		TArray<TRefCountPtr<FNetBlob>> OtherPartialNetBlobs;
+		const uint32 TotalPayloadBitCountForWrongSizeBlob = OtherPartCount*Config->GetMaxPartBitCount() - BlobOverheadBitCount - 1U;
+		const TRefCountPtr<FNetBlob>& OtherBlob = CreateUnreliableMockNetBlob(TotalPayloadBitCountForWrongSizeBlob);
+		MockSequentialPartialNetBlobHandler->SplitNetBlob(OtherBlob, OtherPartialNetBlobs);
+		UE_NET_ASSERT_EQ(OtherPartialNetBlobs.Num(), OtherPartCount);
+
+		WrongSizeSecondBlob = MoveTemp(OtherPartialNetBlobs[1]);
+	}
+
+
+	ServerContext.SerializationContext.GetNetBlobReceiver()->OnNetBlobReceived(ServerContext.SerializationContext, PartialNetBlobs[0]);
+	ServerContext.SerializationContext.GetNetBlobReceiver()->OnNetBlobReceived(ServerContext.SerializationContext, WrongSizeSecondBlob);
+	UE_NET_ASSERT_TRUE(ServerContext.SerializationContext.HasError());
+}
+
 }
diff --git a/Engine/Source/Runtime/Experimental/Iris/Core/Private/Iris/ReplicationSystem/AttachmentReplication.cpp b/Engine/Source/Runtime/Experimental/Iris/Core/Private/Iris/ReplicationSystem/AttachmentReplication.cpp
--- a/Engine/Source/Runtime/Experimental/Iris/Core/Private/Iris/ReplicationSystem/AttachmentReplication.cpp
+++ b/Engine/Source/Runtime/Experimental/Iris/Core/Private/Iris/ReplicationSystem/AttachmentReplication.cpp
@@ -7,7 +7,7 @@
 #include "Iris/ReplicationSystem/ReplicationSystem.h"
 #include "Iris/ReplicationSystem/ReplicationSystemInternal.h"
 #include "Iris/ReplicationSystem/NetBlob/NetBlobAssembler.h"
-#include "Iris/ReplicationSystem/NetBlob/SequentialPartialNetBlobHandler.h"
+#include "Iris/ReplicationSystem/NetBlob/PartialNetObjectAttachmentHandler.h"
 #include "Iris/Serialization/InternalNetSerializationContext.h"
 #include "Iris/Serialization/NetExportContext.h"
 #include "Iris/Serialization/NetSerializationContext.h"
@@ -602,6 +602,16 @@ FNetObjectAttachmentSendQueue* FNetObjectAttachmentsWriter::GetOrCreateQueue(ENe
 class FNetObjectAttachmentReceiveQueue::FDeferredProcessingQueue
 {
 public:
+	struct FInitParams
+	{
+		const UPartialNetObjectAttachmentHandler* PartialNetObjectAttachmentHandler = nullptr;
+	};
+
+	void Init(const FInitParams& InitParams)
+	{
+		PartialNetObjectAttachmentHandler = InitParams.PartialNetObjectAttachmentHandler;
+	}
+
 	bool IsEmpty() const { return Queue.IsEmpty(); }
 	bool HasUnprocessed() const { return !Queue.IsEmpty(); }
 	bool IsSafeToDestroy() const { return Queue.IsEmpty() && !NetBlobAssembler.IsValid(); }
@@ -612,7 +622,10 @@ public:
 		{
 			if (!NetBlobAssembler.IsValid())
 			{
+				FNetBlobAssemblerInitParams InitParams;
+				InitParams.PartialNetBlobHandlerConfig = PartialNetObjectAttachmentHandler ? PartialNetObjectAttachmentHandler->GetConfig() : nullptr;
 				NetBlobAssembler = MakeUnique<FNetBlobAssembler>();
+				NetBlobAssembler->Init(InitParams);
 			}
 
 			NetBlobAssembler->AddPartialNetBlob(Context, NetHandle, reinterpret_cast<const TRefCountPtr<FPartialNetBlob>&>(NetBlob));
@@ -645,6 +658,7 @@ public:
 private:
 	TUniquePtr<FNetBlobAssembler> NetBlobAssembler;
 	TResizableCircularQueue<TRefCountPtr<FNetBlob>> Queue;
+	const UPartialNetObjectAttachmentHandler* PartialNetObjectAttachmentHandler = nullptr;
 };
 
 FNetObjectAttachmentReceiveQueue::FNetObjectAttachmentReceiveQueue()
@@ -661,6 +675,12 @@ FNetObjectAttachmentReceiveQueue::~FNetObjectAttachmentReceiveQueue()
 	delete DeferredProcessingQueue;
 }
 
+void FNetObjectAttachmentReceiveQueue::Init(const FNetObjectAttachmentReceiveQueueInitParams& InitParams)
+{
+	PartialNetObjectAttachmentHandler = InitParams.PartialNetObjectAttachmentHandler;
+	PartialNetBlobType = (InitParams.PartialNetObjectAttachmentHandler ? InitParams.PartialNetObjectAttachmentHandler->GetNetBlobType() : InvalidNetBlobType);
+}
+
 bool FNetObjectAttachmentReceiveQueue::IsSafeToDestroy() const
 {
 	return UnreliableQueue.IsEmpty() && IsDeferredProcessingQueueSafeToDestroy() && (ReliableQueue == nullptr || ReliableQueue->IsSafeToDestroy());
@@ -782,6 +802,9 @@ uint32 FNetObjectAttachmentReceiveQueue::DeserializeReliable(FNetSerializationCo
 		ReliableQueue = new FReliableNetBlobQueue();
 		checkSlow(DeferredProcessingQueue == nullptr);
 		DeferredProcessingQueue = new FDeferredProcessingQueue();
+		FDeferredProcessingQueue::FInitParams InitParams;
+		InitParams.PartialNetObjectAttachmentHandler = PartialNetObjectAttachmentHandler;
+		DeferredProcessingQueue->Init(InitParams);
 	}
 
 	uint32 DeserializedReliableCount = 0;
@@ -870,7 +893,6 @@ uint32 FNetObjectAttachmentReceiveQueue::DeserializeUnreliable(FNetSerialization
 }
 
 FNetObjectAttachmentsReader::FNetObjectAttachmentsReader()
-: PartialNetBlobType(InvalidNetBlobType)
 {
 }
 
@@ -878,6 +900,11 @@ FNetObjectAttachmentsReader::~FNetObjectAttachmentsReader()
 {
 }
 
+void FNetObjectAttachmentsReader::Init(const FNetObjectAttachmentsReaderInitParams& InitParams)
+{
+	PartialNetObjectAttachmentHandler = InitParams.PartialNetObjectAttachmentHandler;
+}
+
 bool FNetObjectAttachmentsReader::HasUnprocessedAttachments(ENetObjectAttachmentType Type, uint32 ObjectIndex) const
 {
 	const FNetObjectAttachmentReceiveQueue* Queue = GetQueue(Type, ObjectIndex);
@@ -930,7 +957,11 @@ FNetObjectAttachmentReceiveQueue* FNetObjectAttachmentsReader::GetOrCreateQueue(
 
 		QueuePtr = MakeUnique<FNetObjectAttachmentReceiveQueue>();
 		Queue = QueuePtr.Get();
-		Queue->SetPartialNetBlobType(PartialNetBlobType);
+		{
+			FNetObjectAttachmentReceiveQueueInitParams InitParams;
+			InitParams.PartialNetObjectAttachmentHandler = PartialNetObjectAttachmentHandler;
+			Queue->Init(InitParams);
+		}
 
 		if (Type == ENetObjectAttachmentType::OutOfBand)
 		{
@@ -953,7 +984,11 @@ FNetObjectAttachmentReceiveQueue* FNetObjectAttachmentsReader::GetOrCreateQueue(
 		}
 
 		Queue = &ObjectToQueue.Add(ObjectIndex);
-		Queue->SetPartialNetBlobType(PartialNetBlobType);
+		{
+			FNetObjectAttachmentReceiveQueueInitParams InitParams;
+			InitParams.PartialNetObjectAttachmentHandler = PartialNetObjectAttachmentHandler;
+			Queue->Init(InitParams);
+		}
 
 		return Queue;
 	}
diff --git a/Engine/Source/Runtime/Experimental/Iris/Core/Private/Iris/ReplicationSystem/AttachmentReplication.h b/Engine/Source/Runtime/Experimental/Iris/Core/Private/Iris/ReplicationSystem/AttachmentReplication.h
--- a/Engine/Source/Runtime/Experimental/Iris/Core/Private/Iris/ReplicationSystem/AttachmentReplication.h
+++ b/Engine/Source/Runtime/Experimental/Iris/Core/Private/Iris/ReplicationSystem/AttachmentReplication.h
@@ -7,6 +7,7 @@
 #include "Iris/ReplicationSystem/NetBlob/ReliableNetBlobQueue.h"
 #include "Containers/Map.h"
 
+class UPartialNetObjectAttachmentHandler;
 namespace UE::Net::Private
 {
 	class FNetBlobHandlerManager;
@@ -133,13 +134,18 @@ private:
 	TUniquePtr<FNetObjectAttachmentSendQueue> SpecialQueues[uint32(ENetObjectAttachmentType::InternalCount)];
 };
 
+struct FNetObjectAttachmentReceiveQueueInitParams
+{
+	const UPartialNetObjectAttachmentHandler* PartialNetObjectAttachmentHandler = nullptr;
+};
+
 class FNetObjectAttachmentReceiveQueue
 {
 public:
 	FNetObjectAttachmentReceiveQueue();
 	~FNetObjectAttachmentReceiveQueue();
 
-	void SetPartialNetBlobType(FNetBlobType InPartialNetBlobType) { PartialNetBlobType = InPartialNetBlobType; }
+	void Init(const FNetObjectAttachmentReceiveQueueInitParams& InitParams);
 
 	bool IsSafeToDestroy() const;
 	bool HasUnprocessed() const;
@@ -170,6 +176,12 @@ private:
 	FDeferredProcessingQueue* DeferredProcessingQueue;
 	uint32 MaxUnreliableCount;
 	FNetBlobType PartialNetBlobType;
+	const UPartialNetObjectAttachmentHandler* PartialNetObjectAttachmentHandler = nullptr;
+};
+
+struct FNetObjectAttachmentsReaderInitParams
+{
+	const UPartialNetObjectAttachmentHandler* PartialNetObjectAttachmentHandler = nullptr;
 };
 
 class FNetObjectAttachmentsReader
@@ -178,7 +190,7 @@ public:
 	FNetObjectAttachmentsReader();
 	~FNetObjectAttachmentsReader();
 
-	void SetPartialNetBlobType(FNetBlobType InPartialNetBlobType) { PartialNetBlobType = InPartialNetBlobType; }
+	void Init(const FNetObjectAttachmentsReaderInitParams& InitParams);
 
 	bool HasUnprocessedAttachments(ENetObjectAttachmentType Type, uint32 ObjectIndex) const;
 
@@ -195,7 +207,7 @@ private:
 
 	TMap<uint32, FNetObjectAttachmentReceiveQueue> ObjectToQueue;
 	TUniquePtr<FNetObjectAttachmentReceiveQueue> SpecialQueues[uint32(ENetObjectAttachmentType::InternalCount)];
-	FNetBlobType PartialNetBlobType;
+	const UPartialNetObjectAttachmentHandler* PartialNetObjectAttachmentHandler = nullptr;
 };
 
 }
diff --git a/Engine/Source/Runtime/Experimental/Iris/Core/Private/Iris/ReplicationSystem/NetBlob/NetBlobAssembler.cpp b/Engine/Source/Runtime/Experimental/Iris/Core/Private/Iris/ReplicationSystem/NetBlob/NetBlobAssembler.cpp
--- a/Engine/Source/Runtime/Experimental/Iris/Core/Private/Iris/ReplicationSystem/NetBlob/NetBlobAssembler.cpp
+++ b/Engine/Source/Runtime/Experimental/Iris/Core/Private/Iris/ReplicationSystem/NetBlob/NetBlobAssembler.cpp
@@ -2,6 +2,7 @@
 #include "Iris/ReplicationSystem/NetBlob/NetBlobAssembler.h"
 #include "Iris/ReplicationSystem/NetBlob/NetBlobHandler.h"
 #include "Iris/ReplicationSystem/NetBlob/RawDataNetBlob.h"
+#include "Iris/ReplicationSystem/NetBlob/SequentialPartialNetBlobHandler.h"
 #include "Iris/Serialization/NetBitStreamReader.h"
 #include "Iris/Serialization/NetSerializationContext.h"
 
@@ -11,15 +12,25 @@ namespace UE::Net
 static const FName NetError_PartialNetBlobSequenceError("Out of order PartialNetBlob.");
 
 FNetBlobAssembler::FNetBlobAssembler()
-: NetBlobCreationInfo({})
-, NextPartIndex(0)
-, PartCount(0)
-, MaxByteCountPerPart(0)
 {
 }
 
+void FNetBlobAssembler::Init(const FNetBlobAssemblerInitParams& InitParams)
+{
+	PartialNetBlobHandlerConfig = InitParams.PartialNetBlobHandlerConfig;
+	ensureMsgf(PartialNetBlobHandlerConfig != nullptr, TEXT("NetBlobAssembler requires a PartialNetBlobHandlerConfig"));
+}
+
 void FNetBlobAssembler::AddPartialNetBlob(FNetSerializationContext& Context, FNetHandle InNetHandle, const TRefCountPtr<FPartialNetBlob>& PartialNetBlob)
 {
+	bIsReadyToAssemble = false;
+
+	if (!ensureMsgf(PartialNetBlobHandlerConfig != nullptr, TEXT("NetBlobAssembler requires a PartialNetBlobHandlerConfig")))
+	{
+		Context.SetError(GNetError_InvalidValue);
+		return;
+	}
+
 	const uint32 PartIndex = PartialNetBlob->GetPartIndex();
 	if (PartIndex != NextPartIndex)
 	{
@@ -27,8 +38,6 @@ void FNetBlobAssembler::AddPartialNetBlob(FNetSerializationContext& Context, FNe
 		return;
 	}
 
-	++NextPartIndex;
-
 	if (PartIndex == 0U)
 	{
 		PartCount = PartialNetBlob->GetPartCount();
@@ -38,15 +47,29 @@ void FNetBlobAssembler::AddPartialNetBlob(FNetSerializationContext& Context, FNe
 			return;
 		}
 
+		if (PartCount > PartialNetBlobHandlerConfig->GetMaxPartCount())
+		{
+			Context.SetError(GNetError_InvalidValue);
+			return;
+		}
+
+		const uint32 PayloadBitCount = PartialNetBlob->GetPayloadBitCount();
+		// We allow part sizes to be lower than the config value in case it was hotfixed.
+		if (PayloadBitCount > PartialNetBlobHandlerConfig->GetMaxPartBitCount())
+		{
+			Context.SetError(GNetError_InvalidValue);
+			return;
+		}
+
 		NetHandle = InNetHandle;
 
 		NetBlobCreationInfo = PartialNetBlob->GetOriginalCreationInfo();
 
 		// The first part must be greater or equal to any subsequent parts or we will report an error.
-		const uint32 PayloadBitCount = PartialNetBlob->GetPayloadBitCount();
-		MaxByteCountPerPart = (PayloadBitCount + 7U)/8U;
+		FirstPayloadBitCount = PayloadBitCount;
 
 		// Prepare a bitstream that can hold the full blob.
+		const uint32 MaxByteCountPerPart = (PayloadBitCount + 7U)/8U;
 		Payload.SetNumUninitialized((MaxByteCountPerPart*PartCount + 3U)/4U);
 		BitWriter.InitBytes(Payload.GetData(), Payload.Num()*4U);
 
@@ -63,7 +86,8 @@ void FNetBlobAssembler::AddPartialNetBlob(FNetSerializationContext& Context, FNe
 
 		const uint32 PayloadBitCount = PartialNetBlob->GetPayloadBitCount();
 		const uint32 PayloadByteCount = (PayloadBitCount + 7U)/8U;
-		if (PayloadByteCount > MaxByteCountPerPart)
+		// All parts except the last one is expected to match the first part. The last part may be smaller.
+		if (((PayloadBitCount != FirstPayloadBitCount) && ((PartIndex + 1U) < PartCount)) || (PayloadBitCount > FirstPayloadBitCount))
 		{
 			Context.SetError(GNetError_InvalidValue);
 			return;
@@ -72,10 +96,19 @@ void FNetBlobAssembler::AddPartialNetBlob(FNetSerializationContext& Context, FNe
 		// Store partial payload.
 		BitWriter.WriteBitStream(PartialNetBlob->GetPayload(), 0, PayloadBitCount);
 	}
+
+	++NextPartIndex;
+	if (NextPartIndex == PartCount)
+	{
+		bIsReadyToAssemble = true;
+		NextPartIndex = 0;
+	}
 }
 
 TRefCountPtr<FNetBlob> FNetBlobAssembler::Assemble(FNetSerializationContext& Context)
 {
+	bIsReadyToAssemble = false;
+
 	INetBlobReceiver* BlobHandler = Context.GetNetBlobReceiver();
 	checkSlow(BlobHandler != nullptr);
 
diff --git a/Engine/Source/Runtime/Experimental/Iris/Core/Private/Iris/ReplicationSystem/NetBlob/PartialNetObjectAttachmentHandler.h b/Engine/Source/Runtime/Experimental/Iris/Core/Private/Iris/ReplicationSystem/NetBlob/PartialNetObjectAttachmentHandler.h
--- a/Engine/Source/Runtime/Experimental/Iris/Core/Private/Iris/ReplicationSystem/NetBlob/PartialNetObjectAttachmentHandler.h
+++ b/Engine/Source/Runtime/Experimental/Iris/Core/Private/Iris/ReplicationSystem/NetBlob/PartialNetObjectAttachmentHandler.h
@@ -45,5 +45,10 @@ public:
 	/** Splits a RawDataNetBlob. The blob must have been created by a registered NetBlobHandler in order to be reconstructed on the receiving side. */
 	bool SplitRawDataNetBlob(const TRefCountPtr<UE::Net::FRawDataNetBlob>& Blob, TArray<TRefCountPtr<FNetBlob>>& OutPartialBlobs, const UE::Net::FNetDebugName* InDebugName) const;
 
-	const UPartialNetObjectAttachmentHandlerConfig* GetConfig() const { return static_cast<const UPartialNetObjectAttachmentHandlerConfig*>(Super::GetConfig()); }
+	const UPartialNetObjectAttachmentHandlerConfig* GetConfig() const;
 };
+
+inline const UPartialNetObjectAttachmentHandlerConfig* UPartialNetObjectAttachmentHandler::GetConfig() const
+{
+	return CastChecked<const UPartialNetObjectAttachmentHandlerConfig>(Super::GetConfig(), ECastCheckedType::NullAllowed);
+}
diff --git a/Engine/Source/Runtime/Experimental/Iris/Core/Private/Iris/ReplicationSystem/ReplicationReader.cpp b/Engine/Source/Runtime/Experimental/Iris/Core/Private/Iris/ReplicationSystem/ReplicationReader.cpp
--- a/Engine/Source/Runtime/Experimental/Iris/Core/Private/Iris/ReplicationSystem/ReplicationReader.cpp
+++ b/Engine/Source/Runtime/Experimental/Iris/Core/Private/Iris/ReplicationSystem/ReplicationReader.cpp
@@ -149,9 +149,11 @@ void FReplicationReader::Init(const FReplicationParameters& InParameters)
 	ReplicationBridge = Parameters.ReplicationSystem->GetReplicationBridge();
 
 	// Find out if there's a PartialNetObjectAttachmentHandler so we can re-assemble split blobs
-	if (const UNetBlobHandler* Handler = ReplicationSystemInternal->GetNetBlobManager().GetPartialNetObjectAttachmentHandler())
+	if (const UPartialNetObjectAttachmentHandler* Handler = ReplicationSystemInternal->GetNetBlobManager().GetPartialNetObjectAttachmentHandler())
 	{
-		Attachments.SetPartialNetBlobType(Handler->GetNetBlobType());
+		FNetObjectAttachmentsReaderInitParams InitParams;
+		InitParams.PartialNetObjectAttachmentHandler = Handler;
+		Attachments.Init(InitParams);
 	}
 
 	if (const UNetBlobHandler* Handler = ReplicationSystemInternal->GetNetBlobManager().GetNetObjectBlobHandler())
diff --git a/Engine/Source/Runtime/Experimental/Iris/Core/Public/Iris/ReplicationSystem/NetBlob/NetBlobAssembler.h b/Engine/Source/Runtime/Experimental/Iris/Core/Public/Iris/ReplicationSystem/NetBlob/NetBlobAssembler.h
--- a/Engine/Source/Runtime/Experimental/Iris/Core/Public/Iris/ReplicationSystem/NetBlob/NetBlobAssembler.h
+++ b/Engine/Source/Runtime/Experimental/Iris/Core/Public/Iris/ReplicationSystem/NetBlob/NetBlobAssembler.h
@@ -4,9 +4,16 @@
 #include "Iris/ReplicationSystem/NetBlob/PartialNetBlob.h"
 #include "Iris/Serialization/NetBitStreamWriter.h"
 
+class USequentialPartialNetBlobHandlerConfig;
+
 namespace UE::Net
 {
 
+struct FNetBlobAssemblerInitParams
+{
+	const USequentialPartialNetBlobHandlerConfig* PartialNetBlobHandlerConfig = nullptr;
+};
+
 /**
  * Utility class to reassemble split blobs.
  * Partial blobs need to be added in order and represent the same blob until it's assembled or an error occurs.
@@ -16,6 +23,9 @@ class FNetBlobAssembler
 public:
 	IRISCORE_API FNetBlobAssembler();
 
+	/** Initialize the NetBlobAssembler. */
+	IRISCORE_API void Init(const FNetBlobAssemblerInitParams& InitParams);
+
 	/**
 	 * Add the next expected part of the split blob. If it's the first part of a blob the effects of prior calls to this function are discarded.
 	 * @param Context FNetSerializationContext for error reporting. Call HasError() on it afterwards to check whether something went wrong.
@@ -24,7 +34,7 @@ public:
 	IRISCORE_API void AddPartialNetBlob(FNetSerializationContext& Context, FNetHandle NetHandle, const TRefCountPtr<FPartialNetBlob>& PartialNetBlob);
 
 	/** Returns true if all parts of the split blob have been added and is ready to be assembled. */
-	bool IsReadyToAssemble() const { return PartCount > 0 && NextPartIndex == PartCount; }
+	bool IsReadyToAssemble() const { return bIsReadyToAssemble; }
 
 	/**
 	 * Assemble all parts of the split blob.
@@ -40,9 +50,11 @@ private:
 	TArray<uint32> Payload;
 	FNetBitStreamWriter BitWriter;
 	FNetHandle NetHandle;
-	uint32 NextPartIndex;
-	uint32 PartCount;
-	uint32 MaxByteCountPerPart;
+	uint32 NextPartIndex = 0;
+	uint32 PartCount = 0;
+	uint32 FirstPayloadBitCount = 0;
+	bool bIsReadyToAssemble = false;
+	const USequentialPartialNetBlobHandlerConfig* PartialNetBlobHandlerConfig = nullptr;
 };
 
 }
