Dataflow Graph:
- renamed DataflowNodeColorsFactory to DataflowNodeColorsRegistry

diff --git a/Engine/Source/Runtime/Experimental/Dataflow/Core/Private/Dataflow/DataflowNodeColorsFactory.cpp b/Engine/Source/Runtime/Experimental/Dataflow/Core/Private/Dataflow/DataflowNodeColorsFactory.cpp
deleted file mode 100644
diff --git a/Engine/Source/Runtime/Experimental/Dataflow/Core/Private/Dataflow/DataflowNodeColorsRegistry.cpp b/Engine/Source/Runtime/Experimental/Dataflow/Core/Private/Dataflow/DataflowNodeColorsRegistry.cpp
new file mode 100644
--- /dev/null
+++ b/Engine/Source/Runtime/Experimental/Dataflow/Core/Private/Dataflow/DataflowNodeColorsRegistry.cpp
@@ -0,0 +1,130 @@
+// Copyright Epic Games, Inc. All Rights Reserved.
+
+#include "Dataflow/DataflowNodeColorsRegistry.h"
+
+#include "Dataflow/DataflowNode.h"
+//#include "Misc/MessageDialog.h"
+#include "Misc/LazySingleton.h"
+
+namespace Dataflow
+{
+	FNodeColorsRegistry::FNodeColorsRegistry()
+	{
+		DataflowSettings = GetMutableDefault<UDataflowSettings>();
+		DataflowSettingsChangedDelegateHandle = DataflowSettings->GetOnDataflowSettingsChangedDelegate().AddRaw(this, &FNodeColorsRegistry::NodeColorsChangedInSettings);
+
+		const FNodeColorsMap NodeColorsMap = DataflowSettings->GetNodeColorsMap();
+
+		for (auto& Elem : NodeColorsMap)
+		{
+			if (ColorsMap.Contains(Elem.Key))
+			{
+				ColorsMap[Elem.Key] = Elem.Value;
+			}
+			else
+			{
+				ColorsMap.Add(Elem.Key, Elem.Value);
+			}
+		}
+	}
+
+	FNodeColorsRegistry::~FNodeColorsRegistry()
+	{
+		DataflowSettings->GetOnDataflowSettingsChangedDelegate().Remove(DataflowSettingsChangedDelegateHandle);
+	}
+
+	FNodeColorsRegistry& FNodeColorsRegistry::Get()
+	{
+		return TLazySingleton<FNodeColorsRegistry>::Get();
+	}
+
+	void FNodeColorsRegistry::TearDown()
+	{
+		return TLazySingleton<FNodeColorsRegistry>::TearDown();
+	}
+
+	void FNodeColorsRegistry::RegisterNodeColors(const FName& Category, const FNodeColors& NodeColors)
+	{
+		if (!ColorsMap.Contains(Category))
+		{
+			ColorsMap.Add(Category, NodeColors);
+		}
+
+		// Register colors in DataflowSettings
+		GetMutableDefault<UDataflowSettings>()->RegisterColors(Category, NodeColors);
+	}
+
+	FLinearColor FNodeColorsRegistry::GetNodeTitleColor(const FName& Category)
+	{
+		if (ColorsMap.Contains(Category))
+		{
+			return ColorsMap[Category].NodeTitleColor;
+		}
+		else
+		{
+			// Check if any of the parent category registered
+			FString CategoryString = Category.ToString();
+			if (CategoryString.Contains(TEXT("|")))
+			{
+				do {
+					FString LfString, RtString;
+					if (CategoryString.Split(TEXT("|"), &LfString, &RtString, ESearchCase::IgnoreCase, ESearchDir::FromEnd))
+					{
+						if (ColorsMap.Contains(FName(*LfString)))
+						{
+							return ColorsMap[FName(*LfString)].NodeTitleColor;
+						}
+
+						CategoryString = LfString;
+					}
+				} while (CategoryString.Contains(TEXT("|")));
+			}
+		}
+		return FNodeColors().NodeTitleColor;
+	}
+
+	FLinearColor FNodeColorsRegistry::GetNodeBodyTintColor(const FName& Category)
+	{
+		if (ColorsMap.Contains(Category))
+		{
+			return ColorsMap[Category].NodeBodyTintColor;
+		}
+		else
+		{
+			// Check if any of the parent category registered
+			FString CategoryString = Category.ToString();
+			if (CategoryString.Contains(TEXT("|")))
+			{
+				do {
+					FString LfString, RtString;
+					if (CategoryString.Split(TEXT("|"), &LfString, &RtString, ESearchCase::IgnoreCase, ESearchDir::FromEnd))
+					{
+						if (ColorsMap.Contains(FName(*LfString)))
+						{
+							return ColorsMap[FName(*LfString)].NodeBodyTintColor;
+						}
+
+						CategoryString = LfString;
+					}
+				} while (CategoryString.Contains(TEXT("|")));
+			}
+		}
+		return FNodeColors().NodeBodyTintColor;
+	}
+
+	void FNodeColorsRegistry::NodeColorsChangedInSettings(const FNodeColorsMap& NodeColorsMap)
+	{
+		for (auto& Elem : NodeColorsMap)
+		{
+			if (ColorsMap.Contains(Elem.Key))
+			{
+				ColorsMap[Elem.Key] = Elem.Value;
+			}
+			else
+			{
+				ColorsMap.Add(Elem.Key, Elem.Value);
+			}
+		}
+	}
+
+}
\ No newline at end of file
diff --git a/Engine/Source/Runtime/Experimental/Dataflow/Core/Private/Dataflow/DataflowSettings.cpp b/Engine/Source/Runtime/Experimental/Dataflow/Core/Private/Dataflow/DataflowSettings.cpp
--- a/Engine/Source/Runtime/Experimental/Dataflow/Core/Private/Dataflow/DataflowSettings.cpp
+++ b/Engine/Source/Runtime/Experimental/Dataflow/Core/Private/Dataflow/DataflowSettings.cpp
@@ -1,7 +1,6 @@
 // Copyright Epic Games, Inc. All Rights Reserved.
 
 #include "Dataflow/DataflowSettings.h"
-#include "Dataflow/DataflowNodeColorsFactory.h"
 
 UDataflowSettings::UDataflowSettings(const FObjectInitializer& ObjectInitlaizer)
 	: Super(ObjectInitlaizer)
diff --git a/Engine/Source/Runtime/Experimental/Dataflow/Core/Public/Dataflow/DataflowCore.h b/Engine/Source/Runtime/Experimental/Dataflow/Core/Public/Dataflow/DataflowCore.h
--- a/Engine/Source/Runtime/Experimental/Dataflow/Core/Public/Dataflow/DataflowCore.h
+++ b/Engine/Source/Runtime/Experimental/Dataflow/Core/Public/Dataflow/DataflowCore.h
@@ -7,5 +7,5 @@
 #include "Dataflow/DataflowInputOutput.h"
 #include "Dataflow/DataflowNode.h"
 #include "Dataflow/DataflowNodeFactory.h"
-#include "Dataflow/DataflowNodeColorsFactory.h"
+#include "Dataflow/DataflowNodeColorsRegistry.h"
 #include "Dataflow/DataflowNodeParameters.h"
diff --git a/Engine/Source/Runtime/Experimental/Dataflow/Core/Public/Dataflow/DataflowNode.h b/Engine/Source/Runtime/Experimental/Dataflow/Core/Public/Dataflow/DataflowNode.h
--- a/Engine/Source/Runtime/Experimental/Dataflow/Core/Public/Dataflow/DataflowNode.h
+++ b/Engine/Source/Runtime/Experimental/Dataflow/Core/Public/Dataflow/DataflowNode.h
@@ -221,7 +221,7 @@ private:
 
 #define DATAFLOW_NODE_REGISTER_CREATION_FACTORY_NODE_COLORS_BY_CATEGORY(A, C1, C2)	\
 {																					\
-	FNodeColorsFactory::Get().RegisterNodeColors(A, {C1, C2});						\
+	FNodeColorsRegistry::Get().RegisterNodeColors(A, {C1, C2});						\
 }																					\
 
 }
diff --git a/Engine/Source/Runtime/Experimental/Dataflow/Core/Public/Dataflow/DataflowNodeColorsFactory.h b/Engine/Source/Runtime/Experimental/Dataflow/Core/Public/Dataflow/DataflowNodeColorsFactory.h
deleted file mode 100644
diff --git a/Engine/Source/Runtime/Experimental/Dataflow/Core/Public/Dataflow/DataflowNodeColorsRegistry.h b/Engine/Source/Runtime/Experimental/Dataflow/Core/Public/Dataflow/DataflowNodeColorsRegistry.h
new file mode 100644
--- /dev/null
+++ b/Engine/Source/Runtime/Experimental/Dataflow/Core/Public/Dataflow/DataflowNodeColorsRegistry.h
@@ -0,0 +1,42 @@
+// Copyright Epic Games, Inc. All Rights Reserved.
+
+#pragma once
+
+#include "CoreMinimal.h"
+#include "Dataflow/DataflowGraph.h"
+#include "Dataflow/DataflowSettings.h"
+#include "ChaosLog.h"
+
+struct FDataflowNode;
+struct FDataflowConnection;
+	
+class FLazySingleton;
+
+namespace Dataflow
+{
+	//
+	//
+	//
+	class DATAFLOWCORE_API FNodeColorsRegistry
+	{
+	public:
+		static FNodeColorsRegistry& Get();
+		static void TearDown();
+
+		void RegisterNodeColors(const FName& Category, const FNodeColors& NodeColors);
+		FLinearColor GetNodeTitleColor(const FName& Category);
+		FLinearColor GetNodeBodyTintColor(const FName& Category);
+		void NodeColorsChangedInSettings(const FNodeColorsMap& NodeColorsMap);
+
+	private:
+		FNodeColorsRegistry();
+		~FNodeColorsRegistry();
+
+		TMap<FName, FNodeColors > ColorsMap;					// [Category] -> Colors
+		FDelegateHandle DataflowSettingsChangedDelegateHandle;
+		UDataflowSettings* DataflowSettings;
+
+		friend FLazySingleton;
+	};
+
+}
diff --git a/Engine/Source/Runtime/Experimental/Dataflow/Engine/Private/Dataflow/DataflowEdNode.cpp b/Engine/Source/Runtime/Experimental/Dataflow/Engine/Private/Dataflow/DataflowEdNode.cpp
--- a/Engine/Source/Runtime/Experimental/Dataflow/Engine/Private/Dataflow/DataflowEdNode.cpp
+++ b/Engine/Source/Runtime/Experimental/Dataflow/Engine/Private/Dataflow/DataflowEdNode.cpp
@@ -127,7 +127,7 @@ FLinearColor UDataflowEdNode::GetNodeTitleColor() const
 		{
 			if (TSharedPtr<FDataflowNode> DataflowNode = DataflowGraph->FindBaseNode(DataflowNodeGuid))
 			{
-				return Dataflow::FNodeColorsFactory::Get().GetNodeTitleColor(DataflowNode->GetCategory());
+				return Dataflow::FNodeColorsRegistry::Get().GetNodeTitleColor(DataflowNode->GetCategory());
 			}
 		}
 	}
@@ -142,7 +142,7 @@ FLinearColor UDataflowEdNode::GetNodeBodyTintColor() const
 		{
 			if (TSharedPtr<FDataflowNode> DataflowNode = DataflowGraph->FindBaseNode(DataflowNodeGuid))
 			{
-				return Dataflow::FNodeColorsFactory::Get().GetNodeBodyTintColor(DataflowNode->GetCategory());
+				return Dataflow::FNodeColorsRegistry::Get().GetNodeBodyTintColor(DataflowNode->GetCategory());
 			}
 		}
 	}
