Broad Phase AABBTree optimization: Simple sorted insertion in AABBTree
Only SingleRaycast and SingleSweep benefit from this optimization

diff --git a/Engine/Source/Runtime/Experimental/Chaos/Public/Chaos/AABBTree.h b/Engine/Source/Runtime/Experimental/Chaos/Public/Chaos/AABBTree.h
--- a/Engine/Source/Runtime/Experimental/Chaos/Public/Chaos/AABBTree.h
+++ b/Engine/Source/Runtime/Experimental/Chaos/Public/Chaos/AABBTree.h
@@ -2686,13 +2686,48 @@ private:
 			{
 				PHYSICS_CSV_SCOPED_VERY_EXPENSIVE(PhysicsVerbose, NodeTraverse_Branch);
 				int32 Idx = 0;
-				for (const TAABB<T, 3>& AABB : Node.ChildrenBounds)
+				
+				if constexpr(Query != EAABBQueryType::Overlap)
 				{
-					if(TAABBTreeIntersectionHelper<TQueryFastData, Query>::Intersects(Start, CurData, TOI, FAABB3(AABB.Min(), AABB.Max()), QueryBounds, QueryHalfExtents, Dir, InvDir, bParallel))
+					FReal TOI0;
+					FAABB3 AABB0(Node.ChildrenBounds[0].Min(), Node.ChildrenBounds[0].Max());
+					bool bIntersect0 = TAABBTreeIntersectionHelper<TQueryFastData, Query>::Intersects(Start, CurData, TOI0, AABB0, QueryBounds, QueryHalfExtents, Dir, InvDir, bParallel);
+
+					FReal TOI1;
+					FAABB3 AABB1(Node.ChildrenBounds[1].Min(), Node.ChildrenBounds[1].Max());
+					bool bIntersect1 = TAABBTreeIntersectionHelper<TQueryFastData, Query>::Intersects(Start, CurData, TOI1, AABB1, QueryBounds, QueryHalfExtents, Dir, InvDir, bParallel);
+					if (bIntersect0 && bIntersect1)
 					{
-						NodeStack.Add(FNodeQueueEntry{ Node.ChildrenNodes[Idx], TOI });
+						if (TOI1 > TOI0)
+						{
+							NodeStack.Add(FNodeQueueEntry{ Node.ChildrenNodes[1], TOI1 });
+							NodeStack.Add(FNodeQueueEntry{ Node.ChildrenNodes[0], TOI0 });
+						}
+						else
+						{
+							NodeStack.Add(FNodeQueueEntry{ Node.ChildrenNodes[0], TOI0 });
+							NodeStack.Add(FNodeQueueEntry{ Node.ChildrenNodes[1], TOI1 });
+						}
+					}
+					else if (bIntersect0)
+					{
+						NodeStack.Add(FNodeQueueEntry{ Node.ChildrenNodes[0], TOI0 });
+					}
+					else if (bIntersect1)
+					{
+						NodeStack.Add(FNodeQueueEntry{ Node.ChildrenNodes[1], TOI1 });
+					}
+				}
+				else
+				{
+					for (const TAABB<T, 3>&AABB : Node.ChildrenBounds)
+					{
+						if (TAABBTreeIntersectionHelper<TQueryFastData, Query>::Intersects(Start, CurData, TOI, FAABB3(AABB.Min(), AABB.Max()), QueryBounds, QueryHalfExtents, Dir, InvDir, bParallel))
+						{
+							NodeStack.Add(FNodeQueueEntry{ Node.ChildrenNodes[Idx], TOI });
+						}
+						++Idx;
 					}
-					++Idx;
 				}
 			}
 		}
