Optimization : Use only one broad phase AABBTree
BoundingVolume is no longer used (Still could be activated by CVars)

diff --git a/Engine/Source/Runtime/Experimental/Chaos/Public/Chaos/AABBTree.h b/Engine/Source/Runtime/Experimental/Chaos/Public/Chaos/AABBTree.h
--- a/Engine/Source/Runtime/Experimental/Chaos/Public/Chaos/AABBTree.h
+++ b/Engine/Source/Runtime/Experimental/Chaos/Public/Chaos/AABBTree.h
@@ -650,7 +650,7 @@ public:
 	using PayloadType = TPayloadType;
 	static constexpr int D = 3;
 	using TType = T;
-	static constexpr T DefaultMaxPayloadBounds = 100000;
+	static constexpr T DefaultMaxPayloadBounds = 10000;
 	static constexpr int32 DefaultMaxChildrenInLeaf = 12;
 	static constexpr int32 DefaultMaxTreeDepth = 16;
 	static constexpr int32 DefaultMaxNumToProcess = 0; // 0 special value for processing all without timeslicing
diff --git a/Engine/Source/Runtime/Experimental/Chaos/Public/Chaos/PBDRigidsEvolution.h b/Engine/Source/Runtime/Experimental/Chaos/Public/Chaos/PBDRigidsEvolution.h
--- a/Engine/Source/Runtime/Experimental/Chaos/Public/Chaos/PBDRigidsEvolution.h
+++ b/Engine/Source/Runtime/Experimental/Chaos/Public/Chaos/PBDRigidsEvolution.h
@@ -41,6 +41,7 @@ struct FBroadPhaseConfig
 		TreeOfGridAndGrid = 4
 	};
 
+	// Broadphase Type from the enum above
 	int32 BroadphaseType;
 	int32 BVNumCells;
 	int32 MaxChildrenInLeaf;
@@ -52,7 +53,7 @@ struct FBroadPhaseConfig
 
 	FBroadPhaseConfig()
 	{
-		BroadphaseType = 3;
+		BroadphaseType = Tree;
 		BVNumCells = 35;
 		MaxChildrenInLeaf = 5;
 		MaxTreeDepth = 200;
