Chaos : geometry collection : fix removal logic for internal clusters
- better check for internal cluster parent and its state
- refactor logic for both remove on sleep and remove on break
- reverted previous fix that was not accounting for some of special cases ( CL 21715207 )

diff --git a/Engine/Source/Runtime/Experimental/Chaos/Private/GeometryCollectionProxyData.cpp b/Engine/Source/Runtime/Experimental/Chaos/Private/GeometryCollectionProxyData.cpp
--- a/Engine/Source/Runtime/Experimental/Chaos/Private/GeometryCollectionProxyData.cpp
+++ b/Engine/Source/Runtime/Experimental/Chaos/Private/GeometryCollectionProxyData.cpp
@@ -111,6 +111,12 @@ bool FGeometryCollectionDynamicStateFacade::HasBrokenOff(int32 TransformIndex) c
 	return IsActive && (!HasParent) && IsDynamicOrSleeping(TransformIndex);
 }
 
+bool FGeometryCollectionDynamicStateFacade::HasInternalClusterParent(int32 TransformIndex) const
+{
+	const uint8 InternalParenttype = InternalClusterParentTypeAttribute.Get()[TransformIndex];
+	return InternalParenttype != (uint8)Chaos::EInternalClusterType::None;
+}
+
 bool FGeometryCollectionDynamicStateFacade::HasDynamicInternalClusterParent(int32 TransformIndex) const
 {
 	const uint8 InternalParenttype = InternalClusterParentTypeAttribute.Get()[TransformIndex];
diff --git a/Engine/Source/Runtime/Experimental/Chaos/Public/GeometryCollectionProxyData.h b/Engine/Source/Runtime/Experimental/Chaos/Public/GeometryCollectionProxyData.h
--- a/Engine/Source/Runtime/Experimental/Chaos/Public/GeometryCollectionProxyData.h
+++ b/Engine/Source/Runtime/Experimental/Chaos/Public/GeometryCollectionProxyData.h
@@ -114,6 +114,9 @@ public:
 	bool HasBrokenOff(int32 TransformIndex) const;
 
 	/** return true if the transform has an internal cluster parent */
+	bool HasInternalClusterParent(int32 TransformIndex) const;
+
+	/** return true if the transform has an internal cluster parent in a dynamic state */
 	bool HasDynamicInternalClusterParent(int32 TransformIndex) const;
 	
 private:
diff --git a/Engine/Source/Runtime/Experimental/GeometryCollectionEngine/Private/GeometryCollection/GeometryCollectionComponent.cpp b/Engine/Source/Runtime/Experimental/GeometryCollectionEngine/Private/GeometryCollection/GeometryCollectionComponent.cpp
--- a/Engine/Source/Runtime/Experimental/GeometryCollectionEngine/Private/GeometryCollection/GeometryCollectionComponent.cpp
+++ b/Engine/Source/Runtime/Experimental/GeometryCollectionEngine/Private/GeometryCollection/GeometryCollectionComponent.cpp
@@ -3778,11 +3778,6 @@ struct FGeometryCollectionDecayContext
 
 void UGeometryCollectionComponent::UpdateDecay(int32 TransformIdx, float UpdatedDecay, bool bUseClusterCrumbling, bool bHasDynamicInternalClusterParent, FGeometryCollectionDecayContext& ContextInOut)
 {
-	if (bHasDynamicInternalClusterParent && !bUseClusterCrumbling)
-	{
-		return;
-	}
-
 	TManagedArray<float>& Decay = ContextInOut.DecayFacade.DecayAttribute.Modify();
 	if (UpdatedDecay > Decay[TransformIdx])
 	{
@@ -3834,31 +3829,46 @@ void UGeometryCollectionComponent::IncrementSleepTimer(float DeltaTime)
 			&& DynamicStateFacade.IsValid())
 		{
 			FGeometryCollectionDecayContext DecayContext(*PhysicsProxy, DecayFacade);
-			
+
+			const TManagedArray<int32>& OriginalParents = RestCollection->GetGeometryCollection()->Parent;
+
 			TManagedArray<float>& Decay = DecayFacade.DecayAttribute.Modify();
 			for (int32 TransformIdx = 0; TransformIdx < Decay.Num(); ++TransformIdx)
 			{
-				// update the sleeping timer accordingly
-				const bool IsRemovalActive = RemoveOnSleepFacade.IsRemovalActive(TransformIdx);
-				const bool HasBrokenOff = DynamicStateFacade.HasBrokenOff(TransformIdx);
-				const bool HasDynamicInternalClusterParent = DynamicStateFacade.HasDynamicInternalClusterParent(TransformIdx);
-				if (IsRemovalActive && (HasBrokenOff || HasDynamicInternalClusterParent))
+				const bool HasInternalClusterParent = DynamicStateFacade.HasInternalClusterParent(TransformIdx);
+				if (HasInternalClusterParent)
 				{
-					// if decay has started we do not need to check slow moving or sleeping state anymore  
-					bool ShouldUpdateTimer = (Decay[TransformIdx] > 0);
-					if (!ShouldUpdateTimer && RestCollection->bSlowMovingAsSleeping)
+					// this children has an dynamic internal cluster parent so it can't be removed but we need tyo process the internal cluster by looking at the original parent properties
+					const int32 OriginalParentIdx = OriginalParents[TransformIdx];
+					const bool HasDynamicInternalClusterParent = DynamicStateFacade.HasDynamicInternalClusterParent(TransformIdx);
+					if (OriginalParentIdx > INDEX_NONE && HasDynamicInternalClusterParent && RemoveOnSleepFacade.IsRemovalActive(OriginalParentIdx))
 					{
-						const FVector CurrentPosition = DynamicCollection->Transform[TransformIdx].GetTranslation();
-						ShouldUpdateTimer |= RemoveOnSleepFacade.ComputeSlowMovingState(TransformIdx, CurrentPosition, DeltaTime, RestCollection->SlowMovingVelocityThreshold);
+						const bool UseClusterCrumbling = true; // with sleep removal : internal clusters always crumble - this will change when we merge the removal feature together
+						const float UpdatedBreakDecay = UE_SMALL_NUMBER; // since we crumble we can only pass a timy number since this will be ignore ( but need to be >0 to ake sure Update Decay works properly )
+						UpdateDecay(TransformIdx, UpdatedBreakDecay, UseClusterCrumbling, HasDynamicInternalClusterParent, DecayContext);
 					}
-					if (ShouldUpdateTimer || DynamicStateFacade.IsSleeping(TransformIdx))
+				}
+				else if (RemoveOnSleepFacade.IsRemovalActive(TransformIdx) && DynamicStateFacade.HasBrokenOff(TransformIdx))
+				{
+					// root bone should not be affected by remove on sleep 
+					if (OriginalParents[TransformIdx] > INDEX_NONE)
 					{
-						RemoveOnSleepFacade.UpdateSleepTimer(TransformIdx, DeltaTime);
+						// if decay has started we do not need to check slow moving or sleeping state anymore  
+						bool ShouldUpdateTimer = (Decay[TransformIdx] > 0);
+						if (!ShouldUpdateTimer && RestCollection->bSlowMovingAsSleeping)
+						{
+							const FVector CurrentPosition = DynamicCollection->Transform[TransformIdx].GetTranslation();
+							ShouldUpdateTimer |= RemoveOnSleepFacade.ComputeSlowMovingState(TransformIdx, CurrentPosition, DeltaTime, RestCollection->SlowMovingVelocityThreshold);
+						}
+						if (ShouldUpdateTimer || DynamicStateFacade.IsSleeping(TransformIdx))
+						{
+							RemoveOnSleepFacade.UpdateSleepTimer(TransformIdx, DeltaTime);
+						}
+
+						// update the decay and disable the particle when decay has completed 
+						const float UpdatedDecay = RemoveOnSleepFacade.ComputeDecay(TransformIdx);
+						UpdateDecay(TransformIdx, UpdatedDecay, DynamicStateFacade.HasChildren(TransformIdx), false, DecayContext);
 					}
-			 
-					// update the decay and disable the particle when decay has completed 
-					const float UpdatedDecay = RemoveOnSleepFacade.ComputeDecay(TransformIdx);
-					UpdateDecay(TransformIdx, UpdatedDecay, DynamicStateFacade.HasChildren(TransformIdx), HasDynamicInternalClusterParent, DecayContext);
 				}
 			}
 
@@ -3890,19 +3900,25 @@ void UGeometryCollectionComponent::IncrementBreakTimer(float DeltaTime)
 			TManagedArray<float>& Decay = DecayFacade.DecayAttribute.Modify();
 			for (int32 TransformIdx = 0; TransformIdx < Decay.Num(); ++TransformIdx)
 			{
-				const bool HasDynamicInternalClusterParent = DynamicStateFacade.HasDynamicInternalClusterParent(TransformIdx);
-				if (RemoveOnBreakFacade.IsRemovalActive(TransformIdx)
-					&& (DynamicStateFacade.HasBrokenOff(TransformIdx) || HasDynamicInternalClusterParent))
+				const bool HasInternalClusterParent = DynamicStateFacade.HasInternalClusterParent(TransformIdx);
+				if (HasInternalClusterParent)
 				{
-					bool UseClusterCrumbling = RemoveOnBreakFacade.UseClusterCrumbling(TransformIdx);
-					if (HasDynamicInternalClusterParent && OriginalParents[TransformIdx] > INDEX_NONE)
+					// this children has an internal cluster parent so it can't be removed but we need tyo process the internal cluster by looking at the original parent properties
+					const int32 OriginalParentIdx = OriginalParents[TransformIdx];
+					const bool HasDynamicInternalClusterParent = DynamicStateFacade.HasDynamicInternalClusterParent(TransformIdx);
+
+					if (OriginalParentIdx > INDEX_NONE && HasDynamicInternalClusterParent && RemoveOnBreakFacade.IsRemovalActive(OriginalParentIdx))
 					{
-						// look up the state of the original parent to decide if the internal cluster parent should cumble 
-						UseClusterCrumbling = RemoveOnBreakFacade.UseClusterCrumbling(OriginalParents[TransformIdx]);
+						const bool UseClusterCrumbling = RemoveOnBreakFacade.UseClusterCrumbling(OriginalParentIdx);
+						const float UpdatedBreakDecay = RemoveOnBreakFacade.UpdateBreakTimerAndComputeDecay(TransformIdx, DeltaTime);
+						UpdateDecay(TransformIdx, UpdatedBreakDecay, UseClusterCrumbling, HasDynamicInternalClusterParent, DecayContext);
 					}
-					
+				} 
+				else if (RemoveOnBreakFacade.IsRemovalActive(TransformIdx) && DynamicStateFacade.HasBrokenOff(TransformIdx))
+				{
+					const bool UseClusterCrumbling = RemoveOnBreakFacade.UseClusterCrumbling(TransformIdx);
 					const float UpdatedBreakDecay = RemoveOnBreakFacade.UpdateBreakTimerAndComputeDecay(TransformIdx, DeltaTime);
-					UpdateDecay(TransformIdx, UpdatedBreakDecay, UseClusterCrumbling, HasDynamicInternalClusterParent, DecayContext);
+					UpdateDecay(TransformIdx, UpdatedBreakDecay, UseClusterCrumbling, false, DecayContext);
 				}
 			}
 
