Dataflow Graph:
- fixed ERROR: heap-buffer-overflow in Array.h

diff --git a/Engine/Source/Runtime/Experimental/Dataflow/Core/Private/Dataflow/DataflowNodeColorsRegistry.cpp b/Engine/Source/Runtime/Experimental/Dataflow/Core/Private/Dataflow/DataflowNodeColorsRegistry.cpp
--- a/Engine/Source/Runtime/Experimental/Dataflow/Core/Private/Dataflow/DataflowNodeColorsRegistry.cpp
+++ b/Engine/Source/Runtime/Experimental/Dataflow/Core/Private/Dataflow/DataflowNodeColorsRegistry.cpp
@@ -3,14 +3,13 @@
 #include "Dataflow/DataflowNodeColorsRegistry.h"
 
 #include "Dataflow/DataflowNode.h"
-//#include "Misc/MessageDialog.h"
 #include "Misc/LazySingleton.h"
 
 namespace Dataflow
 {
 	FNodeColorsRegistry::FNodeColorsRegistry()
 	{
-		DataflowSettings = GetMutableDefault<UDataflowSettings>();
+		UDataflowSettings* DataflowSettings = GetMutableDefault<UDataflowSettings>();
 		DataflowSettingsChangedDelegateHandle = DataflowSettings->GetOnDataflowSettingsChangedDelegate().AddRaw(this, &FNodeColorsRegistry::NodeColorsChangedInSettings);
 
 		const FNodeColorsMap NodeColorsMap = DataflowSettings->GetNodeColorsMap();
@@ -30,7 +29,11 @@ namespace Dataflow
 
 	FNodeColorsRegistry::~FNodeColorsRegistry()
 	{
-		DataflowSettings->GetOnDataflowSettingsChangedDelegate().Remove(DataflowSettingsChangedDelegateHandle);
+		if (IsClassLoaded<UDataflowSettings>())
+		{
+			UDataflowSettings* DataflowSettings = GetMutableDefault<UDataflowSettings>();
+			DataflowSettings->GetOnDataflowSettingsChangedDelegate().Remove(DataflowSettingsChangedDelegateHandle);
+		}
 	}
 
 	FNodeColorsRegistry& FNodeColorsRegistry::Get()
diff --git a/Engine/Source/Runtime/Experimental/Dataflow/Core/Public/Dataflow/DataflowNodeColorsRegistry.h b/Engine/Source/Runtime/Experimental/Dataflow/Core/Public/Dataflow/DataflowNodeColorsRegistry.h
--- a/Engine/Source/Runtime/Experimental/Dataflow/Core/Public/Dataflow/DataflowNodeColorsRegistry.h
+++ b/Engine/Source/Runtime/Experimental/Dataflow/Core/Public/Dataflow/DataflowNodeColorsRegistry.h
@@ -34,7 +34,6 @@ namespace Dataflow
 
 		TMap<FName, FNodeColors > ColorsMap;					// [Category] -> Colors
 		FDelegateHandle DataflowSettingsChangedDelegateHandle;
-		UDataflowSettings* DataflowSettings;
 
 		friend FLazySingleton;
 	};
