Constraints:
- Look for possible prerequisite if any when the constraint's parent can't tick (static meshes for instance). In a context of adding several constraints, we want to make sure that the evaluation order is the right one
- use the ControlRig dependency map to check for dependencies between controls
- display all constraints in sorted order when the constraints actoor is selected

diff --git a/Engine/Plugins/Animation/ControlRig/Source/ControlRig/Private/Constraints/ControlRigTransformableHandle.cpp b/Engine/Plugins/Animation/ControlRig/Source/ControlRig/Private/Constraints/ControlRigTransformableHandle.cpp
--- a/Engine/Plugins/Animation/ControlRig/Source/ControlRig/Private/Constraints/ControlRigTransformableHandle.cpp
+++ b/Engine/Plugins/Animation/ControlRig/Source/ControlRig/Private/Constraints/ControlRigTransformableHandle.cpp
@@ -183,7 +183,7 @@ bool UTransformableControlHandle::HasDirectDependencyWith(const UTransformableHa
 		}
 	}
 	
-	const FRigControlElement* ControlElement = GetControlElement();
+	FRigControlElement* ControlElement = GetControlElement();
 	if (!ControlElement)
 	{
 		return false;
@@ -191,12 +191,59 @@ bool UTransformableControlHandle::HasDirectDependencyWith(const UTransformableHa
 
 	// check whether the other handle is one of the control parent within the CR hierarchy
 	static constexpr bool bRecursive = true;
-	const FRigBaseElementParentArray AllParents = ControlRig->GetHierarchy()->GetParents(ControlElement, bRecursive);
-	return  AllParents.ContainsByPredicate([this, OtherHash](const FRigBaseElement* Parent)
+	const URigHierarchy* Hierarchy = ControlRig->GetHierarchy();
+	const FRigBaseElementParentArray AllParents = Hierarchy->GetParents(ControlElement, bRecursive);
+	const bool bIsParent = AllParents.ContainsByPredicate([this, OtherHash](const FRigBaseElement* Parent)
 	{
 		const uint32 ParentHash = ComputeHash(ControlRig.Get(), Parent->GetName());
 		return ParentHash == OtherHash;		
 	});
+
+	if (bIsParent)
+	{
+		return true;
+	}
+
+	// otherwise, check if there are any dependency in the graph that would cause a cycle
+	const TArray<FRigControlElement*> AllControls = Hierarchy->GetControls();
+	const int32 IndexOfPossibleParent = AllControls.IndexOfByPredicate([this, OtherHash](const FRigBaseElement* Parent)
+	{
+		const uint32 ChildHash = ComputeHash(ControlRig.Get(), Parent->GetName());
+		return ChildHash == OtherHash;
+	});
+
+	if (IndexOfPossibleParent != INDEX_NONE)
+	{
+		FRigControlElement* Parent = AllControls[IndexOfPossibleParent];
+
+		FString FailureReason;
+
+#if WITH_EDITOR
+		const URigHierarchy::TElementDependencyMap DependencyMap = Hierarchy->GetDependenciesForVM(ControlRig->GetVM());
+		const bool bIsParentedTo = Hierarchy->IsParentedTo(ControlElement, Parent, DependencyMap);
+#else
+		const bool bIsParentedTo = Hierarchy->IsParentedTo(ControlElement, Parent);
+#endif
+		
+		if (bIsParentedTo)
+		{
+			return true;
+		}
+	}
+
+	return false;
+}
+
+// if there's no skeletal mesh bound then the handle is not valid so no need to do anything else
+FTickPrerequisite UTransformableControlHandle::GetPrimaryPrerequisite() const
+{
+	if (FTickFunction* TickFunction = GetTickFunction())
+	{
+		return FTickPrerequisite( GetSkeletalMesh(), *TickFunction); 
+	}
+	
+	static const FTickPrerequisite DummyPrerex;
+	return DummyPrerex;
 }
 
 FRigControlElement* UTransformableControlHandle::GetControlElement() const
diff --git a/Engine/Plugins/Animation/ControlRig/Source/ControlRig/Public/Constraints/ControlRigTransformableHandle.h b/Engine/Plugins/Animation/ControlRig/Source/ControlRig/Public/Constraints/ControlRigTransformableHandle.h
--- a/Engine/Plugins/Animation/ControlRig/Source/ControlRig/Public/Constraints/ControlRigTransformableHandle.h
+++ b/Engine/Plugins/Animation/ControlRig/Source/ControlRig/Public/Constraints/ControlRigTransformableHandle.h
@@ -75,6 +75,9 @@ public:
 	/** Check for direct dependencies (ie hierarchy + skeletal mesh) with InOther. */
 	virtual bool HasDirectDependencyWith(const UTransformableHandle& InOther) const override;
 
+	/** Look for a possible tick function that can be used as a prerequisite. */
+	virtual FTickPrerequisite GetPrimaryPrerequisite() const override;
+
 #if WITH_EDITOR
 	/** Returns labels used for UI. */
 	virtual FString GetLabel() const override;
diff --git a/Engine/Plugins/Animation/ControlRig/Source/ControlRigEditor/Private/Editor/Constraints/SConstraintsWidget.cpp b/Engine/Plugins/Animation/ControlRig/Source/ControlRigEditor/Private/Editor/Constraints/SConstraintsWidget.cpp
--- a/Engine/Plugins/Animation/ControlRig/Source/ControlRigEditor/Private/Editor/Constraints/SConstraintsWidget.cpp
+++ b/Engine/Plugins/Animation/ControlRig/Source/ControlRigEditor/Private/Editor/Constraints/SConstraintsWidget.cpp
@@ -781,7 +781,7 @@ void SConstraintsEditionWidget::RemoveItem(const TSharedPtr<FEditableConstraintI
 	FScopedTransaction Transaction(LOCTEXT("RemoveConstraint", "Remove Constraint"));
 	UWorld* World = GCurrentLevelEditingViewportClient->GetWorld();
 	const FConstraintsManagerController& Controller = FConstraintsManagerController::Get(World);
-		
+	
 	const FName ConstraintName(Item->Name.ToString());
 	Controller.RemoveConstraint(ConstraintName);
 
@@ -813,7 +813,8 @@ int32 SConstraintsEditionWidget::RefreshConstraintList()
 	if (bIsConstraintsActor)
 	{
 		const FConstraintsManagerController& Controller = FConstraintsManagerController::Get(World);
-		Constraints = Controller.GetConstraintsArray();
+		static constexpr bool bSorted = true;
+		Constraints = Controller.GetAllConstraints(bSorted);
 	}
 	else
 	{
diff --git a/Engine/Source/Runtime/Experimental/Animation/Constraints/Private/ConstraintsManager.cpp b/Engine/Source/Runtime/Experimental/Animation/Constraints/Private/ConstraintsManager.cpp
--- a/Engine/Source/Runtime/Experimental/Animation/Constraints/Private/ConstraintsManager.cpp
+++ b/Engine/Source/Runtime/Experimental/Animation/Constraints/Private/ConstraintsManager.cpp
@@ -108,6 +108,16 @@ void UTickableConstraint::PostEditChangeProperty(FPropertyChangedEvent& Property
 	}
 }
 
+void UTickableConstraint::PostEditUndo()
+{
+	Super::PostEditUndo();
+
+	// make sure we update ticking when undone
+	const bool bActiveTick = IsValid(this);
+	ConstraintTick.SetTickFunctionEnable(bActiveTick);
+	ConstraintTick.bCanEverTick = bActiveTick;
+}
+
 #endif
 
 /** 
@@ -427,8 +437,6 @@ bool FConstraintsManagerController::RemoveConstraint(const int32 InConstraintInd
 		return false;
 	}
 
-	Manager->Modify();
-
 	const FName ConstraintName = Manager->Constraints[InConstraintIndex]->GetFName();
 	UTickableConstraint* Constraint = Manager->Constraints[InConstraintIndex];
 	
@@ -436,6 +444,7 @@ bool FConstraintsManagerController::RemoveConstraint(const int32 InConstraintInd
 	ConstraintRemoved.Broadcast(ConstraintName);
 	Manager->OnConstraintRemoved_BP.Broadcast(Manager, Constraint);
 
+	Manager->Modify();
 	Manager->Constraints[InConstraintIndex]->ConstraintTick.UnRegisterTickFunction();
 	Manager->Constraints.RemoveAt(InConstraintIndex);
 
@@ -554,4 +563,40 @@ const TArray< TObjectPtr<UTickableConstraint> >& FConstraintsManagerController::
 	}
 
 	return Manager->Constraints;
+}
+
+TArray< TObjectPtr<UTickableConstraint> > FConstraintsManagerController::GetAllConstraints(const bool bSorted) const
+{
+	UConstraintsManager* Manager = FindManager();
+	if (!Manager)
+	{
+		static const TArray< TObjectPtr<UTickableConstraint> > Empty;
+		return Empty;
+	}
+
+	if (!bSorted)
+	{
+		return Manager->Constraints;
+	}
+
+	using ConstraintPtr = TObjectPtr<UTickableConstraint>;
+	TArray< TObjectPtr<UTickableConstraint> > SortedConstraints(Manager->Constraints);
+	// LHS ticks before RHS = LHS is a prerex of RHS 
+	auto TicksBefore = [](const UTickableConstraint& LHS, const UTickableConstraint& RHS)
+	{
+		const TArray<FTickPrerequisite>& RHSPrerex = RHS.ConstraintTick.GetPrerequisites();
+		const FConstraintTickFunction& LHSTickFunction = LHS.ConstraintTick;
+		const bool bIsLHSAPrerexOfRHS = RHSPrerex.ContainsByPredicate([&LHSTickFunction](const FTickPrerequisite& Prerex)
+		{
+			return Prerex.PrerequisiteTickFunction == &LHSTickFunction;
+		});
+		return bIsLHSAPrerexOfRHS;
+	};
+	
+	Algo::StableSort(SortedConstraints, [TicksBefore](const ConstraintPtr& LHS, const ConstraintPtr& RHS)
+	{
+		return TicksBefore(*LHS, *RHS);
+	});
+	
+	return SortedConstraints;
 }
\ No newline at end of file
diff --git a/Engine/Source/Runtime/Experimental/Animation/Constraints/Private/TransformConstraint.cpp b/Engine/Source/Runtime/Experimental/Animation/Constraints/Private/TransformConstraint.cpp
--- a/Engine/Source/Runtime/Experimental/Animation/Constraints/Private/TransformConstraint.cpp
+++ b/Engine/Source/Runtime/Experimental/Animation/Constraints/Private/TransformConstraint.cpp
@@ -916,6 +916,92 @@ UTickableTransformConstraint* FTransformConstraintUtils::CreateFromType(
 	return Constraint;
 }
 
+namespace
+{
+
+// we suppose that both InParentHandle and InChildHandle are safe to use
+bool HasConstraintDependencyWith(UWorld* InWorld, const UTransformableHandle* InParentHandle, const UTransformableHandle* InChildHandle)
+{
+	const FConstraintsManagerController& Controller = FConstraintsManagerController::Get(InWorld);
+	static constexpr bool bSorted = false;
+
+	using HandlePtr = TObjectPtr<UTransformableHandle>;
+	TArray<TObjectPtr<UTransformableHandle>> ParentHandles;
+	
+	using ConstraintPtr = TObjectPtr<UTickableConstraint>;
+	const TArray<ConstraintPtr> Constraints = Controller.GetParentConstraints(InParentHandle->GetHash(), bSorted);
+
+	// get parent handles
+	for (const ConstraintPtr& Constraint: Constraints)
+	{
+		if (const UTickableTransformConstraint* TransformConstraint = Cast<UTickableTransformConstraint>(Constraint.Get()))
+		{
+			if (TransformConstraint->ParentTRSHandle)
+			{
+				ParentHandles.Add(TransformConstraint->ParentTRSHandle);
+			}
+		}
+	}
+
+	// check if InChildHandle is one of them
+	const uint32 ChildHash = InChildHandle->GetHash();
+	const bool bIsParentADependency = ParentHandles.ContainsByPredicate([ChildHash](const HandlePtr& InHandle)
+	{
+		return InHandle->GetHash() == ChildHash;
+	});
+
+	if (bIsParentADependency)
+	{
+		return true;
+	}
+
+	// if not, recurse
+	for (const HandlePtr& ParentHandle: ParentHandles)
+	{
+		if (HasConstraintDependencyWith(InWorld, ParentHandle, InChildHandle))
+		{
+			return true;
+		}
+	}
+
+	return false;
+}
+	
+bool AreHandlesConstrainable( UWorld* InWorld, UTransformableHandle* InParentHandle, UTransformableHandle* InChildHandle)
+{
+	static const TCHAR* ErrorPrefix = TEXT("Dependency error:");
+
+	if (InChildHandle->GetHash() == InParentHandle->GetHash())
+	{
+		UE_LOG(LogTemp, Error, TEXT("%s handles are pointing at the same object."), ErrorPrefix);
+		return false;
+	}
+
+	// check for direct transform dependencies (ei hierarchy)
+	if (InParentHandle->HasDirectDependencyWith(*InChildHandle))
+	{
+#if WITH_EDITOR
+		UE_LOG(LogTemp, Error, TEXT("%s: %s has a direct dependency with %s."),
+			ErrorPrefix, *InParentHandle->GetLabel(), *InChildHandle->GetLabel());
+#endif
+		return false;
+	}
+
+	// check for indirect transform dependencies (ei constraint chain)
+	if (HasConstraintDependencyWith(InWorld, InParentHandle, InChildHandle))
+	{
+#if WITH_EDITOR
+		UE_LOG(LogTemp, Error, TEXT("%s: %s has an indirect dependency with %s."),
+			ErrorPrefix, *InParentHandle->GetLabel(), *InChildHandle->GetLabel());
+#endif
+		return false;
+	}
+
+	return true;
+}
+
+}
+
 UTickableTransformConstraint* FTransformConstraintUtils::CreateAndAddFromActors(
 	UWorld* InWorld,
 	AActor* InParent,
@@ -952,23 +1038,8 @@ UTickableTransformConstraint* FTransformConstraintUtils::CreateAndAddFromActors(
 		return nullptr;
 	}
 
-	bool bDeleteHandles = false;
-	if (ChildHandle->GetHash() == ParentHandle->GetHash())
-	{
-		UE_LOG(LogTemp, Error, TEXT("%s handles are pointing at the same object."), ErrorPrefix);
-		bDeleteHandles = true;
-	}
-	
-	if (ParentHandle->HasDirectDependencyWith(*ChildHandle))
-	{
-#if WITH_EDITOR
-		UE_LOG(LogTemp, Error, TEXT("%s: %s has a direct dependecy to %s."),
-			ErrorPrefix, *ParentHandle->GetLabel(), *ChildHandle->GetLabel());
-#endif
-		bDeleteHandles = true;
-	}
-	
-	if (bDeleteHandles)
+	const bool bCanConstrain = AreHandlesConstrainable(InWorld, ParentHandle, ChildHandle);
+	if (!bCanConstrain)
 	{
 		ChildHandle->MarkAsGarbage();
 		ParentHandle->MarkAsGarbage();
@@ -1008,7 +1079,18 @@ bool FTransformConstraintUtils::AddConstraint(
 		return false;
 	}
 
+	// store previous child constraints
+	const FConstraintsManagerController& Controller = FConstraintsManagerController::Get(InWorld);
+	const TArray<TObjectPtr<UTickableConstraint>> ChildParentConstraints = Controller.GetParentConstraints(InChildHandle->GetHash(), true);
 
+	// add the new one
+	const bool bConstraintAdded = Controller.AddConstraint(Constraint);
+	if (!bConstraintAdded)
+	{
+		return false;
+	}
+
+	// setup the constraint
 	auto SetupConstraint = [InParentHandle, InChildHandle, bMaintainOffset](UTickableTransformConstraint* InConstraint)
 	{
 		InConstraint->ParentTRSHandle = InParentHandle;
@@ -1017,12 +1099,44 @@ bool FTransformConstraintUtils::AddConstraint(
 		InConstraint->Setup();
 	};
 	
-	const FConstraintsManagerController& Controller = FConstraintsManagerController::Get(InWorld);
-	
-	Controller.AddConstraint(Constraint);
-
 	SetupConstraint(Constraint);
 
+	// add dependencies with the last child constraint
+	const FName NewConstraintName = Constraint->GetFName();
+	if (!ChildParentConstraints.IsEmpty())
+	{
+		const FName LastChildConstraintName = ChildParentConstraints.Last()->GetFName();
+		Controller.SetConstraintsDependencies( LastChildConstraintName, NewConstraintName);
+	}
+
+	// make sure we tick after the parent.
+	const FTickFunction* ParentTickFunction = InParentHandle->GetTickFunction();
+	const FTickFunction* ChildTickFunction = InChildHandle->GetTickFunction();
+	if (ParentTickFunction && (ChildTickFunction != ParentTickFunction))
+	{
+		const TArray<FTickPrerequisite>& ParentPrerequisites =  Constraint->ConstraintTick.GetPrerequisites();
+		if (ParentPrerequisites.IsEmpty())
+		{
+			// if the constraint has no prerex at this stage, this means that the parent tick function
+			// is not registered or can't tick (static meshes for instance) so look for the first master tick function if any.
+			// In a context of adding several constraints, we want to make sure that the evaluation order is the right one
+			FTickPrerequisite PrimaryPrerex = InParentHandle->GetPrimaryPrerequisite();
+			if (FTickFunction* PotentialFunction = PrimaryPrerex.Get())
+			{
+				UObject* Target = PrimaryPrerex.PrerequisiteObject.Get();
+				Constraint->ConstraintTick.AddPrerequisite(Target, *PotentialFunction);
+			}
+		}
+	}
+
+	// if child handle is the parent of some other constraints, ensure they will tick after that new one
+	TArray<TObjectPtr<UTickableConstraint>> ChildChildConstraints;
+	GetChildrenConstraints(InWorld, InChildHandle, ChildChildConstraints);
+	for (const TObjectPtr<UTickableConstraint>& ChildConstraint: ChildChildConstraints)
+	{
+		Controller.SetConstraintsDependencies( NewConstraintName, ChildConstraint->GetFName());
+	}
+	
 	return true;
 }
 
@@ -1123,4 +1237,36 @@ TOptional<FTransform> FTransformConstraintUtils::GetConstraintRelativeTransform(
 	const FTransform ParentGlobal = Constraint->GetParentGlobalTransform();
 
 	return ComputeRelativeTransform(ChildLocal, ChildGlobal, ParentGlobal, Constraint);
+}
+
+void FTransformConstraintUtils::GetChildrenConstraints(
+	UWorld* World,
+	const UTransformableHandle* InParentHandle,
+	TArray< TObjectPtr<UTickableConstraint> >& OutConstraints)
+{
+	const uint32 ParentHash = InParentHandle->GetHash();
+
+	const FConstraintsManagerController& Controller = FConstraintsManagerController::Get(World);
+
+	using ConstraintPtr = TObjectPtr<UTickableConstraint>;
+	const TArray<ConstraintPtr>& AllConstraints = Controller.GetConstraintsArray();
+
+	const TArray<ConstraintPtr> FilteredConstraints =
+	AllConstraints.FilterByPredicate( [ParentHash](const ConstraintPtr& Constraint)
+	{
+		const UTickableTransformConstraint* TransformConstraint = Cast<UTickableTransformConstraint>(Constraint.Get());
+		if (!TransformConstraint)
+		{
+			return false;
+		}
+		
+		if (TransformConstraint->ParentTRSHandle && (TransformConstraint->ParentTRSHandle->GetHash() == ParentHash))
+		{
+			return true;
+		}
+
+		return false;
+	} );
+
+	OutConstraints.Append(FilteredConstraints);
 }
\ No newline at end of file
diff --git a/Engine/Source/Runtime/Experimental/Animation/Constraints/Private/TransformableHandle.cpp b/Engine/Source/Runtime/Experimental/Animation/Constraints/Private/TransformableHandle.cpp
--- a/Engine/Source/Runtime/Experimental/Animation/Constraints/Private/TransformableHandle.cpp
+++ b/Engine/Source/Runtime/Experimental/Animation/Constraints/Private/TransformableHandle.cpp
@@ -108,9 +108,9 @@ bool UTransformableComponentHandle::HasDirectDependencyWith(const UTransformable
 		return false;
 	}
 
-	// check whether the other handle is one of the component's parent
 	if (Component.IsValid())
 	{
+		// check whether the other handle is one of the component's parent
 		for (const USceneComponent* Comp=Component->GetAttachParent(); Comp!=nullptr; Comp=Comp->GetAttachParent() )
 		{
 			const uint32 AttachParentHash = GetTypeHash(Comp);
@@ -120,10 +120,88 @@ bool UTransformableComponentHandle::HasDirectDependencyWith(const UTransformable
 			}
 		}
 	}
-	
+
 	return false;
 }
 
+namespace
+{
+	
+FTickPrerequisite LookForPrimaryPrerequisite(USceneComponent* Component)
+{
+	static const FTickPrerequisite DummyPrerex;
+	
+	if (!Component)
+	{
+		return DummyPrerex;
+	}
+
+	auto IsValidTickFunction = [](const FTickFunction* InTickFunction)
+	{
+		if (InTickFunction)
+		{
+			return (InTickFunction->bCanEverTick || InTickFunction->IsTickFunctionRegistered());
+		};
+
+		return false;
+	};
+
+	// check if this can tick
+	if ( IsValidTickFunction(&Component->PrimaryComponentTick) )
+	{
+		return FTickPrerequisite(Component, Component->PrimaryComponentTick);
+	}
+
+	// check if a master constraint can tick
+	if (UWorld* World = Component->GetWorld())
+	{
+		static constexpr bool bSorted = true;
+		
+		const FConstraintsManagerController& Controller = FConstraintsManagerController::Get(World);
+		const TArray<TObjectPtr<UTickableConstraint>> ParentConstraints =
+			Controller.GetParentConstraints(GetTypeHash(Component), bSorted);
+		
+		for (int Index = ParentConstraints.Num()-1; Index >= 0; Index--)
+		{
+			if (UTickableConstraint* Constraint = ParentConstraints[Index])
+			{
+				if(IsValidTickFunction(&Constraint->ConstraintTick))
+				{
+					return FTickPrerequisite(Constraint->GetOuter(), Constraint->ConstraintTick);
+				}
+			}
+		}	
+	}
+
+	// check if parent ticks
+	if (USceneComponent* Parent = Component->GetAttachParent())
+	{
+		FTickFunction& ParentTickFunction = Parent->PrimaryComponentTick;
+		if (IsValidTickFunction(&ParentTickFunction))
+		{
+			return FTickPrerequisite(Parent, ParentTickFunction);
+		}
+
+		// get up the hierarchy if not found
+		return LookForPrimaryPrerequisite(Parent);
+	}
+
+	return DummyPrerex;
+}
+	
+}
+
+FTickPrerequisite UTransformableComponentHandle::GetPrimaryPrerequisite() const
+{
+	if (Component.IsValid())
+	{
+		return LookForPrimaryPrerequisite(Component.Get());
+	}
+
+	static const FTickPrerequisite DummyPrerex;
+	return DummyPrerex;
+}
+
 void UTransformableComponentHandle::UnregisterDelegates() const
 {
 #if WITH_EDITOR
diff --git a/Engine/Source/Runtime/Experimental/Animation/Constraints/Public/ConstraintsManager.h b/Engine/Source/Runtime/Experimental/Animation/Constraints/Public/ConstraintsManager.h
--- a/Engine/Source/Runtime/Experimental/Animation/Constraints/Public/ConstraintsManager.h
+++ b/Engine/Source/Runtime/Experimental/Animation/Constraints/Public/ConstraintsManager.h
@@ -102,6 +102,7 @@ public:
 
 	// UObject interface
 	virtual void PostEditChangeProperty(FPropertyChangedEvent& PropertyChangedEvent) override;
+	virtual void PostEditUndo() override;
 	// End of UObject interface
 #endif
 	/** @todo documentation. */
@@ -228,6 +229,9 @@ public:
 	/** Get read-only access to the array of constraints. */
 	const TArray< TObjectPtr<UTickableConstraint> >& GetConstraintsArray() const;
 
+	/** @todo document */
+	TArray< TObjectPtr<UTickableConstraint> > GetAllConstraints(const bool bSorted = false) const;
+
 	/** Get parent constraints of the specified child. If bSorted is true, then the constraints will be sorted by dependency. */
 	TArray< TObjectPtr<UTickableConstraint> > GetParentConstraints(
 		const uint32 InTargetHash,
diff --git a/Engine/Source/Runtime/Experimental/Animation/Constraints/Public/TransformConstraint.h b/Engine/Source/Runtime/Experimental/Animation/Constraints/Public/TransformConstraint.h
--- a/Engine/Source/Runtime/Experimental/Animation/Constraints/Public/TransformConstraint.h
+++ b/Engine/Source/Runtime/Experimental/Animation/Constraints/Public/TransformConstraint.h
@@ -333,7 +333,7 @@ struct CONSTRAINTS_API FTransformConstraintUtils
 		UWorld* World,
 		const AActor* InChild,
 		TArray< TObjectPtr<UTickableConstraint> >& OutConstraints);
-		
+
 	/** Create a handle for the scene component.*/
 	static UTransformableComponentHandle* CreateHandleForSceneComponent(
 		USceneComponent* InSceneComponent,
@@ -371,4 +371,10 @@ struct CONSTRAINTS_API FTransformConstraintUtils
 
 	/** Computes the current constraint space local transform. */
 	static TOptional<FTransform> GetConstraintRelativeTransform(UWorld* InWorld, const uint32 InHandleHash);
+
+	/** Fills a constraint array that InParentHandle is the parent of. */
+	static void GetChildrenConstraints(
+		UWorld* World,
+		const UTransformableHandle* InParentHandle,
+		TArray< TObjectPtr<UTickableConstraint> >& OutConstraints);
 };
diff --git a/Engine/Source/Runtime/Experimental/Animation/Constraints/Public/TransformableHandle.h b/Engine/Source/Runtime/Experimental/Animation/Constraints/Public/TransformableHandle.h
--- a/Engine/Source/Runtime/Experimental/Animation/Constraints/Public/TransformableHandle.h
+++ b/Engine/Source/Runtime/Experimental/Animation/Constraints/Public/TransformableHandle.h
@@ -3,6 +3,7 @@
 #pragma once
 
 #include "UObject/Object.h"
+#include "Engine/EngineBaseTypes.h"
 
 #include "TransformableHandle.generated.h"
 
@@ -86,6 +87,9 @@ public:
 
 	/** Check for direct dependencies with InOther. */
 	virtual bool HasDirectDependencyWith(const UTransformableHandle& InOther) const PURE_VIRTUAL(HasDirectDependencyWith, return false;);
+
+	/** Look for a possible tick function that can be used as a prerequisite. */
+	virtual FTickPrerequisite GetPrimaryPrerequisite() const PURE_VIRTUAL(GetPrimaryPrerequisite, return FTickPrerequisite(););
 	
 	FHandleModifiedEvent& HandleModified();
 
@@ -137,6 +141,9 @@ public:
 	/** Check for direct dependencies (ie hierarchy) with InOther. */
 	virtual bool HasDirectDependencyWith(const UTransformableHandle& InOther) const override;
 
+	/** Look for a possible tick function that can be used as a prerequisite. */
+	virtual FTickPrerequisite GetPrimaryPrerequisite() const override;
+
 	/** Get the array of float channels for the specified section*/
 	virtual TArrayView<FMovieSceneFloatChannel*>  GetFloatChannels(const UMovieSceneSection* InSection) const override;
 	/** Get the array of double channels for the specified section*/
