fixed unacceptable words

diff --git a/Engine/Source/Runtime/Experimental/Animation/Constraints/Private/TransformConstraint.cpp b/Engine/Source/Runtime/Experimental/Animation/Constraints/Private/TransformConstraint.cpp
--- a/Engine/Source/Runtime/Experimental/Animation/Constraints/Private/TransformConstraint.cpp
+++ b/Engine/Source/Runtime/Experimental/Animation/Constraints/Private/TransformConstraint.cpp
@@ -1118,7 +1118,7 @@ bool FTransformConstraintUtils::AddConstraint(
 		if (ParentPrerequisites.IsEmpty())
 		{
 			// if the constraint has no prerex at this stage, this means that the parent tick function
-			// is not registered or can't tick (static meshes for instance) so look for the first master tick function if any.
+			// is not registered or can't tick (static meshes for instance) so look for the first parent tick function if any.
 			// In a context of adding several constraints, we want to make sure that the evaluation order is the right one
 			FTickPrerequisite PrimaryPrerex = InParentHandle->GetPrimaryPrerequisite();
 			if (FTickFunction* PotentialFunction = PrimaryPrerex.Get())
diff --git a/Engine/Source/Runtime/Experimental/Animation/Constraints/Private/TransformableHandle.cpp b/Engine/Source/Runtime/Experimental/Animation/Constraints/Private/TransformableHandle.cpp
--- a/Engine/Source/Runtime/Experimental/Animation/Constraints/Private/TransformableHandle.cpp
+++ b/Engine/Source/Runtime/Experimental/Animation/Constraints/Private/TransformableHandle.cpp
@@ -152,7 +152,7 @@ FTickPrerequisite LookForPrimaryPrerequisite(USceneComponent* Component)
 		return FTickPrerequisite(Component, Component->PrimaryComponentTick);
 	}
 
-	// check if a master constraint can tick
+	// check if a parent constraint can tick
 	if (UWorld* World = Component->GetWorld())
 	{
 		static constexpr bool bSorted = true;
