Refactored Chaos code to reduce templated code bloat

diff --git a/Engine/Source/Runtime/Experimental/Chaos/Private/Chaos/Collision/GJKContactPointSwept.cpp b/Engine/Source/Runtime/Experimental/Chaos/Private/Chaos/Collision/GJKContactPointSwept.cpp
--- a/Engine/Source/Runtime/Experimental/Chaos/Private/Chaos/Collision/GJKContactPointSwept.cpp
+++ b/Engine/Source/Runtime/Experimental/Chaos/Private/Chaos/Collision/GJKContactPointSwept.cpp
@@ -8,8 +8,7 @@
 
 namespace Chaos
 {
-	template <typename GeometryA, typename GeometryB>
-	FContactPoint GJKContactPointSwept(const GeometryA& A, const FRigidTransform3& AStartTM, const FRigidTransform3& AEndTM, const GeometryB& B, const FRigidTransform3& BStartTM, const FRigidTransform3& BEndTM, const FVec3& Dir, const FReal Length, const FReal IgnorePenetration, const FReal TargetPenetration, FReal& OutTOI)
+	FContactPoint GJKContactPointSwept(const FGeomGJKHelperSIMD& A, const FRigidTransform3& AStartTM, const FRigidTransform3& AEndTM, const FGeomGJKHelperSIMD& B, const FRigidTransform3& BStartTM, const FRigidTransform3& BEndTM, const FVec3& Dir, const FReal Length, const FReal IgnorePenetration, const FReal TargetPenetration, FReal& OutTOI)
 	{
 		FContactPoint Contact;
 		/** B could be static or dynamic. In both cases, we compute contact point in B local space.  
diff --git a/Engine/Source/Runtime/Experimental/Chaos/Private/Chaos/TriangleMeshImplicitObject.cpp b/Engine/Source/Runtime/Experimental/Chaos/Private/Chaos/TriangleMeshImplicitObject.cpp
--- a/Engine/Source/Runtime/Experimental/Chaos/Private/Chaos/TriangleMeshImplicitObject.cpp
+++ b/Engine/Source/Runtime/Experimental/Chaos/Private/Chaos/TriangleMeshImplicitObject.cpp
@@ -407,7 +407,7 @@ struct FTriangleMeshOverlapVisitorNoMTD
 
 		FTriangleRegister Tri(ASimd, BSimd, CSimd);
 
-		if (GJKIntersectionSimd(Tri, QueryGeom, Translation, Rotation, Thickness, InitialDir))
+		if (GJKIntersectionSimd(Tri, Tri.BoundingBox().Extents().Max(), QueryGeom, QueryGeom.BoundingBox().Extents().Max(), Translation, Rotation, Thickness, InitialDir))
 		{
 			bFoundIntersection = true;
 			return false;
diff --git a/Engine/Source/Runtime/Experimental/Chaos/Private/GeometryCollection/GeometryCollectionConvexUtility.cpp b/Engine/Source/Runtime/Experimental/Chaos/Private/GeometryCollection/GeometryCollectionConvexUtility.cpp
--- a/Engine/Source/Runtime/Experimental/Chaos/Private/GeometryCollection/GeometryCollectionConvexUtility.cpp
+++ b/Engine/Source/Runtime/Experimental/Chaos/Private/GeometryCollection/GeometryCollectionConvexUtility.cpp
@@ -421,7 +421,7 @@ void CreateNonoverlappingConvexHulls(
 			return false;
 		}
 		const Chaos::TRigidTransform<Chaos::FReal, 3> IdentityTransform = Chaos::TRigidTransform<Chaos::FReal, 3>::Identity;
-		return GJKIntersection(*Convexes[ConvexA], *Convexes[ConvexB], IdentityTransform);
+		return GJKIntersection(*Convexes[ConvexA], Convexes[ConvexA]->BoundingBox().Extents().Max(), *Convexes[ConvexB], Convexes[ConvexB]->BoundingBox().Extents().Max(), IdentityTransform);
 	};
 
 	auto IsBoneColliding = [&TransformToConvexIndices, &Convexes, &IsColliding](int32 BoneA, int32 BoneB)
diff --git a/Engine/Source/Runtime/Experimental/Chaos/Public/Chaos/EPA.h b/Engine/Source/Runtime/Experimental/Chaos/Public/Chaos/EPA.h
--- a/Engine/Source/Runtime/Experimental/Chaos/Public/Chaos/EPA.h
+++ b/Engine/Source/Runtime/Experimental/Chaos/Public/Chaos/EPA.h
@@ -433,8 +433,9 @@ inline const bool IsEPASuccess(EEPAResult EPAResult)
 // Expanding Polytope Algorithm for finding the contact point for overlapping convex polyhedra.
 // See e.g., "Collision Detection in Interactive 3D Environments" (Gino van den Bergen, 2004)
 // or "Real-time Collision Detection with Implicit Objects" (Leif Olvang, 2010)
-template <typename T, typename TSupportA, typename TSupportB>
-EEPAResult EPA(TArray<TVec3<T>>& VertsABuffer, TArray<TVec3<T>>& VertsBBuffer, const TSupportA& SupportA, const TSupportB& SupportB, T& OutPenetration, TVec3<T>& OutDir, TVec3<T>& WitnessA, TVec3<T>& WitnessB, const FReal Eps = 1.e-2f)
+template <typename T>
+EEPAResult EPA(TArray<TVec3<T>>& VertsABuffer, TArray<TVec3<T>>& VertsBBuffer, const TFunctionRef<TVector<T, 3>(const TVec3<T>& V)>& SupportA,
+	const TFunctionRef<TVector<T, 3>(const TVec3<T>& V)>& SupportB, T& OutPenetration, TVec3<T>& OutDir, TVec3<T>& WitnessA, TVec3<T>& WitnessB, const FReal Eps = 1.e-2f)
 {
 	struct FEPAEntryWrapper
 	{
diff --git a/Engine/Source/Runtime/Experimental/Chaos/Public/Chaos/GJK.h b/Engine/Source/Runtime/Experimental/Chaos/Public/Chaos/GJK.h
--- a/Engine/Source/Runtime/Experimental/Chaos/Public/Chaos/GJK.h
+++ b/Engine/Source/Runtime/Experimental/Chaos/Public/Chaos/GJK.h
@@ -35,6 +35,46 @@
 
 namespace Chaos
 {
+	/*
+	* To avoid code bloat from ~2k instances of templated GJKRaycast2ImplSimd, it was rewritten to operate on support function pointers and void* object pointers
+	* This saved around 4MB in .text section and provided a noticeable CPU performance increase on a bottom line hardware
+	*/
+	struct FGeomGJKHelperSIMD
+	{
+		typedef VectorRegister4Float(*SupportFunc)(const void* Geom, FRealSingle Margin, const VectorRegister4Float V);
+
+		const void* Geometry;
+		FRealSingle Margin;
+		FRealSingle Radius;
+		SupportFunc Func;
+
+		template<class T>
+		FGeomGJKHelperSIMD(const T& Geom)
+			: Geometry(&Geom), Margin((FRealSingle)Geom.GetMargin()), Radius((FRealSingle)Geom.GetRadius()), Func(&SupportCoreSimd<T>)
+		{}
+
+		FRealSingle GetRadius() const { return Radius; }
+		FRealSingle GetMargin() const { return Margin; }
+
+		VectorRegister4Float operator()(const VectorRegister4Float V) const { return SupportFunction(V); }
+		VectorRegister4Float SupportFunction(const VectorRegister4Float V) const { return Func(Geometry, Margin, V); }
+		VectorRegister4Float SupportFunction(const VectorRegister4Float V, FRealSingle InMargin) const { return Func(Geometry, InMargin, V); }
+
+		VectorRegister4Float operator()(const VectorRegister4Float AToBRotation, const VectorRegister4Float BToARotation, const VectorRegister4Float V) const { return SupportFunction(AToBRotation, BToARotation, V); }
+		VectorRegister4Float SupportFunction(const VectorRegister4Float AToBRotation, const VectorRegister4Float BToARotation, const VectorRegister4Float V) const {
+			const VectorRegister4Float VInB = VectorQuaternionRotateVector(AToBRotation, V);
+			const VectorRegister4Float SupportBLocal = Func(Geometry, Margin, VInB);
+			return VectorQuaternionRotateVector(BToARotation, SupportBLocal);
+		}
+
+	private:
+		template<class T>
+		static VectorRegister4Float SupportCoreSimd(const void* Geom, FRealSingle InMargin, const VectorRegister4Float V)
+		{
+			return ((const T*)Geom)->SupportCoreSimd(V, InMargin);
+		}
+	};
+
 	// Check the GJK iteration count against a limit to prevent infinite loops. We should never really hit the limit but 
 	// we are seeing it happen in some cases and more often on some platforms than others so for now we have a warning when it hits.
 	// @todo(chaos): track this issue down (see UE-156361)
@@ -86,8 +126,8 @@ namespace Chaos
 		@param InitialDir The first direction we use to search the CSO
 		@return True if the geometries overlap, False otherwise 
 	 */
-	template <typename T, typename TGeometryA, typename TGeometryB>
-	bool GJKIntersection(const TGeometryA& RESTRICT A, const TGeometryB& RESTRICT B, const TRigidTransform<T, 3>& BToATM, const T InThicknessA = 0, const TVector<T, 3>& InitialDir = TVector<T, 3>(1, 0, 0))
+	template <typename T>
+	bool GJKIntersection(const FGeomGJKHelperSIMD& RESTRICT A, FReal BBoxMaxExtentsA, const FGeomGJKHelperSIMD& RESTRICT B, FReal BBoxMaxExtentsB, const TRigidTransform<T, 3>& BToATM, const T InThicknessA = 0, const TVector<T, 3>& InitialDir = TVector<T, 3>(1, 0, 0))
 	{
 		const UE::Math::TQuat<T>& RotationDouble = BToATM.GetRotation();
 		VectorRegister4Float RotationSimd = MakeVectorRegisterFloatFromDouble(MakeVectorRegister(RotationDouble.X, RotationDouble.Y, RotationDouble.Z, RotationDouble.W));
@@ -99,14 +139,12 @@ namespace Chaos
 
 		const VectorRegister4Float InitialDirSimd = MakeVectorRegisterFloatFromDouble(MakeVectorRegister(InitialDir[0], InitialDir[1], InitialDir[2], 0.0));
 		
-		return GJKIntersectionSimd(A, B, TranslationSimd, RotationSimd, InThicknessA, InitialDirSimd);
+		return GJKIntersectionSimd(A, BBoxMaxExtentsA, B, BBoxMaxExtentsB, TranslationSimd, RotationSimd, InThicknessA, InitialDirSimd);
 	}
 
 
-	template <typename TGeometryA, typename TGeometryB>
-	bool GJKIntersectionSimd(const TGeometryA& RESTRICT A, const TGeometryB& RESTRICT B, const VectorRegister4Float& Translation, const VectorRegister4Float& Rotation, FReal InThicknessA, const VectorRegister4Float& InitialDir)
+	inline bool GJKIntersectionSimd(const FGeomGJKHelperSIMD& RESTRICT A, FReal BBoxMaxExtentsA, const FGeomGJKHelperSIMD& RESTRICT B, FReal BBoxMaxExtentsB, const VectorRegister4Float& Translation, const VectorRegister4Float& Rotation, FReal InThicknessA, const VectorRegister4Float& InitialDir)
 	{
-
 		VectorRegister4Float V = VectorNegate(InitialDir);
 		V = VectorNormalizeSafe(V, MakeVectorRegisterFloatConstant(-1.f, 0.f, 0.f, 0.f));
 
@@ -125,7 +163,7 @@ namespace Chaos
 		CalculateQueryMargins(A, B, ThicknessA, ThicknessB);
 		ThicknessA += InThicknessA;
 
-		FReal EpsilonScale = FMath::Max<FReal>(A.TGeometryA::BoundingBox().Extents().Max(), B.TGeometryB::BoundingBox().Extents().Max());
+		FReal EpsilonScale = FMath::Max<FReal>(BBoxMaxExtentsA, BBoxMaxExtentsB);
 		EpsilonScale = FMath::Min<FReal>(EpsilonScale, 1e5);
 		
 		const FReal InflationReal = ThicknessA + ThicknessB + 1e-6 * EpsilonScale;
@@ -140,9 +178,9 @@ namespace Chaos
 				break;	//if taking too long just stop. This should never happen
 			}
 			const VectorRegister4Float NegVSimd = VectorNegate(V);
-			const VectorRegister4Float SupportASimd = A.SupportCoreSimd(NegVSimd, ThicknessA);
+			const VectorRegister4Float SupportASimd = A.SupportFunction(NegVSimd, (FRealSingle)ThicknessA);
 			const VectorRegister4Float VInBSimd = VectorQuaternionRotateVector(AToBRotation, V);
-			const VectorRegister4Float SupportBLocalSimd = B.SupportCoreSimd(VInBSimd, ThicknessB);
+			const VectorRegister4Float SupportBLocalSimd = B.SupportFunction(VInBSimd, (FRealSingle)ThicknessB);
 			const VectorRegister4Float SupportBSimd = VectorAdd(VectorQuaternionRotateVector(Rotation, SupportBLocalSimd), Translation);
 			const VectorRegister4Float WSimd = VectorSubtract(SupportASimd, SupportBSimd);
 
@@ -376,6 +414,51 @@ namespace Chaos
 	// GJK warm-start data at default numeric precision
 	using FGJKSimplexData = TGJKSimplexData<FReal>;
 
+	struct FGeomGJKHelper
+	{
+		typedef FVector(*SupportFunc)(const void* Geom, const FVec3& Direction, FReal* OutSupportDelta, int32& VertexIndex);
+
+		const void* Geometry;
+		SupportFunc Func;
+		FReal Margin;
+
+		template<class T>
+		FGeomGJKHelper(const T& Geom) : Geometry(&Geom), Func(&SupportCore<T>), Margin(Geom.GetMargin()) {}
+
+		FVector SupportFunction(const FVec3& V, int32& VertexIndex) const { return Func(Geometry, V, nullptr, VertexIndex); }
+		FVector SupportFunction(const FVec3& V, FReal* OutSupportDelta, int32& VertexIndex) const { return Func(Geometry, V, OutSupportDelta, VertexIndex); }
+
+		FVector SupportFunction(const FRotation3& AToBRotation, const FVec3& V, FReal* OutSupportDelta, int32& VertexIndex) const
+		{
+			const FVec3 VInB = AToBRotation * V;
+			return Func(Geometry, VInB, OutSupportDelta, VertexIndex);
+		}
+
+		FVector SupportFunction(const FRigidTransform3& BToATM, const FRotation3& AToBRotation, const FVec3& V, int32& VertexIndex) const
+		{
+			const FVector VInB = AToBRotation * V;
+			const FVector SupportBLocal = Func(Geometry, VInB, nullptr, VertexIndex);
+			return BToATM.TransformPositionNoScale(SupportBLocal);
+		}
+
+		FVector SupportFunction(const FRigidTransform3& BToATM, const FRotation3& AToBRotation, const FVec3& V, FReal* OutSupportDelta, int32& VertexIndex) const
+		{
+			const FVector VInB = AToBRotation * V;
+			const FVector SupportBLocal = Func(Geometry, VInB, OutSupportDelta, VertexIndex);
+			return BToATM.TransformPositionNoScale(SupportBLocal);
+		}
+
+		FReal GetMargin() const { return Margin; }
+
+	private:
+		template<class T>
+		static FVector SupportCore(const void* Geom, const FVec3& Direction, FReal* OutSupportDelta, int32& VertexIndex)
+		{
+			const T* Geometry = (const T*)Geom;
+			return Geometry->SupportCore(Direction, Geometry->GetMargin(), OutSupportDelta, VertexIndex);
+		}
+	};
+
 	/**
 	 * @brief Calculate the penetration data for two shapes using GJK and a warm-start buffer.
 	 *
@@ -517,7 +600,7 @@ namespace Chaos
 
 			T Penetration;
 			TVec3<T> MTD, ClosestA, ClosestBInA;
-			const EEPAResult EPAResult = EPA(VertsA, VertsB, SupportAFunc, SupportBInAFunc, Penetration, MTD, ClosestA, ClosestBInA, EPAEpsilon);
+			const EEPAResult EPAResult = EPA<T>(VertsA, VertsB, SupportAFunc, SupportBInAFunc, Penetration, MTD, ClosestA, ClosestBInA, EPAEpsilon);
 
 			switch (EPAResult)
 			{
@@ -688,7 +771,7 @@ namespace Chaos
 
 			T Penetration;
 			TVec3<T> MTD, ClosestA, ClosestBInA;
-			const EEPAResult EPAResult = EPA(VertsA, VertsB, SupportAFunc, SupportBInAFunc, Penetration, MTD, ClosestA, ClosestBInA, EPAEpsilon);
+			const EEPAResult EPAResult = EPA<T>(VertsA, VertsB, SupportAFunc, SupportBInAFunc, Penetration, MTD, ClosestA, ClosestBInA, EPAEpsilon);
 
 			switch (EPAResult)
 			{
@@ -866,7 +949,7 @@ namespace Chaos
 
 			T Penetration;
 			TVec3<T> MTD, ClosestA, ClosestB;
-			const EEPAResult EPAResult = EPA(VertsA, VertsB, SupportAFunc, SupportBFunc, Penetration, MTD, ClosestA, ClosestB, EPAEpsilon);
+			const EEPAResult EPAResult = EPA<T>(VertsA, VertsB, SupportAFunc, SupportBFunc, Penetration, MTD, ClosestA, ClosestB, EPAEpsilon);
 
 			switch (EPAResult)
 			{
@@ -1024,7 +1107,7 @@ namespace Chaos
 
 			T Penetration;
 			TVec3<T> MTD, ClosestA, ClosestB;
-			const EEPAResult EPAResult = EPA(VertsA, VertsB, SupportAFunc, SupportBFunc, Penetration, MTD, ClosestA, ClosestB, EPAEpsilon);
+			const EEPAResult EPAResult = EPA<T>(VertsA, VertsB, SupportAFunc, SupportBFunc, Penetration, MTD, ClosestA, ClosestB, EPAEpsilon);
 
 			switch (EPAResult)
 			{
@@ -1078,27 +1161,14 @@ namespace Chaos
 		}
 	}
 
-
-	template <typename TGeometryA, typename TGeometryB, bool bNegativePenetrationAllowed = false, typename T>
-	bool GJKPenetrationImpl(const TGeometryA& A, const TGeometryB& B, const TRigidTransform<T, 3>& BToATM, T& OutPenetration, TVec3<T>& OutClosestA, TVec3<T>& OutClosestB, TVec3<T>& OutNormal, int32& OutClosestVertexIndexA, int32& OutClosestVertexIndexB, const T InThicknessA = 0.0f, const T InThicknessB = 0.0f, const TVector<T, 3>& InitialDir = TVector<T, 3>(1, 0, 0), const T Epsilon = 1.e-3f)
+	template <bool bNegativePenetrationAllowed = false, typename T>
+	bool GJKPenetrationImpl(const FGeomGJKHelper& A, const FGeomGJKHelper& B, const TRigidTransform<T, 3>& BToATM, T& OutPenetration, TVec3<T>& OutClosestA, TVec3<T>& OutClosestB, TVec3<T>& OutNormal, int32& OutClosestVertexIndexA, int32& OutClosestVertexIndexB, const T InThicknessA = 0.0f, const T InThicknessB = 0.0f, const TVector<T, 3>& InitialDir = TVector<T, 3>(1, 0, 0), const T Epsilon = 1.e-3f)
 	{
 		int32 VertexIndexA = INDEX_NONE;
 		int32 VertexIndexB = INDEX_NONE;
 
-		auto SupportAFunc = [&A, &VertexIndexA](const TVec3<T>& V)
-		{
-			return A.SupportCore(V, A.GetMargin(), nullptr, VertexIndexA);
-		};
-
 		const TRotation<T, 3> AToBRotation = BToATM.GetRotation().Inverse();
 
-		auto SupportBFunc = [&B, &BToATM, &AToBRotation, &VertexIndexB](const TVec3<T>& V)
-		{
-			const TVector<T, 3> VInB = AToBRotation * V;
-			const TVector<T, 3> SupportBLocal = B.SupportCore(VInB, B.GetMargin(), nullptr, VertexIndexB);
-			return BToATM.TransformPositionNoScale(SupportBLocal);
-		};
-
 		//todo: refactor all of these similar functions
 		TVector<T, 3> V = -InitialDir;
 		if (V.SafeNormalize() == 0)
@@ -1127,8 +1197,8 @@ namespace Chaos
 				break;	//if taking too long just stop. This should never happen
 			}
 			const TVector<T, 3> NegV = -V;
-			const TVector<T, 3> SupportA = SupportAFunc(NegV);
-			const TVector<T, 3> SupportB = SupportBFunc(V);
+			const TVector<T, 3> SupportA = A.SupportFunction(NegV, VertexIndexA);
+			const TVector<T, 3> SupportB = B.SupportFunction(BToATM, AToBRotation, V, VertexIndexB);
 			const TVector<T, 3> W = SupportA - SupportB;
 
 			const T VW = TVector<T, 3>::DotProduct(V, W);
@@ -1182,9 +1252,32 @@ namespace Chaos
 				VertsB.Add(Bs[i]);
 			}
 
+			struct FAHelper
+			{
+				const FGeomGJKHelper& Geom;
+				int32* VertexIndex;
+
+				FAHelper(const FGeomGJKHelper& Geom, int32& VertexIndex) : Geom(Geom), VertexIndex(&VertexIndex) {}
+
+				TVec3<T> operator()(const TVec3<T>& V) const { return Geom.SupportFunction(V, *VertexIndex); }
+			};
+
+			struct FBHelper
+			{
+				const FGeomGJKHelper& Geom;
+				int32* VertexIndex;
+				const FRigidTransform3& BToATM;
+				const FRotation3& AToBRotation;
+
+				TVec3<T> operator()(const TVec3<T>& V) const { return Geom.SupportFunction(BToATM, AToBRotation, V, *VertexIndex); }
+			};
+
+			FAHelper AHelper(A, VertexIndexA);
+			FBHelper BHelper = { B, &VertexIndexB, BToATM, AToBRotation };
+
 			T Penetration;
 			TVec3<T> MTD, ClosestA, ClosestBInA;
-			const EEPAResult EPAResult = EPA(VertsA, VertsB, SupportAFunc, SupportBFunc, Penetration, MTD, ClosestA, ClosestBInA);
+			const EEPAResult EPAResult = EPA<T>(VertsA, VertsB, AHelper, BHelper, Penetration, MTD, ClosestA, ClosestBInA);
 			
 			switch (EPAResult)
 			{
@@ -1268,7 +1361,7 @@ namespace Chaos
 	template <bool bNegativePenetrationAllowed = false, typename T, typename TGeometryA, typename TGeometryB>
 	bool GJKPenetration(const TGeometryA& A, const TGeometryB& B, const TRigidTransform<T, 3>& BToATM, T& OutPenetration, TVec3<T>& OutClosestA, TVec3<T>& OutClosestB, TVec3<T>& OutNormal, int32& OutClosestVertexIndexA, int32& OutClosestVertexIndexB, const T InThicknessA = 0.0f, const T InThicknessB = 0.0f, const TVector<T, 3>& InitialDir = TVector<T, 3>(1, 0, 0), const T Epsilon = 1.e-3f)
 	{
-		return GJKPenetrationImpl<TGeometryA, TGeometryB, bNegativePenetrationAllowed, T>(A, B, BToATM, OutPenetration, OutClosestA, OutClosestB, OutNormal, OutClosestVertexIndexA, OutClosestVertexIndexB, InThicknessA, InThicknessB, InitialDir, Epsilon);
+		return GJKPenetrationImpl<bNegativePenetrationAllowed, T>(A, B, BToATM, OutPenetration, OutClosestA, OutClosestB, OutNormal, OutClosestVertexIndexA, OutClosestVertexIndexB, InThicknessA, InThicknessB, InitialDir, Epsilon);
 	}
 	/** Sweeps one geometry against the other
 	 @A The first geometry
@@ -1425,9 +1518,9 @@ namespace Chaos
 	 @return True if the geometries overlap during the sweep, False otherwise
 	 @note If A overlaps B at the start of the ray ("initial overlap" condition) then this function returns true, and sets OutTime = 0, but does not set any other output variables.
 	 */
-	template <typename TGeometryA, typename TGeometryB, typename T>
-	bool GJKRaycast2ImplSimd(const TGeometryA& A, const TGeometryB& B, const VectorRegister4Float& BToARotation, const VectorRegister4Float& StartPoint, const VectorRegister4Float& RayDir, const T RayLength,
-		T& OutTime, VectorRegister4Float& OutPosition, VectorRegister4Float& OutNormal, bool bComputeMTD, const VectorRegister4Float& InitialDir)
+	template <typename T>
+	bool GJKRaycast2ImplSimd(const FGeomGJKHelperSIMD& A, const FGeomGJKHelperSIMD& B, const VectorRegister4Float& BToARotation, const VectorRegister4Float& StartPoint, const VectorRegister4Float& RayDir,
+		T RayLength, T& OutTime, VectorRegister4Float& OutPosition, VectorRegister4Float& OutNormal, bool bComputeMTD, const VectorRegister4Float& InitialDir)
 	{
 		ensure(RayLength > 0);
 
@@ -1470,24 +1563,10 @@ namespace Chaos
 
 		VectorRegister4Float AToBRotation = VectorQuaternionInverse(BToARotation);
 
-		auto SupportAFunc = [&A, MarginA](const VectorRegister4Float& V)
-		{
-			return A.SupportCoreSimd(V, MarginA);
-		};
-
-		auto SupportBFunc = [&B, MarginB, &AToBRotation, &BToARotation](const VectorRegister4Float& V)
-		{
-			const VectorRegister4Float VInB = VectorQuaternionRotateVector(AToBRotation, V);
-			const VectorRegister4Float SupportBLocal = B.SupportCoreSimd(VInB, MarginB);
-
-			return VectorQuaternionRotateVector(BToARotation, SupportBLocal);
-		};
-
-
-		VectorRegister4Float SupportA = SupportAFunc(InitialDir);
+		VectorRegister4Float SupportA = A.SupportFunction(InitialDir);
 		As[0] = SupportA;
 
-		VectorRegister4Float SupportB = SupportBFunc(VectorNegate(InitialDir));
+		VectorRegister4Float SupportB = B.SupportFunction(AToBRotation, BToARotation, VectorNegate(InitialDir));
 		Bs[0] = SupportB;
 
 		VectorRegister4Float Lambda = VectorZeroFloat();
@@ -1502,7 +1581,7 @@ namespace Chaos
 
 		FRealSingle Inflation2;
 		VectorStoreFloat1(Inflation2Simd, &Inflation2);
-		
+
 		constexpr FRealSingle Eps2 = 1e-6f;
 
 		//mtd needs to find closest point even in inflation region, so can only skip if we found the closest points
@@ -1523,8 +1602,8 @@ namespace Chaos
 
 			V = VectorNormalizeAccurate(V);
 
-			SupportA = SupportAFunc(V);
-			SupportB = SupportBFunc(VectorNegate(V));
+			SupportA = A.SupportFunction(V);
+			SupportB = B.SupportFunction(AToBRotation, BToARotation, VectorNegate(V));
 			const VectorRegister4Float P = VectorSubtract(SupportA, SupportB);
 			const VectorRegister4Float W = VectorSubtract(X, P);
 
@@ -1737,19 +1816,20 @@ namespace Chaos
 						VertsB.Add(VectorAdd(Bs[i], X));
 					}
 
-					auto SupportBAtOriginFunc = [&B, MarginB, &BToARotation, &StartPoint, &AToBRotation](const VectorRegister4Float& Dir)
+					struct SupportBAtOriginHelper
 					{
-						int32 VertexIndex = INDEX_NONE;
-						const VectorRegister4Float DirInB = VectorQuaternionRotateVector(AToBRotation, Dir);
-						const VectorRegister4Float SupportBLocal = B.SupportCoreSimd(DirInB, MarginB);
+						const FGeomGJKHelperSIMD& B;
+						const VectorRegister4Float& AToBRotation;
+						const VectorRegister4Float& BToARotation;
+						const VectorRegister4Float& StartPoint;
 
-						const VectorRegister4Float RotatedVec = VectorQuaternionRotateVector(BToARotation, SupportBLocal);
-						return VectorAdd(RotatedVec, StartPoint);
+						VectorRegister4Float operator()(VectorRegister4Float V) const { return VectorAdd(B.SupportFunction(AToBRotation, BToARotation, V), StartPoint); }
 					};
 
+					SupportBAtOriginHelper SupportBAtOrigin = { B, AToBRotation, BToARotation, StartPoint };
 					VectorRegister4Float Penetration;
 					VectorRegister4Float MTD, ClosestA, ClosestBInA;
-					const EEPAResult EPAResult = VectorEPA(VertsA, VertsB, SupportAFunc, SupportBAtOriginFunc, Penetration, MTD, ClosestA, ClosestBInA);
+					const EEPAResult EPAResult = VectorEPA(VertsA, VertsB, A, SupportBAtOrigin, Penetration, MTD, ClosestA, ClosestBInA);
 					if (IsEPASuccess(EPAResult))
 					{
 						OutNormal = MTD;
@@ -1789,9 +1869,9 @@ namespace Chaos
 		return true;
 	}
 
-	template <typename T, typename TGeometryA, typename TGeometryB>
-	bool GJKRaycast2(const TGeometryA& A, const TGeometryB& B, const TRigidTransform<T, 3>& StartTM, const TVector<T, 3>& RayDir, const T RayLength,
-		T& OutTime, TVector<T, 3>& OutPosition, TVector<T, 3>& OutNormal, const T GivenThicknessA = 0, bool bComputeMTD = false, const TVector<T, 3>& InitialDir = TVector<T, 3>(1, 0, 0), const T GivenThicknessB = 0)
+	template <typename T = FReal>
+	bool GJKRaycast2Impl(const FGeomGJKHelperSIMD& A, const FGeomGJKHelperSIMD& B, const TRigidTransform<T, 3>& StartTM, const TVector<T, 3>& RayDir, const T RayLength, 
+		T& OutTime, TVector<T, 3>& OutPosition, TVector<T, 3>& OutNormal, const T GivenThicknessA, bool bComputeMTD, const TVector<T, 3>& InitialDir, const T GivenThicknessB)
 	{
 		const UE::Math::TQuat<T>& RotationDouble = StartTM.GetRotation();
 		VectorRegister4Float Rotation = MakeVectorRegisterFloatFromDouble(MakeVectorRegister(RotationDouble.X, RotationDouble.Y, RotationDouble.Z, RotationDouble.W));
@@ -1801,17 +1881,16 @@ namespace Chaos
 
 		// Normalize rotation
 		Rotation = VectorNormalizeSafe(Rotation, GlobalVectorConstants::Float0001);
-		
+
 		const VectorRegister4Float InitialDirSimd = MakeVectorRegisterFloatFromDouble(MakeVectorRegister(InitialDir[0], InitialDir[1], InitialDir[2], 0.0));
 		const VectorRegister4Float RayDirSimd = MakeVectorRegisterFloatFromDouble(MakeVectorRegister(RayDir[0], RayDir[1], RayDir[2], 0.0));
 
 		FRealSingle OutTimeFloat = 0.0f;
 		VectorRegister4Float OutPositionSimd, OutNormalSimd;
-		bool result = GJKRaycast2ImplSimd(A, B, Rotation, Translation, RayDirSimd, static_cast<FRealSingle>(RayLength), OutTimeFloat, OutPositionSimd, OutNormalSimd, bComputeMTD, InitialDirSimd);
+		const bool Result = GJKRaycast2ImplSimd(A, B, Rotation, Translation, RayDirSimd, static_cast<FRealSingle>(RayLength), OutTimeFloat, OutPositionSimd, OutNormalSimd, bComputeMTD, InitialDirSimd);
 
 		OutTime = static_cast<double>(OutTimeFloat);
 
-
 		alignas(16) FRealSingle OutFloat[4];
 		VectorStoreAligned(OutNormalSimd, OutFloat);
 		OutNormal.X = OutFloat[0];
@@ -1823,7 +1902,14 @@ namespace Chaos
 		OutPosition.Y = OutFloat[1];
 		OutPosition.Z = OutFloat[2];
 
-		return result;
+		return Result;
+	}
+
+	template <typename TGeometryA, typename TGeometryB, typename T = FReal>
+	bool GJKRaycast2(const TGeometryA& A, const TGeometryB& B, const TRigidTransform<T, 3>& StartTM, const TVector<T, 3>& RayDir, const T RayLength,
+		T& OutTime, TVector<T, 3>& OutPosition, TVector<T, 3>& OutNormal, const T GivenThicknessA = 0, bool bComputeMTD = false, const TVector<T, 3>& InitialDir = TVector<T, 3>(1, 0, 0), const T GivenThicknessB = 0)
+	{
+		return GJKRaycast2Impl(A, B, StartTM, RayDir, RayLength, OutTime, OutPosition, OutNormal, GivenThicknessA, bComputeMTD, InitialDir, GivenThicknessB);
 	}
 
 
diff --git a/Engine/Source/Runtime/Experimental/Chaos/Public/Chaos/GeometryQueries.h b/Engine/Source/Runtime/Experimental/Chaos/Public/Chaos/GeometryQueries.h
--- a/Engine/Source/Runtime/Experimental/Chaos/Public/Chaos/GeometryQueries.h
+++ b/Engine/Source/Runtime/Experimental/Chaos/Public/Chaos/GeometryQueries.h
@@ -74,7 +74,7 @@ namespace Chaos
 			}
 			else
 			{
-				return Utilities::CastHelper(A, BToATM, [&](const auto& AConcrete, const auto& BToAFullTM) { return GJKIntersection<FReal>(AConcrete, B, BToAFullTM, Thickness, Offset.SizeSquared() < 1e-4 ? FVec3(1, 0, 0) : Offset); });
+				return Utilities::CastHelper(A, BToATM, [&](const auto& AConcrete, const auto& BToAFullTM) { return GJKIntersection<FReal>(AConcrete, AConcrete.BoundingBox().Extents().Max(), B, B.BoundingBox().Extents().Max(), BToAFullTM, Thickness, Offset.SizeSquared() < 1e-4 ? FVec3(1, 0, 0) : Offset); });
 			}
 		}
 		else
