Fixed ChaosLowLevelTests compilation

diff --git a/Engine/Source/Runtime/Experimental/Chaos/Public/Chaos/EPA.h b/Engine/Source/Runtime/Experimental/Chaos/Public/Chaos/EPA.h
--- a/Engine/Source/Runtime/Experimental/Chaos/Public/Chaos/EPA.h
+++ b/Engine/Source/Runtime/Experimental/Chaos/Public/Chaos/EPA.h
@@ -433,6 +433,14 @@ inline const bool IsEPASuccess(EEPAResult EPAResult)
 // Expanding Polytope Algorithm for finding the contact point for overlapping convex polyhedra.
 // See e.g., "Collision Detection in Interactive 3D Environments" (Gino van den Bergen, 2004)
 // or "Real-time Collision Detection with Implicit Objects" (Leif Olvang, 2010)
+template <typename T, typename TSupportA, typename TSupportB>
+EEPAResult EPA(TArray<TVec3<T>>& VertsABuffer, TArray<TVec3<T>>& VertsBBuffer, const TSupportA& SupportA, const TSupportB& SupportB, T& OutPenetration, TVec3<T>& OutDir, TVec3<T>& WitnessA, TVec3<T>& WitnessB, const FReal Eps = 1.e-2f)
+{
+	const TFunctionRef<TVector<T, 3>(const TVec3<T>& V)> SupportFuncA(SupportA);
+	const TFunctionRef<TVector<T, 3>(const TVec3<T>& V)> SupportFuncB(SupportB);
+	return EPA<T>(VertsABuffer, VertsBBuffer, SupportFuncA, SupportFuncB, OutPenetration, OutDir, WitnessA, WitnessB, Eps);
+}
+
 template <typename T>
 EEPAResult EPA(TArray<TVec3<T>>& VertsABuffer, TArray<TVec3<T>>& VertsBBuffer, const TFunctionRef<TVector<T, 3>(const TVec3<T>& V)>& SupportA,
 	const TFunctionRef<TVector<T, 3>(const TVec3<T>& V)>& SupportB, T& OutPenetration, TVec3<T>& OutDir, TVec3<T>& WitnessA, TVec3<T>& WitnessB, const FReal Eps = 1.e-2f)
diff --git a/Engine/Source/Runtime/Experimental/Chaos/Public/Chaos/GJK.h b/Engine/Source/Runtime/Experimental/Chaos/Public/Chaos/GJK.h
--- a/Engine/Source/Runtime/Experimental/Chaos/Public/Chaos/GJK.h
+++ b/Engine/Source/Runtime/Experimental/Chaos/Public/Chaos/GJK.h
@@ -1905,7 +1905,7 @@ namespace Chaos
 		return Result;
 	}
 
-	template <typename TGeometryA, typename TGeometryB, typename T = FReal>
+	template <typename T = FReal, typename TGeometryA, typename TGeometryB>
 	bool GJKRaycast2(const TGeometryA& A, const TGeometryB& B, const TRigidTransform<T, 3>& StartTM, const TVector<T, 3>& RayDir, const T RayLength,
 		T& OutTime, TVector<T, 3>& OutPosition, TVector<T, 3>& OutNormal, const T GivenThicknessA = 0, bool bComputeMTD = false, const TVector<T, 3>& InitialDir = TVector<T, 3>(1, 0, 0), const T GivenThicknessB = 0)
 	{
