Fixed ChaosLowLevelTests compilation

diff --git a/Engine/Source/Runtime/Experimental/Chaos/Public/Chaos/GJK.h b/Engine/Source/Runtime/Experimental/Chaos/Public/Chaos/GJK.h
--- a/Engine/Source/Runtime/Experimental/Chaos/Public/Chaos/GJK.h
+++ b/Engine/Source/Runtime/Experimental/Chaos/Public/Chaos/GJK.h
@@ -126,6 +126,12 @@ namespace Chaos
 		@param InitialDir The first direction we use to search the CSO
 		@return True if the geometries overlap, False otherwise 
 	 */
+	template <typename T, typename TGeometryA, typename TGeometryB>
+	bool GJKIntersection(const TGeometryA& RESTRICT A, const TGeometryB& RESTRICT B, const TRigidTransform<T, 3>& BToATM, const T InThicknessA = 0, const TVector<T, 3>& InitialDir = TVector<T, 3>(1, 0, 0))
+	{
+		return GJKIntersection(A, A.TGeometryA::BoundingBox().Extents().Max(), B, B.TGeometryB::BoundingBox().Extents().Max(), BToATM, InThicknessA, InitialDir);
+	}
+
 	template <typename T>
 	bool GJKIntersection(const FGeomGJKHelperSIMD& RESTRICT A, FReal BBoxMaxExtentsA, const FGeomGJKHelperSIMD& RESTRICT B, FReal BBoxMaxExtentsB, const TRigidTransform<T, 3>& BToATM, const T InThicknessA = 0, const TVector<T, 3>& InitialDir = TVector<T, 3>(1, 0, 0))
 	{
