Chaos - Particle  :

Added ability to reset the active particle arrays when particle ranges are deleted.
- Exposed ArrayCollection::RemoveAt
- Exposed DeactivateParticleRanges
- Exposed AddParticleRange

diff --git a/Engine/Source/Runtime/Experimental/Chaos/Public/Chaos/PBDEvolution.h b/Engine/Source/Runtime/Experimental/Chaos/Public/Chaos/PBDEvolution.h
--- a/Engine/Source/Runtime/Experimental/Chaos/Public/Chaos/PBDEvolution.h
+++ b/Engine/Source/Runtime/Experimental/Chaos/Public/Chaos/PBDEvolution.h
@@ -43,6 +43,13 @@ class CHAOS_API FPBDEvolution : public TArrayCollection
 	// Set a block of particles active or inactive, using the index of the first added particle to identify the block.
 	void ActivateParticleRange(int32 Offset, bool bActivate)  { MParticlesActiveView.ActivateRange(Offset, bActivate); }
 
+	// Clear all ranges.
+	void DeactivateParticleRanges() { MParticlesActiveView.Reset(0); }
+
+	// Set a block of particles active or inactive, using the index of the first added particle to identify the block.
+	int32 AddParticleRange(int32 NumItems, bool bActivate = true) { return MParticlesActiveView.AddRange(NumItems, bActivate); }
+
+
 	// Particles accessors
 	const FSolverParticles& Particles() const { return MParticles; }
 	FSolverParticles& Particles() { return MParticles; }
diff --git a/Engine/Source/Runtime/Experimental/ChaosCore/Public/Chaos/ArrayCollection.h b/Engine/Source/Runtime/Experimental/ChaosCore/Public/Chaos/ArrayCollection.h
--- a/Engine/Source/Runtime/Experimental/ChaosCore/Public/Chaos/ArrayCollection.h
+++ b/Engine/Source/Runtime/Experimental/ChaosCore/Public/Chaos/ArrayCollection.h
@@ -47,6 +47,11 @@ public:
 		}
 	}
 
+	void RemoveAt(int32 Index, int32 Count)
+	{
+		RemoveAtHelper(Index, Count);
+	}
+
 	uint32 Size() const 
 	{ return MSize; }
 
