Fix minor issues in Chaos/Deformable/Utilities.h to allow its use in more places
(mainly mark functions as inline so they are not multiply-defined if the header is included in multiple files)

diff --git a/Engine/Source/Runtime/Experimental/Chaos/Public/Chaos/Deformable/Utilities.h b/Engine/Source/Runtime/Experimental/Chaos/Public/Chaos/Deformable/Utilities.h
--- a/Engine/Source/Runtime/Experimental/Chaos/Public/Chaos/Deformable/Utilities.h
+++ b/Engine/Source/Runtime/Experimental/Chaos/Public/Chaos/Deformable/Utilities.h
@@ -20,7 +20,7 @@ namespace Chaos
 	namespace Utilities
 	{
 
-		void GetTetFaces(
+		inline void GetTetFaces(
 			const FIntVector4& Tet,
 			FIntVector3& Face1,
 			FIntVector3& Face2,
@@ -48,12 +48,12 @@ namespace Chaos
 			}
 		}
 
-		int32
+		inline int32
 			GetMin(const FIntVector3& V)
 		{
 			return FMath::Min3(V[0], V[1], V[2]);
 		}
-		int32
+		inline int32
 			GetMid(const FIntVector3& V)
 		{
 			const int32 X = V[0]; const int32 Y = V[1]; const int32 Z = V[2];
@@ -62,17 +62,17 @@ namespace Chaos
 			const int32 XmZ = X - Z;
 			return (XmY * YmZ > -1 ? Y : XmY * XmZ < 1 ? X : Z);
 		}
-		int32
+		inline int32
 			GetMax(const FIntVector3& V)
 		{
 			return FMath::Max3(V[0], V[1], V[2]);
 		}
-		FIntVector3
+		inline FIntVector3
 			GetOrdered(const FIntVector3& V)
 		{
 			return FIntVector3(GetMin(V), GetMid(V), GetMax(V));
 		}
-		FIntVector4
+		inline FIntVector4
 			GetOrdered(const FIntVector4& V)
 		{
 			TArray<int32> VA = { V[0], V[1], V[2], V[3] };
@@ -80,7 +80,7 @@ namespace Chaos
 			return FIntVector4(VA[0], VA[1], VA[2], VA[3]);
 		}
 
-		void
+		inline void
 			GetSurfaceElements(
 				const TArray<FIntVector4>& Tets,
 				TArray<FIntVector3>& SurfaceElements,
@@ -118,7 +118,7 @@ namespace Chaos
 						Tet[1] != Tet[3] &&
 						Tet[2] != Tet[3]))
 					{
-						UE_LOG(LogChaosFlesh, Display, TEXT("Skipping degenerate tet %d of %d."), Idx, Tets.Num());
+						UE_LOG(LogChaos, Display, TEXT("Skipping degenerate tet %d of %d."), Idx, Tets.Num());
 						continue;
 					}
 
@@ -174,12 +174,12 @@ namespace Chaos
 					}
 				}
 				if (nonManifold)
-					UE_LOG(LogChaosFlesh, Display, TEXT("WARNING: Encountered %d non-manifold tetrahedral mesh faces."), nonManifold);
+					UE_LOG(LogChaos, Display, TEXT("WARNING: Encountered %d non-manifold tetrahedral mesh faces."), nonManifold);
 			}
 		}
 	
 
-		void FindClosestTriangle(
+		inline void FindClosestTriangle(
 			const TArray<TVector<FRealSingle, 3>>& Positions,
 			const TArray<FIntVector3>& SurfaceElements,
 			const TArray<TVector<FRealSingle, 3>>& TrianglePositions,
