Fixed Chaos bug in tests where Raycast margins were not updated correctly

diff --git a/Engine/Source/Runtime/Experimental/Chaos/Public/Chaos/GJK.h b/Engine/Source/Runtime/Experimental/Chaos/Public/Chaos/GJK.h
--- a/Engine/Source/Runtime/Experimental/Chaos/Public/Chaos/GJK.h
+++ b/Engine/Source/Runtime/Experimental/Chaos/Public/Chaos/GJK.h
@@ -44,7 +44,7 @@ namespace Chaos
 		typedef VectorRegister4Float(*SupportFunc)(const void* Geom, FRealSingle Margin, const VectorRegister4Float V);
 
 		const void* Geometry;
-		FRealSingle Margin;
+		mutable FRealSingle Margin;
 		FRealSingle Radius;
 		SupportFunc Func;
 
@@ -1547,11 +1547,11 @@ namespace Chaos
 		const T SweepMarginB = (bHasRadiusA || bHasRadiusB) ? 0.0f : (bAIsSmallest ? 0.0f : SweepMarginScale * static_cast<T>(B.GetMargin()));
 
 		// Net margin (note: both SweepMargins are zero if either Radius is non-zero, and only one SweepMargin can be non-zero)
-		const T MarginA = RadiusA + SweepMarginA;
-		const T MarginB = RadiusB + SweepMarginB;
+		A.Margin = RadiusA + SweepMarginA;
+		B.Margin = RadiusB + SweepMarginB;
 
-		const VectorRegister4Float MarginASimd = VectorLoadFloat1(&MarginA);
-		const VectorRegister4Float MarginBSimd = VectorLoadFloat1(&MarginB);
+		const VectorRegister4Float MarginASimd = VectorLoadFloat1(&A.Margin);
+		const VectorRegister4Float MarginBSimd = VectorLoadFloat1(&B.Margin);
 
 		VectorRegister4Float Simplex[4] = { VectorZeroFloat(), VectorZeroFloat(), VectorZeroFloat(), VectorZeroFloat() };
 		VectorRegister4Float As[4] = { VectorZeroFloat(), VectorZeroFloat(), VectorZeroFloat(), VectorZeroFloat() };
@@ -1840,7 +1840,7 @@ namespace Chaos
 					{
 						OutNormal = MTD;
 						VectorStoreFloat1(Penetration, &OutTime);
-						OutTime = -OutTime - (MarginA + MarginB);
+						OutTime = -OutTime - (A.Margin + B.Margin);
 						OutPosition = ClosestA;
 					}
 					else if (EPAResult == EEPAResult::NoValidContact)
@@ -1850,7 +1850,7 @@ namespace Chaos
 					else
 					{
 						//assume touching hit
-						OutTime = -(MarginA + MarginB);
+						OutTime = -(A.Margin + B.Margin);
 						OutNormal = MTD;
 						OutPosition = VectorMultiplyAdd(OutNormal, MarginASimd, As[0]);
 					}
@@ -1859,7 +1859,7 @@ namespace Chaos
 				else
 				{
 					//didn't even go into gjk loop, touching hit
-					OutTime = -(MarginA + MarginB);
+					OutTime = -(A.Margin + B.Margin);
 					OutNormal = MakeVectorRegisterFloat(0.0f, 0.0f, 1.0f, 0.0f);
 					OutPosition = VectorMultiplyAdd(OutNormal, MarginASimd, As[0]);
 				}
