ChaosFlesh : Support for dynamic addition and removal.

diff --git a/Engine/Source/Runtime/Experimental/Chaos/Private/Chaos/Deformable/ChaosDeformableSolver.cpp b/Engine/Source/Runtime/Experimental/Chaos/Private/Chaos/Deformable/ChaosDeformableSolver.cpp
--- a/Engine/Source/Runtime/Experimental/Chaos/Private/Chaos/Deformable/ChaosDeformableSolver.cpp
+++ b/Engine/Source/Runtime/Experimental/Chaos/Private/Chaos/Deformable/ChaosDeformableSolver.cpp
@@ -35,10 +35,11 @@ DECLARE_CYCLE_STAT(TEXT("DeformableSolver.Advance"), STAT_DeformableSolver_Advan
 DEFINE_LOG_CATEGORY_STATIC(LogChaosDeformableSolver, Log, All);
 namespace Chaos::Softs
 {
+	FCriticalSection FDeformableSolver::InitializationMutex;
+	FCriticalSection FDeformableSolver::RemovalMutex;
 	FCriticalSection FDeformableSolver::PackageOutputMutex;
 	FCriticalSection FDeformableSolver::PackageInputMutex;
 
-
 	FDeformableSolver::FDeformableSolver(FDeformableSolverProperties InProp)
 		: CurrentInputPackage(TUniquePtr < FDeformablePackage >(nullptr))
 		, PreviousInputPackage(TUniquePtr < FDeformablePackage >(nullptr))
@@ -62,7 +63,7 @@ namespace Chaos::Softs
 		{
 			SurfaceElements.Reset(new TArray<Chaos::TVec3<int32>>());
 		}
-		
+
 		if (Property.bDoSelfCollision)
 		{
 			SurfaceTriangleMesh.Reset(new Chaos::FTriangleMesh());
@@ -71,11 +72,24 @@ namespace Chaos::Softs
 		{
 			AllElements.Reset(new TArray<Chaos::TVec4<int32>>());
 		}
+
+		InitializeKinematicConstraint();
 		Frame = 0;
 		Time = 0.f;
 	}
 
-	bool FDeformableSolver::Advance(FSolverReal DeltaTime)
+	void FDeformableSolver::Simulate(FSolverReal DeltaTime)
+	{
+		if (Property.NumSolverIterations)
+		{
+			RemoveSimulationObjects();
+			InitializeSimulationObjects();
+			UpdateProxyInputPackages();
+			AdvanceDt(DeltaTime);
+		}
+	}
+
+	bool FDeformableSolver::AdvanceDt(FSolverReal DeltaTime)
 	{
 		SCOPE_CYCLE_COUNTER(STAT_DeformableSolver_Advance);
 		TRACE_CPUPROFILER_EVENT_SCOPE(DeformableSolver_Advance);
@@ -101,188 +115,196 @@ namespace Chaos::Softs
 	{
 		TRACE_CPUPROFILER_EVENT_SCOPE(DeformableSolver_InitializeSimulationObjects);
 
-		for (TUniquePtr<FThreadingProxy>& Proxy : UninitializedProxys)
-		{
-			InitializeSimulationObject(*Proxy);
-			InitializeKinematicState(*Proxy);
-
-			FThreadingProxy::FKey Key = Proxy->GetOwner();
-			Proxies.Add(Key, TUniquePtr<FThreadingProxy>(Proxy.Release()));
-		}
-		
-		if (UninitializedProxys.Num() != 0)
 		{
-			if (Property.bDoSelfCollision)
+			FScopeLock Lock(&InitializationMutex); // @todo(flesh) : change to threaded task based commands to prevent the lock. 
+			if (UninitializedProxys_External.Num())
 			{
-				InitializeSelfCollisionVariables();
-			}
+				for (TUniquePtr<FThreadingProxy>& Proxy : UninitializedProxys_External)
+				{
+					InitializeSimulationObject(*Proxy);
 
-			if (Property.bUseGridBasedConstraints)
-			{
-				InitializeGridBasedConstraintVariables();
+					FThreadingProxy::FKey Key = Proxy->GetOwner();
+					Proxies.Add(Key, TUniquePtr<FThreadingProxy>(Proxy.Release()));
+				}
+
+				if (UninitializedProxys_External.Num() != 0)
+				{
+					if (Property.bDoSelfCollision)
+					{
+						InitializeSelfCollisionVariables();
+					}
+
+					if (Property.bUseGridBasedConstraints)
+					{
+						InitializeGridBasedConstraintVariables();
+					}
+				}
+				UninitializedProxys_External.SetNum(0, true);
 			}
 		}
-		UninitializedProxys.SetNum(0, true);
+
 		InitializeCollisionBodies();
 	}
 
 	void FDeformableSolver::InitializeSimulationObject(FThreadingProxy& InProxy)
 	{
 		TRACE_CPUPROFILER_EVENT_SCOPE(DeformableSolver_InitializeSimulationObject);
-
 		if (FFleshThreadingProxy* Proxy = InProxy.As<FFleshThreadingProxy>())
 		{
-			if (const FManagedArrayCollection* Dynamic = &Proxy->GetDynamicCollection())
-			{
-				if (const FManagedArrayCollection* Rest = &Proxy->GetRestCollection())
-				{
-					const TManagedArray<FSolverReal>* MassArray = Rest->FindAttribute<FSolverReal>("Mass", FGeometryCollection::VerticesGroup);
+			InitializeDeformableParticles(*Proxy);
+			InitializeKinematicParticles(*Proxy);
+			InitializeTetrahedralConstraint(*Proxy);
+			InitializeGidBasedConstraints(*Proxy);
+		}
+		Time = 0.f;
+	}
+
+	void FDeformableSolver::InitializeDeformableParticles(FFleshThreadingProxy& Proxy)
+	{
+		const FManagedArrayCollection& Dynamic = Proxy.GetDynamicCollection();
+		const FManagedArrayCollection& Rest = Proxy.GetRestCollection();
+
+		const TManagedArray<FVector3f>& DynamicVertex = Dynamic.GetAttribute<FVector3f>("Vertex", FGeometryCollection::VerticesGroup);
+		const TManagedArray<FSolverReal>* MassArray = Rest.FindAttribute<FSolverReal>("Mass", FGeometryCollection::VerticesGroup);
+		uint32 NumParticles = Rest.NumElements(FGeometryCollection::VerticesGroup);
+		FSolverReal Mass = 100.0;// @todo(chaos) : make user attributes
+
+		auto ChaosVert = [](FVector3d V) { return Chaos::FVec3(V.X, V.Y, V.Z); };
+		auto ChaosM = [](FSolverReal M, const TManagedArray<float>* AM, int32 Index, int32 Num) { return FSolverReal((AM != nullptr) ? (*AM)[Index] : M / FSolverReal(Num)); };
+		auto ChaosInvM = [](FSolverReal M) { return FSolverReal(FMath::IsNearlyZero(M) ? 0.0 : 1 / M); };
+		auto DoubleVert = [](FVector3f V) { return FVector3d(V.X, V.Y, V.Z); };
+
+		const FTransform& InitialTransform = Proxy.GetInitialTransform();
+		int32 ParticleStart = Evolution->AddParticleRange(NumParticles, 1, true);
+		for (uint32 vdx = 0; vdx < NumParticles; ++vdx)
+		{
+			int32 SolverParticleIndex = ParticleStart + vdx;
+			Evolution->Particles().X(SolverParticleIndex) = ChaosVert(InitialTransform.TransformPosition(DoubleVert(DynamicVertex[vdx])));
+			Evolution->Particles().V(SolverParticleIndex) = Chaos::FVec3(0.f, 0.f, 0.f);
+			Evolution->Particles().M(SolverParticleIndex) = ChaosM(Mass, MassArray, vdx, NumParticles);
+			Evolution->Particles().InvM(SolverParticleIndex) = ChaosInvM(Evolution->Particles().M(SolverParticleIndex));
+			Evolution->Particles().PAndInvM(SolverParticleIndex).InvM = Evolution->Particles().InvM(SolverParticleIndex);
+			MObjects[SolverParticleIndex] = Proxy.GetOwner();
+		}
+		Proxy.SetSolverParticleRange(ParticleStart, NumParticles);
+	}
 
-					// @todo(chaos) : make user attributes
-					FSolverReal Mass = 100.0;
+	void FDeformableSolver::InitializeKinematicParticles(FFleshThreadingProxy& Proxy)
+	{
+		TRACE_CPUPROFILER_EVENT_SCOPE(DeformableSolver_InitializeSimulationObject);
 
-					const TManagedArray<FVector3f>& Vertex = Rest->GetAttribute<FVector3f>("Vertex", FGeometryCollection::VerticesGroup);
-					const TManagedArray<FIntVector>& Indices = Rest->GetAttribute<FIntVector>("Indices", FGeometryCollection::FacesGroup);
-					const TManagedArray<FVector3f>& DynamicVertex = Dynamic->GetAttribute<FVector3f>("Vertex", FGeometryCollection::VerticesGroup);
+		const FManagedArrayCollection& Rest = Proxy.GetRestCollection();
+		const FIntVector2& Range = Proxy.GetSolverParticleRange();
+
+		if (Property.bEnableKinematics)
+		{
+			typedef Chaos::Facades::FKinematicBindingFacade Kinematics;
 
-					if (uint32 NumParticles = Vertex.Num())
+			// Add Kinematics Node
+			for (int i = Kinematics::NumKinematicBindings(&Rest) - 1; i >= 0; i--)
+			{
+				Kinematics::FBindingKey Key = Kinematics::GetKinematicBindingKey(&Rest, i);
+
+				int32 BoneIndex = INDEX_NONE;
+				TArray<int32> BoundVerts;
+				TArray<float> BoundWeights;
+				Kinematics::GetBoneBindings(&Rest, Key, BoneIndex, BoundVerts, BoundWeights);
+
+				for (int32 vdx : BoundVerts)
+				{
+					if (BoneIndex != INDEX_NONE)
 					{
-						if (uint32 NumSurfaceElements = Indices.Num())
-						{
-							// @todo(chaos) : reduce conversions
-							auto ChaosVert = [](FVector3d V) { return Chaos::FVec3(V.X, V.Y, V.Z); };
-							auto DoubleVert = [](FVector3f V) { return FVector3d(V.X, V.Y, V.Z); };
-							auto ChaosM = [](FSolverReal M, const TManagedArray<float>* AM, int32 Index, int32 Num) { return FSolverReal((AM != nullptr) ? (*AM)[Index] : M / FSolverReal(Num)); };
-							auto ChaosInvM = [](FSolverReal M) { return FSolverReal(FMath::IsNearlyZero(M) ? 0.0 : 1 / M); };
-							auto ChaosTet = [](FIntVector4 V, int32 dp) { return Chaos::TVec4<int32>(dp + V.X, dp + V.Y, dp + V.Z, dp + V.W); };
-							auto ChaosTri = [](FIntVector V, int32 dp) { return  Chaos::TVec3<int32>(dp + V.X, dp + V.Y, dp + V.Z); };
-							auto ChaosIntVec2 = [](int32 A, int32 B) { return  Chaos::TVec2<int32>(A, B); };
-
-							//  Add Simulation Particles Node
-							const FTransform& InitialTransform = Proxy->GetInitialTransform();
-							int32 ParticleStart = Evolution->AddParticleRange(NumParticles, 1, true);
-							for (uint32 vdx = 0; vdx < NumParticles; ++vdx)
-							{
-								int32 SolverParticleIndex = ParticleStart + vdx;
-								Evolution->Particles().X(SolverParticleIndex) = ChaosVert(InitialTransform.TransformPosition(DoubleVert(DynamicVertex[vdx])));
-								Evolution->Particles().V(SolverParticleIndex) = Chaos::FVec3(0.f, 0.f, 0.f);
-								Evolution->Particles().M(SolverParticleIndex) = ChaosM(Mass, MassArray, vdx, NumParticles);
-								Evolution->Particles().InvM(SolverParticleIndex) = ChaosInvM(Evolution->Particles().M(SolverParticleIndex));
-								Evolution->Particles().PAndInvM(SolverParticleIndex).InvM = Evolution->Particles().InvM(SolverParticleIndex);
-								MObjects[SolverParticleIndex] = Proxy->GetOwner();
-							}
-							Proxy->SetSolverParticleRange(ParticleStart, NumParticles);
+						int32 ParticleIndex = Range[0] + vdx;
+						Evolution->Particles().InvM(ParticleIndex) = 0.f;
+						Evolution->Particles().PAndInvM(ParticleIndex).InvM = 0.f;
+					}
+				}
+			}
+		}
+	}
 
-							if (Property.bEnableKinematics)
-							{
-								typedef Chaos::Facades::FKinematicBindingFacade Kinematics;
-
-								// Add Kinematics Node
-								for (int i = Kinematics::NumKinematicBindings(Rest) - 1; i >= 0; i--)
-								{
-									Kinematics::FBindingKey Key = Kinematics::GetKinematicBindingKey(Rest, i);
-
-									int32 BoneIndex = INDEX_NONE;
-									TArray<int32> BoundVerts;
-									TArray<float> BoundWeights;
-									Kinematics::GetBoneBindings(Rest, Key, BoneIndex, BoundVerts, BoundWeights);
-
-									for (int32 vdx : BoundVerts)
-									{
-										if (BoneIndex != INDEX_NONE)
-										{
-											int32 ParticleIndex = ParticleStart + vdx;
-											Evolution->Particles().InvM(ParticleIndex) = 0.f;
-											Evolution->Particles().PAndInvM(ParticleIndex).InvM = 0.f;
-										}
-									}
-								}
-							}
 
-							const TManagedArray<FIntVector4>& Tetrahedron = Rest->GetAttribute<FIntVector4>("Tetrahedron", "Tetrahedral");
+	void FDeformableSolver::InitializeTetrahedralConstraint(FFleshThreadingProxy& Proxy)
+	{
+		TRACE_CPUPROFILER_EVENT_SCOPE(DeformableSolver_InitializeSimulationObject);
 
-							if (uint32 NumElements = Tetrahedron.Num())
-							{
-								const FIntVector2& Range = Proxy->GetSolverParticleRange();
-
-								// Add Tetrahedral Elements Node
-								TArray<Chaos::TVec4<int32>> Elements;
-								Elements.SetNum(NumElements);
-								for (uint32 edx = 0; edx < NumElements; ++edx)
-								{
-									Elements[edx] = ChaosTet(Tetrahedron[edx], Range[0]);
-								}
-
-								if (Property.bUseGridBasedConstraints)
-								{
-									int32 ElementsOffset = AllElements->Num();
-									AllElements->SetNum(ElementsOffset + NumElements);
-									for (uint32 edx = 0; edx < NumElements; ++edx)
-									{
-										(*AllElements)[edx + ElementsOffset] = ChaosTet(Tetrahedron[edx], Range[0]);
-									}
-
-								}
-
-
-								FXPBDCorotatedConstraints<FSolverReal, FSolverParticles>* CorotatedConstraint =
-									new FXPBDCorotatedConstraints<FSolverReal, FSolverParticles>(
-										Evolution->Particles(), Elements, /*bRecordMetric = */false, (FSolverReal)100000.0);
-
-								int32 InitIndex = Evolution->AddConstraintInitRange(1, true);
-								Evolution->ConstraintInits()[InitIndex] =
-									[CorotatedConstraint](FSolverParticles& InParticles, const FSolverReal Dt)
-								{
-									CorotatedConstraint->Init();
-								};
-
-
-								int32 ConstraintIndex = Evolution->AddConstraintRuleRange(1, true);
-								Evolution->ConstraintRules()[ConstraintIndex] =
-									[CorotatedConstraint](FSolverParticles& InParticles, const FSolverReal Dt)
-								{
-									CorotatedConstraint->ApplyInParallel(InParticles, Dt);
-
-								};
-								CorotatedConstraints.Add(TUniquePtr<FXPBDCorotatedConstraints<FSolverReal, FSolverParticles>>(CorotatedConstraint));
-							}
+		const FManagedArrayCollection& Rest = Proxy.GetRestCollection();
 
-							if (Property.bDoSelfCollision || Property.CacheToFile)
-							{
-								// Add Surface Elements Node
-								int32 SurfaceElementsOffset = SurfaceElements->Num();
-								SurfaceElements->SetNum(NumSurfaceElements + SurfaceElementsOffset);
-								for (uint32 edx = 0; edx < NumSurfaceElements; ++edx)
-								{
-									(*SurfaceElements)[edx + SurfaceElementsOffset] = ChaosTri(Indices[edx], ParticleStart);
-								}
+		auto ChaosTet = [](FIntVector4 V, int32 dp) { return Chaos::TVec4<int32>(dp + V.X, dp + V.Y, dp + V.Z, dp + V.W); };
 
-							}
+		const TManagedArray<FIntVector4>& Tetrahedron = Rest.GetAttribute<FIntVector4>("Tetrahedron", "Tetrahedral");
+		if (uint32 NumElements = Tetrahedron.Num())
+		{
+			const FIntVector2& Range = Proxy.GetSolverParticleRange();
 
+			// Add Tetrahedral Elements Node
+			TArray<Chaos::TVec4<int32>> Elements;
+			Elements.SetNum(NumElements);
+			for (uint32 edx = 0; edx < NumElements; ++edx)
+			{
+				Elements[edx] = ChaosTet(Tetrahedron[edx], Range[0]);
+			}
 
-							Time = 0.f;
-						}
-					}
+			if (Property.bUseGridBasedConstraints)
+			{
+				int32 ElementsOffset = AllElements->Num();
+				AllElements->SetNum(ElementsOffset + NumElements);
+				for (uint32 edx = 0; edx < NumElements; ++edx)
+				{
+					(*AllElements)[edx + ElementsOffset] = ChaosTet(Tetrahedron[edx], Range[0]);
 				}
+
 			}
+
+
+			FXPBDCorotatedConstraints<FSolverReal, FSolverParticles>* CorotatedConstraint =
+				new FXPBDCorotatedConstraints<FSolverReal, FSolverParticles>(
+					Evolution->Particles(), Elements, /*bRecordMetric = */false, (FSolverReal)100000.0);
+
+			int32 InitIndex = Evolution->AddConstraintInitRange(1, true);
+			Evolution->ConstraintInits()[InitIndex] =
+				[CorotatedConstraint](FSolverParticles& InParticles, const FSolverReal Dt)
+			{
+				CorotatedConstraint->Init();
+			};
+
+
+			int32 ConstraintIndex = Evolution->AddConstraintRuleRange(1, true);
+			Evolution->ConstraintRules()[ConstraintIndex] =
+				[CorotatedConstraint](FSolverParticles& InParticles, const FSolverReal Dt)
+			{
+				CorotatedConstraint->ApplyInParallel(InParticles, Dt);
+			};
+			CorotatedConstraints.Add(TUniquePtr<FXPBDCorotatedConstraints<FSolverReal, FSolverParticles>>(CorotatedConstraint));
 		}
 	}
 
-	void FDeformableSolver::InitializeCollisionBodies()
+	void FDeformableSolver::InitializeGidBasedConstraints(FFleshThreadingProxy& Proxy)
 	{
-		TRACE_CPUPROFILER_EVENT_SCOPE(DeformableSolver_InitializeCollisionBodies);
-		if (Property.bUseFloor && Evolution->CollisionParticles().Size()==0)
+		TRACE_CPUPROFILER_EVENT_SCOPE(DeformableSolver_InitializeSimulationObject);
+		if (Property.bUseGridBasedConstraints)
 		{
-			Chaos::FVec3 Position(0.f);
-			Chaos::FVec3 EulerRot(0.f);
-			Evolution->AddCollisionParticleRange(1, 1, true);
-			Evolution->CollisionParticles().X(0) = Position;
-			Evolution->CollisionParticles().R(0) = Chaos::TRotation<Chaos::FReal, 3>::MakeFromEuler(EulerRot);
-			Evolution->CollisionParticles().SetDynamicGeometry(0, MakeUnique<Chaos::TPlane<Chaos::FReal, 3>>(Chaos::FVec3(0.f, 0.f, 0.f), Chaos::FVec3(0.f, 0.f, 1.f)));
+			auto ChaosTet = [](FIntVector4 V, int32 dp) { return Chaos::TVec4<int32>(dp + V.X, dp + V.Y, dp + V.Z, dp + V.W); };
+
+			const FManagedArrayCollection& Rest = Proxy.GetRestCollection();
+			const TManagedArray<FIntVector4>& Tetrahedron = Rest.GetAttribute<FIntVector4>("Tetrahedron", "Tetrahedral");
+
+			if (uint32 NumElements = Tetrahedron.Num())
+			{
+				const FIntVector2& Range = Proxy.GetSolverParticleRange();
+
+				int32 ElementsOffset = AllElements->Num();
+				AllElements->SetNum(ElementsOffset + NumElements);
+				for (uint32 edx = 0; edx < NumElements; ++edx)
+				{
+					(*AllElements)[edx + ElementsOffset] = ChaosTet(Tetrahedron[edx], Range[0]);
+				}
+			}
 		}
 	}
 
-	void FDeformableSolver::InitializeKinematicState(FThreadingProxy& InProxy)
+
+	void FDeformableSolver::InitializeKinematicConstraint()
 	{
 		auto MKineticUpdate = [this](FSolverParticles& MParticles, const FSolverReal Dt, const FSolverReal MTime, const int32 Index)
 		{
@@ -311,6 +333,20 @@ namespace Chaos::Softs
 		Evolution->SetKinematicUpdateFunction(MKineticUpdate);
 	}
 
+	void FDeformableSolver::InitializeCollisionBodies()
+	{
+		TRACE_CPUPROFILER_EVENT_SCOPE(DeformableSolver_InitializeCollisionBodies);
+		if (Property.bUseFloor && Evolution->CollisionParticles().Size() == 0)
+		{
+			Chaos::FVec3 Position(0.f);
+			Chaos::FVec3 EulerRot(0.f);
+			Evolution->AddCollisionParticleRange(1, 1, true);
+			Evolution->CollisionParticles().X(0) = Position;
+			Evolution->CollisionParticles().R(0) = Chaos::TRotation<Chaos::FReal, 3>::MakeFromEuler(EulerRot);
+			Evolution->CollisionParticles().SetDynamicGeometry(0, MakeUnique<Chaos::TPlane<Chaos::FReal, 3>>(Chaos::FVec3(0.f, 0.f, 0.f), Chaos::FVec3(0.f, 0.f, 1.f)));
+		}
+	}
+
 	void FDeformableSolver::InitializeSelfCollisionVariables()
 	{
 		TRACE_CPUPROFILER_EVENT_SCOPE(DeformableSolver_InitializeSelfCollisionVariables);
@@ -404,6 +440,67 @@ namespace Chaos::Softs
 		}
 	}
 
+	void FDeformableSolver::RemoveSimulationObjects()
+	{
+		TRACE_CPUPROFILER_EVENT_SCOPE(DeformableSolver_RemoveSimulationObjects);
+		TArray< FThreadingProxy* > RemovedProxies;
+		{
+			FScopeLock Lock(&RemovalMutex); // @todo(flesh) : change to threaded task based commands to prevent the lock. 
+			RemovedProxies = TArray< FThreadingProxy* >(RemovedProxys_External);
+			RemovedProxys_External.Empty();
+		}
+
+		if (RemovedProxies.Num())
+		{
+			Evolution->ResetConstraintRules();
+			Evolution->DeactivateParticleRanges();
+
+			// delete the simulated particles in block moves
+			for (FThreadingProxy* BaseProxy : RemovedProxies)
+			{
+				if (FFleshThreadingProxy* Proxy = BaseProxy->As<FFleshThreadingProxy>())
+				{
+					FIntVector2 Indices = Proxy->GetSolverParticleRange();
+					Proxies.FindAndRemoveChecked(MObjects[Indices[0]]);
+					Evolution->Particles().RemoveAt(Indices[0], Indices[1]);
+				}
+			}
+
+			// reindex ranges on moved particles in the proxies. 
+			const UObject* CurrentObject = nullptr;
+			for (int Index = 0; Index < MObjects.Num(); Index++)
+			{
+				if (MObjects[Index] != CurrentObject)
+				{
+					CurrentObject = MObjects[Index];
+					if (CurrentObject)
+					{
+						if (ensure(Proxies.Contains(CurrentObject)))
+						{
+							if (FFleshThreadingProxy* MovedProxy = Proxies[CurrentObject]->As<FFleshThreadingProxy>())
+							{
+								FIntVector2 Range = MovedProxy->GetSolverParticleRange();
+								MovedProxy->SetSolverParticleRange(Index, Range[1]);
+								int32 Offset = Evolution->AddParticleRange(Range[1]);
+								//ensure(Offset == Range[0]);
+							}
+						}
+					}
+				}
+			}
+
+			// regenerate all constraints
+			for (TPair< FThreadingProxy::FKey, TUniquePtr<FThreadingProxy> >& BaseProxyPair : Proxies)
+			{
+				if (FFleshThreadingProxy* Proxy = BaseProxyPair.Value->As<FFleshThreadingProxy>())
+				{
+					InitializeTetrahedralConstraint(*Proxy);
+					InitializeGidBasedConstraints(*Proxy);
+				}
+			}
+		}
+	}
+
 	void FDeformableSolver::TickSimulation(FSolverReal DeltaTime)
 	{
 		TRACE_CPUPROFILER_EVENT_SCOPE(DeformableSolver_TickSimulation);
@@ -456,9 +553,18 @@ namespace Chaos::Softs
 		return TUniquePtr<FDeformablePackage>(nullptr);
 	}
 
-	void FDeformableSolver::AddProxy(TUniquePtr<FThreadingProxy> InObject)
+	void FDeformableSolver::AddProxy(FThreadingProxy* InObject)
+	{
+		FScopeLock Lock(&InitializationMutex);
+		UninitializedProxys_External.Add(TUniquePtr< FThreadingProxy>(InObject));
+		InitializedObjects_External.Add(InObject->GetOwner());
+	}
+
+	void FDeformableSolver::RemoveProxy(FThreadingProxy* InObject)
 	{
-		UninitializedProxys.Add(TUniquePtr< FThreadingProxy>(InObject.Release()));
+		FScopeLock Lock(&RemovalMutex);
+		RemovedProxys_External.Add(InObject);
+		InitializedObjects_External.Remove(InObject->GetOwner());
 	}
 
 	void FDeformableSolver::UpdateOutputState(FThreadingProxy& InProxy)
diff --git a/Engine/Source/Runtime/Experimental/Chaos/Private/Chaos/Deformable/ChaosDeformableSolverThreadingWrapper.cpp b/Engine/Source/Runtime/Experimental/Chaos/Private/Chaos/Deformable/ChaosDeformableSolverThreadingWrapper.cpp
--- a/Engine/Source/Runtime/Experimental/Chaos/Private/Chaos/Deformable/ChaosDeformableSolverThreadingWrapper.cpp
+++ b/Engine/Source/Runtime/Experimental/Chaos/Private/Chaos/Deformable/ChaosDeformableSolverThreadingWrapper.cpp
@@ -35,86 +35,130 @@ namespace Chaos::Softs
 {
 	void FDeformableSolver::FPhysicsThreadAccess::Reset(const FDeformableSolverProperties& InProps)
 	{
-		Solver.Reset(InProps);
+		if(Solver) 
+			Solver->Reset(InProps);
 	}
 
-	bool FDeformableSolver::FPhysicsThreadAccess::Advance(FSolverReal DeltaTime)
+	bool FDeformableSolver::FPhysicsThreadAccess::AdvanceDt(FSolverReal DeltaTime)
 	{
-		return Solver.Advance(DeltaTime);
+		if (Solver)
+			return Solver->AdvanceDt(DeltaTime);
+		return false;
+	}
+
+	void FDeformableSolver::FPhysicsThreadAccess::Simulate(FSolverReal DeltaTime)
+	{
+		if(Solver) 
+			Solver->Simulate(DeltaTime);
 	}
 
 	void FDeformableSolver::FPhysicsThreadAccess::InitializeSimulationObjects()
 	{
-		Solver.InitializeSimulationObjects();
+		if(Solver) 
+			Solver->InitializeSimulationObjects();
 	}
 
 	void FDeformableSolver::FPhysicsThreadAccess::InitializeSimulationObject(FThreadingProxy& InProxy)
 	{
-		Solver.InitializeSimulationObject(InProxy);
+		if(Solver) 
+			Solver->InitializeSimulationObject(InProxy);
 	}
 
 	void FDeformableSolver::FPhysicsThreadAccess::InitializeCollisionBodies()
 	{
-		return Solver.InitializeCollisionBodies();
+		if (Solver)
+			Solver->InitializeCollisionBodies();
 	}
 
-	void FDeformableSolver::FPhysicsThreadAccess::InitializeKinematicState(FThreadingProxy& InProxy)
+	void FDeformableSolver::FPhysicsThreadAccess::InitializeKinematicConstraint()
 	{
-		Solver.InitializeKinematicState(InProxy);
+		if(Solver) Solver->InitializeKinematicConstraint();
 	}
 
 	void FDeformableSolver::FPhysicsThreadAccess::InitializeSelfCollisionVariables()
 	{
-		Solver.InitializeSelfCollisionVariables();
+		if(Solver) 
+			Solver->InitializeSelfCollisionVariables();
+	}
+
+	bool FDeformableSolver::FGameThreadAccess::HasObject(UObject* InObject) const
+	{
+		if (Solver)
+			return Solver->HasObject(InObject);
+		return false;
 	}
 
 	void FDeformableSolver::FGameThreadAccess::PushInputPackage(int32 InFrame, FDeformableDataMap&& InPackage)
 	{
-		Solver.PushInputPackage(InFrame, MoveTemp(InPackage));
+		if(Solver) 
+			Solver->PushInputPackage(InFrame, MoveTemp(InPackage));
 	}
 
 	TUniquePtr<FDeformablePackage> FDeformableSolver::FPhysicsThreadAccess::PullInputPackage()
 	{
-		return Solver.PullInputPackage();
+		if (Solver)
+			return Solver->PullInputPackage();
+		return TUniquePtr<FDeformablePackage>(nullptr);
 	}
 
 	void FDeformableSolver::FPhysicsThreadAccess::UpdateProxyInputPackages()
 	{
-		return Solver.UpdateProxyInputPackages();
+		if (Solver)
+			Solver->UpdateProxyInputPackages();
+	}
+
+	void FDeformableSolver::FPhysicsThreadAccess::RemoveSimulationObjects()
+	{
+		if(Solver) 
+			Solver->RemoveSimulationObjects();
 	}
 
 	void FDeformableSolver::FPhysicsThreadAccess::TickSimulation(FSolverReal DeltaTime)
 	{
-		Solver.TickSimulation(DeltaTime);
+		if(Solver)
+			Solver->TickSimulation(DeltaTime);
 	}
 
 	void FDeformableSolver::FPhysicsThreadAccess::PushOutputPackage(int32 InFrame, FDeformableDataMap&& InPackage)
 	{
-		Solver.PushOutputPackage(InFrame, MoveTemp(InPackage));
+		if(Solver)
+			Solver->PushOutputPackage(InFrame, MoveTemp(InPackage));
 	}
 
 	TUniquePtr<FDeformablePackage>  FDeformableSolver::FGameThreadAccess::PullOutputPackage()
 	{
-		return Solver.PullOutputPackage();
+		if (Solver)
+			return  Solver->PullOutputPackage();
+		return TUniquePtr<FDeformablePackage>(nullptr);
+	}
+
+	void  FDeformableSolver::FGameThreadAccess::AddProxy(FThreadingProxy* InObject)
+	{
+		if (Solver)
+			Solver->AddProxy(InObject);
 	}
 
-	void  FDeformableSolver::FGameThreadAccess::AddProxy(TUniquePtr<FThreadingProxy> InObject)
+	void  FDeformableSolver::FGameThreadAccess::RemoveProxy(FThreadingProxy* InObject)
 	{
-		return Solver.AddProxy(TUniquePtr<FThreadingProxy>(InObject.Release()));
+		if(Solver) 
+			Solver->RemoveProxy(InObject);
 	}
 
 	void  FDeformableSolver::FPhysicsThreadAccess::UpdateOutputState(FThreadingProxy& InProxy)
 	{
-		Solver.UpdateOutputState(InProxy);
+		if(Solver) 
+			Solver->UpdateOutputState(InProxy);
 	}
 
 	void  FDeformableSolver::FPhysicsThreadAccess::WriteFrame(FThreadingProxy& InProxy, const FSolverReal DeltaTime)
 	{
-		Solver.WriteFrame(InProxy, DeltaTime);
+		if(Solver) 
+			Solver->WriteFrame(InProxy, DeltaTime);
 	}
 
 	void  FDeformableSolver::FPhysicsThreadAccess::WriteTrisGEO(const FSolverParticles& Particles, const TArray<TVec3<int32>>& Mesh)
 	{
-		Solver.WriteTrisGEO(Particles, Mesh);
+		if(Solver) 
+			Solver->WriteTrisGEO(Particles, Mesh);
 	}
 }; // Namespace Chaos::Softs
\ No newline at end of file
diff --git a/Engine/Source/Runtime/Experimental/Chaos/Public/Chaos/Deformable/ChaosDeformableSolver.h b/Engine/Source/Runtime/Experimental/Chaos/Public/Chaos/Deformable/ChaosDeformableSolver.h
--- a/Engine/Source/Runtime/Experimental/Chaos/Public/Chaos/Deformable/ChaosDeformableSolver.h
+++ b/Engine/Source/Runtime/Experimental/Chaos/Public/Chaos/Deformable/ChaosDeformableSolver.h
@@ -20,7 +20,6 @@
 
 namespace Chaos::Softs
 {
-	// @todo(Deformable) : Explore using the chaos ISolverBase, IProxyBase instead. 
 	class CHAOS_API FDeformableSolver
 	{
 		friend class FGameThreadAccess;
@@ -35,11 +34,13 @@ namespace Chaos::Softs
 		class CHAOS_API FPhysicsThreadAccess
 		{
 		public:
-			FPhysicsThreadAccess(FDeformableSolver& InSolver, const FPhysicsThreadAccessor&) : Solver(InSolver) {}
+			FPhysicsThreadAccess(FDeformableSolver* InSolver, const FPhysicsThreadAccessor&) : Solver(InSolver) {}
+			bool operator()() { return Solver != nullptr; }
 
 			/* Simulation Advance */
 			void UpdateProxyInputPackages();
-			bool Advance(FSolverReal DeltaTime);
+			void Simulate(FSolverReal DeltaTime);
+			bool AdvanceDt(FSolverReal DeltaTime);
 			void Reset(const FDeformableSolverProperties&);
 			void TickSimulation(FSolverReal DeltaTime);
 			void UpdateOutputState(FThreadingProxy&);
@@ -50,24 +51,25 @@ namespace Chaos::Softs
 			void InitializeSimulationObjects();
 			void InitializeSimulationObject(FThreadingProxy&);
 			void InitializeCollisionBodies();
-			void InitializeKinematicState(FThreadingProxy&);
+			void InitializeKinematicConstraint();
 			void InitializeSelfCollisionVariables();
+			void RemoveSimulationObjects();
 
 			/*IO Utility*/
 			void WriteTrisGEO(const Softs::FSolverParticles& Particles, const TArray<TVec3<int32>>& Mesh);
 			void WriteFrame(FThreadingProxy& , const FSolverReal DeltaTime);
 
-			const FDeformableSolverProperties& GetProperties() const { return Solver.GetProperties(); }
+			const FDeformableSolverProperties& GetProperties() const { return Solver->GetProperties(); }
 
-			FPBDEvolution* GetEvolution() { return Solver.Evolution.Get(); }
-			const FPBDEvolution* GetEvolution() const { return Solver.Evolution.Get(); }
+			FPBDEvolution* GetEvolution() { return Solver->Evolution.Get(); }
+			const FPBDEvolution* GetEvolution() const { return Solver->Evolution.Get(); }
 
-			TArrayCollectionArray<const UObject*>& GetObjectsMap() { return Solver.MObjects; }
-			const TArrayCollectionArray<const UObject*>& GetObjectsMap() const { return Solver.MObjects; }
+			TArrayCollectionArray<const UObject*>& GetObjectsMap() { return Solver->MObjects; }
+			const TArrayCollectionArray<const UObject*>& GetObjectsMap() const { return Solver->MObjects; }
 
 
 		private:
-			FDeformableSolver& Solver;
+			FDeformableSolver* Solver;
 		};
 
 
@@ -75,17 +77,19 @@ namespace Chaos::Softs
 		class CHAOS_API FGameThreadAccess
 		{
 		public:
-			FGameThreadAccess(FDeformableSolver& InSolver, const FGameThreadAccessor&) : Solver(InSolver) {}
-
-			int32 GetFrame() const { return Solver.GetFrame(); }
-			void AddProxy(TUniquePtr<FThreadingProxy> InObject);
+			FGameThreadAccess(FDeformableSolver* InSolver, const FGameThreadAccessor&) : Solver(InSolver) {}
+			bool operator()() { return Solver != nullptr; }
 
+			int32 GetFrame() const { return Solver->GetFrame(); }
+			bool HasObject(UObject* InObject) const;
+			void AddProxy(FThreadingProxy* InObject);
+			void RemoveProxy(FThreadingProxy* InObject);
 			void PushInputPackage(int32 Frame, FDeformableDataMap&& InPackage);
 
 			TUniquePtr<FDeformablePackage> PullOutputPackage();
 
 		private:
-			FDeformableSolver& Solver;
+			FDeformableSolver* Solver;
 		};
 
 	protected:
@@ -93,7 +97,8 @@ namespace Chaos::Softs
 		/* Simulation Advance */
 		int32 GetFrame() const { return Frame; }
 		void UpdateProxyInputPackages();
-		bool Advance(FSolverReal DeltaTime);
+		void Simulate(FSolverReal DeltaTime);
+		bool AdvanceDt(FSolverReal DeltaTime);
 		void Reset(const FDeformableSolverProperties&);
 		void TickSimulation(FSolverReal DeltaTime);
 		void UpdateOutputState(FThreadingProxy&);
@@ -103,17 +108,25 @@ namespace Chaos::Softs
 		/* Iteration Advance */
 		void InitializeSimulationObjects();
 		void InitializeSimulationObject(FThreadingProxy&);
+		void InitializeDeformableParticles(FFleshThreadingProxy&);
+		void InitializeKinematicParticles(FFleshThreadingProxy&);
+		void InitializeTetrahedralConstraint(FFleshThreadingProxy&);
+		void InitializeGidBasedConstraints(FFleshThreadingProxy&);
+		void InitializeKinematicConstraint();
 		void InitializeCollisionBodies();
-		void InitializeKinematicState(FThreadingProxy&);
 		void InitializeSelfCollisionVariables();
 		void InitializeGridBasedConstraintVariables();
+		void RemoveSimulationObjects();
+
 
 		/*IO Utility*/
 		void WriteTrisGEO(const FSolverParticles& Particles, const TArray<TVec3<int32>>& Mesh);
 		void WriteFrame(FThreadingProxy&, const FSolverReal DeltaTime);
 
 		/*Game Thread API*/
-		void AddProxy(TUniquePtr<FThreadingProxy> InObject);
+		bool HasObject(UObject* InObject) const { return InitializedObjects_External.Contains(InObject); }
+		void AddProxy(FThreadingProxy* InObject);
+		void RemoveProxy(FThreadingProxy* InObject);
 		TUniquePtr<FDeformablePackage> PullOutputPackage();
 		void PushInputPackage(int32 Frame, FDeformableDataMap&& InPackage);
 
@@ -122,15 +135,21 @@ namespace Chaos::Softs
 	private:
 
 		// connections outside the solver.
+		static FCriticalSection	InitializationMutex; // @todo(flesh) : change to threaded commands to prevent the lock. 
+		static FCriticalSection	RemovalMutex; // @todo(flesh) : change to threaded commands to prevent the lock. 
 		static FCriticalSection	PackageOutputMutex;
 		static FCriticalSection	PackageInputMutex;
-		TArray< TUniquePtr<FThreadingProxy> > UninitializedProxys;
-		TMap< FThreadingProxy::FKey, TUniquePtr<FThreadingProxy> > Proxies;
+		TArray< FThreadingProxy* > RemovedProxys_External;
+		TArray< TUniquePtr<FThreadingProxy> > UninitializedProxys_External;
 		TArray< TUniquePtr<FDeformablePackage>  > BufferedInputPackages;
 		TArray< TUniquePtr<FDeformablePackage>  > BufferedOutputPackages;
 		TUniquePtr < FDeformablePackage > CurrentInputPackage;
 		TUniquePtr < FDeformablePackage > PreviousInputPackage;
 
+		TSet< const UObject* > InitializedObjects_External;
+		TMap< FThreadingProxy::FKey, TUniquePtr<FThreadingProxy> > Proxies;
+
+
 		// User Configuration
 		FDeformableSolverProperties Property;
 
