Dataflow Graph:
- fixed crash in FNodeColorsRegistry::~FNodeColorsRegistry() when GetMutableDefault<UDataflowSettings>() returns invalid value

diff --git a/Engine/Source/Runtime/Experimental/Dataflow/Core/Private/Dataflow/DataflowNodeColorsRegistry.cpp b/Engine/Source/Runtime/Experimental/Dataflow/Core/Private/Dataflow/DataflowNodeColorsRegistry.cpp
--- a/Engine/Source/Runtime/Experimental/Dataflow/Core/Private/Dataflow/DataflowNodeColorsRegistry.cpp
+++ b/Engine/Source/Runtime/Experimental/Dataflow/Core/Private/Dataflow/DataflowNodeColorsRegistry.cpp
@@ -31,8 +31,11 @@ namespace Dataflow
 	{
 		if (IsClassLoaded<UDataflowSettings>())
 		{
-			UDataflowSettings* DataflowSettings = GetMutableDefault<UDataflowSettings>();
-			DataflowSettings->GetOnDataflowSettingsChangedDelegate().Remove(DataflowSettingsChangedDelegateHandle);
+			if (UObjectInitialized())
+			{
+				UDataflowSettings* DataflowSettings = GetMutableDefault<UDataflowSettings>();
+				DataflowSettings->GetOnDataflowSettingsChangedDelegate().Remove(DataflowSettingsChangedDelegateHandle);
+			}
 		}
 	}
 
