Remove comment ed conditionnal causing some branch merging issues

diff --git a/Engine/Source/Runtime/Experimental/Chaos/Private/PhysicsProxy/GeometryCollectionPhysicsProxy.cpp b/Engine/Source/Runtime/Experimental/Chaos/Private/PhysicsProxy/GeometryCollectionPhysicsProxy.cpp
--- a/Engine/Source/Runtime/Experimental/Chaos/Private/PhysicsProxy/GeometryCollectionPhysicsProxy.cpp
+++ b/Engine/Source/Runtime/Experimental/Chaos/Private/PhysicsProxy/GeometryCollectionPhysicsProxy.cpp
@@ -543,12 +543,9 @@ void FGeometryCollectionPhysicsProxy::Initialize(Chaos::FPBDRigidsEvolutionBase
 								}
 								else
 								{
-								//	if (!(GameThreadCollection.StatusFlags[ChildIndex] & FGeometryCollection::ENodeFlags::FS_IgnoreCollisionInParentCluster))
-									{
-										TUniquePtr<Chaos::FImplicitObject> CopiedChildImplicit = GameThreadCollection.Implicits[ChildIndex]->DeepCopy();
-										FImplicitObjectTransformed* TransformedChildImplicit = new FImplicitObjectTransformed(MoveTemp(CopiedChildImplicit), RelativeShapeTransform);  
-										ChildImplicits.Add(TUniquePtr<FImplicitObjectTransformed>(TransformedChildImplicit));
-									}
+									TUniquePtr<Chaos::FImplicitObject> CopiedChildImplicit = GameThreadCollection.Implicits[ChildIndex]->DeepCopy();
+									FImplicitObjectTransformed* TransformedChildImplicit = new FImplicitObjectTransformed(MoveTemp(CopiedChildImplicit), RelativeShapeTransform);  
+									ChildImplicits.Add(TUniquePtr<FImplicitObjectTransformed>(TransformedChildImplicit));
 								}
 							}
 						}
