ChaosFlesh : Proxy initilization clean up.

diff --git a/Engine/Source/Runtime/Experimental/Chaos/Private/Chaos/Deformable/ChaosDeformableSolver.cpp b/Engine/Source/Runtime/Experimental/Chaos/Private/Chaos/Deformable/ChaosDeformableSolver.cpp
--- a/Engine/Source/Runtime/Experimental/Chaos/Private/Chaos/Deformable/ChaosDeformableSolver.cpp
+++ b/Engine/Source/Runtime/Experimental/Chaos/Private/Chaos/Deformable/ChaosDeformableSolver.cpp
@@ -49,6 +49,16 @@ namespace Chaos::Softs
 		Reset(Property);
 	}
 
+	FDeformableSolver::~FDeformableSolver()
+	{
+		for (FThreadingProxy* Proxy : UninitializedProxys_Internal)
+		{
+			delete Proxy;
+		}
+		UninitializedProxys_Internal.Empty();
+	}
+
+
 	void FDeformableSolver::Reset(const FDeformableSolverProperties& InProps)
 	{
 		TRACE_CPUPROFILER_EVENT_SCOPE(DeformableSolver_Reset);
@@ -117,17 +127,17 @@ namespace Chaos::Softs
 
 		{
 			FScopeLock Lock(&InitializationMutex); // @todo(flesh) : change to threaded task based commands to prevent the lock. 
-			if (UninitializedProxys_External.Num())
+			if (UninitializedProxys_Internal.Num())
 			{
-				for (TUniquePtr<FThreadingProxy>& Proxy : UninitializedProxys_External)
+				for (FThreadingProxy* Proxy : UninitializedProxys_Internal)
 				{
 					InitializeSimulationObject(*Proxy);
 
 					FThreadingProxy::FKey Key = Proxy->GetOwner();
-					Proxies.Add(Key, TUniquePtr<FThreadingProxy>(Proxy.Release()));
+					Proxies.Add(Key, TUniquePtr<FThreadingProxy>(Proxy));
 				}
 
-				if (UninitializedProxys_External.Num() != 0)
+				if (UninitializedProxys_Internal.Num() != 0)
 				{
 					if (Property.bDoSelfCollision)
 					{
@@ -139,7 +149,7 @@ namespace Chaos::Softs
 						InitializeGridBasedConstraintVariables();
 					}
 				}
-				UninitializedProxys_External.SetNum(0, true);
+				UninitializedProxys_Internal.SetNum(0, true);
 			}
 		}
 
@@ -446,8 +456,8 @@ namespace Chaos::Softs
 		TArray< FThreadingProxy* > RemovedProxies;
 		{
 			FScopeLock Lock(&RemovalMutex); // @todo(flesh) : change to threaded task based commands to prevent the lock. 
-			RemovedProxies = TArray< FThreadingProxy* >(RemovedProxys_External);
-			RemovedProxys_External.Empty();
+			RemovedProxies = TArray< FThreadingProxy* >(RemovedProxys_Internal);
+			RemovedProxys_Internal.Empty();
 		}
 
 		if (RemovedProxies.Num())
@@ -553,18 +563,39 @@ namespace Chaos::Softs
 		return TUniquePtr<FDeformablePackage>(nullptr);
 	}
 
-	void FDeformableSolver::AddProxy(FThreadingProxy* InObject)
+	void FDeformableSolver::AddProxy(FThreadingProxy* InProxy)
 	{
 		FScopeLock Lock(&InitializationMutex);
-		UninitializedProxys_External.Add(TUniquePtr< FThreadingProxy>(InObject));
-		InitializedObjects_External.Add(InObject->GetOwner());
+		UninitializedProxys_Internal.Add(InProxy);
+		InitializedObjects_External.Add(InProxy->GetOwner());
 	}
 
-	void FDeformableSolver::RemoveProxy(FThreadingProxy* InObject)
+	void FDeformableSolver::RemoveProxy(FThreadingProxy* InProxy)
 	{
 		FScopeLock Lock(&RemovalMutex);
-		RemovedProxys_External.Add(InObject);
-		InitializedObjects_External.Remove(InObject->GetOwner());
+
+		InitializedObjects_External.Remove(InProxy->GetOwner());
+
+		// If a proxy has not been initialized yet, then we need
+		// to clean up the internal buffers. 
+		int32 Index = UninitializedProxys_Internal.IndexOfByKey(InProxy);
+		if(Index!=INDEX_NONE)
+		{
+			UninitializedProxys_Internal.RemoveAtSwap(Index);
+			if (Proxies.Contains(InProxy->GetOwner()))
+			{
+				RemovedProxys_Internal.Add(InProxy);
+			}
+			else
+			{
+				delete InProxy;
+			}
+		}
+		else if(Proxies.Contains(InProxy->GetOwner()))
+		{
+			RemovedProxys_Internal.Add(InProxy);
+		}
+
 	}
 
 	void FDeformableSolver::UpdateOutputState(FThreadingProxy& InProxy)
diff --git a/Engine/Source/Runtime/Experimental/Chaos/Public/Chaos/Deformable/ChaosDeformableSolver.h b/Engine/Source/Runtime/Experimental/Chaos/Public/Chaos/Deformable/ChaosDeformableSolver.h
--- a/Engine/Source/Runtime/Experimental/Chaos/Public/Chaos/Deformable/ChaosDeformableSolver.h
+++ b/Engine/Source/Runtime/Experimental/Chaos/Public/Chaos/Deformable/ChaosDeformableSolver.h
@@ -29,6 +29,7 @@ namespace Chaos::Softs
 	public:
 
 		FDeformableSolver(FDeformableSolverProperties InProp = FDeformableSolverProperties());
+		virtual ~FDeformableSolver();
 
 		/* Physics Thread Access API */
 		class CHAOS_API FPhysicsThreadAccess
@@ -139,8 +140,8 @@ namespace Chaos::Softs
 		static FCriticalSection	RemovalMutex; // @todo(flesh) : change to threaded commands to prevent the lock. 
 		static FCriticalSection	PackageOutputMutex;
 		static FCriticalSection	PackageInputMutex;
-		TArray< FThreadingProxy* > RemovedProxys_External;
-		TArray< TUniquePtr<FThreadingProxy> > UninitializedProxys_External;
+		TArray< FThreadingProxy* > RemovedProxys_Internal;
+		TArray< FThreadingProxy* > UninitializedProxys_Internal;
 		TArray< TUniquePtr<FDeformablePackage>  > BufferedInputPackages;
 		TArray< TUniquePtr<FDeformablePackage>  > BufferedOutputPackages;
 		TUniquePtr < FDeformablePackage > CurrentInputPackage;
