Dataflow : Check for subsystem init first.

diff --git a/Engine/Source/Runtime/Experimental/Dataflow/Core/Private/Dataflow/DataflowNodeColorsRegistry.cpp b/Engine/Source/Runtime/Experimental/Dataflow/Core/Private/Dataflow/DataflowNodeColorsRegistry.cpp
--- a/Engine/Source/Runtime/Experimental/Dataflow/Core/Private/Dataflow/DataflowNodeColorsRegistry.cpp
+++ b/Engine/Source/Runtime/Experimental/Dataflow/Core/Private/Dataflow/DataflowNodeColorsRegistry.cpp
@@ -29,9 +29,9 @@ namespace Dataflow
 
 	FNodeColorsRegistry::~FNodeColorsRegistry()
 	{
-		if (IsClassLoaded<UDataflowSettings>())
+		if (UObjectInitialized())
 		{
-			if (UObjectInitialized())
+			if (IsClassLoaded<UDataflowSettings>())
 			{
 				UDataflowSettings* DataflowSettings = GetMutableDefault<UDataflowSettings>();
 				DataflowSettings->GetOnDataflowSettingsChangedDelegate().Remove(DataflowSettingsChangedDelegateHandle);
