Chaos: Modified TObjectPool to avoid unnecessary operations.

The real motivation behind this change is that on some platforms Clang generates invalid binary (infinite loop) for this code in optimized builds (Test configuration).

#rnx

diff --git a/Engine/Source/Runtime/Experimental/ChaosCore/Public/Chaos/ObjectPool.h b/Engine/Source/Runtime/Experimental/ChaosCore/Public/Chaos/ObjectPool.h
--- a/Engine/Source/Runtime/Experimental/ChaosCore/Public/Chaos/ObjectPool.h
+++ b/Engine/Source/Runtime/Experimental/ChaosCore/Public/Chaos/ObjectPool.h
@@ -73,37 +73,32 @@ namespace Chaos
 			// We know there's a free item somewhere, find one - shuffling full blocks to the end of the list
 			check(Blocks.Num() > 0);
 
-			// Set BackIndex to the first empty block from the end
-			int32 BackIndex = Blocks.Num() - 1;
-			while(BackIndex > 0)
+			if (Blocks[0].IsFull())
 			{
-				if(Blocks[BackIndex].IsFull())
+				// Set BackIndex to the first empty block from the end
+				int32 BackIndex = Blocks.Num() - 1;
+				while(BackIndex > 0)
 				{
-					--BackIndex;
-				}
-				else
-				{
-					break;
+					if(Blocks[BackIndex].IsFull())
+					{
+						--BackIndex;
+					}
+					else
+					{
+						break;
+					}
 				}
-			}
 
-			int32 RunCount = 0;
-			while(true)
-			{
-				checkf(RunCount++ < Blocks.Num(), TEXT("Could not find an empty block"));
-
-				if(Blocks[0].IsFull())
-				{
-					Swap(Blocks[0], Blocks[BackIndex--]);
-					continue;
-				}
+				checkf(!Blocks[BackIndex].IsFull(), TEXT("Could not find an empty block"));
 
-				// Blocks[0] is now a block with at least one free element
-				ObjectType* NewPtr = Blocks[0].GetNextFree();
-				Construct<ObjectType>(NewPtr, Forward<TArgs>(Args)...);
-				--FreeCount;
-				return NewPtr;
+				Swap(Blocks[0], Blocks[BackIndex]);
 			}
+
+			// Blocks[0] is now a block with at least one free element
+			ObjectType* NewPtr = Blocks[0].GetNextFree();
+			Construct<ObjectType>(NewPtr, Forward<TArgs>(Args)...);
+			--FreeCount;
+			return NewPtr;
 		}
 
 		/**
