Chaos : geometry collection : fix collision generation for clusters that should use a union of their children collision shapes
- make sure we only process a cluster when all the children have been properly processed first
- optimize collision hierarchy by combining transformed of transformed implicit into one

diff --git a/Engine/Source/Runtime/Experimental/Chaos/Private/PhysicsProxy/GeometryCollectionPhysicsProxy.cpp b/Engine/Source/Runtime/Experimental/Chaos/Private/PhysicsProxy/GeometryCollectionPhysicsProxy.cpp
--- a/Engine/Source/Runtime/Experimental/Chaos/Private/PhysicsProxy/GeometryCollectionPhysicsProxy.cpp
+++ b/Engine/Source/Runtime/Experimental/Chaos/Private/PhysicsProxy/GeometryCollectionPhysicsProxy.cpp
@@ -515,46 +515,76 @@ void FGeometryCollectionPhysicsProxy::Initialize(Chaos::FPBDRigidsEvolutionBase
 					{
 						const Chaos::FRigidTransform3 ParentShapeTransform =  GameThreadCollection.MassToLocal[ParentToFixIndex];
 				
-						// Make a union of the children geometry
-						TArray<TUniquePtr<Chaos::FImplicitObject>> ChildImplicits;
-						for (const int32& ChildIndex: GameThreadCollection.Children[ParentToFixIndex])
+						// let's make sure all our children have an implicit defined, other wise, postpone to next iteration 
+						bool bAllChildrenHaveCollision = true;
+						for (const int32& ChildIndex : GameThreadCollection.Children[ParentToFixIndex])
 						{
-							using FImplicitObjectTransformed = Chaos::TImplicitObjectTransformed<Chaos::FReal, 3>;
+							// defer if any of the children is a cluster with no collision yet generated 
+							if (GameThreadCollection.Implicits[ChildIndex] == nullptr && GameThreadCollection.Children[ChildIndex].Num() > 0)
+							{
+								ParentToPotentiallyFix.Add(ParentToFixIndex);
+								bAllChildrenHaveCollision = false;
+								break;
+							}
+						}
 
-							Chaos::FGeometryParticle* ChildParticle = GTParticles[ChildIndex].Get();
-							const FGeometryDynamicCollection::FSharedImplicit& ChildImplicit = GameThreadCollection.Implicits[ChildIndex];
-							if (ChildImplicit)
+						if (bAllChildrenHaveCollision)
+						{
+							// Make a union of the children geometry
+							TArray<TUniquePtr<Chaos::FImplicitObject>> ChildImplicits;
+							for (const int32& ChildIndex : GameThreadCollection.Children[ParentToFixIndex])
 							{
-								const Chaos::FRigidTransform3 ChildShapeTransform =  GameThreadCollection.MassToLocal[ChildIndex] * GameThreadCollection.Transform[ChildIndex];
-								const Chaos::FRigidTransform3 RelativeShapeTransform = ChildShapeTransform.GetRelativeTransform(ParentShapeTransform);
-						
-								// assumption that we only have can only have one level of union for any child
-								if (ChildImplicit->GetType() == Chaos::ImplicitObjectType::Union)
+								using FImplicitObjectTransformed = Chaos::TImplicitObjectTransformed<Chaos::FReal, 3>;
+
+								Chaos::FGeometryParticle* ChildParticle = GTParticles[ChildIndex].Get();
+								const FGeometryDynamicCollection::FSharedImplicit& ChildImplicit = GameThreadCollection.Implicits[ChildIndex];
+								if (ChildImplicit)
 								{
-									if (Chaos::FImplicitObjectUnion* Union = ChildImplicit->GetObject<Chaos::FImplicitObjectUnion>())
+									const Chaos::FRigidTransform3 ChildShapeTransform = GameThreadCollection.MassToLocal[ChildIndex] * GameThreadCollection.Transform[ChildIndex];
+									const Chaos::FRigidTransform3 RelativeShapeTransform = ChildShapeTransform.GetRelativeTransform(ParentShapeTransform);
+
+									// assumption that we only have can only have one level of union for any child
+									if (ChildImplicit->GetType() == Chaos::ImplicitObjectType::Union)
+									{
+										if (Chaos::FImplicitObjectUnion* Union = ChildImplicit->GetObject<Chaos::FImplicitObjectUnion>())
+										{
+											for (const TUniquePtr<Chaos::FImplicitObject>& ImplicitObject : Union->GetObjects())
+											{
+												TUniquePtr<Chaos::FImplicitObject> CopiedChildImplicit = ImplicitObject->DeepCopy();
+												FImplicitObjectTransformed* TransformedChildImplicit = new FImplicitObjectTransformed(MoveTemp(CopiedChildImplicit), RelativeShapeTransform);
+												ChildImplicits.Add(TUniquePtr<FImplicitObjectTransformed>(TransformedChildImplicit));
+											}
+										}
+									}
+									else
 									{
-										for (const TUniquePtr<Chaos::FImplicitObject>& ImplicitObject : Union->GetObjects())
+										TUniquePtr<Chaos::FImplicitObject> CopiedChildImplicit = GameThreadCollection.Implicits[ChildIndex]->DeepCopy();
+										if (CopiedChildImplicit->GetType() == Chaos::ImplicitObjectType::Transformed)
 										{
-											TUniquePtr<Chaos::FImplicitObject> CopiedChildImplicit = ImplicitObject->DeepCopy();
-											FImplicitObjectTransformed* TransformedChildImplicit = new FImplicitObjectTransformed(MoveTemp(CopiedChildImplicit), RelativeShapeTransform);  
+											if (FImplicitObjectTransformed* Transformed = CopiedChildImplicit->GetObject<FImplicitObjectTransformed>())
+											{
+												Transformed->SetTransform(Transformed->GetTransform() * RelativeShapeTransform);
+												ChildImplicits.Add(TUniquePtr<FImplicitObjectTransformed>(Transformed));
+											}
+										}
+										else
+										{
+											FImplicitObjectTransformed* TransformedChildImplicit = new FImplicitObjectTransformed(MoveTemp(CopiedChildImplicit), RelativeShapeTransform);
 											ChildImplicits.Add(TUniquePtr<FImplicitObjectTransformed>(TransformedChildImplicit));
 										}
 									}
 								}
-								else
-								{
 									TUniquePtr<Chaos::FImplicitObject> CopiedChildImplicit = GameThreadCollection.Implicits[ChildIndex]->DeepCopy();
 									FImplicitObjectTransformed* TransformedChildImplicit = new FImplicitObjectTransformed(MoveTemp(CopiedChildImplicit), RelativeShapeTransform);  
 									ChildImplicits.Add(TUniquePtr<FImplicitObjectTransformed>(TransformedChildImplicit));
-								}
 							}
+							if (ChildImplicits.Num() > 0)
+							{
+								Chaos::FImplicitObject* UnionImplicit = new Chaos::FImplicitObjectUnion(MoveTemp(ChildImplicits));
+								GameThreadCollection.Implicits[ParentToFixIndex] = FGeometryDynamicCollection::FSharedImplicit(UnionImplicit);
+							}
+							GTParticles[ParentToFixIndex]->SetGeometry(GameThreadCollection.Implicits[ParentToFixIndex]);
 						}
-						if (ChildImplicits.Num() > 0)
-						{
-							Chaos::FImplicitObject* UnionImplicit = new Chaos::FImplicitObjectUnion(MoveTemp(ChildImplicits));
-							GameThreadCollection.Implicits[ParentToFixIndex] = FGeometryDynamicCollection::FSharedImplicit(UnionImplicit);
-						}
-						GTParticles[ParentToFixIndex]->SetGeometry(GameThreadCollection.Implicits[ParentToFixIndex]);
 					}
 				}
 
