AABBTree revert unnecessary changes, Default values are overwritten by BroadPhase config

diff --git a/Engine/Source/Runtime/Experimental/Chaos/Public/Chaos/AABBTree.h b/Engine/Source/Runtime/Experimental/Chaos/Public/Chaos/AABBTree.h
--- a/Engine/Source/Runtime/Experimental/Chaos/Public/Chaos/AABBTree.h
+++ b/Engine/Source/Runtime/Experimental/Chaos/Public/Chaos/AABBTree.h
@@ -650,7 +650,7 @@ public:
 	using PayloadType = TPayloadType;
 	static constexpr int D = 3;
 	using TType = T;
-	static constexpr T DefaultMaxPayloadBounds = 10000;
+	static constexpr T DefaultMaxPayloadBounds = 100000;
 	static constexpr int32 DefaultMaxChildrenInLeaf = 12;
 	static constexpr int32 DefaultMaxTreeDepth = 16;
 	static constexpr int32 DefaultMaxNumToProcess = 0; // 0 special value for processing all without timeslicing
