Chaos - fix non-linear joint solver kinematic body transforms

diff --git a/Engine/Source/Runtime/Experimental/Chaos/Private/Chaos/Joint/PBDJointSolverGaussSeidel.cpp b/Engine/Source/Runtime/Experimental/Chaos/Private/Chaos/Joint/PBDJointSolverGaussSeidel.cpp
--- a/Engine/Source/Runtime/Experimental/Chaos/Private/Chaos/Joint/PBDJointSolverGaussSeidel.cpp
+++ b/Engine/Source/Runtime/Experimental/Chaos/Private/Chaos/Joint/PBDJointSolverGaussSeidel.cpp
@@ -46,11 +46,15 @@ namespace Chaos
 		InitConnectorRs[1] = R(1) * LocalConnectorXs[1].GetRotation();
 		InitConnectorRs[1].EnforceShortestArcWith(InitConnectorRs[0]);
 
-		ConnectorXs[0] = InitConnectorXs[0];
-		ConnectorRs[0] = InitConnectorRs[0];
-
-		ConnectorXs[1] = InitConnectorXs[1];
-		ConnectorRs[1] = InitConnectorRs[1];
+		const FVec3 BodyP0 = P(0);
+		const FRotation3 BodyQ0 = Q(0);
+		const FVec3 BodyP1 = P(1);
+		const FRotation3 BodyQ1 = Q(1);
+		ConnectorXs[0] = BodyP0 + BodyQ0 * LocalConnectorXs[0].GetTranslation();
+		ConnectorXs[1] = BodyP1 + BodyQ1 * LocalConnectorXs[1].GetTranslation();
+		ConnectorRs[0] = BodyQ0 * LocalConnectorXs[0].GetRotation();
+		ConnectorRs[1] = BodyQ1 * LocalConnectorXs[1].GetRotation();
+		ConnectorRs[1].EnforceShortestArcWith(ConnectorRs[0]);
 	}
 
 
