Chaos - constraint graph optimizations
- graph structure is now persistent between frames
- particles and constraints are now added/removed only when enabled/disabled
- constraint sorting now performed in level assignment
- no longer needs to build the islands for just-puty-to-sleep islands
- more unit tests

diff --git a/Engine/Source/Programs/HeadlessChaos/Private/HeadlessChaosTestCCD.cpp b/Engine/Source/Programs/HeadlessChaos/Private/HeadlessChaosTestCCD.cpp
--- a/Engine/Source/Programs/HeadlessChaos/Private/HeadlessChaosTestCCD.cpp
+++ b/Engine/Source/Programs/HeadlessChaos/Private/HeadlessChaosTestCCD.cpp
@@ -66,6 +66,9 @@ namespace ChaosTest
 		Dynamic->SetCCDEnabled(false);
 		Dynamic->SetGravityEnabled(false);
 
+		Evolution.EnableParticle(Static);
+		Evolution.EnableParticle(Dynamic);
+
 
 		Dynamic->V() = FVec3(0, 0, -InitialSpeed);
 
@@ -130,6 +133,9 @@ namespace ChaosTest
 
 		Dynamic->V() = FVec3(0, 0, -InitialSpeed);
 
+		Evolution.EnableParticle(Static);
+		Evolution.EnableParticle(Dynamic);
+
 		for (int i = 0; i < 1; ++i)
 		{
 			Evolution.AdvanceOneTimeStep(Dt);
@@ -194,6 +200,9 @@ namespace ChaosTest
 
 		Dynamic->V() = FVec3(0, 0, -InitialSpeed);
 
+		Evolution.EnableParticle(Static);
+		Evolution.EnableParticle(Dynamic);
+
 		for (int i = 0; i < 1; ++i)
 		{
 			Evolution.AdvanceOneTimeStep(Dt);
@@ -284,13 +293,13 @@ namespace ChaosTest
 		Dynamic->SetGravityEnabled(false);
 
 		// IMPORTANT : this is required to make sure the particles internal representation will reflect the sim data
-		Evolution.DirtyParticle(*ContainerFaces[0]);
-		Evolution.DirtyParticle(*ContainerFaces[1]);
-		Evolution.DirtyParticle(*ContainerFaces[2]);
-		Evolution.DirtyParticle(*ContainerFaces[3]);
-		Evolution.DirtyParticle(*ContainerFaces[4]);
-		Evolution.DirtyParticle(*ContainerFaces[5]);
-		Evolution.DirtyParticle(*Dynamic);
+		Evolution.EnableParticle(ContainerFaces[0]);
+		Evolution.EnableParticle(ContainerFaces[1]);
+		Evolution.EnableParticle(ContainerFaces[2]);
+		Evolution.EnableParticle(ContainerFaces[3]);
+		Evolution.EnableParticle(ContainerFaces[4]);
+		Evolution.EnableParticle(ContainerFaces[5]);
+		Evolution.EnableParticle(Dynamic);
 
 		Dynamic->V() = InitialVelocity;
 		///////////////////////////////////
diff --git a/Engine/Source/Programs/HeadlessChaos/Private/HeadlessChaosTestClustering.cpp b/Engine/Source/Programs/HeadlessChaos/Private/HeadlessChaosTestClustering.cpp
--- a/Engine/Source/Programs/HeadlessChaos/Private/HeadlessChaosTestClustering.cpp
+++ b/Engine/Source/Programs/HeadlessChaos/Private/HeadlessChaosTestClustering.cpp
@@ -34,12 +34,14 @@ namespace ChaosTest {
 		uint32 BoxId = FirstId++;
 		FPBDRigidParticleHandle* Box2 = AppendClusteredParticleBox(Particles, FVec3((FReal)100, (FReal)100, (FReal)100));
 		uint32 Box2Id = FirstId++;
-		
+
 		Box2->X() = FVec3((FReal)100, (FReal)0, (FReal)0);
 		Box2->P() = Box2->X();
 
+		Evolution.EnableParticle(Box1);
+		Evolution.EnableParticle(Box2);
+
 		Evolution.AdvanceOneTimeStep(0);	//hack to initialize islands
-		//Evolution.InitializeAccelerationStructures();	//make sure islands are created
 		FClusterCreationParameters ClusterParams;
 		
 		TArray<Chaos::FPBDRigidParticleHandle*> ClusterChildren;
@@ -87,6 +89,8 @@ namespace ChaosTest {
 			Box->X() = FVec3((FReal)i * (FReal)100, (FReal)0, (FReal)0);
 			Box->P() = Box->X();
 			Boxes.Add(Box);
+
+			Evolution.EnableParticle(Box);
 		}
 
 		Evolution.AdvanceOneTimeStep(0);	//hack to generate islands
@@ -206,6 +210,8 @@ namespace ChaosTest {
 			Box->X() = FVec3((FReal)i * (FReal)100, (FReal)0, (FReal)0);
 			Box->P() = Box->X();
 			Boxes.Add(Box);
+
+			Evolution.EnableParticle(Box);
 		}
 
 		Evolution.AdvanceOneTimeStep(0);	//hack to generate islands
diff --git a/Engine/Source/Programs/HeadlessChaos/Private/HeadlessChaosTestCollisions.cpp b/Engine/Source/Programs/HeadlessChaos/Private/HeadlessChaosTestCollisions.cpp
--- a/Engine/Source/Programs/HeadlessChaos/Private/HeadlessChaosTestCollisions.cpp
+++ b/Engine/Source/Programs/HeadlessChaos/Private/HeadlessChaosTestCollisions.cpp
@@ -8,7 +8,6 @@
 #include "Chaos/Pair.h"
 #include "Chaos/PBDRigidsEvolution.h"
 #include "Chaos/PBDRigidParticles.h"
-#include "Chaos/PBDCollisionConstraintsPGS.h"
 #include "Chaos/CollisionResolutionTypes.h"
 #include "Chaos/Sphere.h"
 #include "Chaos/Utilities.h"
diff --git a/Engine/Source/Programs/HeadlessChaos/Private/HeadlessChaosTestConstraintHandle.cpp b/Engine/Source/Programs/HeadlessChaos/Private/HeadlessChaosTestConstraintHandle.cpp
--- a/Engine/Source/Programs/HeadlessChaos/Private/HeadlessChaosTestConstraintHandle.cpp
+++ b/Engine/Source/Programs/HeadlessChaos/Private/HeadlessChaosTestConstraintHandle.cpp
@@ -113,6 +113,10 @@ namespace ChaosTest {
 		TEvolution Evolution(Particles, PhysicalMaterials);
 
 		TArray<FPBDRigidParticleHandle*> ParticleHandles = Evolution.CreateDynamicParticles(5);
+		for (FPBDRigidParticleHandle* ParticleHandle : ParticleHandles)
+		{
+			Evolution.EnableParticle(ParticleHandle);
+		}
 
 		FPBDJointConstraints Constraints;
 		TArray<FPBDJointConstraintHandle*> ConstraintsHandles =
@@ -134,6 +138,10 @@ namespace ChaosTest {
 		THandleArray<FChaosPhysicsMaterial> PhysicalMaterials;
 		TEvolution Evolution(Particles, PhysicalMaterials);
 		TArray<FPBDRigidParticleHandle*> ParticleHandles = Evolution.CreateDynamicParticles(5);
+		for (FPBDRigidParticleHandle* ParticleHandle : ParticleHandles)
+		{
+			Evolution.EnableParticle(ParticleHandle);
+		}
 
 		FPBDPositionConstraints Constraints;
 		TArray<FPBDPositionConstraintHandle*> ConstraintsHandles =
@@ -155,6 +163,10 @@ namespace ChaosTest {
 		THandleArray<FChaosPhysicsMaterial> PhysicalMaterials;
 		TEvolution Evolution(Particles, PhysicalMaterials);
 		TArray<FPBDRigidParticleHandle*> ParticleHandles = Evolution.CreateDynamicParticles(5);
+		for (FPBDRigidParticleHandle* ParticleHandle : ParticleHandles)
+		{
+			Evolution.EnableParticle(ParticleHandle);
+		}
 
 		FPBDRigidSpringConstraints Constraints;
 		TArray<FPBDRigidSpringConstraintHandle*> ConstraintsHandles =
@@ -177,6 +189,10 @@ namespace ChaosTest {
 		THandleArray<FChaosPhysicsMaterial> PhysicalMaterials;
 		TEvolution Evolution(Particles, PhysicalMaterials);
 		TArray<FPBDRigidParticleHandle*> ParticleHandles = Evolution.CreateDynamicParticles(5);
+		for (FPBDRigidParticleHandle* ParticleHandle : ParticleHandles)
+		{
+			Evolution.EnableParticle(ParticleHandle);
+		}
 
 		FPBDRigidDynamicSpringConstraints Constraints;
 		TArray<FPBDRigidDynamicSpringConstraintHandle*> ConstraintsHandles =
diff --git a/Engine/Source/Programs/HeadlessChaos/Private/HeadlessChaosTestConstraintSolverCore.cpp b/Engine/Source/Programs/HeadlessChaos/Private/HeadlessChaosTestConstraintSolverCore.cpp
--- a/Engine/Source/Programs/HeadlessChaos/Private/HeadlessChaosTestConstraintSolverCore.cpp
+++ b/Engine/Source/Programs/HeadlessChaos/Private/HeadlessChaosTestConstraintSolverCore.cpp
@@ -4,13 +4,11 @@
 #include "HeadlessChaosTestUtility.h"
 
 #include "Chaos/CollisionResolutionTypes.h"
-//#include "Chaos/Collision/PBDCollisionSolver.h"
 #include "Chaos/Evolution/SolverBody.h"
 #include "Chaos/GJK.h"
 #include "Chaos/Pair.h"
 #include "Chaos/PBDRigidsEvolution.h"
 #include "Chaos/PBDRigidParticles.h"
-#include "Chaos/PBDCollisionConstraintsPGS.h"
 #include "Chaos/Sphere.h"
 #include "Chaos/Utilities.h"
 
diff --git a/Engine/Source/Programs/HeadlessChaos/Private/HeadlessChaosTestConstraints.cpp b/Engine/Source/Programs/HeadlessChaos/Private/HeadlessChaosTestConstraints.cpp
--- a/Engine/Source/Programs/HeadlessChaos/Private/HeadlessChaosTestConstraints.cpp
+++ b/Engine/Source/Programs/HeadlessChaos/Private/HeadlessChaosTestConstraints.cpp
@@ -33,6 +33,8 @@ namespace ChaosTest {
 			THandleArray<FChaosPhysicsMaterial> PhysicalMaterials;
 			TEvolution Evolution(Particles, PhysicalMaterials);
 			TArray<FPBDRigidParticleHandle*> Dynamics = Evolution.CreateDynamicParticles(1);
+			Evolution.EnableParticle(Dynamics[0]);
+
 			TArray<FVec3> Positions = { FVec3(0) };
 			FPBDPositionConstraints PositionConstraints(MoveTemp(Positions), MoveTemp(Dynamics), 1.f);
 			InitEvolutionSettings(Evolution);
@@ -50,6 +52,7 @@ namespace ChaosTest {
 			InitEvolutionSettings(Evolution);
 			TArray<FPBDRigidParticleHandle*> Dynamics = Evolution.CreateDynamicParticles(1);
 			Dynamics[0]->SetGravityEnabled(false);
+			Evolution.EnableParticle(Dynamics[0]);
 
 			TArray<FVec3> Positions = { FVec3(1) };
 			FPBDPositionConstraints PositionConstraints(MoveTemp(Positions), MoveTemp(Dynamics), 0.5f);
@@ -106,6 +109,9 @@ namespace ChaosTest {
 		JointConstraints.AddConstraint(JointParticles, FRigidTransform3(JointConstraintPosition, FRotation3::FromIdentity()));
 		Evolution.AddConstraintContainer(JointConstraints);
 
+		Evolution.EnableParticle(Dynamics[0]);
+		Evolution.EnableParticle(Dynamics[1]);
+
 		FReal Dt = 0.1f;
 		for (int32 TimeIndex = 0; TimeIndex < 100; ++TimeIndex)
 		{
@@ -167,6 +173,7 @@ namespace ChaosTest {
 			SuspensionConstraints.AddConstraint(DynamicParticle, SuspensionLocalLocationA, SuspensionSettings);
 
 			Evolution.AddConstraintContainer(SuspensionConstraints);
+			Evolution.EnableParticle(DynamicParticle);
 
 			Evolution.AdvanceOneTimeStep(0.1);
 			Evolution.EndFrame(0.1);
@@ -217,6 +224,7 @@ namespace ChaosTest {
 			SuspensionConstraints.AddConstraint(DynamicParticle, SuspensionLocalLocationB, SuspensionSettings);
 
 			Evolution.AddConstraintContainer(SuspensionConstraints);
+			Evolution.EnableParticle(DynamicParticle);
 
 			const FVec3& Pos = Evolution.GetParticleHandles().Handle(0)->X();
 			const FRotation3& Rot = Evolution.GetParticleHandles().Handle(0)->R();
@@ -293,6 +301,7 @@ namespace ChaosTest {
 		}
 
 		Evolution.AddConstraintContainer(SuspensionConstraints);
+		Evolution.EnableParticle(DynamicParticle);
 
 		const FVec3& Pos = Evolution.GetParticleHandles().Handle(0)->X();
 		const FRotation3& Rot = Evolution.GetParticleHandles().Handle(0)->R();
diff --git a/Engine/Source/Programs/HeadlessChaos/Private/HeadlessChaosTestForces.cpp b/Engine/Source/Programs/HeadlessChaos/Private/HeadlessChaosTestForces.cpp
--- a/Engine/Source/Programs/HeadlessChaos/Private/HeadlessChaosTestForces.cpp
+++ b/Engine/Source/Programs/HeadlessChaos/Private/HeadlessChaosTestForces.cpp
@@ -22,6 +22,7 @@ namespace ChaosTest
 		TEvolution Evolution(Particles, PhysicalMaterials);
 		
 		TArray<FPBDRigidParticleHandle*> Dynamics = Evolution.CreateDynamicParticles(1);
+		Evolution.EnableParticle(Dynamics[0]);
 		Evolution.AdvanceOneTimeStep(0.1);
 		EXPECT_LT(Dynamics[0]->X()[2], 0);
 	}
diff --git a/Engine/Source/Programs/HeadlessChaos/Private/HeadlessChaosTestGraph.cpp b/Engine/Source/Programs/HeadlessChaos/Private/HeadlessChaosTestGraph.cpp
--- a/Engine/Source/Programs/HeadlessChaos/Private/HeadlessChaosTestGraph.cpp
+++ b/Engine/Source/Programs/HeadlessChaos/Private/HeadlessChaosTestGraph.cpp
@@ -25,6 +25,13 @@ namespace ChaosTest {
 	template<int32 T_TYPEID>
 	class TMockGraphConstraintHandle;
 
+	struct FMockConstraint
+	{
+		TVec2<FGeometryParticleHandle*> ConstrainedParticles;
+		bool bEnabled;
+		bool bSleeping;
+	};
+
 	/**
 	 * Constraint Container with minimal API required to test the Graph.
 	 * We can pretend we have many constraint containers of different types
@@ -35,12 +42,6 @@ namespace ChaosTest {
 	{
 	public:
 		using FConstraintContainerHandle = TMockGraphConstraintHandle<T_TYPEID>;
-		struct FMockConstraint
-		{
-			TVec2<int32> ConstrainedParticles;
-			bool bEnabled;
-			bool bSleeping;
-		};
 
 		TMockGraphConstraints()
 			: TPBDIndexedConstraintContainer<TMockGraphConstraints<T_TYPEID>>(FConstraintContainerHandle::StaticType())
@@ -48,12 +49,14 @@ namespace ChaosTest {
 			TPBDIndexedConstraintContainer<TMockGraphConstraints<T_TYPEID>>::SetContainerId(T_TYPEID);
 		}
 
+		using FPBDConstraintContainer::GetContainerId;
+
 		//
 		// FConstraintContainer Implementation
 		//
 		virtual int32 GetNumConstraints() const override final { return Constraints.Num(); }
 		virtual void ResetConstraints() override final {}
-		virtual void AddConstraintsToGraph(FPBDIslandManager& IslandManager) override final {}
+		virtual void AddConstraintsToGraph(FPBDIslandManager& IslandManager) override final;
 		virtual void PrepareTick() override final {}
 		virtual void UnprepareTick() override final {}
 
@@ -77,60 +80,38 @@ namespace ChaosTest {
 		void ApplyVelocityConstraints(const TArrayView<int32>& ConstraintIndices, const FReal Dt, const int32 It, const int32 NumIts) {}
 		void ApplyProjectionConstraints(const TArrayView<int32>& ConstraintIndices, const FReal Dt, const int32 It, const int32 NumIts) {}
 
-		TVec2<int32> ConstraintParticleIndices(const int32 ConstraintIndex) const
+		TVec2<FGeometryParticleHandle*> ConstrainedParticles(const int32 ConstraintIndex) const
 		{
 			return Constraints[ConstraintIndex].ConstrainedParticles;
 		}
 
-		void AddConstraint(const TVec2<int32>& InConstraintedParticles)
+		void AddConstraint(const TVec2<FGeometryParticleHandle*>& InConstraintedParticles)
 		{
 			Constraints.Emplace(FMockConstraint({ InConstraintedParticles, true, false }));
 			Handles.Emplace(HandleAllocator.AllocHandle(this, Handles.Num()));
 		}
 
-		bool AreConstraintsIndependent(const TPBDRigidParticles<FReal, 3>& InParticles, const TArray<FConstraintHandle*>& InConstraintHandles)
+		// Veryify that all constraints in the set are independen (no two reference the same particle). Used to validate coloring.
+		bool AreConstraintsIndependent(const TArray<FConstraintHandle*>& InConstraintHandles)
 		{
 			bool bIsValid = true;
-			TSet<int32> ParticleSet;
+			TSet<const FGeometryParticleHandle*> ParticleSet;
 			for (FConstraintHandle* ConstraintHandle : InConstraintHandles)
 			{
 				const int32 ConstraintIndex = ConstraintHandle->As<FConstraintContainerHandle>()->GetConstraintIndex();
 				FMockConstraint& Constraint = Constraints[ConstraintIndex];
-				if (ParticleSet.Contains(Constraint.ConstrainedParticles[0]) && !InParticles.HasInfiniteMass(Constraint.ConstrainedParticles[0]))
+				FConstGenericParticleHandle Particle0 = Constraint.ConstrainedParticles[0];
+				FConstGenericParticleHandle Particle1 = Constraint.ConstrainedParticles[1];
+				if (ParticleSet.Contains(Particle0->Handle()) && Particle0->IsDynamic())
 				{
 					bIsValid = false;
 				}
-				if (ParticleSet.Contains(Constraint.ConstrainedParticles[1]) && !InParticles.HasInfiniteMass(Constraint.ConstrainedParticles[1]))
+				if (ParticleSet.Contains(Particle1->Handle()) && Particle1->IsDynamic())
 				{
 					bIsValid = false;
 				}
-				ParticleSet.Add(Constraint.ConstrainedParticles[0]);
-				ParticleSet.Add(Constraint.ConstrainedParticles[1]);
-			}
-			return bIsValid;
-		}
-
-		bool AreConstraintsIndependent(const TArray<TGeometryParticleHandle<FReal, 3>*>& Particles, const TArray<FConstraintHandle*>& InConstraintHandles)
-		{
-			bool bIsValid = true;
-			TSet<TGeometryParticleHandle<FReal, 3>*> ParticleSet;
-			for (FConstraintHandle* ConstraintHandle : InConstraintHandles)
-			{
-				const int32 ConstraintIndex = ConstraintHandle->As<FConstraintContainerHandle>()->GetConstraintIndex();
-				FMockConstraint& Constraint = Constraints[ConstraintIndex];
-				TGeometryParticleHandle<FReal, 3>* Particle0 = Particles[Constraint.ConstrainedParticles[0]];
-				if (ParticleSet.Contains(Particle0) && Particle0->CastToRigidParticle() && Particle0->ObjectState() == EObjectStateType::Dynamic)
-				{
-					bIsValid = false;
-				}
-
-				TGeometryParticleHandle<FReal, 3>* Particle1 = Particles[Constraint.ConstrainedParticles[1]];
-				if (ParticleSet.Contains(Particle1) && Particle1->CastToRigidParticle() && Particle1->ObjectState() == EObjectStateType::Dynamic)
-				{
-					bIsValid = false;
-				}
-				ParticleSet.Add(Particle0);
-				ParticleSet.Add(Particle1);
+				ParticleSet.Add(Particle0->Handle());
+				ParticleSet.Add(Particle1->Handle());
 			}
 			return bIsValid;
 		}
@@ -177,7 +158,7 @@ namespace ChaosTest {
 			return ConcreteContainer()->Constraints[GetConstraintIndex()].bSleeping;
 		}
 
-		virtual FParticlePair GetConstrainedParticles() const override { return FParticlePair(nullptr, nullptr); }
+		virtual FParticlePair GetConstrainedParticles() const override { return ConcreteContainer()->Constraints[GetConstraintIndex()].ConstrainedParticles; }
 
 		static const FConstraintHandleTypeID& StaticType()
 		{
@@ -186,6 +167,15 @@ namespace ChaosTest {
 		}
 	};
 
+	template<int32 T_TYPEID>
+	void TMockGraphConstraints<T_TYPEID>::AddConstraintsToGraph(FPBDIslandManager& IslandManager)
+	{
+		for (TMockGraphConstraintHandle<T_TYPEID>* Handle : Handles)
+		{
+			IslandManager.AddConstraint(GetContainerId(), Handle, Handle->GetConstrainedParticles());
+		}
+	}
+
 
 	template class TMockGraphConstraints<0>;
 	template class TMockGraphConstraints<1>;
@@ -244,13 +234,13 @@ namespace ChaosTest {
 		TMockGraphConstraints<0> ConstraintsOfType0;
 		for (const auto& ConstrainedParticles : ConstrainedParticles0)
 		{
-			ConstraintsOfType0.AddConstraint(ConstrainedParticles);
+			ConstraintsOfType0.AddConstraint({ AllParticles[ConstrainedParticles[0]], AllParticles[ConstrainedParticles[1]] });
 		}
 
 		TMockGraphConstraints<1> ConstraintsOfType1;
 		for (const auto& ConstrainedParticles : ConstrainedParticles1)
 		{
-			ConstraintsOfType0.AddConstraint(ConstrainedParticles);
+			ConstraintsOfType0.AddConstraint({ AllParticles[ConstrainedParticles[0]], AllParticles[ConstrainedParticles[1]] });
 		}
 
 		// Set up the particle graph
@@ -258,34 +248,27 @@ namespace ChaosTest {
 		Graph.AddConstraintContainer(ConstraintsOfType0);
 		Graph.AddConstraintContainer(ConstraintsOfType1);
 
-		Graph.InitializeGraph(SOAs.GetNonDisabledDynamicView());
-
-		Graph.ReserveConstraints(ConstraintsOfType0.GetNumConstraints());
-		for (int32 ConstraintIndex = 0; ConstraintIndex < ConstraintsOfType0.GetNumConstraints(); ++ConstraintIndex)
+		for (FGeometryParticleHandle* Particle : AllParticles)
 		{
-			TVec2<int32> Indices = ConstraintsOfType0.ConstraintParticleIndices(ConstraintIndex);
-			Graph.AddConstraint(0, ConstraintsOfType0.Handles[ConstraintIndex], TVec2<TGeometryParticleHandle<FReal, 3>*>(AllParticles[Indices[0]], AllParticles[Indices[1]]));
+			Graph.AddParticle(Particle);
 		}
 
-		Graph.ReserveConstraints(ConstraintsOfType1.GetNumConstraints());
-		for (int32 ConstraintIndex = 0; ConstraintIndex < ConstraintsOfType1.GetNumConstraints(); ++ConstraintIndex)
-		{
-			TVec2<int32> Indices = ConstraintsOfType1.ConstraintParticleIndices(ConstraintIndex);
-			Graph.AddConstraint(1, ConstraintsOfType1.Handles[ConstraintIndex], TVec2<TGeometryParticleHandle<FReal, 3>*>(AllParticles[Indices[0]], AllParticles[Indices[1]]));
-		}
-		
-		// Generate the constraint/particle islands
-		Graph.UpdateGraph();
+		Graph.InitializeGraph(SOAs.GetNonDisabledDynamicView());
+
+		ConstraintsOfType0.AddConstraintsToGraph(Graph);
+		ConstraintsOfType1.AddConstraintsToGraph(Graph);
+
 		Graph.UpdateIslands(SOAs);
 		
-		// Islands should be end up with the following particles (note: particle 17 is infinite mass and can appear in multiple islands)
+		// Islands should be end up with the following particles (note: particle 17 is infinite mass and is referenced by multiple islands, but
+		// is not included in the list of island particles)
 		TArray<TSet<TGeometryParticleHandle<FReal, 3>*>> ExpectedIslandParticles = {
 			{AllParticles[0], AllParticles[1], AllParticles[2], AllParticles[3], AllParticles[4], AllParticles[5], AllParticles[6]},
 			{AllParticles[7], AllParticles[8], AllParticles[9], AllParticles[10], AllParticles[11], AllParticles[13], AllParticles[18]},
 			{AllParticles[12]},
-			{AllParticles[14], AllParticles[15], AllParticles[16], AllParticles[17]},
+			{AllParticles[14], AllParticles[15], AllParticles[16]},
 			{AllParticles[19]},
-			{AllParticles[17], AllParticles[20]},
+			{AllParticles[20]},
 		};
 
 		// Get the Island indices which map to the ExpectedIslandParticles
@@ -318,6 +301,8 @@ namespace ChaosTest {
 				EXPECT_TRUE(ExpectedIslandParticles[ExpectedIslandIndex].Contains(CalculatedIslandParticle.GetParticle()));
 			}
 		}
+
+		Graph.Reset();
 	}
 
 	void CheckIslandIntegrity(TArray<TSet<TGeometryParticleHandle<FReal,3>*>>& ExpectedIslandParticles, const TArray<int32> CalculatedIslandIndices, FPBDConstraintGraph& Graph)
@@ -398,6 +383,9 @@ namespace ChaosTest {
 
 		IterationData.Push(Data0);
 
+		// NOTE: only dynamic particles are included in the generated islands
+		// i.e. particle 2 will be absent
+
 		//////////////////////////////////////////////////////////////////////////
 		// Iteration 1
 		FIterationData Data1;
@@ -412,7 +400,7 @@ namespace ChaosTest {
 
 		// Islands should be end up with the following particles
 		Data1.ExpectedIslandParticleIndices = {
-			{0,1,2},
+			{0,1},
 			{3,4,5}
 		};
 
@@ -440,8 +428,8 @@ namespace ChaosTest {
 		};
 
 		Data2.ExpectedIslandParticleIndices = {
-			{0,1,2},
-			{2,3,4,5},
+			{0,1},
+			{3,4,5},
 		};
 
 		Data2.ExpectedIslandEdges = {
@@ -468,7 +456,7 @@ namespace ChaosTest {
 
 		Data3.ExpectedIslandParticleIndices = {
 			{0,1},
-			{2,3},
+			{3},
 			{4,5}
 		};
 
@@ -495,7 +483,7 @@ namespace ChaosTest {
 
 		Data4.ExpectedIslandParticleIndices = {
 			{0,1},
-			{2,3},
+			{3},
 			{4},
 			{5}
 		};
@@ -524,8 +512,8 @@ namespace ChaosTest {
 
 		Data5.ExpectedIslandParticleIndices = {
 			{1},
-			{0,2},
-			{2,3},
+			{0},
+			{3},
 			{4},
 			{5}
 		};
@@ -547,7 +535,11 @@ namespace ChaosTest {
 		for (int Iteration = 0; Iteration < IterationData.Num(); Iteration++)
 		{
 			FIterationData& IterData = IterationData[Iteration];
-			
+
+			FPBDConstraintGraph Graph;
+			//FPBDConstraintColor GraphColor;
+			const int32 ContainerId = 0;
+
 			// convert ExpectedIslandParticleIndices to ExpectedIslandParticles for ease of later comparison
 			for (int32 I = 0; I < IterData.ExpectedIslandParticleIndices.Num(); I++)
 			{
@@ -557,46 +549,30 @@ namespace ChaosTest {
 				for (int32 Index : Set)
 				{
 					HandlesSet.Add(AllParticles[Index]);
+					Graph.AddParticle(AllParticles[Index]);
 				}
 				IterData.ExpectedIslandParticles.Add(HandlesSet);
 			}
 
-			FPBDConstraintGraph Graph;
-			//FPBDConstraintColor GraphColor;
-			const int32 ContainerId = 0;
-
-			// Set up the particle graph
-			Graph.InitializeGraph(SOAs.GetNonDisabledDynamicView());
-
 			// add constraints
 			TMockGraphConstraints<0> ConstraintsOfType0;
 			Graph.AddConstraintContainer(ConstraintsOfType0);
 
 			for (const auto& ConstrainedParticles : IterData.ConstrainedParticles)
 			{
-				ConstraintsOfType0.AddConstraint(ConstrainedParticles);
-			}
-
-			Graph.ReserveConstraints(ConstraintsOfType0.GetNumConstraints());
-			for (int32 ConstraintIndex = 0; ConstraintIndex < ConstraintsOfType0.GetNumConstraints(); ++ConstraintIndex)
-			{
-				TVec2<int32> Indices = ConstraintsOfType0.ConstraintParticleIndices(ConstraintIndex);
-				Graph.AddConstraint(ContainerId, ConstraintsOfType0.Handles[ConstraintIndex], TVec2<TGeometryParticleHandle<FReal, 3>*>(AllParticles[Indices[0]], AllParticles[Indices[1]]));
+				ConstraintsOfType0.AddConstraint({ AllParticles[ConstrainedParticles[0]],AllParticles[ConstrainedParticles[1]] });
 			}
 
 			// Generate the constraint/particle islands
-			Graph.UpdateGraph();
+			Graph.InitializeGraph(SOAs.GetNonDisabledDynamicView());
+
+			ConstraintsOfType0.AddConstraintsToGraph(Graph);
+
 			Graph.UpdateIslands(SOAs);
 
 			// Check the number of islands
 			EXPECT_EQ(IterData.ExpectedIslandParticleIndices.Num(), Graph.NumIslands());
 
-			// Compute levels and colors
-			Graph.GetIslandGraph()->InitSorting();
-
-			Graph.GetIslandGraph()->ComputeLevels();
-			Graph.GetIslandGraph()->ComputeColors(0);
-
 			// Assign color to constraints
 			//GraphColor.InitializeColor(Graph);
 			//for (int32 IslandIndex = 0; IslandIndex < Graph.NumIslands(); ++IslandIndex)
@@ -667,16 +643,11 @@ namespace ChaosTest {
 	{
 
 		SOAs.ClearTransientDirty();
+
 		Graph.InitializeGraph(SOAs.GetNonDisabledDynamicView());
 
-		Graph.ReserveConstraints(Constraints.GetNumConstraints());
-		for(int32 ConstraintIndex = 0; ConstraintIndex < Constraints.GetNumConstraints(); ++ConstraintIndex)
-		{
-			const TVec2<int32> Indices = Constraints.ConstraintParticleIndices(ConstraintIndex);
-			Graph.AddConstraint(0,Constraints.Handles[ConstraintIndex],TVec2<TGeometryParticleHandle<FReal,3>*>(Particles[Indices[0]],Particles[Indices[1]]));
-		}
+		Constraints.AddConstraintsToGraph(Graph);
 
-		Graph.UpdateGraph();
 		Graph.UpdateIslands(SOAs);
 		for (int32 IslandIndex = 0; IslandIndex < Graph.NumIslands(); ++IslandIndex)
 		{
@@ -760,11 +731,16 @@ namespace ChaosTest {
 			TMockGraphConstraints<0> Constraints;
 			for (const auto& ConstrainedParticleIndices : ConstrainedParticles)
 			{
-				Constraints.AddConstraint(ConstrainedParticleIndices);
+				Constraints.AddConstraint({ Particles[ConstrainedParticleIndices[0]], Particles[ConstrainedParticleIndices[1]] });
 			}
 
 			FPBDConstraintGraph Graph;
 			Graph.AddConstraintContainer(Constraints);
+			for (FGeometryParticleHandle* Particle : Particles)
+			{
+				Graph.AddParticle(Particle);
+			}
+
 			for (int32 LoopIndex = 0; LoopIndex < 5 + PhysicalMaterial->SleepCounterThreshold; ++LoopIndex)
 			{
 				// @todo(chaos): redo this test - sleeping now used a damped velocity rather than current velocity
@@ -796,6 +772,8 @@ namespace ChaosTest {
 				EXPECT_EQ(ContainsHelper(SOAs.GetDirtyParticlesView(),Particles[4]),bIsDirty);
 				EXPECT_EQ(ContainsHelper(SOAs.GetDirtyParticlesView(),Particles[5]),bIsDirty);
 			}
+
+			Graph.Reset();
 		}
 	}
 
@@ -848,11 +826,15 @@ namespace ChaosTest {
 			TMockGraphConstraints<0> Constraints;
 			for (const auto& ConstrainedParticleIndices : ConstrainedParticles)
 			{
-				Constraints.AddConstraint(ConstrainedParticleIndices);
+				Constraints.AddConstraint({ Particles[ConstrainedParticleIndices[0]], Particles[ConstrainedParticleIndices[1]] });
 			}
 
 			FPBDConstraintGraph Graph;
 			Graph.AddConstraintContainer(Constraints);
+			for (FGeometryParticleHandle* Particle : Particles)
+			{
+				Graph.AddParticle(Particle);
+			}
 
 			const int32 WakeUpFrame = 5 + PhysicalMaterial->SleepCounterThreshold;
 			for (int32 LoopIndex = 0; LoopIndex < WakeUpFrame + 5; ++LoopIndex)
@@ -868,7 +850,7 @@ namespace ChaosTest {
 				{
 					for (const auto& ConstrainedParticleIndices : ConstrainedParticlesOnWake)
 					{
-						Constraints.AddConstraint(ConstrainedParticleIndices);
+						Constraints.AddConstraint({ Particles[ConstrainedParticleIndices[0]], Particles[ConstrainedParticleIndices[1]] });
 					}
 				}
 
@@ -897,6 +879,8 @@ namespace ChaosTest {
 				EXPECT_EQ(ContainsHelper(SOAs.GetDirtyParticlesView(),Particles[4]),bIsDirty);
 				EXPECT_EQ(ContainsHelper(SOAs.GetDirtyParticlesView(),Particles[5]),b5IsDirty);
 			}
+
+			Graph.Reset();
 		}
 	}
 
@@ -949,14 +933,18 @@ namespace ChaosTest {
 			TMockGraphConstraints<0> Constraints;
 			for (const auto& ConstrainedParticleIndices : ConstrainedParticles)
 			{
-				Constraints.AddConstraint(ConstrainedParticleIndices);
+				Constraints.AddConstraint({ Particles[ConstrainedParticleIndices[0]], Particles[ConstrainedParticleIndices[1]] });
 			}
 
 			FPBDConstraintGraph Graph;
 			Graph.AddConstraintContainer(Constraints);
-
+			for (FGeometryParticleHandle* Particle : Particles)
+			{
+				Graph.AddParticle(Particle);
+			}
+			
 			const int32 MergeFrame = 5 + PhysicalMaterial->SleepCounterThreshold;
-			const int32 SleepAfterMergeFrame = MergeFrame + PhysicalMaterial->SleepCounterThreshold + 1;	//first frame after merge must wake up
+			const int32 SleepAfterMergeFrame = MergeFrame + PhysicalMaterial->SleepCounterThreshold;
 			for(int32 LoopIndex = 0; LoopIndex < SleepAfterMergeFrame + 5; ++LoopIndex)
 			{
 				if(LoopIndex == MergeFrame)
@@ -977,7 +965,7 @@ namespace ChaosTest {
 				{
 					for (const auto& ConstrainedParticleIndices : ConstrainedParticlesOnWake)
 					{
-						Constraints.AddConstraint(ConstrainedParticleIndices);
+						Constraints.AddConstraint({ Particles[ConstrainedParticleIndices[0]], Particles[ConstrainedParticleIndices[1]] });
 					}
 				}
 
@@ -1014,6 +1002,8 @@ namespace ChaosTest {
 				EXPECT_EQ(ContainsHelper(SOAs.GetDirtyParticlesView(),Particles[4]),bIsDirty);
 				EXPECT_EQ(ContainsHelper(SOAs.GetDirtyParticlesView(),Particles[5]),b5IsDirty);
 			}
+
+			Graph.Reset();
 		}
 	}
 
@@ -1127,23 +1117,22 @@ namespace ChaosTest {
 		TMockGraphConstraints<0> Constraints;
 		for (const auto& ConstrainedParticleIndices : ConstrainedParticles)
 		{
-			Constraints.AddConstraint(ConstrainedParticleIndices);
+			Constraints.AddConstraint({ AllParticles[ConstrainedParticleIndices[0]], AllParticles[ConstrainedParticleIndices[1]] });
 		}
 		
 		// Build the connectivity graph and islands
 		FPBDConstraintGraph Graph;
 		Graph.AddConstraintContainer(Constraints);
 		//FPBDConstraintColor GraphColor;
-		const int32 ContainerId = 0;
-		
-		Graph.InitializeGraph(SOAs.GetNonDisabledDynamicView());
-		Graph.ReserveConstraints(Constraints.GetNumConstraints());
-		for (int32 ConstraintIndex = 0; ConstraintIndex < Constraints.GetNumConstraints(); ++ConstraintIndex)
+		for (FGeometryParticleHandle* Particle : AllParticles)
 		{
-			const TVec2<int32>& Indices = Constraints.ConstraintParticleIndices(ConstraintIndex);
-			Graph.AddConstraint(ContainerId, Constraints.Handles[ConstraintIndex], TVec2<TGeometryParticleHandle<FReal, 3>*>(AllParticles[Indices[0]], AllParticles[Indices[1]]) );
+			Graph.AddParticle(Particle);
 		}
-		Graph.UpdateGraph();
+
+		Graph.InitializeGraph(SOAs.GetNonDisabledDynamicView());
+		
+		Constraints.AddConstraintsToGraph(Graph);
+
 		Graph.UpdateIslands(SOAs);
 
 		// It's a connected grid, so only one island
@@ -1203,6 +1192,8 @@ namespace ChaosTest {
 		//	//EXPECT_LE(MaxIslandColor, MaxNumGreedyColors - 1);
 		//	EXPECT_LE(MaxIslandColor, MaxNumGreedyColors);
 		//}
+
+		Graph.Reset();
 	}
 
 	/**
@@ -1282,10 +1273,6 @@ namespace ChaosTest {
 		TArray<FPBDRigidParticleHandle*> Dynamics = SOAs.CreateDynamicParticles(500);
 		for (auto Dynamic : Dynamics) { AllParticles.Add(Dynamic); }
 
-		// Set up the particle graph
-		FPBDConstraintGraph Graph;
-		Graph.InitializeGraph(SOAs.GetNonDisabledDynamicView());
-		
 		const int32 NumParticles = AllParticles.Num();
 
 		// Put together constraint lists
@@ -1302,61 +1289,86 @@ namespace ChaosTest {
 			ConstrainedParticles2.Add({ (ConstraintIdx * 6) % NumParticles, (ConstraintIdx * 7) % NumParticles });
 		}
 		
+		int32 NextContainerId = 0;
+
 		// Create some position constraints
-		Graph.ReserveConstraints(ConstrainedParticles0.Num());
 		FPBDPositionConstraints PositionConstraints;
+		PositionConstraints.SetContainerId(NextContainerId++);
 		for (const TVec2<int32>& Indices : ConstrainedParticles0)
 		{
 			FPBDRigidParticleHandle* Particle0 = AllParticles[Indices[0]]->CastToRigidParticle();
-			auto* NewConstraint = PositionConstraints.AddConstraint(Particle0, FVec3(0, 0, 0 ));
-			Graph.AddConstraint(0, NewConstraint, TVec2<FGeometryParticleHandle*>(Particle0, nullptr));
+			PositionConstraints.AddConstraint(Particle0, FVec3(0, 0, 0 ));
 		}
 
 		// Create some suspension constraints
-		Graph.ReserveConstraints(ConstrainedParticles1.Num());
 		FPBDSuspensionConstraints SuspensionConstraints;
+		SuspensionConstraints.SetContainerId(NextContainerId++);
 		for (const TVec2<int32>& Indices : ConstrainedParticles1)
 		{
 			FPBDSuspensionSettings Settings;
 			FPBDRigidParticleHandle* Particle0 = AllParticles[Indices[0]]->CastToRigidParticle();
 			FPBDRigidParticleHandle* Particle1 = AllParticles[Indices[1]]->CastToRigidParticle();
-			auto* NewConstraint = SuspensionConstraints.AddConstraint(Particle0, FVec3(0, 0, 0), Settings);
-			Graph.AddConstraint(1, NewConstraint, TVec2<FGeometryParticleHandle*>(Particle0, Particle1));
+			SuspensionConstraints.AddConstraint(Particle0, FVec3(0, 0, 0), Settings);
 		}
 
 		// Create some dynamic spring constraints.
-		Graph.ReserveConstraints(ConstrainedParticles2.Num());
 		FPBDRigidDynamicSpringConstraints SpringConstraints;
+		SpringConstraints.SetContainerId(NextContainerId++);
 		for (const TVec2<int32>& Indices : ConstrainedParticles2)
 		{
 			FPBDRigidParticleHandle* Particle0 = AllParticles[Indices[0]]->CastToRigidParticle();
 			FPBDRigidParticleHandle* Particle1 = AllParticles[Indices[1]]->CastToRigidParticle();
-			auto* NewConstraint = SpringConstraints.AddConstraint(FPBDRigidDynamicSpringConstraints::FConstrainedParticlePair(Particle0, Particle1));
-			Graph.AddConstraint(2, NewConstraint, TVec2<FGeometryParticleHandle*>(Particle0, Particle1));
+			SpringConstraints.AddConstraint(FPBDRigidDynamicSpringConstraints::FConstrainedParticlePair(Particle0, Particle1));
 		}
+
+		// Set up the particle graph
+		FPBDConstraintGraph Graph;
+		Graph.AddConstraintContainer(PositionConstraints);
+		Graph.AddConstraintContainer(SuspensionConstraints);
+		Graph.AddConstraintContainer(SpringConstraints);
+		for (FGeometryParticleHandle* Particle : AllParticles)
+		{
+			Graph.AddParticle(Particle);
+		}
+		Graph.InitializeGraph(SOAs.GetNonDisabledDynamicView());
 		
+		PositionConstraints.AddConstraintsToGraph(Graph);
+		SuspensionConstraints.AddConstraintsToGraph(Graph);
+		SpringConstraints.AddConstraintsToGraph(Graph);
+
+		Graph.EndTick();
+
 		ConstraintGraphValidation::ValidateConstraintGraphParticleHandles(Graph);
 
+
 		// Add new particles
 		TArray<TPBDRigidParticleHandle<FReal, 3>*> NewDynamics = SOAs.CreateDynamicParticles(20);
 		ConstraintGraphValidation::ValidateConstraintGraphParticleHandles(Graph);
 
 		// Remove some constraints
+		Graph.RemoveConstraint(PositionConstraints.GetContainerId(), PositionConstraints.GetConstraintHandle(0));
+		Graph.RemoveConstraint(SuspensionConstraints.GetContainerId(), SuspensionConstraints.GetConstraintHandle(0));
+		Graph.RemoveConstraint(SpringConstraints.GetContainerId(), SpringConstraints.GetConstraintHandle(0));
 		PositionConstraints.RemoveConstraint(0);
 		SuspensionConstraints.RemoveConstraint(0);
 		SpringConstraints.RemoveConstraint(0);
 		ConstraintGraphValidation::ValidateConstraintGraphParticleHandles(Graph);
 		
 		// Add new constraints
-		Graph.ReserveConstraints(ConstrainedParticles1.Num());
 		FPBDJointConstraints JointConstraints;
+		JointConstraints.SetContainerId(NextContainerId++);
 		for (const TVec2<int32>& Indices : ConstrainedParticles1)
 		{
 			FPBDRigidParticleHandle* Particle0 = AllParticles[Indices[0]]->CastToRigidParticle();
 			FPBDRigidParticleHandle* Particle1 = AllParticles[Indices[1]]->CastToRigidParticle();
-			auto* NewConstraint = JointConstraints.AddConstraint(FParticlePair(Particle0,Particle1), FRigidTransform3());
-			Graph.AddConstraint(3, NewConstraint, TVec2<FGeometryParticleHandle*>(Particle0, Particle1));
+			JointConstraints.AddConstraint(FParticlePair(Particle0,Particle1), FRigidTransform3());
 		}
+
+		Graph.AddConstraintContainer(JointConstraints);
+		Graph.InitializeGraph(SOAs.GetNonDisabledDynamicView());
+		JointConstraints.AddConstraintsToGraph(Graph);
+		Graph.EndTick();
+
 		ConstraintGraphValidation::ValidateConstraintGraphParticleHandles(Graph);
 
 		// Remove a block of particles
@@ -1370,7 +1382,7 @@ namespace ChaosTest {
 		// Disable a block of particles
 		for (int32 Idx = 10; Idx < 20; ++Idx)
 		{
-			Graph.DisableParticle({ Dynamics[Idx] });
+			Graph.RemoveParticle({ Dynamics[Idx] });
 			SOAs.DisableParticle(Dynamics[Idx]);
 		}
 		ConstraintGraphValidation::ValidateConstraintGraphParticleHandles(Graph);
@@ -1378,11 +1390,12 @@ namespace ChaosTest {
 		// Re-enable the block of particles
 		for (int32 Idx = 10; Idx < 20; ++Idx)
 		{
-			Graph.EnableParticle(Dynamics[Idx], nullptr);
+			Graph.AddParticle(Dynamics[Idx]);
 			SOAs.EnableParticle(Dynamics[Idx]);
 		}
 		ConstraintGraphValidation::ValidateConstraintGraphParticleHandles(Graph);
 		
+		Graph.Reset();
 	}
 
 
diff --git a/Engine/Source/Programs/HeadlessChaos/Private/HeadlessChaosTestGraphEvolution.cpp b/Engine/Source/Programs/HeadlessChaos/Private/HeadlessChaosTestGraphEvolution.cpp
--- a/Engine/Source/Programs/HeadlessChaos/Private/HeadlessChaosTestGraphEvolution.cpp
+++ b/Engine/Source/Programs/HeadlessChaos/Private/HeadlessChaosTestGraphEvolution.cpp
@@ -34,6 +34,11 @@ namespace ChaosTest
 
 			// Create some particles and constraints connecting them in a chain: 0-1-2-3-...-N
 			ParticleHandles = Evolution.CreateDynamicParticles(NumParticles);
+
+			for (FGeometryParticleHandle* ParticleHandle : ParticleHandles)
+			{
+				Evolution.EnableParticle(ParticleHandle);
+			}
 		}
 
 		// Connect all the particles in a chain
@@ -213,11 +218,11 @@ namespace ChaosTest
 
 		Test.Advance();
 
-		// Constraint[1] is not in the graph or an island
+		// Constraint[1] is not in an island, but is still in the graph (we have not explicitly removed it)
 		EXPECT_NE(Test.ConstraintHandles[0]->GetConstraintGraphIndex(), INDEX_NONE);
-		EXPECT_EQ(Test.ConstraintHandles[1]->GetConstraintGraphIndex(), INDEX_NONE);	// Not in graph
+		EXPECT_NE(Test.ConstraintHandles[1]->GetConstraintGraphIndex(), INDEX_NONE);				// Still in graph
+		EXPECT_EQ(Test.IslandManager->GetConstraintIsland(Test.ConstraintHandles[1]), INDEX_NONE);	// Not in an island
 		EXPECT_NE(Test.ConstraintHandles[2]->GetConstraintGraphIndex(), INDEX_NONE);
-		EXPECT_TRUE(Test.IslandManager->DebugCheckConstraintNotInGraph(Test.ConstraintHandles[1]));
 
 		// Should now have 2 islands
 		EXPECT_EQ(Test.IslandManager->NumIslands(), 2);
@@ -587,8 +592,36 @@ namespace ChaosTest
 		EXPECT_TRUE(FConstGenericParticleHandle(Test.ParticleHandles[3])->Sleeping());
 	}
 
-	// Add a constraint between a sleeping and a kinematic body.
+	// Add a constraint between a sleeping and a kinematic body and tick.
+	// Nothing should change.
+	//
+	//		(d=dynamic, s=sleeping, k=kinematic)
+	//		Ak - Bs    =>    Ak - Bs
+	//
+	GTEST_TEST(GraphEvolutionTests, TestConstraintGraph_SleepingKinematicConstraint)
+	{
+		FGraphEvolutionTest Test(2);
+
+		// Make A kinematic, B sleeping
+		Test.Evolution.SetParticleObjectState(Test.ParticleHandles[0], EObjectStateType::Kinematic);
+		Test.Evolution.SetParticleObjectState(Test.ParticleHandles[1], EObjectStateType::Sleeping);
+
+		// Add a constraint A-B
+		Test.ConstraintHandles.Add(Test.Constraints.AddConstraint({ Test.ParticleHandles[0], Test.ParticleHandles[1] }));
+
+		Test.Advance();
+
+		// Everything asleep
+		EXPECT_TRUE(FConstGenericParticleHandle(Test.ParticleHandles[1])->Sleeping());
+		EXPECT_TRUE(Test.ConstraintHandles[0]->IsSleeping());
+	}
+
+	// Add a constraint between a sleeping and a kinematic body, one tick after the bodies were added.
+	// 
+	// This differs from TestConstraintGraph_SleepingKinematicConstraint in that we tick the scene one
+	// time before adding the constraint, which means the particles are already in separate islands.
 	// Nothing should wake and the constraint should be flagged as sleeping.
+	// 
 	// This behaviour is required for streaming to work since scene creation may
 	// be amortized over multiple frames and constraints may be made betweens
 	// sleeping particles in a later tick.
@@ -596,7 +629,7 @@ namespace ChaosTest
 	//		(d=dynamic, s=sleeping, k=kinematic)
 	//		Ak  Bs    =>    Ak - Bs
 	//
-	GTEST_TEST(GraphEvolutionTests, TestConstraintGraph_SleepingKinematicConstraint)
+	GTEST_TEST(GraphEvolutionTests, TestConstraintGraph_SleepingKinematicConstraint2)
 	{
 		FGraphEvolutionTest Test(2);
 
@@ -732,4 +765,141 @@ namespace ChaosTest
 		EXPECT_TRUE(Test.ConstraintHandles[2]->IsSleeping());
 	}
 
+	// Add constraints beween objects on the tick where their island goes to sleep, and one either side just to be sure.
+	//
+	//		(d=dynamic, s=sleeping, k=kinematic)
+	//		Ad   Bd  =>    As - Bs
+	//
+	GTEST_TEST(GraphEvolutionTests, TestConstraintGraph_SleepingSleepingConstraint_Timing)
+	{
+		// Count how many frames it takes the simulation to sleep
+		FGraphEvolutionTest SleepTest(2);
+		SleepTest.AdvanceUntilSleeping();
+		const int32 TicksToSleep = SleepTest.TickCount;
+
+		// Create a new simulation up to the sleep tick +/- a tick
+		// Verify that adding a constraint on that tick leaves the scene as expected
+		for (int32 SleepRelativeTickCount = -1; SleepRelativeTickCount < 2; ++SleepRelativeTickCount)
+		{
+			FGraphEvolutionTest Test(2);
+			for (int32 Frame = 0; Frame < TicksToSleep + SleepRelativeTickCount; ++Frame)
+			{
+				Test.Advance();
+			}
+			const bool bExpectSleep = (SleepRelativeTickCount >= 0);
+
+			// Should have 2 islands in a state that depends on whether we hit the sleep tick yet
+			EXPECT_EQ(Test.IslandManager->NumIslands(), 2);
+			EXPECT_EQ(FConstGenericParticleHandle(Test.ParticleHandles[0])->Sleeping(), bExpectSleep);
+			EXPECT_EQ(FConstGenericParticleHandle(Test.ParticleHandles[1])->Sleeping(), bExpectSleep);
+			EXPECT_EQ(Test.IslandManager->GetIsland(0)->IsSleeping(), bExpectSleep);
+			EXPECT_EQ(Test.IslandManager->GetIsland(1)->IsSleeping(), bExpectSleep);
+
+			// Add a constraint A-B and tick
+			Test.ConstraintHandles.Add(Test.Constraints.AddConstraint({ Test.ParticleHandles[0], Test.ParticleHandles[1] }));
+			Test.Advance();
+
+			// Should now have 1 island and it should be asleep
+			EXPECT_EQ(Test.IslandManager->NumIslands(), 1);
+			EXPECT_TRUE(FConstGenericParticleHandle(Test.ParticleHandles[0])->Sleeping());
+			EXPECT_TRUE(FConstGenericParticleHandle(Test.ParticleHandles[1])->Sleeping());
+			EXPECT_TRUE(Test.IslandManager->GetIsland(0)->IsSleeping());
+
+			// Constraint should also be asleep
+			EXPECT_TRUE(Test.ConstraintHandles[0]->IsSleeping());
+		}
+	}
+
+	// Test an edge case bug that is probably easy to accidentally reintroduce. This would leave
+	// a dangling pointer in the constraint graph due to a collision constraint being deleted while
+	// in a sleeping island.
+	// 
+	// The fix was to ensure that we build the Island particle and constraint lists for islands that
+	// have just been put to sleep (we still don't bother for thoise that were already asleep) so
+	// that we can visit all the particles and constraints to set the sleep state.
+	// 
+	// 1: A dynamic particle is in its own awake island
+	// - Tick
+	// 2a: The particle is manually put to sleep
+	// 2b: A constraint is added between the particle and a kinematic
+	// - Tick
+	// During the graph update on this tick, the particle's island is put to sleep in UpdateGraph
+	// because all particles in it are asleep. However, the constraint was added this tick as well,
+	// but when it was added the island was awake, so the constraint starts in the awake state.
+	// 
+	// Verify that the constraint does actually get put to sleep at some point in the graph update.
+	//
+	//		(d=dynamic, s=sleeping, k=kinematic)
+	//		Ak   Bd  =>    As - Bs
+	//
+	// NOTE: the transition to sleep is by a user call, not the automatic sleep-when-not-moving system
+	//
+	GTEST_TEST(GraphEvolutionTests, TestConstraintGraph_SleepingSleepingConstraint_Timing2)
+	{
+		FGraphEvolutionTest Test(2);
+
+		// Make A kinematic, B dynamic
+		Test.Evolution.SetParticleObjectState(Test.ParticleHandles[0], EObjectStateType::Kinematic);
+		Test.Evolution.SetParticleObjectState(Test.ParticleHandles[1], EObjectStateType::Dynamic);
+
+		Test.Advance();
+
+		EXPECT_EQ(Test.IslandManager->NumIslands(), 1);
+		EXPECT_FALSE(Test.IslandManager->GetIsland(0)->IsSleeping());
+
+		// Explicitly put B to sleep (its island will still be considered awake until the tick)
+		Test.Evolution.SetParticleObjectState(Test.ParticleHandles[1], EObjectStateType::Sleeping);
+
+		// Add a constraint A-B. B is asleep, but its island is awake
+		Test.ConstraintHandles.Add(Test.Constraints.AddConstraint({ Test.ParticleHandles[0], Test.ParticleHandles[1] }));
+
+		Test.Advance();
+
+		// Everything should be asleep
+		// The bug was that the constraint was still flagged as awake, but in a sleeping island.
+		EXPECT_EQ(Test.IslandManager->NumIslands(), 1);
+		EXPECT_TRUE(Test.IslandManager->GetIsland(0)->IsSleeping());
+		EXPECT_TRUE(FConstGenericParticleHandle(Test.ParticleHandles[1])->Sleeping());
+		EXPECT_TRUE(Test.ConstraintHandles[0]->IsSleeping());
+	}
+
+	// This is a very similar test to TestConstraintGraph_SleepingSleepingConstraint_Timing2
+	// in that is exposes the same bug where an island that is implicitly put to sleep because
+	// all its particles were explicitly put to sleep did not put its constraints to sleep.
+	//
+	//		(d=dynamic, s=sleeping, k=kinematic)
+	//		Ad -  Bd  =>    As - Bs
+	// 
+	// NOTE: the transition to sleep is by a user call, not the automatic sleep-when-not-moving system
+	//
+	GTEST_TEST(GraphEvolutionTests, TestConstraintGraph_SleepingSleepingConstraint_Timing3)
+	{
+		FGraphEvolutionTest Test(2);
+
+		// Make A, B dynamic
+		Test.Evolution.SetParticleObjectState(Test.ParticleHandles[0], EObjectStateType::Dynamic);
+		Test.Evolution.SetParticleObjectState(Test.ParticleHandles[1], EObjectStateType::Dynamic);
+
+		// Add a constraint A-B
+		Test.ConstraintHandles.Add(Test.Constraints.AddConstraint({ Test.ParticleHandles[0], Test.ParticleHandles[1] }));
+
+		Test.Advance();
+
+		EXPECT_EQ(Test.IslandManager->NumIslands(), 1);
+		EXPECT_FALSE(Test.IslandManager->GetIsland(0)->IsSleeping());
+
+		// Explicitly put both particles (and therefore their island) to sleep
+		Test.Evolution.SetParticleObjectState(Test.ParticleHandles[0], EObjectStateType::Sleeping);
+		Test.Evolution.SetParticleObjectState(Test.ParticleHandles[1], EObjectStateType::Sleeping);
+
+		Test.Advance();
+
+		// Everything should be asleep
+		// The bug was that the constraint was still flagged as awake, but in a sleeping island.
+		EXPECT_EQ(Test.IslandManager->NumIslands(), 1);
+		EXPECT_TRUE(Test.IslandManager->GetIsland(0)->IsSleeping());
+		EXPECT_TRUE(FConstGenericParticleHandle(Test.ParticleHandles[0])->Sleeping());
+		EXPECT_TRUE(FConstGenericParticleHandle(Test.ParticleHandles[1])->Sleeping());
+		EXPECT_TRUE(Test.ConstraintHandles[0]->IsSleeping());
+	}
 }
diff --git a/Engine/Source/Programs/HeadlessChaos/Private/HeadlessChaosTestJoints.cpp b/Engine/Source/Programs/HeadlessChaos/Private/HeadlessChaosTestJoints.cpp
--- a/Engine/Source/Programs/HeadlessChaos/Private/HeadlessChaosTestJoints.cpp
+++ b/Engine/Source/Programs/HeadlessChaos/Private/HeadlessChaosTestJoints.cpp
@@ -45,16 +45,23 @@ namespace ChaosTest {
 			Evolution.AddConstraintContainer(Joints);
 		}
 
-		FPBDJointConstraintHandle* AddJoint(const TVec2<FGeometryParticleHandle*>& InConstrainedParticleIndices, const FVec3& InLocation)
+		FPBDJointConstraintHandle* AddJoint(const TVec2<FGeometryParticleHandle*>& InConstrainedParticles, const FVec3& InLocation)
 		{
-			return Joints.AddConstraint(InConstrainedParticleIndices, FRigidTransform3(InLocation, FRotation3::FromIdentity()));
+			FPBDJointConstraintHandle* Constraint = Joints.AddConstraint(InConstrainedParticles, FRigidTransform3(InLocation, FRotation3::FromIdentity()));
+
+			// @todo(chaos): this should be automatic, but it's handled by the proxy. Fix this
+			InConstrainedParticles[0]->ParticleConstraints().Add(Constraint);
+			InConstrainedParticles[1]->ParticleConstraints().Add(Constraint);
+
+			return Constraint;
 		}
 
 		virtual void Create()
 		{
 			for (int32 ParticleIndex = 0; ParticleIndex < ParticlePositions.Num(); ++ParticleIndex)
 			{
-				AddParticleBox(ParticlePositions[ParticleIndex], FRotation3::MakeFromEuler(FVec3(0.f, 0.f, 0.f)).GetNormalized(), ParticleSizes[ParticleIndex], ParticleMasses[ParticleIndex]);
+				auto* Particle = AddParticleBox(ParticlePositions[ParticleIndex], FRotation3::MakeFromEuler(FVec3(0.f, 0.f, 0.f)).GetNormalized(), ParticleSizes[ParticleIndex], ParticleMasses[ParticleIndex]);
+				Evolution.EnableParticle(Particle);
 			}
 
 			for (int32 JointIndex = 0; JointIndex < JointPositions.Num(); ++JointIndex)
@@ -494,10 +501,14 @@ namespace ChaosTest {
 		Evolution.SetPhysicsMaterial(StaticBox, MakeSerializable(PhysicalMaterial));
 		Evolution.SetPhysicsMaterial(Box2, MakeSerializable(PhysicalMaterial));
 
+		Evolution.EnableParticle(StaticBox);
+		Evolution.EnableParticle(Box2);
+
 		FPBDRigidSpringConstraints JointConstraints;
 		JointConstraints.AddConstraint(ConstrainedParticles, Points, 1.0f, 0.0f, (Points[0] - Points[1]).Size());
 		Evolution.AddConstraintContainer(JointConstraints);
 
+		
 		const FReal Dt = 0.01f;
 		for (int32 i = 0; i < 100; ++i)
 		{
@@ -536,6 +547,9 @@ namespace ChaosTest {
 			Evolution.SetPhysicsMaterial(&StaticBox, MakeSerializable(PhysicalMaterial));
 			Evolution.SetPhysicsMaterial(&Box2, MakeSerializable(PhysicalMaterial));
 
+			Evolution.EnableParticle(&StaticBox);
+			Evolution.EnableParticle(&Box2);
+
 			Chaos::FPBDRigidDynamicSpringConstraints SpringConstraints(MoveTemp(Constraints));
 			Evolution.AddConstraintContainer(SpringConstraints);
 
@@ -566,6 +580,9 @@ namespace ChaosTest {
 			Evolution.SetPhysicsMaterial(&StaticBox, MakeSerializable(PhysicalMaterial));
 			Evolution.SetPhysicsMaterial(&Box2, MakeSerializable(PhysicalMaterial));
 
+			Evolution.EnableParticle(&StaticBox);
+			Evolution.EnableParticle(&Box2);
+
 			Chaos::FPBDRigidDynamicSpringConstraints SpringConstraints(MoveTemp(Constraints), 400);
 			Evolution.AddConstraintContainer(SpringConstraints);
 
diff --git a/Engine/Source/Programs/HeadlessChaos/Private/HeadlessChaosTestPerf.cpp b/Engine/Source/Programs/HeadlessChaos/Private/HeadlessChaosTestPerf.cpp
--- a/Engine/Source/Programs/HeadlessChaos/Private/HeadlessChaosTestPerf.cpp
+++ b/Engine/Source/Programs/HeadlessChaos/Private/HeadlessChaosTestPerf.cpp
@@ -99,6 +99,10 @@ namespace ChaosTest {
 				SphereDynamic->InvI() = TVec3<FRealSingle>(1. / 100000.);
 				Evolution.SetPhysicsMaterial(SphereDynamic, Material);
 				ParticleHandles.Add(SphereDynamic);
+
+				Evolution.EnableParticle(BoxDynamic);
+				Evolution.EnableParticle(SphereDynamic);
+
 #if USE_CONSTRAINTS
 				FVec3 JointLocation(x * Interval, y * Interval, Height * (FReal)0.75);
 				FPBDJointConstraints::FConstraintContainerHandle* NewJoint = Joints.AddConstraint({ BoxDynamic, SphereDynamic }, FRigidTransform3(Offset + JointLocation, FRotation3::FromIdentity()));
diff --git a/Engine/Source/Programs/HeadlessChaos/Private/HeadlessChaosTestProbes.cpp b/Engine/Source/Programs/HeadlessChaos/Private/HeadlessChaosTestProbes.cpp
--- a/Engine/Source/Programs/HeadlessChaos/Private/HeadlessChaosTestProbes.cpp
+++ b/Engine/Source/Programs/HeadlessChaos/Private/HeadlessChaosTestProbes.cpp
@@ -54,6 +54,9 @@ namespace ChaosTest
 		// Make sure the particles would collide if Dynamic wasn't a probe
 		::ChaosTest::SetParticleSimDataToCollide({ Static,Dynamic });
 
+		Evolution.EnableParticle(Static);
+		Evolution.EnableParticle(Dynamic);
+
 		// 1 step should be enough to cause a collision event
 		Evolution.AdvanceOneTimeStep(Dt);
 		Evolution.EndFrame(Dt);
@@ -102,6 +105,9 @@ namespace ChaosTest
 		// Make sure the particles would collide if Dynamic wasn't a probe
 		::ChaosTest::SetParticleSimDataToCollide({ Static,Dynamic });
 
+		Evolution.EnableParticle(Static);
+		Evolution.EnableParticle(Dynamic);
+
 		// 1 step should be enough to cause a collision event
 		Evolution.AdvanceOneTimeStep(Dt);
 		Evolution.EndFrame(Dt);
diff --git a/Engine/Source/Programs/HeadlessChaos/Private/HeadlessChaosTestSim.cpp b/Engine/Source/Programs/HeadlessChaos/Private/HeadlessChaosTestSim.cpp
--- a/Engine/Source/Programs/HeadlessChaos/Private/HeadlessChaosTestSim.cpp
+++ b/Engine/Source/Programs/HeadlessChaos/Private/HeadlessChaosTestSim.cpp
@@ -90,8 +90,8 @@ namespace ChaosTest {
 		::ChaosTest::SetParticleSimDataToCollide({ Static,Dynamic });
 
 		// IMPORTANT : this is required to make sure the particles internal representation will reflect the sim data
-		Evolution.DirtyParticle(*Static);
-		Evolution.DirtyParticle(*Dynamic);
+		Evolution.EnableParticle(Static);
+		Evolution.EnableParticle(Dynamic);
 
 		const FReal Dt = 1 / 60.f;
 		for (int i = 0; i < 200; ++i)
@@ -128,8 +128,8 @@ namespace ChaosTest {
 		::ChaosTest::SetParticleSimDataToCollide({ Static,Dynamic });
 
 		// IMPORTANT : this is required to make sure the particles internal representation will reflect the sim data
-		Evolution.DirtyParticle(*Static);
-		Evolution.DirtyParticle(*Dynamic);
+		Evolution.EnableParticle(Static);
+		Evolution.EnableParticle(Dynamic);
 
 		const FReal Dt = 1 / 60.f;
 		for (int i = 0; i < 100; ++i)
@@ -173,6 +173,9 @@ namespace ChaosTest {
 
 		::ChaosTest::SetParticleSimDataToCollide({ Static,Dynamic });
 
+		Evolution.EnableParticle(Static);
+		Evolution.EnableParticle(Dynamic);
+
 		for (int i = 0; i < 100; ++i)
 		{
 			Evolution.AdvanceOneTimeStep(1 / 60.f);
@@ -225,9 +228,9 @@ namespace ChaosTest {
 		::ChaosTest::SetParticleSimDataToCollide({ Static,Dynamic1,Dynamic2 });
 
 		// IMPORTANT : this is required to make sure the particles internal representation will reflect the sim data
-		Evolution.DirtyParticle(*Static);
-		Evolution.DirtyParticle(*Dynamic1);
-		Evolution.DirtyParticle(*Dynamic2);
+		Evolution.EnableParticle(Static);
+		Evolution.EnableParticle(Dynamic1);
+		Evolution.EnableParticle(Dynamic2);
 
 		bool Dynamic1WentToSleep = false;
 		bool Dynamic1HasWokeAgain = false;
@@ -374,8 +377,8 @@ namespace ChaosTest {
 		Evolution.SetPhysicsMaterial(Dynamic2, MakeSerializable(PhysicsMaterial));
 
 		// IMPORTANT : this is required to make sure the particles internal representation will reflect the sim data
-		Evolution.DirtyParticle(*Dynamic1);
-		Evolution.DirtyParticle(*Dynamic2);
+		Evolution.EnableParticle(Dynamic1);
+		Evolution.EnableParticle(Dynamic2);
 
 		const auto& SleepData = Evolution.GetParticles().GetDynamicParticles().GetSleepData();
 		EXPECT_EQ(SleepData.Num(), 0);
diff --git a/Engine/Source/Programs/HeadlessChaos/Private/HeadlessChaosTestSolverBody.cpp b/Engine/Source/Programs/HeadlessChaos/Private/HeadlessChaosTestSolverBody.cpp
--- a/Engine/Source/Programs/HeadlessChaos/Private/HeadlessChaosTestSolverBody.cpp
+++ b/Engine/Source/Programs/HeadlessChaos/Private/HeadlessChaosTestSolverBody.cpp
@@ -9,7 +9,6 @@
 #include "Chaos/Pair.h"
 #include "Chaos/PBDRigidsEvolution.h"
 #include "Chaos/PBDRigidParticles.h"
-#include "Chaos/PBDCollisionConstraintsPGS.h"
 #include "Chaos/Sphere.h"
 #include "Chaos/Utilities.h"
 
diff --git a/Engine/Source/Programs/HeadlessChaos/Private/HeadlessChaosTestVehicles.cpp b/Engine/Source/Programs/HeadlessChaos/Private/HeadlessChaosTestVehicles.cpp
--- a/Engine/Source/Programs/HeadlessChaos/Private/HeadlessChaosTestVehicles.cpp
+++ b/Engine/Source/Programs/HeadlessChaos/Private/HeadlessChaosTestVehicles.cpp
@@ -1525,6 +1525,8 @@ namespace ChaosTest
 		Dynamic->I() = TVec3<FRealSingle>(100000.0f);
 		Dynamic->InvI() = TVec3<FRealSingle>(1.0f / 100000.0f);
 
+		Evolution.EnableParticle(Dynamic);
+
 		float AccumulatedTime = 0.f;
 		const FReal Dt = 1 / 30.f;
 		for (int i = 0; i < 500; ++i)
diff --git a/Engine/Source/Programs/HeadlessChaos/Public/HeadlessChaosTestConstraints.h b/Engine/Source/Programs/HeadlessChaos/Public/HeadlessChaosTestConstraints.h
--- a/Engine/Source/Programs/HeadlessChaos/Public/HeadlessChaosTestConstraints.h
+++ b/Engine/Source/Programs/HeadlessChaos/Public/HeadlessChaosTestConstraints.h
@@ -59,6 +59,8 @@ namespace ChaosTest
 			}
 			Evolution.SetPhysicsMaterial(&Particle, MakeSerializable(PhysicalMaterial));
 
+			Evolution.EnableParticle(&Particle);
+
 			return &Particle;
 		}
 
diff --git a/Engine/Source/Runtime/Experimental/Chaos/Private/Chaos/Collision/PBDCollisionConstraint.cpp b/Engine/Source/Runtime/Experimental/Chaos/Private/Chaos/Collision/PBDCollisionConstraint.cpp
--- a/Engine/Source/Runtime/Experimental/Chaos/Private/Chaos/Collision/PBDCollisionConstraint.cpp
+++ b/Engine/Source/Runtime/Experimental/Chaos/Private/Chaos/Collision/PBDCollisionConstraint.cpp
@@ -144,8 +144,8 @@ namespace Chaos
 		return Constraint;
 	}
 
-	FPBDCollisionConstraint FPBDCollisionConstraint::MakeCopy(
-		const FPBDCollisionConstraint& Source)
+	// Used by the resim system. Constraints copied this way will later be returned via RestoreFrom()
+	FPBDCollisionConstraint FPBDCollisionConstraint::MakeCopy(const FPBDCollisionConstraint& Source)
 	{
 		// @todo(chaos): The resim cache version probably doesn't need all the data, so maybe try to cut this down?
 		FPBDCollisionConstraint Constraint = Source;
@@ -160,6 +160,22 @@ namespace Chaos
 		return Constraint;
 	}
 
+	// Used by the resim system. Restore a constraint from the past that we saved with MakeCopy
+	void FPBDCollisionConstraint::RestoreFrom(const FPBDCollisionConstraint& Source)
+	{
+		// We do not want to overwrite these properties because they are managed by the systems
+		// where the constraint is registsred (allocator, container, graph).
+		const FPBDCollisionConstraintContainerCookie Cookie = ContainerCookie;
+		const int32 ConstraintGraphIndex = GetConstraintGraphIndex();
+
+		// Copy everything
+		*this = Source;
+
+		// Restore the system state
+		ContainerCookie = Cookie;
+		SetConstraintGraphIndex(ConstraintGraphIndex);
+	}
+
 	FPBDCollisionConstraint::FPBDCollisionConstraint()
 		: ImplicitTransform{ FRigidTransform3(), FRigidTransform3() }
 		, Particle{ nullptr, nullptr }
diff --git a/Engine/Source/Runtime/Experimental/Chaos/Private/Chaos/Collision/ParticlePairMidPhase.cpp b/Engine/Source/Runtime/Experimental/Chaos/Private/Chaos/Collision/ParticlePairMidPhase.cpp
--- a/Engine/Source/Runtime/Experimental/Chaos/Private/Chaos/Collision/ParticlePairMidPhase.cpp
+++ b/Engine/Source/Runtime/Experimental/Chaos/Private/Chaos/Collision/ParticlePairMidPhase.cpp
@@ -471,6 +471,7 @@ namespace Chaos
 		}
 	}
 
+	// This fucntion is called by the resim system to add/replace a collision with one from the history buffer
 	void FSingleShapePairCollisionDetector::SetCollision(const FPBDCollisionConstraint& SourceConstraint)
 	{
 		const int32 CurrentEpoch = MidPhase.GetCollisionAllocator().GetCurrentEpoch();
@@ -483,10 +484,8 @@ namespace Chaos
 			Constraint->GetContainerCookie().CreationEpoch = CurrentEpoch;
 		}
 
-		// Copy the constraint over the existing one, taking care to leave the cookie intact
-		const FPBDCollisionConstraintContainerCookie Cookie = Constraint->GetContainerCookie();
-		*Constraint = SourceConstraint;
-		Constraint->GetContainerCookie() = Cookie;
+		// Copy the constraint data over the existing one (ensure we do not replace data required by the graph and the allocator/container)
+		Constraint->RestoreFrom(SourceConstraint);
 
 		// Add the constraint to the active list
 		// If the constraint already existed and was already active, this will do nothing
diff --git a/Engine/Source/Runtime/Experimental/Chaos/Private/Chaos/Island/IslandGraph.cpp b/Engine/Source/Runtime/Experimental/Chaos/Private/Chaos/Island/IslandGraph.cpp
--- a/Engine/Source/Runtime/Experimental/Chaos/Private/Chaos/Island/IslandGraph.cpp
+++ b/Engine/Source/Runtime/Experimental/Chaos/Private/Chaos/Island/IslandGraph.cpp
@@ -53,24 +53,18 @@ void FIslandGraph<NodeType, EdgeType, IslandType, OwnerType>::ParentIslands(cons
 		GraphIslands[SecondIsland].ChildrenIslands.Add(FirstIsland);
 		GraphIslands[FirstIsland].ChildrenIslands.Add(SecondIsland);
 
-		if(bIsEdgeMoving)
-		{
-			// If we are adding a constraint the island is no longer persistent
-			GraphIslands[FirstIsland].bIsPersistent = false;
-			GraphIslands[SecondIsland].bIsPersistent = false;
-		}
+		// NOTE: We used to set the islands to non-persistent if either nodes in the edge was moving, but we cannot do that here
+		// because this is called as part of the graph update, and both particles may not yet be updated so their moving flag may not
+		// be appropriately set. Instead we manage island persistance in the actual merge phase (see ReassignIslands)
 	}
 }
 	
 template<typename NodeType, typename EdgeType, typename IslandType, typename OwnerType>
-void FIslandGraph<NodeType, EdgeType, IslandType, OwnerType>::UpdateNode(const NodeType& NodeItem, const bool bValidNode, const int32 IslandIndex, const bool bStationaryNode, const int32 NodeIndex)
+void FIslandGraph<NodeType, EdgeType, IslandType, OwnerType>::UpdateNode(const int32 NodeIndex, const bool bValidNode, const bool bStationaryNode)
 {
 	if (GraphNodes.IsValidIndex(NodeIndex))
 	{
 		FGraphNode& GraphNode = GraphNodes[NodeIndex];
-#if CHAOS_CONSTRAINTHANDLE_DEBUG_ENABLED
-		ensure(GraphNode.NodeItem == NodeItem);
-#endif
 
 		// Update node state (must come first because IsEdgeMoving relies in state on Stationary being up to date)
 		const bool bWasValidNode = GraphNode.bValidNode;
@@ -96,7 +90,6 @@ void FIslandGraph<NodeType, EdgeType, IslandType, OwnerType>::UpdateNode(const N
 			// Put the valid node into one of the islands - they will be merged anyway so doesn't matter which
 			// If we did not have an island (no edges), one will be assigned later
 			GraphNode.IslandIndex = MasterIsland;
-			GraphNode.NodeIslands.Reset();
 		}
 
 		// If we are changing to invalid (Dynamic/Sleeping->Kinematic) wake the island
@@ -110,7 +103,6 @@ void FIslandGraph<NodeType, EdgeType, IslandType, OwnerType>::UpdateNode(const N
 
 			// Invalid node island lists are built later (see PopulateIslands)
 			GraphNode.IslandIndex = INDEX_NONE;
-			GraphNode.NodeIslands.Reset();
 		}
 
 		// If we change validity, we may have to change the validity of some of the edges
@@ -127,32 +119,25 @@ void FIslandGraph<NodeType, EdgeType, IslandType, OwnerType>::UpdateNode(const N
 template<typename NodeType, typename EdgeType, typename IslandType, typename OwnerType>
 int32 FIslandGraph<NodeType, EdgeType, IslandType, OwnerType>::AddNode(const NodeType& NodeItem, const bool bValidNode, const int32 IslandIndex, const bool bStationaryNode)
 {
-	int32 NodeIndex = INDEX_NONE;
-	int32* ItemIndex = ItemNodes.Find(NodeItem);
-	if (ItemIndex && GraphNodes.IsValidIndex(*ItemIndex))
+	// Do not try to add a node that is already in the graph
+	check(ItemNodes.Find(NodeItem) == nullptr);
+
+	FGraphNode GraphNode;
+	GraphNode.NodeItem = NodeItem;
+
+	if (bValidNode)
 	{
-		UpdateNode(NodeItem,bValidNode,IslandIndex,bStationaryNode,*ItemIndex);
-		NodeIndex = *ItemIndex;
+		GraphNode.IslandIndex = IslandIndex;
 	}
-	else
-	{
-		FGraphNode GraphNode;
-		GraphNode.NodeItem = NodeItem;
 
-		if (bValidNode)
-		{
-			GraphNode.IslandIndex = IslandIndex;
-		}
-		
-		GraphNode.bValidNode = bValidNode;
-		GraphNode.bStationaryNode = bStationaryNode;
+	GraphNode.bValidNode = bValidNode;
+	GraphNode.bStationaryNode = bStationaryNode;
 
-		NodeIndex =  ItemNodes.Add(NodeItem, GraphNodes.Emplace(GraphNode));
+	const int32 NodeIndex =  ItemNodes.Add(NodeItem, GraphNodes.Emplace(GraphNode));
 
-		if (Owner != nullptr)
-		{
-			Owner->GraphNodeAdded(NodeItem, NodeIndex);
-		}
+	if (Owner != nullptr)
+	{
+		Owner->GraphNodeAdded(NodeItem, NodeIndex);
 	}
 	return NodeIndex;
 }
@@ -166,29 +151,28 @@ void FIslandGraph<NodeType, EdgeType, IslandType, OwnerType>::RemoveNode(const N
 		if (GraphNodes.IsValidIndex(NodeIndex))
 		{
 			FGraphNode& GraphNode = GraphNodes[NodeIndex];
-			
-			// If only one node and zero edges, we invalidate the node island
-			if(GraphNode.NodeEdges.Num() == 0  && GraphIslands.IsValidIndex(GraphNode.IslandIndex))
+
+			// Remove all edges attached to the node, and flag any islands that contain the node 
+			// as non-persistent so they will be woken on the next update
+			if (GraphNode.IslandIndex != INDEX_NONE)
 			{
+				check(GraphIslands.IsValidIndex(GraphNode.IslandIndex));
 				GraphIslands[GraphNode.IslandIndex].bIsPersistent = false;
 			}
-			else
+			for (int32 NodeEdgeIndex = GraphNode.NodeEdges.GetMaxIndex() - 1; NodeEdgeIndex >= 0; --NodeEdgeIndex)
 			{
-				// Otherwise we loop over all the connected edges to remove the node inside the nodeedges list
-				for (int32 NodeEdgeIndex = GraphNode.NodeEdges.GetMaxIndex() - 1; NodeEdgeIndex >= 0; --NodeEdgeIndex)
+				if (GraphNode.NodeEdges.IsValidIndex(NodeEdgeIndex))
 				{
-					if (GraphNode.NodeEdges.IsValidIndex(NodeEdgeIndex))
+					const int32 GraphEdgeIndex = GraphNode.NodeEdges[NodeEdgeIndex];
+					const int32 GraphIslandIndex = GraphEdges[GraphEdgeIndex].IslandIndex;
+					if (GraphIslandIndex != INDEX_NONE)
 					{
-						const int32 GraphEdgeIndex = GraphNode.NodeEdges[NodeEdgeIndex];
-						const int32 GraphIslandIndex = GraphEdges[GraphEdgeIndex].IslandIndex;
-						if (GraphIslandIndex != INDEX_NONE)
-						{
-							GraphIslands[GraphIslandIndex].bIsPersistent = false;
-						}
-						RemoveEdge(GraphEdgeIndex);
+						GraphIslands[GraphIslandIndex].bIsPersistent = false;
 					}
+					RemoveEdge(GraphEdgeIndex);
 				}
 			}
+
 			ItemNodes.Remove(NodeItem);
 			GraphNodes.RemoveAt(NodeIndex);
 
@@ -353,46 +337,42 @@ void FIslandGraph<NodeType, EdgeType, IslandType, OwnerType>::AttachIslands(cons
 template<typename NodeType, typename EdgeType, typename IslandType, typename OwnerType>
 int32 FIslandGraph<NodeType, EdgeType, IslandType, OwnerType>::AddEdge(const EdgeType& EdgeItem, const int32 ItermContainer, const int32 FirstNode, const int32 SecondNode)
 {
+	// Should not call AddEdge for edges that are already present
+	check(ItemEdges.Find(EdgeItem) == nullptr);
+
 	int32 EdgeIndex = INDEX_NONE;
+
 	// We only add an edge if one of the 2 nodes are valid
+	// @todo(chaos): we can relax this requirement now if necessary
 	if (GraphNodes.IsValidIndex(FirstNode) || GraphNodes.IsValidIndex(SecondNode))
 	{
-		// @todo(chaos): This test could be quite slow. will probably be better to use the graph index that is stored on the constraint handle
-		int32* ItemIndex = ItemEdges.Find(EdgeItem);
-		if (ItemIndex)
-		{
-			// If the edge is already there, no need to attach islands since it must have been done before
-			EdgeIndex = *ItemIndex;
-		}
-		else
-		{
-			// Create a  new edge and enqueue the linked islands to be merged if necessary
-			FGraphEdge GraphEdge;
-			GraphEdge.EdgeItem = EdgeItem;
-			GraphEdge.FirstNode = FirstNode;
-			GraphEdge.SecondNode = SecondNode;
-			GraphEdge.IslandIndex = INDEX_NONE;
-			GraphEdge.ItemContainer = ItermContainer;
-			GraphEdge.bValidEdge = true;
+		// Create a  new edge and enqueue the linked islands to be merged if necessary
+		FGraphEdge GraphEdge;
+		GraphEdge.EdgeItem = EdgeItem;
+		GraphEdge.FirstNode = FirstNode;
+		GraphEdge.SecondNode = SecondNode;
+		GraphEdge.IslandIndex = INDEX_NONE;
+		GraphEdge.ItemContainer = ItermContainer;
+		GraphEdge.bValidEdge = true;
 
-			EdgeIndex = GraphEdges.Emplace(GraphEdge);
-			ItemEdges.Add(EdgeItem, EdgeIndex);
+		EdgeIndex = GraphEdges.Emplace(GraphEdge);
+		ItemEdges.Add(EdgeItem, EdgeIndex);
 
-			GraphEdges[EdgeIndex].FirstEdge = GraphNodes.IsValidIndex(FirstNode) ? GraphNodes[FirstNode].NodeEdges.Add(EdgeIndex) : INDEX_NONE;
-			GraphEdges[EdgeIndex].SecondEdge = GraphNodes.IsValidIndex(SecondNode) ? GraphNodes[SecondNode].NodeEdges.Add(EdgeIndex) : INDEX_NONE;
+		GraphEdges[EdgeIndex].FirstEdge = GraphNodes.IsValidIndex(FirstNode) ? GraphNodes[FirstNode].NodeEdges.Add(EdgeIndex) : INDEX_NONE;
+		GraphEdges[EdgeIndex].SecondEdge = GraphNodes.IsValidIndex(SecondNode) ? GraphNodes[SecondNode].NodeEdges.Add(EdgeIndex) : INDEX_NONE;
 
-			AttachIslands(EdgeIndex);
+		AttachIslands(EdgeIndex);
 
-			if (Owner != nullptr)
-			{
-				Owner->GraphEdgeAdded(EdgeItem, EdgeIndex);
-			}
+		if (Owner != nullptr)
+		{
+			Owner->GraphEdgeAdded(EdgeItem, EdgeIndex);
 		}
 	}	
 	else
 	{
 		UE_LOG(LogChaos, Error, TEXT("Island Graph : Trying to add an edge with invalid node indices %d  %d in a list of nodes of size %d"), FirstNode, SecondNode, GraphNodes.Num());
 	}
+
 	return EdgeIndex;
 }
 
@@ -452,7 +432,7 @@ void FIslandGraph<NodeType, EdgeType, IslandType, OwnerType>::MergeIslands(const
 		{
 			GraphIslands[CurrentIndex].IslandCounter = GraphCounter;
 			GraphIslands[CurrentIndex].ParentIsland = ParentIndex;
-	
+
 			// Recursively iterate over all the children ones
 			for (auto& MergedIsland : GraphIslands[CurrentIndex].ChildrenIslands)
 			{
@@ -525,7 +505,6 @@ void FIslandGraph<NodeType, EdgeType, IslandType, OwnerType>::InitSorting()
 	// Reset island max number of levels and colors
 	for(auto& GraphIsland : GraphIslands)
 	{
-		GraphIsland.MaxLevels = INDEX_NONE;
 		GraphIsland.MaxColors = INDEX_NONE;
 	}
 }
@@ -545,21 +524,22 @@ void FIslandGraph<NodeType, EdgeType, IslandType, OwnerType>::UpdateLevels(const
 			ensure(!GraphEdge.bValidEdge || GraphIslands.IsValidIndex(GraphEdge.IslandIndex));
 #endif
 
-			// Do nothing if the edge is assigned a level or if the island is sleeping 
-			if (GraphEdge.bValidEdge && (GraphEdge.LevelIndex == INDEX_NONE) && !GraphIslands[GraphEdge.IslandIndex].bIsSleeping)
+			// Skip edges that already have a level assigned.
+			// Skip sleeping islands, unless they were just put to sleep (this is because we need to know
+			// what particles and constraints are in islands that were just put to sleep so that we can
+			// update their sleep state)
+			if (GraphEdge.bValidEdge && (GraphEdge.LevelIndex == INDEX_NONE) && (!GraphIslands[GraphEdge.IslandIndex].bIsSleeping || !GraphIslands[GraphEdge.IslandIndex].bWasSleeping))
 			{
 				const int32 OtherIndex = (NodeIndex == GraphEdge.FirstNode) ?
 							GraphEdge.SecondNode : GraphEdge.FirstNode;
 
-				GraphEdge.LevelIndex = GraphNode.LevelIndex;
+				SetEdgeLevel(GraphEdge, GraphNode.LevelIndex);
 					
-				GraphIslands[GraphEdge.IslandIndex].MaxLevels = FGenericPlatformMath::Max(
-					GraphIslands[GraphEdge.IslandIndex].MaxLevels, GraphEdge.LevelIndex);
-
 				// If we have another node, append it to our queue on the next level
-				if (GraphNodes.IsValidIndex(OtherIndex) && GraphNodes[OtherIndex].bValidNode && GraphNodes[OtherIndex].LevelIndex == INDEX_NONE)
+				if (GraphNodes.IsValidIndex(OtherIndex) && GraphNodes[OtherIndex].bValidNode && (GraphNodes[OtherIndex].LevelIndex == INDEX_NONE))
 				{
-					GraphNodes[OtherIndex].LevelIndex = GraphEdge.LevelIndex+1;
+					SetNodeLevel(GraphNodes[OtherIndex], GraphEdge.LevelIndex + 1);
+
 					NodeQueue.PushLast(OtherIndex);
 				}
 			}
@@ -570,7 +550,7 @@ void FIslandGraph<NodeType, EdgeType, IslandType, OwnerType>::UpdateLevels(const
 template<typename NodeType, typename EdgeType, typename IslandType, typename OwnerType>
 void FIslandGraph<NodeType, EdgeType, IslandType, OwnerType>::ComputeLevels()
 {
-	// First enqueue all the static/kinematic nodes for levels 0
+	// Set all kinematic nodes to level 0 and enqueue them as roots for the recusrive level assignment.
 	NodeQueue.Reset();
 	for (int32 NodeIndex = 0, NumNodes = GraphNodes.GetMaxIndex(); NodeIndex < NumNodes; ++NodeIndex)
 	{
@@ -579,7 +559,7 @@ void FIslandGraph<NodeType, EdgeType, IslandType, OwnerType>::ComputeLevels()
 			FGraphNode& GraphNode = GraphNodes[NodeIndex];
 			if(!GraphNode.bValidNode)
 			{
-				GraphNode.LevelIndex = 0;
+				SetNodeLevel(GraphNode, 0);
 				UpdateLevels(NodeIndex);
 			}
 		}
@@ -595,12 +575,21 @@ void FIslandGraph<NodeType, EdgeType, IslandType, OwnerType>::ComputeLevels()
 		UpdateLevels(NodeIndex);
 	}
 
-	// An isolated island that is only dynamics will not have been processed above, put everything without a level into level zero
-	for(auto& GraphEdge : GraphEdges)
+	// An island containing only dynamics will not have been assigned any levels. In this case we set all levels to 0.
+	// NOTE: this is only required because we use the level assignment callback to put the particles and constraints
+	// into per-island lists, otherwise it wouldn't really matter as long as levels are the same for all such particles.
+	for(FGraphEdge& GraphEdge : GraphEdges)
 	{
-		if(GraphEdge.bValidEdge)
+		if (GraphEdge.bValidEdge && (GraphEdge.LevelIndex == INDEX_NONE) && (!GraphIslands[GraphEdge.IslandIndex].bIsSleeping || !GraphIslands[GraphEdge.IslandIndex].bWasSleeping))
 		{
-			GraphEdge.LevelIndex = FGenericPlatformMath::Max(GraphEdge.LevelIndex, 0);
+			SetEdgeLevel(GraphEdge, 0);
+		}
+	}
+	for (FGraphNode& GraphNode : GraphNodes)
+	{
+		if (GraphNode.bValidNode && (GraphNode.LevelIndex == INDEX_NONE) && (!GraphIslands[GraphNode.IslandIndex].bIsSleeping || !GraphIslands[GraphNode.IslandIndex].bWasSleeping))
+		{
+			SetNodeLevel(GraphNode, 0);
 		}
 	}
 }
@@ -729,7 +718,6 @@ void FIslandGraph<NodeType, EdgeType, IslandType, OwnerType>::SplitIsland(const
 			{
 				GraphNode.IslandIndex = INDEX_NONE;
 			}
-			GraphNode.NodeIslands.Reset();
 
 			// Loop over the node edges to continue the island discovery
 			for (int32& EdgeIndex : GraphNode.NodeEdges)
@@ -790,7 +778,6 @@ void FIslandGraph<NodeType, EdgeType, IslandType, OwnerType>::SplitIslands()
 				else
 				{
 					RootNode.IslandIndex = INDEX_NONE;
-					RootNode.NodeIslands.Reset();
 				}
 			}
 		}
@@ -810,17 +797,25 @@ void FIslandGraph<NodeType, EdgeType, IslandType, OwnerType>::ReassignIslands()
 	// Update all the edges island indices
 	for (FGraphEdge& GraphEdge : GraphEdges)
 	{
-		int32 EdgeIslandIndex = INDEX_NONE;
-		if (GraphIslands.IsValidIndex(GraphEdge.IslandIndex))
+		int32 EdgeIslandIndex = GraphEdge.IslandIndex;
+		if (EdgeIslandIndex != INDEX_NONE)
 		{
+			check(GraphIslands.IsValidIndex(GraphEdge.IslandIndex))
+
+			// ParentIndex will be same as EdgeIslandIndex if we are not merging
 			const int32 ParentIndex = GraphIslands[GraphEdge.IslandIndex].ParentIsland;
-			if(GraphIslands.IsValidIndex(ParentIndex))
+			check(GraphIslands.IsValidIndex(ParentIndex));
+
+			// If we are moving to another island, that island will be waking up if any of our nodes are moving
+			if ((EdgeIslandIndex != ParentIndex) && IsEdgeMoving(GraphEdge))
 			{
-				EdgeIslandIndex = ParentIndex;
-				GraphIslands[ParentIndex].NumEdges++;
+				GraphIslands[ParentIndex].bIsPersistent = false;
 			}
+
+			// Assign to edge to the (maybe new) island
+			GraphIslands[ParentIndex].NumEdges++;
+			GraphEdge.IslandIndex = ParentIndex;
 		}
-		GraphEdge.IslandIndex = EdgeIslandIndex;
 	}
 
 	// Update all the nodes indices
@@ -831,7 +826,6 @@ void FIslandGraph<NodeType, EdgeType, IslandType, OwnerType>::ReassignIslands()
 			const int32 ParentIndex = GraphIslands[GraphNode.IslandIndex].ParentIsland;
 			if(GraphIslands.IsValidIndex(ParentIndex))
 			{
-				// This is only called from MergeIslands so we know we are awake. NodeIslands will be rebuilt (See PopulateIslands)
 				if (GraphNode.bValidNode)
 				{
 					GraphNode.IslandIndex = ParentIndex;
@@ -840,7 +834,6 @@ void FIslandGraph<NodeType, EdgeType, IslandType, OwnerType>::ReassignIslands()
 				{
 					GraphNode.IslandIndex = INDEX_NONE;
 				}
-				GraphNode.NodeIslands.Reset();
 
 				if(GraphNode.bValidNode)
 				{
@@ -856,10 +849,13 @@ void FIslandGraph<NodeType, EdgeType, IslandType, OwnerType>::UpdateGraph()
 {
 	SCOPE_CYCLE_COUNTER(STAT_MergeIslandGraph);
 
-	// Default to sleeping - the flag will get cleared if an island contains anything that should wake it
+	// Reset the sleep state of the islands. If any node or edge ion the island has changed we will wake
+	// the island. Also if anything was removed from the island since the last update (not persistant)
+	// we set it to awake.
 	for (auto& GraphIsland : GraphIslands)
 	{
-		GraphIsland.bIsSleeping = true;
+		GraphIsland.bWasSleeping = GraphIsland.bIsSleeping;
+		GraphIsland.bIsSleeping = GraphIsland.bIsPersistent;
 	}
 
 	// Merge all the islands if necessary
@@ -875,13 +871,11 @@ void FIslandGraph<NodeType, EdgeType, IslandType, OwnerType>::UpdateGraph()
 			// Add new islands for the all the particles that are not connected into the graph 
 			if (GraphNode.bValidNode && !GraphIslands.IsValidIndex(GraphNode.IslandIndex))
 			{
-				FGraphIsland GraphIsland = { 0, 1 }; // {NumEdges, NumNodes}
+				FGraphIsland GraphIsland = FGraphIsland();
+				GraphIsland.NumEdges = 0;
+				GraphIsland.NumNodes = 1;
+				GraphIsland.bIsSleeping = GraphNode.bStationaryNode;
 				GraphNode.IslandIndex = AddIsland(MoveTemp(GraphIsland));
-
-#if CHAOS_CONSTRAINTHANDLE_DEBUG_ENABLED
-				// Should never have any NodeIslands if we get here
-				ensure(GraphNode.NodeIslands.Num() == 0);
-#endif
 			}
 
 			// Clear the sleeping flag on the island for moving nodes
@@ -913,15 +907,6 @@ void FIslandGraph<NodeType, EdgeType, IslandType, OwnerType>::UpdateGraph()
 			}
 		}
 	}
-
-	// Set the sleeping flag to false if the persistent flag is not set
-	for (FGraphIsland& GraphIsland : GraphIslands)
-	{
-		if (!GraphIsland.bIsPersistent)
-		{
-			GraphIsland.bIsSleeping = false;
-		}
-	}
 	
 	// Split the islands that are persistent and not sleeping if possible
 	SplitIslands();
@@ -929,6 +914,10 @@ void FIslandGraph<NodeType, EdgeType, IslandType, OwnerType>::UpdateGraph()
 	// Remove edges from their island if:
 	// - the island is awake and the edge is invalid (no valid nodes)
 	// - the island was destroyed/merged but the edge wasn't moved to the new island (because it is invalid)
+	// In other words: when we assigned islands we ignored constraints between two kinematic bodies, but
+	// they may be in an island from the last tick (because they were dynamic then) so we reset their island index.
+	// We only do this for awake islands. We leave kinematic-kinematic constraints in sleeping islands in case
+	// the particle(s) get converted back to dynamic before the islands wakes.
 	for (FGraphEdge& GraphEdge : GraphEdges)
 	{
 		const bool bValidIsland = GraphIslands.IsValidIndex(GraphEdge.IslandIndex);
@@ -942,7 +931,7 @@ void FIslandGraph<NodeType, EdgeType, IslandType, OwnerType>::UpdateGraph()
 }
 
 template<typename NodeType, typename EdgeType, typename IslandType, typename OwnerType>
-void FIslandGraph<NodeType, EdgeType, IslandType, OwnerType>::InitIslands()
+void FIslandGraph<NodeType, EdgeType, IslandType, OwnerType>::RemoveAllAwakeEdges()
 {
 	// Remove of all the non sleeping edges
 	for(int32 EdgeIndex = GraphEdges.GetMaxIndex(); EdgeIndex >= 0; --EdgeIndex)
@@ -957,23 +946,6 @@ void FIslandGraph<NodeType, EdgeType, IslandType, OwnerType>::InitIslands()
 	} 
 }
 
-template<typename NodeType, typename EdgeType, typename IslandType, typename OwnerType>
-void FIslandGraph<NodeType, EdgeType, IslandType, OwnerType>::ResetIslands()
-{
-	GraphEdges.Reset();
-	ItemEdges.Reset();
-	
-	// Reset of all the edges + the nodes ones
-	// will probably need to change that with persistent contacts
-	for (int32 NodeIndex = 0, NumNodes = GraphNodes.GetMaxIndex(); NodeIndex < NumNodes; ++NodeIndex)
-	{
-		if (GraphNodes.IsValidIndex(NodeIndex))
-		{
-			GraphNodes[NodeIndex].NodeEdges.Reset();
-		}
-	}
-}
-
 template class FIslandGraph<FGeometryParticleHandle*, FConstraintHandleHolder, FPBDIsland*, FPBDIslandManager>;
 template class FIslandGraph<int32, int32, int32, TNullIslandGraphOwner<int32, int32>>;
 }
diff --git a/Engine/Source/Runtime/Experimental/Chaos/Private/Chaos/Island/IslandManager.cpp b/Engine/Source/Runtime/Experimental/Chaos/Private/Chaos/Island/IslandManager.cpp
--- a/Engine/Source/Runtime/Experimental/Chaos/Private/Chaos/Island/IslandManager.cpp
+++ b/Engine/Source/Runtime/Experimental/Chaos/Private/Chaos/Island/IslandManager.cpp
@@ -36,6 +36,9 @@ FAutoConsoleVariableRef CVarChaosSolverCollisionDefaultAngularSleepThreshold(TEX
 bool bChaosSolverValidateGraph = (CHAOS_CONSTRAINTHANDLE_DEBUG_ENABLED != 0);
 FAutoConsoleVariableRef CVarChaosSolverValidateGraph(TEXT("p.Chaos.Solver.ValidateGraph"), bChaosSolverValidateGraph, TEXT(""));
 
+bool bChaosSolverPersistentGraph = true;
+FAutoConsoleVariableRef CVarChaosSolverPersistentGraph(TEXT("p.Chaos.Solver.PersistentGraph"), bChaosSolverPersistentGraph, TEXT(""));
+
 extern int32 GSingleThreadedPhysics;
 
 	
@@ -105,86 +108,6 @@ inline bool IsIslandSleeping(const FReal MaxLinearSpeed2, const FReal MaxAngular
 	}
 	return false;
 }
-
-/** Add all the graph particle and constraints to the solver islands*/
-inline void PopulateIslands(FPBDIslandManager::GraphType* IslandGraph)
-{
-	auto AddNodeToIsland = [IslandGraph](const int32 IslandIndex, FPBDIslandManager::GraphType::FGraphNode& GraphNode)
-	{
-		if (IslandGraph->GraphIslands.IsValidIndex(IslandIndex))
-		{
-			FPBDIsland* Island = IslandGraph->GraphIslands[IslandIndex].IslandItem;
-			if (!Island->IsSleeping() || Island->SleepingChanged())
-			{
-				Island->AddParticle(GraphNode.NodeItem);
-			}
-		}
-	};
-	for (auto& GraphNode : IslandGraph->GraphNodes)
-	{
-		GraphNode.NodeIslands.Reset();
-
-		// If the node is valid : one and only one island
-		if (GraphNode.bValidNode)
-		{
-#if CHAOS_CONSTRAINTHANDLE_DEBUG_ENABLED
-			ensure(GraphNode.IslandIndex != INDEX_NONE);
-#endif
-			AddNodeToIsland(GraphNode.IslandIndex, GraphNode);
-			GraphNode.NodeIslands.Add(GraphNode.IslandIndex);
-		}
-		else
-		{
-			// A particle could belong to several islands when static/kinematic (not valid)
-			// First we compute the unique particle set of islands
-			for (auto& NodeEdge : GraphNode.NodeEdges)
-			{
-				// The check for INDEX_NONE here filters out kinematic-kinematic constraints which may still
-				// be in the graph (a dynamic was converted to kinematic)
-				const int32 EdgeIslandIndex = IslandGraph->GraphEdges[NodeEdge].IslandIndex;
-				if (EdgeIslandIndex != INDEX_NONE)
-				{
-					GraphNode.NodeIslands.Add(IslandGraph->GraphEdges[NodeEdge].IslandIndex);
-				}
-			}
-			// Loop over the set of islands and add the particle to the solver island
-			for (auto& NodeIsland : GraphNode.NodeIslands)
-			{
-				AddNodeToIsland(NodeIsland, GraphNode);
-			}
-		}
-	}
-	
-	// Loop over the graph edges to transfer them into the solver island (if awake and involves a dynamic particle)
-	// At this point, no awake island should contain any invalid nodes or edges
-	for (auto& GraphEdge : IslandGraph->GraphEdges)
-	{
-		if (IslandGraph->GraphIslands.IsValidIndex(GraphEdge.IslandIndex))
-		{
-			FPBDIsland* Island = IslandGraph->GraphIslands[GraphEdge.IslandIndex].IslandItem;
-
-			// Note: we also populate the island if the sleeping state just changed. This is because we use this
-			// data to sync the constraint's sleeping state in SyncIslands.
-			// @todo(chaos): This is not great - we need a more robust way to sync the sleep state of islands/constraints
-			if (!Island->IsSleeping() || Island->SleepingChanged())
-			{
-#if CHAOS_CONSTRAINTHANDLE_DEBUG_ENABLED
-				ensure(GraphEdge.bValidEdge);
-#endif
-				// The sub sort key is for consistent ordering of constraints when both Level and Color are the same.
-				// Without this the ordering within a level/color would depend on the order they were added to the graph
-				// which depends on collision detection order, which is multi-threaded and not deterministic (by default).
-				// The sort key is based on Particle's NodeIndex because particles remain in the graph with the same index 
-				// as long as they are active. We need to pack the key into 32 bits so we truncate the particle indices. 
-				// This should matter only very rarely. If we hit a conflict we may see a slight non-deterministic behaviour 
-				// difference because the solve  order may vary from frame to frame, but it will still solve ok.
-				// DESYNC WARNING! See above
-				const uint32 SubSortKey = (((uint32(GraphEdge.FirstNode) & 0xFFFF) << 16) | (uint32(GraphEdge.SecondNode) & 0xFFFF));
-				Island->AddConstraint(GraphEdge.EdgeItem, GraphEdge.LevelIndex, GraphEdge.ColorIndex, SubSortKey);
-			}
-		}
-	}
-}
 	
 /** Compute sleeping thresholds given a solver island  */
 inline bool ComputeSleepingThresholds(FPBDIsland* Island,
@@ -245,7 +168,6 @@ inline bool ComputeSleepingThresholds(FPBDIsland* Island,
 FPBDIslandManager::FPBDIslandManager() 
 	: Islands()
 	, IslandGraph(MakeUnique<GraphType>())
-	, MaxNumConstraintContainers(0)
 	, MaxParticleIndex(INDEX_NONE)
 	, bIslandsPopulated(false)
 	, bEndTickCalled(true)
@@ -253,17 +175,6 @@ FPBDIslandManager::FPBDIslandManager()
 	IslandGraph->SetOwner(this);
 }
 
-FPBDIslandManager::FPBDIslandManager(const TParticleView<FPBDRigidParticles>& PBDRigids) 
-	: Islands()
-	, IslandGraph(MakeUnique<GraphType>())
-	, MaxNumConstraintContainers(0)
-	, MaxParticleIndex(INDEX_NONE)
-	, bIslandsPopulated(false)
-	, bEndTickCalled(true)
-{
-	InitializeGraph(PBDRigids);
-}
-
 FPBDIslandManager::~FPBDIslandManager()
 {
 }
@@ -280,7 +191,12 @@ void FPBDIslandManager::AddConstraintContainer(FPBDConstraintContainer& Constrai
 {
 	const int32 ContainerId = ConstraintContainer.GetContainerId();
 
-	MaxNumConstraintContainers = FMath::Max(MaxNumConstraintContainers, ContainerId + 1);
+	if (ConstraintContainers.Num() <= ContainerId)
+	{
+		ConstraintContainers.SetNum(ContainerId + 1);
+	}
+
+	ConstraintContainers[ContainerId] = &ConstraintContainer;
 }
 
 void FPBDIslandManager::InitializeGraph(const TParticleView<FPBDRigidParticles>& PBDRigids)
@@ -305,38 +221,33 @@ void FPBDIslandManager::InitializeGraph(const TParticleView<FPBDRigidParticles>&
 		GraphIsland.ParentIsland = INDEX_NONE;
 	}
 
-	// We are adding all the particles from the solver in
-	// case some were just created/activated
-	for (auto& RigidParticle : PBDRigids)
-	{
-		AddParticle(RigidParticle.Handle());
-	}
-
-	// we update the valid/steady state of the nodes in case any state changed
+	// We update the valid/steady state of the nodes in case any state changed. We don't need to do this for newly added particles, 
+	// so we do it before AddParticles and AddConstraints
+	// @todo(chaos): can we do this when the properties change instead?
 	for(int32 NodeIndex = 0, NumNodes = IslandGraph->GraphNodes.GetMaxIndex(); NodeIndex < NumNodes; ++NodeIndex)
 	{
 		if (IslandGraph->GraphNodes.IsValidIndex(NodeIndex))
 		{
-			GraphType::FGraphNode& GraphNode = IslandGraph->GraphNodes[NodeIndex];
+			FGeometryParticleHandle* ParticleHandle = IslandGraph->GraphNodes[NodeIndex].NodeItem;
 
-			FGeometryParticleHandle* ParticleHandle = GraphNode.NodeItem;
-			IslandGraph->UpdateNode(ParticleHandle, IsDynamicParticle(ParticleHandle), GraphNode.IslandIndex,
-													IsStationaryParticle(ParticleHandle), NodeIndex);
-
-			// @todo(chaos): we should not be doing this here - we should remove particles when they are disabled...is this a leftover temp fix?
-			if (ParticleHandle->CastToRigidParticle() && ParticleHandle->CastToRigidParticle()->Disabled())
-			{
-				RemoveParticle(ParticleHandle);
-			}
+			IslandGraph->UpdateNode(NodeIndex, IsDynamicParticle(ParticleHandle), IsStationaryParticle(ParticleHandle));
 		}
 	}
 
-	// For now we are removing all the constraints but really we should keep the persistent constraints and only remove the transients.
-	// NOTE: This only removes constraints from awake islands and it leaves the sleeping ones. This is important for 
-	// persistent collisions because we do not run collision detection on sleeping particles, so we need the graph
-	// to keep it's sleeping edges so we know what collisions to wake when an island is awakened. The constraint
-	// management system also uses the sleeping state as a lock to prevent constraint destruction.
-	IslandGraph->InitIslands();
+	if (!bChaosSolverPersistentGraph)
+	{
+		// We are adding all the particles from the solver in case some were just created/activated
+		for (auto& RigidParticle : PBDRigids)
+		{
+			AddParticle(RigidParticle.Handle());
+		}
+
+		// NOTE: This only removes constraints from awake islands and it leaves the sleeping ones. This is important for 
+		// semi-persistent collisions because we do not run collision detection on sleeping particles, and we need the graph
+		// to keep it's sleeping edges so that we know what collisions to wake when an island is awakened. The constraint
+		// management system also uses the sleeping state as a lock to prevent constraint destruction.
+		IslandGraph->RemoveAllAwakeEdges();
+	}
 }
 
 void FPBDIslandManager::EndTick()
@@ -346,7 +257,6 @@ void FPBDIslandManager::EndTick()
 	{
 		Island->ClearParticles();
 		Island->ClearConstraints();
-		Island->ResetSleepingChanged();
 	}
 
 	// We can no longer use the Island's particle and constraint lists
@@ -396,79 +306,131 @@ void FPBDIslandManager::ReserveConstraints(const int32 NumConstraints)
 // Callback from the IslandGraph to allow us to store the node index
 void FPBDIslandManager::GraphNodeAdded(FGeometryParticleHandle* ParticleHandle, const int32 NodeIndex)
 {
-	// @todo(chaos): we should store the node index on kinematics and statics too
-	if (FPBDRigidParticleHandle* PBDRigid = ParticleHandle->CastToRigidParticle())
-	{
-		PBDRigid->SetConstraintGraphIndex(NodeIndex);
-	}
+	check(ParticleHandle->ConstraintGraphIndex() == INDEX_NONE);
+
+	ParticleHandle->SetConstraintGraphIndex(NodeIndex);
 }
 
 void FPBDIslandManager::GraphNodeRemoved(FGeometryParticleHandle* ParticleHandle)
 {
-	if (FPBDRigidParticleHandle* PBDRigid = ParticleHandle->CastToRigidParticle())
-	{
-		PBDRigid->SetConstraintGraphIndex(INDEX_NONE);
-	}
+	check(ParticleHandle->ConstraintGraphIndex() != INDEX_NONE);
+
+	ParticleHandle->SetConstraintGraphIndex(INDEX_NONE);
 }
 
 // Callback from the IslandGraph to allow us to store the edge index
 void FPBDIslandManager::GraphEdgeAdded(const FConstraintHandleHolder& ConstraintHandle, const int32 EdgeIndex)
 {
+	check(ConstraintHandle->GetConstraintGraphIndex() == INDEX_NONE);
+
 	ConstraintHandle->SetConstraintGraphIndex(EdgeIndex);
 }
 
 void FPBDIslandManager::GraphEdgeRemoved(const FConstraintHandleHolder& ConstraintHandle)
 {
+	check(ConstraintHandle->GetConstraintGraphIndex() != INDEX_NONE);
+
 	ConstraintHandle->SetConstraintGraphIndex(INDEX_NONE);
 }
 
-int32 FPBDIslandManager::AddParticle(FGeometryParticleHandle* ParticleHandle, const int32 IslandIndex, const bool bOnlyDynamic)
+int32 FPBDIslandManager::AddParticle(FGeometryParticleHandle* ParticleHandle, const FGeometryParticleHandle* ParentParticleHandle)
+{
+	int32 IslandIndex = INDEX_NONE;
+
+	// If we have a parent, assign the child to the same island in the same sleep state
+	if ((ParticleHandle != nullptr) && (ParentParticleHandle != nullptr))
+	{
+		// We are only adding the child particle to the graph if the parent is dynamic
+		const FPBDRigidParticleHandle* ParentPBDRigid = ParentParticleHandle->CastToRigidParticle();
+		if (ParentPBDRigid && ParentPBDRigid->ObjectState() == EObjectStateType::Dynamic)
+		{
+			FPBDRigidParticleHandle* ChildPBDRigid = ParticleHandle->CastToRigidParticle();
+			if (ChildPBDRigid && IsDynamicParticle(ChildPBDRigid))
+			{
+				// If the child particle is dynamic or sleeping we place the child in
+				// the same island as the parent with the same sleep state
+				IslandIndex = ParentPBDRigid->IslandIndex();
+				ChildPBDRigid->SetSleeping(ParentPBDRigid->Sleeping());
+			}
+		}
+	}
+
+	return TryAddParticleIfDynamic(ParticleHandle, IslandIndex);
+}
+
+int32 FPBDIslandManager::TryAddParticle(FGeometryParticleHandle* ParticleHandle, const int32 IslandIndex)
 {
 	if (ParticleHandle)
 	{
-		const bool bIsDynamic = IsDynamicParticle(ParticleHandle);
-		if (!bOnlyDynamic || (bOnlyDynamic && bIsDynamic))
+		if (!ParticleHandle->IsInConstraintGraph())
 		{
 			MaxParticleIndex = FMath::Max(MaxParticleIndex, ParticleHandle->UniqueIdx().Idx);
 
-			FPBDRigidParticleHandle* PBDRigid = ParticleHandle->CastToRigidParticle();
+			const bool bIsDynamic = IsDynamicParticle(ParticleHandle);
 			const bool bIsStationary = IsStationaryParticle(ParticleHandle);
+			IslandGraph->AddNode(ParticleHandle, bIsDynamic, IslandIndex, bIsStationary);
+		}
 
-			// Assign or update the particle's graph node
-			int32 NodeIndex = INDEX_NONE;
-			if ((PBDRigid != nullptr) && (PBDRigid->ConstraintGraphIndex() != INDEX_NONE))
-			{
-				// If the rigid already has a graph index we just update the node information based on the new particles state...
-				NodeIndex = PBDRigid->ConstraintGraphIndex();
-				IslandGraph->UpdateNode(ParticleHandle, bIsDynamic, IslandIndex, bIsStationary, NodeIndex);
-			}
-			else
-			{
-				// If we get here we have a new particle, or a particle that may be in multiple islands (Static/Kinematic)
-				NodeIndex = IslandGraph->AddNode(ParticleHandle, bIsDynamic, IslandIndex, bIsStationary);
-			}
+		return ParticleHandle->ConstraintGraphIndex();
+	}
+	return INDEX_NONE;
+}
+
+int32 FPBDIslandManager::TryAddParticleIfDynamic(FGeometryParticleHandle* ParticleHandle, const int32 IslandIndex)
+{
+	if (ParticleHandle)
+	{
+		if (IsDynamicParticle(ParticleHandle))
+		{
+			return TryAddParticle(ParticleHandle, IslandIndex);
+		}
+
+		return ParticleHandle->ConstraintGraphIndex();
+	}
+	return INDEX_NONE;
+}
 
-			return NodeIndex;
+int32 FPBDIslandManager::TryAddParticleIfNotDynamic(FGeometryParticleHandle* ParticleHandle, const int32 IslandIndex)
+{
+	if (ParticleHandle)
+	{
+		if (!IsDynamicParticle(ParticleHandle))
+		{
+			return TryAddParticle(ParticleHandle, IslandIndex);
 		}
+
+		return ParticleHandle->ConstraintGraphIndex();
 	}
 	return INDEX_NONE;
 }
-	
+
 void FPBDIslandManager::AddConstraint(const uint32 ContainerId, FConstraintHandle* ConstraintHandle, const TVec2<FGeometryParticleHandle*>& ConstrainedParticles)
 {
 	if (ConstraintHandle)
 	{
+		if (ConstraintHandle->GetConstraintGraphIndex() != INDEX_NONE)
+		{
+			// Already in the graph
+			return;
+		}
+
 		// Are the particles dynamic (including asleep)?
-		const bool bValidParticle0 = ConstrainedParticles[0] && IsDynamicParticle(ConstrainedParticles[0]);
-		const bool bValidParticle1 = ConstrainedParticles[1] && IsDynamicParticle(ConstrainedParticles[1]);
+		const bool bDynamicParticle0 = ConstrainedParticles[0] && IsDynamicParticle(ConstrainedParticles[0]);
+		const bool bDynamicParticle1 = ConstrainedParticles[1] && IsDynamicParticle(ConstrainedParticles[1]);
 		
-		// We are checking if one of the 2 particle is dynamic to add the constraint to the graph.
-		// NOTE: This is also where kinematic particles get added to the graph since they were not passed to InitializeGraph.
-		if(bValidParticle0 || bValidParticle1)
+		// We only add the constraint if one particle is non-kinematic
+		// @todo(chaos): we can relax the above restriction now if necessary, but so far no need
+		if (bDynamicParticle0 || bDynamicParticle1)
 		{
-			const int32 NodeIndex0 = AddParticle(ConstrainedParticles[0], INDEX_NONE, false);
-			const int32 NodeIndex1 = AddParticle(ConstrainedParticles[1], INDEX_NONE, false);
+			// If either of the particles is kinematic, it will not have been added by AddParticle, so we must add it here
+			const int32 NodeIndex0 = TryAddParticleIfNotDynamic(ConstrainedParticles[0], /*IslandIndex*/INDEX_NONE);
+			const int32 NodeIndex1 = TryAddParticleIfNotDynamic(ConstrainedParticles[1], /*IslandIndex*/INDEX_NONE);
 
+			// All required particles should be in the graph at this point
+			check((NodeIndex0 != INDEX_NONE) || (ConstrainedParticles[0] == nullptr));
+			check((NodeIndex1 != INDEX_NONE) || (ConstrainedParticles[1] == nullptr));
+
+			// Add the constraint to the graph
 			const int32 EdgeIndex = IslandGraph->AddEdge(ConstraintHandle, ContainerId, NodeIndex0, NodeIndex1);
 
 			// If we were added to a sleeping island, make sure the constraint is flagged as sleeping.
@@ -515,6 +477,21 @@ void FPBDIslandManager::RemoveConstraint(const uint32 ContainerId, FConstraintHa
 	}
 }
 
+void FPBDIslandManager::RemoveConstraints(const uint32 ContainerId)
+{
+	for (int32 EdgeIndex = 0; EdgeIndex < IslandGraph->GraphEdges.Num(); ++EdgeIndex)
+	{
+		if (IslandGraph->GraphEdges.IsValidIndex(EdgeIndex))
+		{
+			FGraphEdge& GraphEdge = IslandGraph->GraphEdges[EdgeIndex];
+			if (GraphEdge.ItemContainer == ContainerId)
+			{
+				IslandGraph->RemoveEdge(EdgeIndex);
+			}
+		}
+	}
+}
+
 void FPBDIslandManager::RemoveParticleConstraints(FGeometryParticleHandle* ParticleHandle, const uint32 ContainerId)
 {
 	if (ParticleHandle)
@@ -563,41 +540,6 @@ void FPBDIslandManager::RemoveParticleConstraints(FGeometryParticleHandle* Parti
 	}
 }
 
-
-void FPBDIslandManager::EnableParticle(FGeometryParticleHandle* ChildParticle, const FGeometryParticleHandle* ParentParticle)
-{
-	if (ParentParticle && ChildParticle)
-	{
-		// We are only adding the child particle to the graph if the parent is dynamic
-		const FPBDRigidParticleHandle* ParentPBDRigid = ParentParticle->CastToRigidParticle();
-		if (ParentPBDRigid && ParentPBDRigid->ObjectState() == EObjectStateType::Dynamic)
-		{
-			int32 IslandIndex = INDEX_NONE;
-			FPBDRigidParticleHandle* ChildPBDRigid = ChildParticle->CastToRigidParticle();
-			if (ChildPBDRigid && IsDynamicParticle(ChildPBDRigid))
-			{
-				// If the child particle is dynamic or sleeping we are transferring
-				// the sleeping flag from the parent to the child
-				// and using the parent island index. if the island index we directly
-				// update the solver island without waiting the next sync
-				IslandIndex = ParentPBDRigid->IslandIndex();
-				ChildPBDRigid->SetSleeping(ParentPBDRigid->Sleeping());	// @todo(chaos): this feels dodgy
-			}
-			// We add the child particle to the graph 
-			AddParticle(ChildParticle, IslandIndex);
-		}
-	}
-}
-
-void FPBDIslandManager::DisableParticle(FGeometryParticleHandle* ParticleHandle)
-{
-	if (ParticleHandle)
-	{
-		// We remove the particle handle from the graph
-		RemoveParticle(ParticleHandle);
-	}
-}
-
 int32 FPBDIslandManager::GetParticleLevel(FGeometryParticleHandle* ParticleHandle) const
 {
 	return IslandGraph->GetNodeItemLevel(ParticleHandle);
@@ -612,7 +554,7 @@ inline bool SolverIslandSortPredicate(const TUniquePtr<FPBDIsland>& SolverIsland
 {
 	return SolverIslandL->GetNumConstraints() < SolverIslandR->GetNumConstraints();
 }
-	
+
 void FPBDIslandManager::GraphIslandAdded(const int32 IslandIndex)
 {
 	FGraphIsland& GraphIsland = IslandGraph->GraphIslands[IslandIndex];
@@ -622,7 +564,7 @@ void FPBDIslandManager::GraphIslandAdded(const int32 IslandIndex)
 	Islands.Reserve(IslandGraph->MaxNumIslands());
 	if (!Islands.IsValidIndex(IslandIndex))
 	{
-		Islands.EmplaceAt(IslandIndex, MakeUnique<FPBDIsland>(MaxNumConstraintContainers));
+		Islands.EmplaceAt(IslandIndex, MakeUnique<FPBDIsland>(ConstraintContainers.Num()));
 	}
 
 	// NOTE: The index may be reused and we may already have an island so we need to reset the island
@@ -631,11 +573,10 @@ void FPBDIslandManager::GraphIslandAdded(const int32 IslandIndex)
 
 }
 
-void FPBDIslandManager::SyncIslands(FPBDRigidsSOAs& Particles)
+void FPBDIslandManager::InitIslands()
 {
 	Islands.Reserve(IslandGraph->MaxNumIslands());
 	IslandIndexing.SetNum(IslandGraph->MaxNumIslands(),false);
-	SortedIslands.SetNum(IslandGraph->MaxNumIslands(),false);
 	int32 LocalIsland = 0;
 
 	// Sync of the solver islands first and reserve the required space
@@ -643,6 +584,8 @@ void FPBDIslandManager::SyncIslands(FPBDRigidsSOAs& Particles)
 	{
 		if (IslandGraph->GraphIslands.IsValidIndex(IslandIndex))
 		{
+			FGraphIsland& GraphIsland = IslandGraph->GraphIslands[IslandIndex];
+
 			// NOTE: We do not check !IsSleeping() here because we need to sync the sleep state of islands
 			// that were just put to sleep this tick. See PopulateIslands.
 
@@ -652,27 +595,31 @@ void FPBDIslandManager::SyncIslands(FPBDRigidsSOAs& Particles)
 			FPBDIsland* Island = Islands[IslandIndex].Get();
 			
 			// We then transfer the persistent flag and the graph dense index to the solver island
-			Island->SetIsPersistent(IslandGraph->GraphIslands[IslandIndex].bIsPersistent);
-			Island->SetIsSleeping(IslandGraph->GraphIslands[IslandIndex].bIsSleeping);
+			Island->SetIsPersistent(GraphIsland.bIsPersistent);
+			Island->SetIsSleeping(GraphIsland.bIsSleeping);
+			Island->SetIsSleepingChanged(GraphIsland.bIsSleeping != GraphIsland.bWasSleeping);
 			Island->SetIslandIndex(LocalIsland);
 
 			// We update the IslandIndexing to retrieve the graph sparse and persistent index from the dense one.
 			IslandIndexing[LocalIsland] = IslandIndex;
-			SortedIslands[LocalIsland] = IslandIndex;
 			LocalIsland++;
 
-			Island->ReserveParticles(IslandGraph->GraphIslands[IslandIndex].NumNodes);
+			Island->ReserveParticles(GraphIsland.NumNodes);
 
 			// Reset of the sleep counter if the island is :
 			// - Non persistent since we are starting incrementing the counter once the island is persistent and if values below the threshold
 			// - Sleeping since as soon as it wakes up we could start incrementing the counter as well
-			// - Just woken because we may not have been asleep for a whole frame to hit the
-			if (!Island->IsPersistent() || Island->IsSleeping() || Island->SleepingChanged())
+			// - Just woken because we may not have been asleep for a whole frame (this probably makes the previous check unnecessary)
+			if (!Island->IsPersistent() || Island->IsSleeping() || Island->IsSleepingChanged())
 			{ 
 				Island->SetSleepCounter(0);
 			}
-			// We reset the persistent flag to be true on the island graph. 
-			IslandGraph->GraphIslands[IslandIndex].bIsPersistent = true;
+
+			// We reset the persistent flag to be true on the island graph for next frame. It may get set to false again
+			// if any particles are added or removed
+			// @todo(chaos): It does not feel right to reset this here. Should probably be in EndTick? (but See AdvanceClustering which
+			// is currently before that, and generates new particles that get added to the graph).
+			GraphIsland.bIsPersistent = true;
 		}
 		else if (Islands.IsValidIndex(IslandIndex))
 		{
@@ -681,14 +628,14 @@ void FPBDIslandManager::SyncIslands(FPBDRigidsSOAs& Particles)
 	}
 
 	IslandIndexing.SetNum(LocalIsland, false);
-	SortedIslands.SetNum(LocalIsland, false);
-
-	PopulateIslands(IslandGraph.Get());
+}
 
-	// Update of the sync and sleep state for each island
+void FPBDIslandManager::SyncIslands(FPBDRigidsSOAs& Particles)
+{
+	// Update of the sync and sleep state of all particles in all awake islands
 	for(auto& Island : Islands)
 	{
-		if(!Island->IsSleeping() || Island->SleepingChanged())
+		if(!Island->IsSleeping() || Island->IsSleepingChanged())
 		{
 			Island->UpdateSyncState(Particles);
 			Island->PropagateSleepState(Particles);
@@ -703,30 +650,92 @@ void FPBDIslandManager::SyncIslands(FPBDRigidsSOAs& Particles)
 	bIslandsPopulated = true;
 }
 
-void FPBDIslandManager::UpdateGraph()
+void FPBDIslandManager::GraphNodeLevelAssigned(FGraphNode& GraphNode)
 {
-	// Merge the graph islands if required
-	IslandGraph->UpdateGraph();
+	// Add the particle to its island. This will be called in the order that levels are assigned
+	// which allows us to skip the sorting-by-level step. NOTE: We only add dynamic particles to 
+	// islands at this stage (kinematics are added to the solver when we gather data).
+	if (GraphNode.bValidNode)
+	{
+		// We have a dynamic particle - add it to its island
+		check(GraphNode.IslandIndex != INDEX_NONE);
+		FPBDIsland* Island = IslandGraph->GraphIslands[GraphNode.IslandIndex].IslandItem;
+
+#if CHAOS_CONSTRAINTHANDLE_DEBUG_ENABLED
+		check(!Island->DebugContainsParticle(GraphNode.NodeItem));
+#endif
+
+		Island->AddParticle(GraphNode.NodeItem);
+	}
+}
+
+void FPBDIslandManager::GraphEdgeLevelAssigned(FGraphEdge& GraphEdge)
+{
+	check(GraphEdge.IslandIndex != INDEX_NONE);
+	FPBDIsland* Island = IslandGraph->GraphIslands[GraphEdge.IslandIndex].IslandItem;
+
+#if CHAOS_CONSTRAINTHANDLE_DEBUG_ENABLED
+	ensure(GraphEdge.bValidEdge);
+	check(!Island->DebugContainsConstraint(GraphEdge.EdgeItem));
+#endif
+
+	// Persistent ordering between ticks is now built-in because the graph is persistent so the constraint order will 
+	// not change from tick to tick even with parallel collision detection (though the order will be different between 
+	// machines unles deterministic mode is set on the evolution). As a result this sort key is not required but for 
+	// now we still calculate it because we can disable persistent islands on a cvar (bChaosSolverPersistentGraph).
+	// 
+	// For non-persistent graph mode:
+	// The sub sort key is for consistent ordering of constraints when both Level and Color are the same.
+	// Without this the ordering within a level/color would depend on the order they were added to the graph
+	// which depends on collision detection order, which is multi-threaded and not deterministic (by default).
+	// The sort key is based on Particle's NodeIndex because particles remain in the graph with the same index 
+	// as long as they are active. We need to pack the key into 32 bits so we truncate the particle indices. 
+	// This should matter only very rarely. If we hit a conflict we may see a slight non-deterministic behaviour 
+	// difference because the solve  order may vary from frame to frame, but it will still solve ok.
+	const uint32 SubSortKey = (((uint32(GraphEdge.FirstNode) & 0xFFFF) << 16) | (uint32(GraphEdge.SecondNode) & 0xFFFF));
+
+	Island->AddConstraint(GraphEdge.EdgeItem, GraphEdge.LevelIndex, GraphEdge.ColorIndex, SubSortKey);
 }
 
 void FPBDIslandManager::UpdateIslands(FPBDRigidsSOAs& Particles)
 {
-	// Assign levels, colors, etc
+	// Assign all particles and constraints to islands, merging and splitting islands as required
+	IslandGraph->UpdateGraph();
+
+#if CHAOS_CONSTRAINTHANDLE_DEBUG_ENABLED
+	DebugCheckParticlesInGraph(Particles.GetNonDisabledDynamicView());
+#endif
+
+	// Initializes our set of islands to match the state of the graph islands, ready to receive the particles and constraints.
+	InitIslands();
+
+	// Reeset all particle and constraint levels, colors etc
 	IslandGraph->InitSorting();
+
+	// Assign a level to every constraint and particle.
+	// NOTE: This will call GraphNodeLevelAssigned/GraphEdgeLevelAssigned for every
+	// particle/constraint in level order. We use this to add constraints and particles
+	// to their islands in level-sorted order, so we don't need to sort again below
 	IslandGraph->ComputeLevels();
+
+	// Assign colors
+	// @todo(chaos): re-enable graph coloring, although it's probably not worth it!
 	//IslandGraph->ComputeColors();
-	
-	// Sync the graph islands with the solver islands objects
+
+	// Sync the state of all particles and constraints (primarily sleep state)
 	SyncIslands(Particles);
 
 	// Sort each island's constraints based on level, color, etc
-	for (TUniquePtr<FPBDIsland>& Island : Islands)
+	// This is not required when we have a persistent graph because we already added constraints to islands in 
+	// Level-order, and the sub-level order does not change from tick to tick when persistence is enabled.
+	if (!bChaosSolverPersistentGraph)
 	{
-		if (!Island->IsSleeping())
+		for (TUniquePtr<FPBDIsland>& Island : Islands)
 		{
-			// Sort constraints based on level
-			// @todo(chaos): can we eliminate this step? We effectively did it already above
-			Island->SortConstraints();
+			if (!Island->IsSleeping())
+			{
+				Island->SortConstraints();
+			}
 		}
 	}
 
@@ -803,10 +812,11 @@ void FPBDIslandManager::SetIslandSleeping(FPBDRigidsSOAs& Particles, const int32
 		{
 			IslandGraph->GraphIslands[GraphIndex].bIsSleeping = bIsSleeping;
 			Islands[GraphIndex]->SetIsSleeping(bIsSleeping);
+			Islands[GraphIndex]->SetIsSleepingChanged(true);
 			Islands[GraphIndex]->PropagateSleepState(Particles);
 		}
 
-		// Reset the sleep counter with every wake call, even if already awake
+		// Reset the sleep counter with every wake call, even if already awake so that we don't immediately go back to sleep
 		if (!bIsSleeping)
 		{
 			Islands[GraphIndex]->SetSleepCounter(0);
@@ -887,7 +897,83 @@ bool FPBDIslandManager::IslandNeedsResim(const int32 IslandIndex) const
 	return Islands[GetGraphIndex(IslandIndex)]->NeedsResim();
 }
 
+///////////////////////////////////////////////////////////////////////////////////////////////////
+///////////////////////////////////////////////////////////////////////////////////////////////////
+///////////////////////////////////////////////////////////////////////////////////////////////////
+//
+//
+// BEGIN DEBUG STUFF
+//
+//
+///////////////////////////////////////////////////////////////////////////////////////////////////
+///////////////////////////////////////////////////////////////////////////////////////////////////
+///////////////////////////////////////////////////////////////////////////////////////////////////
+
 #if CHAOS_CONSTRAINTHANDLE_DEBUG_ENABLED
+bool FPBDIslandManager::DebugCheckParticlesInGraph(const TParticleView<FPBDRigidParticles>& Particles) const
+{
+	// Make sure that all particles that should be in the graph actually are.
+	// We used to loop over all particles every trick and re-add them, but now we only add/remove when necessary.
+	bool bSuccess = true;
+
+	// All dynamic particles should be in the graph
+	for (const auto& ParticleHandle : Particles)
+	{
+		FConstGenericParticleHandle Particle = ParticleHandle.Handle();
+		if (!Particle->IsDynamic())
+		{
+			continue;
+		}
+
+		int32* PNodeIndex = IslandGraph->ItemNodes.Find(Particle->Handle());
+		bSuccess = bSuccess && (PNodeIndex != nullptr);
+		if (ensure(bSuccess))
+		{
+			const int32 NodeIndex = *PNodeIndex;
+			bSuccess = bSuccess && IslandGraph->GraphNodes.IsValidIndex(NodeIndex);
+			if (ensure(bSuccess))
+			{
+				const FGraphNode& GraphNode = IslandGraph->GraphNodes[NodeIndex];
+				bSuccess = bSuccess && (GraphNode.NodeItem == Particle->Handle());
+				ensure(bSuccess);
+			}
+		}
+	}
+
+	for (const FGraphEdge& GraphEdge : IslandGraph->GraphEdges)
+	{
+		const int32 FirstNodeIndex = GraphEdge.FirstNode;
+		const int32 SecondNodeIndex = GraphEdge.SecondNode;
+
+		// Every constraint should reference at least one particle
+		bSuccess = bSuccess && ((FirstNodeIndex != INDEX_NONE) || (SecondNodeIndex != INDEX_NONE));
+		ensure(bSuccess);
+
+		// Every particle referenced by a constraint should be in the graph (including kinematics and statics)
+		// The node must point to the correct particle
+		if (FirstNodeIndex != INDEX_NONE)
+		{
+			bSuccess = bSuccess && IslandGraph->GraphNodes.IsValidIndex(FirstNodeIndex);
+			if (ensure(bSuccess))
+			{
+				bSuccess = bSuccess && (IslandGraph->GraphNodes[FirstNodeIndex].NodeItem == GraphEdge.EdgeItem->GetConstrainedParticles()[0]);
+				ensure(bSuccess);
+			}
+		}
+		if (SecondNodeIndex != INDEX_NONE)
+		{
+			bSuccess = bSuccess && IslandGraph->GraphNodes.IsValidIndex(SecondNodeIndex);
+			if (ensure(bSuccess))
+			{
+				bSuccess = bSuccess && (IslandGraph->GraphNodes[SecondNodeIndex].NodeItem == GraphEdge.EdgeItem->GetConstrainedParticles()[1]);
+				ensure(bSuccess);
+			}
+		}
+	}
+
+	return bSuccess;
+}
+
 bool FPBDIslandManager::DebugCheckParticleNotInGraph(const FGeometryParticleHandle* ParticleHandle) const
 {
 	bool bSuccess = true;
@@ -1078,23 +1164,13 @@ void FPBDIslandManager::DebugCheckParticleIslands(const FGeometryParticleHandle*
 		{
 			// Dynamic particles should be in 1 island
 			ensure(Node.IslandIndex != INDEX_NONE);
-			ensure(Node.NodeIslands.Num() == 1);
-			ensure(Node.NodeIslands.Contains(Node.IslandIndex));
 		}
 		else
 		{
-			// Static/Kinematic particles should be in at least 1 island if it has ever been in an awake island...
-			// We can't check that explicitly, but if we have no edges we may legitimately have no island
+			// There's not much we can check for non-dynamic particles. They should not have an IslandIndex because that's not used
+			// for kinematics, but they may or may not be in an island. E.g., if the kinematic is connected to dynamic it will
+			// be in one or more islands, but if it has no connections or is connected only to another kinematic it will not be in any islands.
 			ensure(Node.IslandIndex == INDEX_NONE);
-			ensure((Node.NodeEdges.Num() == 0) || (Node.NodeIslands.Num() >= 1));
-		}
-
-		// If we found the node in some solver islands, make sure the nodes know about them
-		// NOTE: SolverIsland indexes are into the IslandIndexingArray, not the Islands array
-		for (int32 NodeSolverIsland : NodeSolverIslands)
-		{
-			const int32 NodeIsland = IslandIndexing[NodeSolverIsland];
-			ensure(Node.NodeIslands.Contains(NodeIsland));
 		}
 	}
 }
diff --git a/Engine/Source/Runtime/Experimental/Chaos/Private/Chaos/Island/SolverIsland.cpp b/Engine/Source/Runtime/Experimental/Chaos/Private/Chaos/Island/SolverIsland.cpp
--- a/Engine/Source/Runtime/Experimental/Chaos/Private/Chaos/Island/SolverIsland.cpp
+++ b/Engine/Source/Runtime/Experimental/Chaos/Private/Chaos/Island/SolverIsland.cpp
@@ -25,7 +25,7 @@ void FPBDIsland::Reuse()
 	check(GetNumConstraints() == 0);
 
 	SetSleepCounter(0);
-	ResetSleepingChanged();
+	SetIsSleepingChanged(false);
 }
 
 void FPBDIsland::UpdateParticles()
diff --git a/Engine/Source/Runtime/Experimental/Chaos/Private/Chaos/PBDCollisionConstraints.cpp b/Engine/Source/Runtime/Experimental/Chaos/Private/Chaos/PBDCollisionConstraints.cpp
--- a/Engine/Source/Runtime/Experimental/Chaos/Private/Chaos/PBDCollisionConstraints.cpp
+++ b/Engine/Source/Runtime/Experimental/Chaos/Private/Chaos/PBDCollisionConstraints.cpp
@@ -73,6 +73,8 @@ namespace Chaos
 	bool DebugDrawProbeDetection = false;
 	FAutoConsoleVariableRef CVarDebugDrawProbeDetection(TEXT("p.Chaos.Collision.DebugDrawProbeDetection"), DebugDrawProbeDetection, TEXT("Draw probe constraint detection."));
 	
+	extern bool bChaosSolverPersistentGraph;
+
 #if CHAOS_DEBUG_DRAW
 	namespace CVars
 	{
@@ -358,16 +360,66 @@ namespace Chaos
 
 	void FPBDCollisionConstraints::AddConstraintsToGraph(FPBDIslandManager& IslandManager)
 	{
-		if (GetCollisionsEnabled())
+		// Debugging/diagnosing: if we have collisions disabled, remove all collisions from the graph and don't add any more
+		if (!GetCollisionsEnabled())
 		{
-			for (FPBDCollisionConstraintHandle* ConstraintHandle : GetConstraints())
-			{
-				if (ConstraintHandle->IsEnabled())
+			IslandManager.RemoveConstraints(GetContainerId());
+			return;
+		}
+
+		// If we are running with a persistent graph, remove expired collisions	
+		if (bChaosSolverPersistentGraph)
+		{
+			// Find all expired constraints in the graph
+			TempCollisions.Reset();
+			IslandManager.VisitConstraintsInAwakeIslands(GetContainerId(),
+				[this](FConstraintHandle* ConstraintHandle)
 				{
-					IslandManager.AddConstraint(GetContainerId(), ConstraintHandle, ConstraintHandle->GetConstrainedParticles());
-				}
+					if (FPBDCollisionConstraintHandle* CollisionHandle = ConstraintHandle->As<FPBDCollisionConstraintHandle>())
+					{
+						if (!CollisionHandle->IsEnabled() || ConstraintAllocator.IsConstraintExpired(CollisionHandle->GetContact()))
+						{
+							TempCollisions.Add(CollisionHandle);
+						}
+					}
+				});
+
+			// Remove expired constraints
+			for (FPBDCollisionConstraintHandle* CollisionHandle : TempCollisions)
+			{
+				IslandManager.RemoveConstraint(GetContainerId(), CollisionHandle);
 			}
 		}
+
+		// Collect all the new constraints that need to be added to the graph
+		TempCollisions.Reset();
+		for (FPBDCollisionConstraintHandle* ConstraintHandle : ConstraintAllocator.GetConstraints())
+		{
+			FPBDCollisionConstraint& Constraint = ConstraintHandle->GetContact();
+			if (!Constraint.IsInConstraintGraph() && Constraint.IsEnabled())
+			{
+				TempCollisions.Add(ConstraintHandle);
+			}
+		}
+
+		// Sort new constraints into a predictable order. This isn't strictly required, but without it we
+		// can get fairly different behaviour from run to run because the collisions detection order is 
+		// effectively random on multicore machines
+		TempCollisions.Sort(
+			[](const FPBDCollisionConstraintHandle& L, const FPBDCollisionConstraintHandle& R)
+			{
+				const uint64 LKey = L.GetContact().GetParticlePairKey().GetKey();
+				const uint64 RKey = R.GetContact().GetParticlePairKey().GetKey();
+				return LKey < RKey;
+			});
+
+		// Add the new constraints to the graph
+		for (FPBDCollisionConstraintHandle* ConstraintHandle : TempCollisions)
+		{
+			IslandManager.AddConstraint(GetContainerId(), ConstraintHandle, ConstraintHandle->GetConstrainedParticles());
+		}
+
+		TempCollisions.Reset();
 	}
 
 	const FPBDCollisionConstraint& FPBDCollisionConstraints::GetConstraint(int32 Index) const
diff --git a/Engine/Source/Runtime/Experimental/Chaos/Private/Chaos/PBDRigidClustering.cpp b/Engine/Source/Runtime/Experimental/Chaos/Private/Chaos/PBDRigidClustering.cpp
--- a/Engine/Source/Runtime/Experimental/Chaos/Private/Chaos/PBDRigidClustering.cpp
+++ b/Engine/Source/Runtime/Experimental/Chaos/Private/Chaos/PBDRigidClustering.cpp
@@ -9,7 +9,6 @@
 #include "Chaos/MassProperties.h"
 #include "Chaos/PBDRigidsEvolution.h"
 #include "Chaos/PBDCollisionConstraints.h"
-#include "Chaos/PBDCollisionConstraintsPGS.h"
 #include "Chaos/PBDRigidClusteringAlgo.h"
 #include "Chaos/Sphere.h"
 #include "Chaos/UniformGrid.h"
@@ -89,13 +88,7 @@ namespace Chaos
 		}
 
 		// Must do this so that the constraint graph knows about this particle 
-		// prior to calling CreateIslands().  We could call MEvolution.CreateParticle()
-		// which does the same thing, but also calls DirtyParticle(), which is already
-		// done by MEvolution.CreateClusteredParticles(), and will be done again by
-		// MEvolution.EnableParticle().
-		//MEvolution.GetConstraintGraph().AddParticle(NewParticle); // PBDRigidsEvolutionGBF protects GetConstraintGraph().  Bah!
-		MEvolution.CreateParticle(NewParticle); // Doesn't create, just adds to constraint graph
-		MEvolution.EnableParticle(NewParticle, nullptr); // null for parent skips constraint graph EnableParticle()
+		MEvolution.EnableParticle(NewParticle);
 		NewParticle->SetCollisionGroup(INT_MAX);
 		TopLevelClusterParents.Add(NewParticle);
 
@@ -201,7 +194,6 @@ namespace Chaos
 		{
 			NewParticle = MEvolution.CreateClusteredParticles(1)[0]; // calls Evolution.DirtyParticle()
 		}
-		MEvolution.CreateParticle(NewParticle);
 		MEvolution.EnableParticle(NewParticle, Parent);
 
 		NewParticle->SetCollisionGroup(INT_MAX);
diff --git a/Engine/Source/Runtime/Experimental/Chaos/Private/Chaos/PBDRigidsEvolution.cpp b/Engine/Source/Runtime/Experimental/Chaos/Private/Chaos/PBDRigidsEvolution.cpp
--- a/Engine/Source/Runtime/Experimental/Chaos/Private/Chaos/PBDRigidsEvolution.cpp
+++ b/Engine/Source/Runtime/Experimental/Chaos/Private/Chaos/PBDRigidsEvolution.cpp
@@ -1003,11 +1003,8 @@ namespace Chaos
 	void FPBDRigidsEvolutionBase::SetParticleObjectState(FPBDRigidParticleHandle* Particle, EObjectStateType ObjectState)
 	{
 		const EObjectStateType InitialState = Particle->ObjectState();
-
-		// Should never be called for statics - they cannot change state. Except...
-		// @todo(chaos): We should be able to assert this but FRigidClustering::CreateClusterParticle and ReleaseClusterParticle change the state after creation
-		//check(InitialState != EObjectStateType::Static)
-		//check(ObjectState != EObjectStateType::Static);
+		const bool bWasDynamic = (InitialState == EObjectStateType::Dynamic) || (InitialState == EObjectStateType::Sleeping);
+		const bool bIsDynamic = (ObjectState == EObjectStateType::Dynamic) || (ObjectState == EObjectStateType::Sleeping);
 
 		Particle->SetObjectStateLowLevel(ObjectState);
 
@@ -1031,15 +1028,17 @@ namespace Chaos
 			}
 		}
 
+		// If we are now dynamic and enabled, we need to be in the graph if not already there. We may not be
+		// because Kinematic particles are only in the graph if referenced by a constraint.
+		if (bIsDynamic && !bWasDynamic && !Particle->Disabled())
+		{
+			ConstraintGraph.AddParticle(Particle);
+		}
+
 		// If we switched to/from dynamic the inertia conditioning may need to be refreshed
-		if (Particle->InertiaConditioningEnabled())
+		if (Particle->InertiaConditioningEnabled() && (bWasDynamic != bIsDynamic))
 		{
-			const bool bWasDynamic = (InitialState == EObjectStateType::Dynamic) || (InitialState == EObjectStateType::Sleeping);
-			const bool bIsDynamic = (ObjectState == EObjectStateType::Dynamic) || (ObjectState == EObjectStateType::Sleeping);
-			if (bWasDynamic != bIsDynamic)
-			{
-				Particle->SetInertiaConditioningDirty();
-			}
+			Particle->SetInertiaConditioningDirty();
 		}
 	}
 
@@ -1066,10 +1065,9 @@ namespace Chaos
 		{
 			Particles.DisableParticle(Particle);
 			RemoveParticleFromAccelerationStructure(*Particle);
-			ConstraintGraph.DisableParticle(Particle);
+			DisableConstraints(Particle);
+			ConstraintGraph.RemoveParticle(Particle);
 		}
-
-		DisableConstraints(InParticles);
 	}
 
 	void FPBDRigidsEvolutionBase::DisableParticleWithRemovalEvent(FGeometryParticleHandle* Particle)
diff --git a/Engine/Source/Runtime/Experimental/Chaos/Private/PBDRigidsSolver.cpp b/Engine/Source/Runtime/Experimental/Chaos/Private/PBDRigidsSolver.cpp
--- a/Engine/Source/Runtime/Experimental/Chaos/Private/PBDRigidsSolver.cpp
+++ b/Engine/Source/Runtime/Experimental/Chaos/Private/PBDRigidsSolver.cpp
@@ -1086,7 +1086,7 @@ namespace Chaos
 					Rigid->SetPreObjectStateLowLevel(Rigid->ObjectState());	//created this frame so pre is the initial value
 				}
 				Handle->SetPhysicsProxy(Proxy);
-				GetEvolution()->CreateParticle(Handle);
+				GetEvolution()->RegisterParticle(Handle);
 				Proxy->SetInitialized(GetCurrentFrame());
 			}
 		};
diff --git a/Engine/Source/Runtime/Experimental/Chaos/Private/PhysicsProxy/GeometryCollectionPhysicsProxy.cpp b/Engine/Source/Runtime/Experimental/Chaos/Private/PhysicsProxy/GeometryCollectionPhysicsProxy.cpp
--- a/Engine/Source/Runtime/Experimental/Chaos/Private/PhysicsProxy/GeometryCollectionPhysicsProxy.cpp
+++ b/Engine/Source/Runtime/Experimental/Chaos/Private/PhysicsProxy/GeometryCollectionPhysicsProxy.cpp
@@ -187,8 +187,6 @@ void PopulateSimulatedParticle(
 	Handle->SetW(Chaos::FVec3(0.f));
 	Handle->SetP(Handle->X());
 	Handle->SetQ(Handle->R());
-	Handle->SetIslandIndex(INDEX_NONE);
-	Handle->SetConstraintGraphIndex(INDEX_NONE);
 	Handle->SetCenterOfMass(FVector3f::ZeroVector);
 	Handle->SetRotationOfMass(FQuat::Identity);
 
@@ -757,7 +755,7 @@ void FGeometryCollectionPhysicsProxy::CreateNonClusteredParticles(Chaos::FPBDRig
 			Handle->GTGeometryParticle() = GTParticles[Idx].Get();
 
 			check(SolverParticleHandles[Idx]->GetParticleType() == Handle->GetParticleType());
-			RigidsSolver->GetEvolution()->CreateParticle(Handle);
+			RigidsSolver->GetEvolution()->RegisterParticle(Handle);
 		}
 	}
 
@@ -909,6 +907,8 @@ void FGeometryCollectionPhysicsProxy::InitializeBodiesPT(Chaos::FPBDRigidsSolver
 				{
 					Shape->SetMaterial(Parameters.PhysicalMaterialHandle);
 				}
+
+				RigidsSolver->GetEvolution()->EnableParticle(Handle);
 			}
 		},true);
 
@@ -1299,6 +1299,8 @@ FGeometryCollectionPhysicsProxy::BuildClusters_Internal(
 	const FVector::FReal MassScale = WorldScale.X * WorldScale.Y * WorldScale.Z;
 	const Chaos::FVec3f ScaledInertia = Chaos::Utilities::ScaleInertia<float>((Chaos::FVec3f)InertiaTensor[CollectionClusterIndex], FVector3f(WorldScale), true);  
 	
+	// NOTE: The particle creation call above (CreateClusterParticle) activates the particle, which does various things including adding
+	// it to the constraint graph. It seems a bit dodgy to be completely changing what the particle is after it has been enabled. Maybe we should fix this.
 	PopulateSimulatedParticle(
 		Parent,
 		Parameters.Shared, 
diff --git a/Engine/Source/Runtime/Experimental/Chaos/Public/Chaos/Collision/CollisionConstraintAllocator.h b/Engine/Source/Runtime/Experimental/Chaos/Public/Chaos/Collision/CollisionConstraintAllocator.h
--- a/Engine/Source/Runtime/Experimental/Chaos/Public/Chaos/Collision/CollisionConstraintAllocator.h
+++ b/Engine/Source/Runtime/Experimental/Chaos/Public/Chaos/Collision/CollisionConstraintAllocator.h
@@ -23,8 +23,8 @@ namespace Chaos
 	 * by a FParticlePairMidPhase object. the MidPhase object is what actually calls the Narrow Phase and maintains
 	 * the set of collision constraints for all the shape pairs on the particles.
 	 * 
-	 * Constraints are allocated during the collision detection phase and retained between ticks. An attempt ot create
-	 * a constraint for the same shape pair as seen on the previous tick will return the existing collision constraint, with
+	 * Constraints are allocated during the collision detection phase and retained between ticks. An attempt to create
+	 * a constraint for the same shape pair as seen on the previous tick will return the existing collision constraint with
 	 * all of its data intact.
 	 *
 	 * The allocator also keeps a list of Standard and Swept collision constraints that are active for the current tick.
@@ -36,9 +36,14 @@ namespace Chaos
 	 * counter.
 	 * 
 	 * The Midphase list is pruned at the end of each tick so if particles are destroyed or a particle pair is no longer 
-	 * overlapping.
+	 * overlapping, the collisions will be destroyed.
+	 * 
+	 * When particles are destroyed, we do not immediately destroy the MidPhases (or Collisions) that are associoated with
+	 * the particle. Instead, we clear the particle pointer from them, but leave their destruction to the pruning process.
+	 * This avoids the need to parse collision and midphase lists whenever a particle is disabled.
 	 * 
 	 * NOTE: To reduce RBAN memory use, we do not create any collision blocks until the first call to CreateCollisionConstraint
+	 * (see ConstraintPool initialization)
 	*/
 	class CHAOS_API FCollisionConstraintAllocator
 	{
@@ -138,6 +143,19 @@ namespace Chaos
 			return CurrentEpoch;
 		}
 
+		/**
+		* Has the constraint expired. An expired constraint is one that was not refreshed this tick.
+		* 
+		* @note Sleeping constrainst will report they are expired, but they should not be deleted until awoken.
+		* We deliberately don't check the sleeping flag here because it adds a cache miss. This function
+		* should be called only when you already know the sleep state (or you must also check IsSleeping())
+		*/
+		bool IsConstraintExpired(const FPBDCollisionConstraint& Constraint)
+		{
+			const bool bIsExpired = Constraint.GetContainerCookie().LastUsedEpoch < CurrentEpoch;
+			return bIsExpired;
+		}
+
 		/**
 		 * @brief Destroy all constraints
 		*/
@@ -363,8 +381,7 @@ namespace Chaos
 
 		FParticlePairMidPhase* FindParticlePairMidPhaseImpl(FGeometryParticleHandle* Particle0, FGeometryParticleHandle* Particle1, FGeometryParticleHandle* SearchParticle)
 		{
-			// Find the existing midphase from one of the particle's lists of midphases
-//			FGeometryParticleHandle* SearchParticle2 = (Particle0->ParticleCollisions().Num() <= Particle1->ParticleCollisions().Num()) ? Particle0 : Particle1;
+			check((SearchParticle == Particle0) || (SearchParticle == Particle1));
 
 			const FCollisionParticlePairKey Key = FCollisionParticlePairKey(Particle0, Particle1);
 			return SearchParticle->ParticleCollisions().FindMidPhase(Key.GetKey());
@@ -374,7 +391,7 @@ namespace Chaos
 		{
 			const FCollisionParticlePairKey Key = FCollisionParticlePairKey(Particle0, Particle1);
 
-			// We enqueue a raw pointer and wrap it in a UniquePtr later
+			// We enqueue a raw pointer so that we can use a lock-free queue. We wrap it in a UniquePtr later 
 			FParticlePairMidPhase* MidPhase = new FParticlePairMidPhase();
 			MidPhase->Init(Particle0, Particle1, Key, *this);
 
diff --git a/Engine/Source/Runtime/Experimental/Chaos/Public/Chaos/Collision/CollisionKeys.h b/Engine/Source/Runtime/Experimental/Chaos/Public/Chaos/Collision/CollisionKeys.h
--- a/Engine/Source/Runtime/Experimental/Chaos/Public/Chaos/Collision/CollisionKeys.h
+++ b/Engine/Source/Runtime/Experimental/Chaos/Public/Chaos/Collision/CollisionKeys.h
@@ -57,24 +57,28 @@ namespace Chaos
 	private:
 		void GenerateKey(const FGeometryParticleHandle* Particle0, const FGeometryParticleHandle* Particle1)
 		{
-			const bool bIsLocalID0 = Particle0->ParticleID().LocalID != INDEX_NONE;
-			const bool bIsLocalID1 = Particle1->ParticleID().LocalID != INDEX_NONE;
-			const uint32 ID0 = uint32((bIsLocalID0) ? Particle0->ParticleID().LocalID : Particle0->ParticleID().GlobalID);
-			const uint32 ID1 = uint32((bIsLocalID1) ? Particle1->ParticleID().LocalID : Particle1->ParticleID().GlobalID);
-
-			if (ID0 < ID1)
-			{
-				Key.Key32s[0].Key31 = ID0;
-				Key.Key32s[0].IsLocal = bIsLocalID0;
-				Key.Key32s[1].Key31 = ID1;
-				Key.Key32s[1].IsLocal = bIsLocalID1;
-			}
-			else
+			Key.Key64 = 0;
+			if ((Particle0 != nullptr) && (Particle1 != nullptr))
 			{
-				Key.Key32s[0].Key31 = ID1;
-				Key.Key32s[0].IsLocal = bIsLocalID1;
-				Key.Key32s[1].Key31 = ID0;
-				Key.Key32s[1].IsLocal = bIsLocalID0;
+				const bool bIsLocalID0 = Particle0->ParticleID().LocalID != INDEX_NONE;
+				const bool bIsLocalID1 = Particle1->ParticleID().LocalID != INDEX_NONE;
+				const uint32 ID0 = uint32((bIsLocalID0) ? Particle0->ParticleID().LocalID : Particle0->ParticleID().GlobalID);
+				const uint32 ID1 = uint32((bIsLocalID1) ? Particle1->ParticleID().LocalID : Particle1->ParticleID().GlobalID);
+
+				if (ID0 < ID1)
+				{
+					Key.Key32s[0].Key31 = ID0;
+					Key.Key32s[0].IsLocal = bIsLocalID0;
+					Key.Key32s[1].Key31 = ID1;
+					Key.Key32s[1].IsLocal = bIsLocalID1;
+				}
+				else
+				{
+					Key.Key32s[0].Key31 = ID1;
+					Key.Key32s[0].IsLocal = bIsLocalID1;
+					Key.Key32s[1].Key31 = ID0;
+					Key.Key32s[1].IsLocal = bIsLocalID0;
+				}
 			}
 		}
 
diff --git a/Engine/Source/Runtime/Experimental/Chaos/Public/Chaos/Collision/PBDCollisionConstraint.h b/Engine/Source/Runtime/Experimental/Chaos/Public/Chaos/Collision/PBDCollisionConstraint.h
--- a/Engine/Source/Runtime/Experimental/Chaos/Public/Chaos/Collision/PBDCollisionConstraint.h
+++ b/Engine/Source/Runtime/Experimental/Chaos/Public/Chaos/Collision/PBDCollisionConstraint.h
@@ -273,6 +273,13 @@ namespace Chaos
 		*/
 		static FPBDCollisionConstraint MakeCopy(const FPBDCollisionConstraint& Source);
 
+		/**
+		 * Restore the properties of a collision from the properties of another. Used by the rewind/resim system when restoring
+		 * a collision constraint that was saved (and called MakeCopy).
+		 * This takes care not to overwrite any data maintained for the use of other systems (the ContainerCookie, Graph data, etc)
+		 */
+		void RestoreFrom(const FPBDCollisionConstraint& Source);
+
 		FPBDCollisionConstraint();
 		virtual ~FPBDCollisionConstraint();
 
@@ -649,6 +656,14 @@ namespace Chaos
 			}
 		}
 
+		/**
+		 *	A key used to uniquely identify the constraint (it is based on the two particle IDs)
+		 */
+		FCollisionParticlePairKey GetParticlePairKey() const
+		{
+			return FCollisionParticlePairKey(GetParticle0(), GetParticle1());
+		}
+
 	public:
 		const FPBDCollisionConstraintHandle* GetConstraintHandle() const { return this; }
 		FPBDCollisionConstraintHandle* GetConstraintHandle() { return this; }
diff --git a/Engine/Source/Runtime/Experimental/Chaos/Public/Chaos/GeometryParticles.h b/Engine/Source/Runtime/Experimental/Chaos/Public/Chaos/GeometryParticles.h
--- a/Engine/Source/Runtime/Experimental/Chaos/Public/Chaos/GeometryParticles.h
+++ b/Engine/Source/Runtime/Experimental/Chaos/Public/Chaos/GeometryParticles.h
@@ -457,6 +457,7 @@ namespace Chaos
 			TArrayCollection::AddArray(&MWeakParticleHandle);
 			TArrayCollection::AddArray(&MParticleConstraints);
 			TArrayCollection::AddArray(&MParticleCollisions);
+			TArrayCollection::AddArray(&MGraphIndex);
 			TArrayCollection::AddArray(&MResimType);
 			TArrayCollection::AddArray(&MEnabledDuringResim);
 			TArrayCollection::AddArray(&MLightWeightDisabled);
@@ -497,6 +498,7 @@ namespace Chaos
 			, MWeakParticleHandle(MoveTemp(Other.MWeakParticleHandle))
 			, MParticleConstraints(MoveTemp(Other.MParticleConstraints))
 			, MParticleCollisions(MoveTemp(Other.MParticleCollisions))
+			, MGraphIndex(MoveTemp(Other.MGraphIndex))
 			, MResimType(MoveTemp(Other.MResimType))
 			, MEnabledDuringResim(MoveTemp(Other.MEnabledDuringResim))
 			, MLightWeightDisabled(MoveTemp(Other.MLightWeightDisabled))
@@ -529,6 +531,7 @@ namespace Chaos
 			TArrayCollection::AddArray(&MWeakParticleHandle);
 			TArrayCollection::AddArray(&MParticleConstraints);
 			TArrayCollection::AddArray(&MParticleCollisions);
+			TArrayCollection::AddArray(&MGraphIndex);
 			TArrayCollection::AddArray(&MResimType);
 			TArrayCollection::AddArray(&MEnabledDuringResim);
 			TArrayCollection::AddArray(&MLightWeightDisabled);
@@ -571,6 +574,7 @@ namespace Chaos
 			TArrayCollection::AddArray(&MWeakParticleHandle);
 			TArrayCollection::AddArray(&MParticleConstraints);
 			TArrayCollection::AddArray(&MParticleCollisions);
+			TArrayCollection::AddArray(&MGraphIndex);
 			TArrayCollection::AddArray(&MResimType);
 			TArrayCollection::AddArray(&MEnabledDuringResim);
 			TArrayCollection::AddArray(&MLightWeightDisabled);
@@ -836,6 +840,9 @@ namespace Chaos
 			return MParticleCollisions[Index];
 		}
 
+		FORCEINLINE const int32 ConstraintGraphIndex(const int32 Index) const { return MGraphIndex[Index]; }
+		FORCEINLINE int32& ConstraintGraphIndex(const int32 Index) { return MGraphIndex[Index]; }
+
 		FORCEINLINE EResimType ResimType(const int32 Index) const { return MResimType[Index]; }
 		FORCEINLINE EResimType& ResimType(const int32 Index) { return MResimType[Index]; }
 
@@ -961,6 +968,7 @@ public:
 		TArrayCollectionArray<FWeakParticleHandle> MWeakParticleHandle;
 		TArrayCollectionArray<FConstraintHandleArray> MParticleConstraints;
 		TArrayCollectionArray<FParticleCollisions> MParticleCollisions;
+		TArrayCollectionArray<int32> MGraphIndex;
 		TArrayCollectionArray<EResimType> MResimType;
 		TArrayCollectionArray<bool> MEnabledDuringResim;
 		TArrayCollectionArray<bool> MLightWeightDisabled;
diff --git a/Engine/Source/Runtime/Experimental/Chaos/Public/Chaos/Island/IslandGraph.h b/Engine/Source/Runtime/Experimental/Chaos/Public/Chaos/Island/IslandGraph.h
--- a/Engine/Source/Runtime/Experimental/Chaos/Public/Chaos/Island/IslandGraph.h
+++ b/Engine/Source/Runtime/Experimental/Chaos/Public/Chaos/Island/IslandGraph.h
@@ -13,7 +13,6 @@ namespace Chaos
 	struct CHAOS_API TIslandGraphNode
 	{
 		/** List of edges connected to the node */
-		// @todo(chaos): should we really have a sparse array here? The implementation probably wastes a lot of memory
 		TSparseArray<int32> NodeEdges;
 
 		/**
@@ -22,14 +21,6 @@ namespace Chaos
 		 */
 		int32 IslandIndex = INDEX_NONE;
 
-		/**
-		 * List of islands in which the node is referenced. Should only be 1 item for valied nodes(
-		 * Dynamic/Sleeping Particles), and may be more if invalid (Static/Kinematic).
-		 * @note this array is populated externally (see see PopulateIslands in IslandManager.cpp)
-		 * @todo(chaos): should we really have a set here? The implementation probably wastes a lot of memory
-		 */
-		TSet<int32> NodeIslands;
-
 		/** Check if a node is valid (checked for graph partitioning into islands).
 		 * In the context of particles (nodes) and constraints (edges), bValidNode is true when
 		 * a Particle is Dynamic or Sleeping, but not when Static or Kinematic. This is because we
@@ -108,7 +99,10 @@ namespace Chaos
 		bool bIsPersistent = true;
 
 		/** Boolean to check if an island is sleeping or not */
-		bool bIsSleeping = true;
+		bool bIsSleeping = false;
+
+		/** Boolean to check if an island was sleeping last tick or not */
+		bool bWasSleeping = false;
 
 		/** List of children islands  to be merged */
 		// @todo(chaos): should we really have a set here? The implementation probably wastes a lot of memory
@@ -120,9 +114,6 @@ namespace Chaos
 		/** Island Item that is stored per island*/
 		IslandType IslandItem;
 
-		/** Max levels per island */
-		int32 MaxLevels = INDEX_NONE;
-
 		/** Max color per island */
 		int32 MaxColors = INDEX_NONE;
 	};
@@ -185,14 +176,12 @@ public:
 
 	/**
 	* Given a node item, update the graph node information (valid,discard...)
-	* @param NodeItem Node Item to be stored in the graph node
+	* @param NodeIndex Sparse Index of the node to update
 	* @param bValidNode Check if a node should be considered for graph partitioning (into islands)
-	* @param IslandIndex Potential island index we want the node to belong to
 	* @param bStationaryNode Boolean to check if a node is steady or not
-	*  @param NodeIndex Index to consider to update the graph node information
 	* @return Node index that has just been added
 	*/
-	void UpdateNode(const NodeType& NodeItem, const bool bValidNode, const int32 IslandIndex, const bool bStationaryNode, const int32 NodeIndex);
+	void UpdateNode(const int32 NodeIndex, const bool bValidNode, const bool bStationaryNode);
 
 	/**
 	 * Remove a node from the graph nodes list
@@ -232,17 +221,14 @@ public:
 	 */
 	void RemoveEdge(const int32 EdgeIndex);
 
-	void UpdateEdge(const int32 EdgeIndex);
-
 	/**
-	* Remove all the edges from the graph
-	*/
-	void ResetIslands();
-
+	 * Update edge state based on current node state
+	 */
+	void UpdateEdge(const int32 EdgeIndex);
 	/**
-	* Init the islands (remove edges only for non sleeping islands)
-	*/
-	void InitIslands();
+	 * Remove all non-sleeping edges from the graph
+	 */
+	void RemoveAllAwakeEdges();
 
 	/**
 	 * Get the maximum number of edges that are inside a graph (SparseArray size)
@@ -342,17 +328,20 @@ public:
 
 	/**
 	* Check if at least one of the 2 edge nodes is moving
-	* @param EdgeIndex Index of the edge to check
 	* @return Boolean to check if the edge is moving or not
 	*/
+	bool IsEdgeMoving(const FGraphEdge& GraphEdge)
+	{
+		const bool bFirstNodeMoving = GraphNodes.IsValidIndex(GraphEdge.FirstNode) && !GraphNodes[GraphEdge.FirstNode].bStationaryNode;
+		const bool bSecondNodeMoving = GraphNodes.IsValidIndex(GraphEdge.SecondNode) && !GraphNodes[GraphEdge.SecondNode].bStationaryNode;
+		return bFirstNodeMoving || bSecondNodeMoving;
+	}
+
 	bool IsEdgeMoving(const int32 EdgeIndex)
 	{
 		if(GraphEdges.IsValidIndex(EdgeIndex))
 		{
-			const FGraphEdge& GraphEdge = GraphEdges[EdgeIndex];
-			const bool bFirstNodeMoving = GraphNodes.IsValidIndex(GraphEdge.FirstNode) && !GraphNodes[GraphEdge.FirstNode].bStationaryNode;
-			const bool bSecondNodeMoving = GraphNodes.IsValidIndex(GraphEdge.SecondNode) && !GraphNodes[GraphEdge.SecondNode].bStationaryNode;
-			return bFirstNodeMoving || bSecondNodeMoving;
+			return IsEdgeMoving(GraphEdges[EdgeIndex]);
 		}
 		return false;
 	}
@@ -367,6 +356,25 @@ public:
 		return IslandIndex;
 	}
 
+	void SetNodeLevel(FGraphNode& GraphNode, const int32 LevelIndex)
+	{
+		GraphNode.LevelIndex = LevelIndex;
+		if (Owner != nullptr)
+		{
+			Owner->GraphNodeLevelAssigned(GraphNode);
+		}
+	}
+
+
+	void SetEdgeLevel(FGraphEdge& GraphEdge, const int32 LevelIndex)
+	{
+		GraphEdge.LevelIndex = LevelIndex;
+		if (Owner != nullptr)
+		{
+			Owner->GraphEdgeLevelAssigned(GraphEdge);
+		}
+	}
+
 	FGraphOwner* Owner = nullptr;
 
 	/** List of graph nodes */
@@ -399,9 +407,11 @@ struct TNullIslandGraphOwner
 
 	void GraphNodeAdded(const NodeType& NodeItem, const int32 NodeIndex) {}
 	void GraphNodeRemoved(const NodeType& NodeItem) {}
+	void GraphNodeLevelAssigned(TIslandGraphNode<NodeType>& Node) {}
 
 	void GraphEdgeAdded(const EdgeType& EdgeItem, const int32 EdgeIndex) {}
 	void GraphEdgeRemoved(const EdgeType& EdgeItem) {}
+	void GraphEdgeLevelAssigned(TIslandGraphEdge<EdgeType>& Edge) {}
 };
 
 
diff --git a/Engine/Source/Runtime/Experimental/Chaos/Public/Chaos/Island/IslandGroup.h b/Engine/Source/Runtime/Experimental/Chaos/Public/Chaos/Island/IslandGroup.h
--- a/Engine/Source/Runtime/Experimental/Chaos/Public/Chaos/Island/IslandGroup.h
+++ b/Engine/Source/Runtime/Experimental/Chaos/Public/Chaos/Island/IslandGroup.h
@@ -63,9 +63,10 @@ namespace Chaos
 		FPBDIslandManager& IslandManager;
 
 		// Item counters used to initialize the solver containers
+		static constexpr int32 NumExpectedConstraintTypes = 5;
 		int32 NumParticles;
 		int32 NumConstraints;
-		TArray<int32> NumContainerConstraints;	// Per ContainerId
+		TArray<int32, TInlineAllocator<NumExpectedConstraintTypes>> NumContainerConstraints;	// Per ContainerId
 
 		TArray<FPBDIsland*> Islands;
 	};
diff --git a/Engine/Source/Runtime/Experimental/Chaos/Public/Chaos/Island/IslandManager.h b/Engine/Source/Runtime/Experimental/Chaos/Public/Chaos/Island/IslandManager.h
--- a/Engine/Source/Runtime/Experimental/Chaos/Public/Chaos/Island/IslandManager.h
+++ b/Engine/Source/Runtime/Experimental/Chaos/Public/Chaos/Island/IslandManager.h
@@ -65,12 +65,6 @@ public:
 	*/
 	void RemoveConstraints();
 
-	/**
-	  * Default island manager constructor
-	  * @param Particles List of particles to be used to fill the graph nodes
-	  */
-	FPBDIslandManager(const TParticleView<FPBDRigidParticles>& PBDRigids);
-
 	/**
 	* Clear the manager and set up the particle-to-graph-node mapping for the specified particles
 	* Should be called before AddConstraint.
@@ -122,6 +116,11 @@ public:
 	  */
 	void RemoveConstraint(const uint32 ContainerId, FConstraintHandle* ConstraintHandle);
 
+	/**
+	 * Remove all constraints of the specified type.
+	 */
+	void RemoveConstraints(const uint32 ContainerId);
+
 	/**
 	 * @brief Remove all the constraints of the specified type (ContainerId) belonging to the Particle
 	 * @param Particle the particle to remove the constraints from
@@ -137,13 +136,12 @@ public:
 	int32 ReserveParticles(const int32 NumParticles);
 
 	/**
-	  * Add a particle to the Island graph
-	  * @param AddedParticle Particle to be added to the graph
-	  * @param IslandIndex Potential island index in which the particle will be added
-	  * @param bOnlyDynamic Boolean to add only dynamic particle into the islands
+	  * Add a particle to the graph
+	  * @param ParticleHandle Particle to be added to the graph
+	  * @param ParentParticleHandle Put the ParticleHandle in the same island as the parent with the same sleep state. May be null.
 	  * @return Sparse graph node index
 	  */
-	int32 AddParticle(FGeometryParticleHandle* AddedParticle, const int32 IslandIndex = INDEX_NONE, const bool bOnlyDynamic = true);
+	int32 AddParticle(FGeometryParticleHandle* ParticleHandle, const FGeometryParticleHandle* ParentParticleHandle = nullptr);
 
 	/**
 	  * Remove a particle from the graph
@@ -151,36 +149,38 @@ public:
 	  */
 	void RemoveParticle(FGeometryParticleHandle* Particle);
 
-	/**
-	  * Adds \p ChildParticle to the constraint graph.  Copies the sleeping
-	  * state and island of \p ParentParticle to \p ChildParticle.  Does nothing
-	  * if \p ParentParticle is not supplied.
-	  * @param ChildParticle Particle on which the sleeping state and island index will be copied
-	  * @param ParentParticle Particle from which the sleeping state and island index will be extracted
-	  */
-	void EnableParticle(FGeometryParticleHandle* ChildParticle, const FGeometryParticleHandle* ParentParticle);
-
-	/**
-	  * Disable a particle and remove it from its island.
-	  * @note: this does not remove any attached constraints - constraints need to be re-added and islands will need to be rebuilt.
-	  * @param Particle Particle to be disabled
-	  */
-	void DisableParticle(FGeometryParticleHandle* Particle);
-
 	/**
 	 * Get the Level for the specified particle (distance to a kinmetic in the graph)
 	*/
 	int32 GetParticleLevel(FGeometryParticleHandle* Particle) const;
 
+
 	/**
-	  * Remove all the graph edges
+	  * Call the visitor on all constraints in an awake island. E.g., can be used to check if any constraints
+	  * in the graph are expired and should be removed. However, no constraints should be added or removed
+	  * from the visitor itself - build a list and remove them after.
 	  */
-	void ResetIslands(const TParticleView<FPBDRigidParticles>& PBDRigids);
+	template<typename VisitorType>
+	void VisitConstraintsInAwakeIslands(const int32 ContainerId, const VisitorType& Visitor)
+	{
+		for (FGraphEdge& GraphEdge : IslandGraph->GraphEdges)
+		{
+			// If a constraint is not in any island (kinematic-kinematic for example) it is considered asleep.
+			if ((GraphEdge.ItemContainer == ContainerId) && (GraphEdge.IslandIndex != INDEX_NONE))
+			{
+				const FGraphIsland& GraphIsland = IslandGraph->GraphIslands[GraphEdge.IslandIndex];
+				if (!GraphIsland.bIsSleeping)
+				{
+					Visitor(GraphEdge.EdgeItem);
+				}
+			}
+		}
+	}
 
 	/**
-	 * Process all the particles and constraints in the graph. Assign all t islands.
-	*/
-	void UpdateGraph();
+	  * Clear all the islands
+	  */
+	void ResetIslands(const TParticleView<FPBDRigidParticles>& PBDRigids);
 
 	/**
 	  * Pack all the constraints into their assigned islands in sorted order
@@ -293,10 +293,13 @@ public:
 	void GraphIslandAdded(const int32 IslandIndex);
 	void GraphNodeAdded(FGeometryParticleHandle* NodeItem, const int32 NodeIndex);
 	void GraphNodeRemoved(FGeometryParticleHandle* NodeItem);
+	void GraphNodeLevelAssigned(FGraphNode& GraphNode);
 	void GraphEdgeAdded(const FConstraintHandleHolder& EdgeItem, const int32 EdgeIndex);
 	void GraphEdgeRemoved(const FConstraintHandleHolder& EdgeItem);
+	void GraphEdgeLevelAssigned(FGraphEdge& GraphEdge);
 
 #if CHAOS_CONSTRAINTHANDLE_DEBUG_ENABLED
+	bool DebugCheckParticlesInGraph(const TParticleView<FPBDRigidParticles>& Particles) const;
 	bool DebugCheckParticleNotInGraph(const FGeometryParticleHandle* ParticleHandle) const;
 	bool DebugParticleConstraintsNotInGraph(const FGeometryParticleHandle* Particle, const int32 ContainerId) const;
 	bool DebugCheckConstraintNotInGraph(const FConstraintHandle* ConstraintHandle) const;
@@ -311,6 +314,17 @@ protected:
 	void DebugCheckIslands() const;
 #endif
 
+	// Add a particle to the graph if it is not already in it. Return the node index if present (added or existing)
+	int32 TryAddParticle(FGeometryParticleHandle* ParticleHandle, const int32 IslandIndex);
+
+	// Add a particle to the graph if it is dynamic and not already in it. Return the node index if present, even if not dynamic
+	int32 TryAddParticleIfDynamic(FGeometryParticleHandle* ParticleHandle, const int32 IslandIndex);
+
+	// Add a particle to the graph if it is kinematic and not already in it. Return the node index if present, even if dynamic
+	int32 TryAddParticleIfNotDynamic(FGeometryParticleHandle* ParticleHandle, const int32 IslandIndex);
+
+	void InitIslands();
+
 	/**
 	* Sync the Islands with the IslandGraph
 	*/
@@ -324,6 +338,8 @@ protected:
 
 	void ValidateIslands() const;
 
+	TArray<FPBDConstraintContainer*>		ConstraintContainers;
+
 	/** List of solver islands in the manager */
 	TSparseArray<TUniquePtr<FPBDIsland>>	Islands;
 
@@ -333,14 +349,9 @@ protected:
 	/** Global graph to deal with merging and splitting of islands */
 	TUniquePtr<GraphType>					IslandGraph;
 
-	int32									MaxNumConstraintContainers;
-
 	/** Maximum particle index that we have in the graph*/
 	int32									MaxParticleIndex = INDEX_NONE;
 
-	/** Sorted list of islands */
-	TArray<int32>							SortedIslands;
-
 	/** Islands are only populated with constraints and particles during the solver phase. Trying to access them outside that is an error we can trap */
 	bool									bIslandsPopulated;
 
diff --git a/Engine/Source/Runtime/Experimental/Chaos/Public/Chaos/Island/SolverIsland.h b/Engine/Source/Runtime/Experimental/Chaos/Public/Chaos/Island/SolverIsland.h
--- a/Engine/Source/Runtime/Experimental/Chaos/Public/Chaos/Island/SolverIsland.h
+++ b/Engine/Source/Runtime/Experimental/Chaos/Public/Chaos/Island/SolverIsland.h
@@ -162,16 +162,17 @@ namespace Chaos
 		/**
 		* Members accessors
 		*/
-		FORCEINLINE bool IsSleeping() const {return bIsSleeping;}
-		FORCEINLINE void SetIsSleeping(const bool bIsSleepingIn ) { bSleepingChanged = bSleepingChanged || (bIsSleeping != bIsSleepingIn); bIsSleeping = bIsSleepingIn; }
-		FORCEINLINE void ResetSleepingChanged() { bSleepingChanged = false; }
+		FORCEINLINE bool IsSleeping() const { return bIsSleeping; }
+		FORCEINLINE void SetIsSleeping(const bool bInIsSleepingIn) { bIsSleeping = bInIsSleepingIn; }
+		FORCEINLINE bool IsSleepingChanged() const { return bIsSleepingChanged; }
+		FORCEINLINE void SetIsSleepingChanged(const bool bInIsSleepingChanged) { bIsSleepingChanged = bInIsSleepingChanged; }
+
 		FORCEINLINE bool IsPersistent() const { return bIsPersistent; }
 		FORCEINLINE void SetIsPersistent(const bool bIsPersistentIn) { bIsPersistent = bIsPersistentIn; }
 		FORCEINLINE bool NeedsResim() const { return bNeedsResim; }
 		FORCEINLINE void SetNeedsResim(const bool bNeedsResimIn) { bNeedsResim = bNeedsResimIn; }
 		FORCEINLINE int32 GetSleepCounter() const { return SleepCounter; }
 		FORCEINLINE void SetSleepCounter(const int32 SleepCounterIn) { SleepCounter = SleepCounterIn; }
-		FORCEINLINE bool SleepingChanged() const { return bSleepingChanged; }
 		FORCEINLINE void SetIsUsingCache(const bool bIsUsingCacheIn ) { bIsUsingCache = bIsUsingCacheIn; }
 		FORCEINLINE bool IsUsingCache() const { return bIsUsingCache; }
 	
@@ -182,6 +183,30 @@ namespace Chaos
 
 		void UpdateSyncState(FPBDRigidsSOAs& Particles);
 
+		// Debug functionality
+		bool DebugContainsParticle(const FConstGenericParticleHandle& ParticleHandle) const
+		{
+			for (const FPBDIslandParticle& IslandParticle : IslandParticles)
+			{
+				if (IslandParticle.GetParticle() == ParticleHandle->Handle())
+				{
+					return true;
+				}
+			}
+			return false;
+		}
+		bool DebugContainsConstraint(FConstraintHandle* ConstraintHandle) const
+		{
+			for (const FPBDIslandConstraint& IslandConstraint : IslandConstraintsByType[ConstraintHandle->GetContainerId()])
+			{
+				if (ConstraintHandle == IslandConstraint.GetConstraint())
+				{
+					return true;
+				}
+			}
+			return false;
+		}
+
 	private:
 
 		int32 IslandIndex;
@@ -196,7 +221,7 @@ namespace Chaos
 		bool bIsPersistent = true;
 
 		/** Flag to check if the sleeping state has changed or not */
-		bool bSleepingChanged = false;
+		bool bIsSleepingChanged = false;
 
 		/** Sleep counter to trigger island sleeping */
 		int32 SleepCounter = 0;
@@ -205,7 +230,8 @@ namespace Chaos
 		TArray<FPBDIslandParticle> IslandParticles;
 
 		/** List of all the island constraints handles */
-		TArray<TArray<FPBDIslandConstraint>> IslandConstraintsByType;
+		static constexpr int32 NumExpectedConstraintTypes = 5;
+		TArray<TArray<FPBDIslandConstraint>, TInlineAllocator<NumExpectedConstraintTypes>> IslandConstraintsByType;
 		int32 NumConstraints;
 	
 		/** Check if the island is using the cache or not */
diff --git a/Engine/Source/Runtime/Experimental/Chaos/Public/Chaos/PBDCollisionConstraints.h b/Engine/Source/Runtime/Experimental/Chaos/Public/Chaos/PBDCollisionConstraints.h
--- a/Engine/Source/Runtime/Experimental/Chaos/Public/Chaos/PBDCollisionConstraints.h
+++ b/Engine/Source/Runtime/Experimental/Chaos/Public/Chaos/PBDCollisionConstraints.h
@@ -269,6 +269,7 @@ private:
 
 	FCollisionConstraintAllocator ConstraintAllocator;
 	int32 NumActivePointConstraints;
+	TArray<FPBDCollisionConstraintHandle*> TempCollisions;	// Reused from tick to tick to build contact lists
 
 	TArrayCollectionArray<bool>& MCollided;
 	const TArrayCollectionArray<TSerializablePtr<FChaosPhysicsMaterial>>& MPhysicsMaterials;
diff --git a/Engine/Source/Runtime/Experimental/Chaos/Public/Chaos/PBDRigidsEvolution.h b/Engine/Source/Runtime/Experimental/Chaos/Public/Chaos/PBDRigidsEvolution.h
--- a/Engine/Source/Runtime/Experimental/Chaos/Public/Chaos/PBDRigidsEvolution.h
+++ b/Engine/Source/Runtime/Experimental/Chaos/Public/Chaos/PBDRigidsEvolution.h
@@ -2,7 +2,6 @@
 #pragma once
 
 #include "Chaos/PBDCollisionConstraints.h"
-#include "Chaos/PBDCollisionConstraintsPGS.h"
 #include "Chaos/PBDConstraintGraph.h"
 #include "Chaos/PBDRigidClustering.h"
 #include "Chaos/PBDRigidParticles.h"
@@ -342,7 +341,10 @@ public:
 	CHAOS_API FPBDRigidsSOAs& GetParticles() { return Particles; }
 	CHAOS_API const FPBDRigidsSOAs& GetParticles() const { return Particles; }
 
-	// NOTE: we do not currently support removing containers. In a few places we assume the ContainerId is equal to the array index.
+	/**
+	* Register a constraint container with the evolution. Constraints added to the container will be solved during the tick.
+	* @note we do not currently support removing containers. In a few places we assume the ContainerId is persistent and equal to the array index.
+	*/
 	CHAOS_API void AddConstraintContainer(FPBDConstraintContainer& InContainer, const int32 Priority = 0)
 	{
 		const int32 ContainerId = ConstraintContainers.Add(&InContainer);
@@ -352,36 +354,57 @@ public:
 		IslandGroupManager.AddConstraintContainer(InContainer, Priority);
 	}
 
+	/**
+	* Set the number of iterations to perform in the constraint position-solve phase
+	*/
 	CHAOS_API void SetNumPositionIterations(int32 InNumIterations)
 	{
 		NumPositionIterations = InNumIterations;
 	}
 
+	/**
+	* Get the number of position iterations the solver is running
+	*/
 	CHAOS_API int32 GetNumPositionIterations() const
 	{
 		return NumPositionIterations;
 	}
 
+	/**
+	* Set the number of iterations to perform in the constraint velocity-solve phase
+	*/
 	CHAOS_API void SetNumVelocityIterations(int32 InNumIterations)
 	{
 		NumVelocityIterations = InNumIterations;
 	}
 
+	/**
+	* Get the number of velocity iterations the solver is running
+	*/
 	CHAOS_API int32 GetNumVelocityIterations() const
 	{
 		return NumVelocityIterations;
 	}
 
+	/**
+	* Set the number of iterations to perform in the constraint projection phase
+	*/
 	CHAOS_API void SetNumProjectionIterations(int32 InNumIterations)
 	{
 		NumProjectionIterations = InNumIterations;
 	}
 
+	/**
+	* Get the number of projection iterations the solver is running
+	*/
 	CHAOS_API int32 GetNumProjectionIterations() const
 	{
 		return NumProjectionIterations;
 	}
 
+	/**
+	* Set the kinematic target for a particle. This will exist for only one tick - a new target must be set for the next tick if required.
+	*/
 	CHAOS_API void SetParticleKinematicTarget(FKinematicGeometryParticleHandle* KinematicHandle, const FKinematicTarget& NewKinematicTarget)
 	{
 		if (KinematicHandle)
@@ -399,33 +422,66 @@ public:
 		}
 	}
 	
+	/**
+	* To be called after creating a particle in the Particles container
+	* @todo(chaos): We should add a particle creation API to the evolution
+	* @todo(chaos): This is (or could be) very similar to Enable/Disable - do we really need both?
+	*/
+	CHAOS_API void RegisterParticle(FGeometryParticleHandle* Particle)
+	{
+		// Add to the graph if necessary
+		ConstraintGraph.AddParticle(Particle);
 
+		// Flag as dirty to update the spatial query acceleration structures
+		DirtyParticle(*Particle);
+	}
 
-
-	CHAOS_API void EnableParticle(FGeometryParticleHandle* Particle, const FGeometryParticleHandle* ParentParticle)
+	/**
+	* Enable a particle.Only enabled particles are simulated.
+	* If the particle has constraints connected to it they will also be enabled (assuming the other particles in the constraints are also enabled). 
+	* If a parent particle is specified, the child particle will be placed in the same island as the parent with the same sleep state.
+	*/
+	CHAOS_API void EnableParticle(FGeometryParticleHandle* Particle, const FGeometryParticleHandle* ParentParticle = nullptr)
 	{
 		Particles.EnableParticle(Particle);
 		EnableConstraints(Particle);
-		ConstraintGraph.EnableParticle(Particle, ParentParticle);
+		ConstraintGraph.AddParticle(Particle, ParentParticle);
 		DirtyParticle(*Particle);
 	}
 
+	/**
+	* Disable a particle so that it is no longer simulated. This also disables all constraints connected to the particle.
+	*/
 	CHAOS_API void DisableParticle(FGeometryParticleHandle* Particle)
 	{
 		RemoveParticleFromAccelerationStructure(*Particle);
 		Particles.DisableParticle(Particle);
-		// disable constraint before disabling the particle in the graph because of dependencies 
 		DisableConstraints(Particle);
-		ConstraintGraph.DisableParticle(Particle);
+		ConstraintGraph.RemoveParticle(Particle);
 	}
 
+	/**
+	* To be called when a particle geometry changes. We must clear collisions and anything else that may reference the prior shapes.
+	* @note This is also called during particle creation before we even know if the particle is static or kinematic, so we don't (re)add to 
+	* the graph unless it was already in there. There should be a call to either RegisterParticle or EnableParticle
+	* after creating a particle which will handle initially adding to the graph.
+	*/
 	CHAOS_API void InvalidateParticle(FGeometryParticleHandle* Particle)
 	{
-		// Remove the particle from the constraint graph. This should remove all the constraints too
-		ConstraintGraph.DisableParticle(Particle);
-
 		// Destroy all the transient constraints (collisions) because the particle has changed somehow (e.g. new shapes) and they may have cached the previous state
 		DestroyTransientConstraints(Particle);
+
+		// Remove all persistent constraints (joints etc) from the graph too
+		// @todo(chaos): do we still need to do this? Only if some non-collision constraints hold refs to shapes...
+		if (Particle->IsInConstraintGraph())
+		{
+			// Remove the particle from the constraint graph. This should remove all the constraints too
+			ConstraintGraph.RemoveParticle(Particle);
+
+			// Re-add the particle to the constraint graph
+			// (we could add a RemoveParticleConstraints method to the graph, but removing and adding a particle isn't too bad)
+			ConstraintGraph.AddParticle(Particle);
+		}
 	}
 	
 	CHAOS_API void FlushExternalAccelerationQueue(FAccelerationStructure& Acceleration,FPendingSpatialDataQueue& ExternalQueue);
@@ -510,19 +566,13 @@ public:
 		}
 	}
 
-	CHAOS_API void CreateParticle(FGeometryParticleHandle* ParticleAdded)
-	{
-		ConstraintGraph.AddParticle(ParticleAdded);
-		DirtyParticle(*ParticleAdded);
-	}
-
 	CHAOS_API void SetParticleObjectState(FPBDRigidParticleHandle* Particle, EObjectStateType ObjectState);
 
 	CHAOS_API void SetParticleSleepType(FPBDRigidParticleHandle* Particle, ESleepType InSleepType);
 
 	CHAOS_API void DisableParticles(const TSet<FGeometryParticleHandle*>& InParticles);
 
-	/** remove a constraint from the constraint graph (see AddConstraintsToConstraintGraph) */
+	/** remove a constraint from the constraint graph */
 	CHAOS_API void RemoveConstraintFromConstraintGraph(FConstraintHandle* ConstraintHandle)
 	{
 		if (ConstraintHandle->IsInConstraintGraph())
@@ -531,7 +581,7 @@ public:
 		}
 	}
 
-	/** remove a list of constraints from the constraint graph (see AddConstraintsToConstraintGraph) */
+	/** remove a list of constraints from the constraint graph */
 	CHAOS_API void RemoveConstraintsFromConstraintGraph(const FConstraintHandleArray& Constraints)
 	{
 		for (FConstraintHandle* BaseConstraintHandle : Constraints)
@@ -543,27 +593,12 @@ public:
 		}
 	}
 
-	/** Add a list of constraints to the constraint graph (see RemoveConstraintsFromConstraintGraph) */
-	CHAOS_API void AddConstraintsToConstraintGraph(const FConstraintHandleArray& Constraints)
-	{
-		for (FConstraintHandle* BaseConstraintHandle : Constraints)
-		{
-			if (FPBDJointConstraintHandle* ConstraintHandle = BaseConstraintHandle->As<FPBDJointConstraintHandle>())
-			{
-				if (!ensure(ConstraintHandle->IsInConstraintGraph()))
-				{
-					ConstraintGraph.AddConstraint(ConstraintHandle->GetContainerId(), ConstraintHandle, ConstraintHandle->GetConstrainedParticles());
-				}
-			}
-		}
-	}
-
-
-
-	/** Disconnect constraints from a set of particles to be destroyed. 
+	/** 
+	* Disconnect constraints from a set of particles to be destroyed. 
 	* this will set the constraints to Enbaled = false and set their respective bodies handles to nullptr.
 	* Once this is done, the constraints cannot be re-enabled.
 	* @note This only applies to persistent constraints (joints etc), not transient constraints (collisons)
+	* @see DestroyTransientConstraints()
 	*/
 	CHAOS_API void DisconnectConstraints(const TSet<FGeometryParticleHandle*>& RemovedParticles)
 	{
@@ -578,10 +613,12 @@ public:
 		}
 	}
 
-	/** Disconnect constraints from a particle to be removed (or destroyed)
+	/** 
+	* Disconnect constraints from a particle to be removed (or destroyed)
 	* this will set the constraints to Enabled = false, but leave connections to the particles to support
 	* re-enabling at a later time.
 	* @note This only applies to persistent constraints (joints etc), not transient constraints (collisons)
+	* @see DestroyTransientConstraints()
 	*/
 	CHAOS_API void DisableConstraints(FGeometryParticleHandle* ParticleHandle)
 	{
@@ -596,15 +633,10 @@ public:
 		}
 	}
 
-	CHAOS_API void DisableConstraints(const TSet<FGeometryParticleHandle*>& DisabledParticles)
-	{
-		for (FGeometryParticleHandle* ParticleHandle : DisabledParticles)
-		{
-			DisableConstraints(ParticleHandle);
-		}
-	}
-
-	/** Enable constraints from the enabled particles; constraints will only become enabled if their particle end points are valid */
+	/** 
+	* Enable constraints from the enabled particles; constraints will only become enabled if their particle end points are valid.
+	* @note This only applies to persistent constraints (joints etc), not transient constraints (collisons)
+	*/
 	CHAOS_API void EnableConstraints(FGeometryParticleHandle* ParticleHandle)
 	{
 		for (FConstraintHandle* Constraint : ParticleHandle->ParticleConstraints())
@@ -616,14 +648,9 @@ public:
 		}
 	}
 
-	CHAOS_API void EnableConstraints(const TSet<FGeometryParticleHandle*>& EnabledParticles)
-	{
-		for (FGeometryParticleHandle* ParticleHandle : EnabledParticles)
-		{
-			EnableConstraints(ParticleHandle);
-		}
-	}
-
+	/** 
+	* Clear all constraints from the system reeady for shut down 
+	*/
 	CHAOS_API void ResetConstraints()
 	{
 		// Remove all the constraints from the graph
@@ -646,6 +673,9 @@ public:
 		}
 	}
 
+	/**
+	* Destroy all transient constraints (collisions) onvolving the specified particle.
+	*/
 	virtual void DestroyTransientConstraints(FGeometryParticleHandle* Particle) {}
 
 	const TParticleView<FPBDRigidClusteredParticles>& GetNonDisabledClusteredView() const { return Particles.GetNonDisabledClusteredView(); }
@@ -665,20 +695,6 @@ public:
 		Particle->AuxilaryValue(PhysicsMaterials) = InMaterial;
 	}
 
-	CHAOS_API int32 NumIslands() const { return ConstraintGraph.NumIslands(); }
-
-	void InitializeAccelerationStructures()
-	{
-		ConstraintGraph.InitializeGraph(Particles.GetNonDisabledDynamicView());
-
-		for (FPBDConstraintContainer* Container : ConstraintContainers)
-		{
-			Container->AddConstraintsToGraph(ConstraintGraph);
-		}
-
-		ConstraintGraph.ResetIslands(Particles.GetNonDisabledDynamicView());
-	}
-
 	void PrepareTick()
 	{
 		Collisions::ResetChaosCollisionCounters();
@@ -913,17 +929,20 @@ protected:
 
 	void CreateConstraintGraph()
 	{
-		// Add new particles to the graph
+		// Update the current state of the graph based on existing particles and constraints.
+		// Any new particles (from this tick) should have been added when they were enabled.
 		ConstraintGraph.InitializeGraph(Particles.GetNonDisabledDynamicView());
 
-		// Add all constraints to graph (all constraints except sleeping collisions are removed every tick)
+		// Add all constraints to the graph.
+		// NOTE: in PersistentGraph mode, only new constraints need to be added and expired ones should be removed.
+		// In non-peristent mode, all constraints in awake islands are removed every tick and so all would need to be re-added here.
+		// @todo(chaos): it feels a bit inconsistent that particles are added when enabled, but constraints
+		// are added here. Currently it needs to be this way to properly support user changes to sleep state
+		// of particles, but we could probably make this cleaner.
 		for (FPBDConstraintContainer* ConstraintContainer : ConstraintContainers)
 		{
-			ConstraintContainer->AddConstraintsToGraph(ConstraintGraph);
+			ConstraintContainer->AddConstraintsToGraph(GetConstraintGraph());
 		}
-
-		// Process the addidions
-		ConstraintGraph.UpdateGraph();
 	}
 
 	void CreateIslands()
diff --git a/Engine/Source/Runtime/Experimental/Chaos/Public/Chaos/ParticleHandle.h b/Engine/Source/Runtime/Experimental/Chaos/Public/Chaos/ParticleHandle.h
--- a/Engine/Source/Runtime/Experimental/Chaos/Public/Chaos/ParticleHandle.h
+++ b/Engine/Source/Runtime/Experimental/Chaos/Public/Chaos/ParticleHandle.h
@@ -414,6 +414,7 @@ protected:
 	TGeometryParticleHandleImp()
 		: TParticleHandleBase<T, d>()
 	{
+		SetConstraintGraphIndex(INDEX_NONE);
 	}
 
 	TGeometryParticleHandleImp(TSerializablePtr<TGeometryParticles<T, d>> InParticles, int32 InParticleIdx, int32 InHandleIdx, const FGeometryParticleParameters& Params)
@@ -425,6 +426,7 @@ protected:
 		GeometryParticleDefaultConstruct<T, d>(*this, Params);
 		SetHasBounds(false);
 		SetLightWeightDisabled(false);
+		SetConstraintGraphIndex(INDEX_NONE);
 	}
 
 	template <typename TParticlesType, typename TParams>
@@ -446,6 +448,9 @@ public:
 
 	~TGeometryParticleHandleImp()
 	{
+		// If we weren't removed from the graph, invalid pointer dereferencing is possible
+		check(ConstraintGraphIndex() == INDEX_NONE);
+
 		if (bPersistent)
 		{
 			GeometryParticles->ResetWeakParticleHandle(ParticleIdx);
@@ -708,6 +713,21 @@ public:
 		return GeometryParticles->ParticleCollisions(ParticleIdx);
 	}
 
+	int32 ConstraintGraphIndex() const
+	{ 
+		return GeometryParticles->ConstraintGraphIndex(ParticleIdx);
+	}
+	
+	void SetConstraintGraphIndex(const int32 InGraphIndex)
+	{ 
+		GeometryParticles->ConstraintGraphIndex(ParticleIdx) = InGraphIndex;
+	}
+
+	bool IsInConstraintGraph() const
+	{
+		return ConstraintGraphIndex() != INDEX_NONE;
+	}
+
 protected:
 
 	friend TGeometryParticleHandles<T, d>;
@@ -847,7 +867,6 @@ protected:
 		SetAngularAcceleration(TVector<T, d>(0));
 		SetObjectStateLowLevel(Params.bStartSleeping ? EObjectStateType::Sleeping : EObjectStateType::Dynamic);
 		SetIslandIndex(INDEX_NONE);
-		SetConstraintGraphIndex(INDEX_NONE);
 		SetToBeRemovedOnFracture(false);
 		SetSleepType(ESleepType::MaterialSleep);
 		SetInvIConditioning(TVec3<FRealSingle>(1));
@@ -1026,10 +1045,6 @@ public:
 	int32& IslandIndex() { return PBDRigidParticles->IslandIndex(ParticleIdx); }
 	void SetIslandIndex(const int32 InIslandIndex) { PBDRigidParticles->IslandIndex(ParticleIdx) = InIslandIndex; }
 	
-	int32 ConstraintGraphIndex() const { return PBDRigidParticles->ConstraintGraphIndex(ParticleIdx); }
-	int32& ConstraintGraphIndex() { return PBDRigidParticles->ConstraintGraphIndex(ParticleIdx); }
-	void SetConstraintGraphIndex(const int32 InGraphIndex) { PBDRigidParticles->ConstraintGraphIndex(ParticleIdx) = InGraphIndex; }
-
 	bool ToBeRemovedOnFracture() const { return PBDRigidParticles->ToBeRemovedOnFracture(ParticleIdx); }
 	bool& ToBeRemovedOnFracture() { return PBDRigidParticles->ToBeRemovedOnFracture(ParticleIdx); }
 	void SetToBeRemovedOnFracture(const bool bToBeRemovedOnFracture) { PBDRigidParticles->ToBeRemovedOnFracture(ParticleIdx) = bToBeRemovedOnFracture; }
@@ -1829,6 +1844,21 @@ public:
 		}
 	}
 
+	FORCEINLINE int32 ConstraintGraphIndex() const
+	{
+		return MHandle->ConstraintGraphIndex();
+	}
+
+	FORCEINLINE void SetConstraintGraphIndex(const int32 InGraphIndex)
+	{
+		MHandle->SetConstraintGraphIndex(InGraphIndex);
+	}
+
+	FORCEINLINE bool IsInConstraintGraph() const
+	{
+		return MHandle->IsInConstraintGraph();
+	}
+
 	bool ToBeRemovedOnFracture() const 
 	{
 		if (MHandle->CastToRigidParticle() && MHandle->ObjectState() == EObjectStateType::Dynamic)
diff --git a/Engine/Source/Runtime/Experimental/Chaos/Public/Chaos/RigidParticles.h b/Engine/Source/Runtime/Experimental/Chaos/Public/Chaos/RigidParticles.h
--- a/Engine/Source/Runtime/Experimental/Chaos/Public/Chaos/RigidParticles.h
+++ b/Engine/Source/Runtime/Experimental/Chaos/Public/Chaos/RigidParticles.h
@@ -95,7 +95,6 @@ public:
 		TArrayCollection::AddArray(&MObjectState);
 		TArrayCollection::AddArray(&MPreObjectState);
 		TArrayCollection::AddArray(&MIslandIndex);
-		TArrayCollection::AddArray(&MGraphIndex);
 		TArrayCollection::AddArray(&MToBeRemovedOnFracture);
 		TArrayCollection::AddArray(&MSleepType);
 	}
@@ -125,7 +124,6 @@ public:
 		, MControlFlags(MoveTemp(Other.MControlFlags))
 		, MTransientFlags(MoveTemp(Other.MTransientFlags))
 		, MIslandIndex(MoveTemp(Other.MIslandIndex))
-		, MGraphIndex(MoveTemp(Other.MGraphIndex))
 		, MToBeRemovedOnFracture(MoveTemp(Other.MToBeRemovedOnFracture))
 		, MObjectState(MoveTemp(Other.MObjectState))
 		, MPreObjectState(MoveTemp(Other.MPreObjectState))
@@ -158,7 +156,6 @@ public:
 		TArrayCollection::AddArray(&MObjectState);
 		TArrayCollection::AddArray(&MPreObjectState);
 		TArrayCollection::AddArray(&MIslandIndex);
-		TArrayCollection::AddArray(&MGraphIndex);
 		TArrayCollection::AddArray(&MToBeRemovedOnFracture);
 		TArrayCollection::AddArray(&MSleepType);
 	}
@@ -288,9 +285,6 @@ public:
 	FORCEINLINE const int32 IslandIndex(const int32 Index) const { return MIslandIndex[Index]; }
 	FORCEINLINE int32& IslandIndex(const int32 Index) { return MIslandIndex[Index]; }
 
-	FORCEINLINE const int32 ConstraintGraphIndex(const int32 Index) const { return MGraphIndex[Index]; }
-	FORCEINLINE int32& ConstraintGraphIndex(const int32 Index) { return MGraphIndex[Index]; }
-
 	FORCEINLINE FString ToString(int32 index) const
 	{
 		FString BaseString = TKinematicGeometryParticles<T, d>::ToString(index);
@@ -376,7 +370,6 @@ private:
 	TArrayCollectionArray<FRigidParticleControlFlags> MControlFlags;
 	TArrayCollectionArray<FRigidParticleTransientFlags> MTransientFlags;
 	TArrayCollectionArray<int32> MIslandIndex;
-	TArrayCollectionArray<int32> MGraphIndex;
 	TArrayCollectionArray<bool> MToBeRemovedOnFracture;
 	TArrayCollectionArray<EObjectStateType> MObjectState;
 	TArrayCollectionArray<EObjectStateType> MPreObjectState;
