FNames containing asset paths are deprecated. FSoftObjectPath or FTopLevelAssetPath should be used instead.
Fixups for misc engine code.

diff --git a/Engine/Source/Runtime/Engine/Private/ObjectLibrary.cpp b/Engine/Source/Runtime/Engine/Private/ObjectLibrary.cpp
--- a/Engine/Source/Runtime/Engine/Private/ObjectLibrary.cpp
+++ b/Engine/Source/Runtime/Engine/Private/ObjectLibrary.cpp
@@ -503,7 +503,7 @@ int32 UObjectLibrary::LoadAssetsFromAssetData()
 		for (int32 AssetIdx = 0; AssetIdx < AssetDataList.Num(); AssetIdx++)
 		{
 			FAssetData& Data = AssetDataList[AssetIdx];
-			AssetsToStream.AddUnique(FSoftObjectPath(Data.PackageName));
+			AssetsToStream.AddUnique(FSoftObjectPath(Data.PackageName, {}, {}));
 		}
 
 		if (AssetsToStream.Num())
diff --git a/Engine/Source/Runtime/Engine/Private/PhysicsEngine/PhysicalMaterialMask.cpp b/Engine/Source/Runtime/Engine/Private/PhysicsEngine/PhysicalMaterialMask.cpp
--- a/Engine/Source/Runtime/Engine/Private/PhysicsEngine/PhysicalMaterialMask.cpp
+++ b/Engine/Source/Runtime/Engine/Private/PhysicsEngine/PhysicalMaterialMask.cpp
@@ -43,7 +43,7 @@ static void OnDumpPhysicalMaterialMaskData(const TArray< FString >& Arguments)
 	if (Arguments.Num() > 0)
 	{
 		const FAssetRegistryModule& AssetRegistryModule = FModuleManager::LoadModuleChecked<FAssetRegistryModule>(TEXT("AssetRegistry"));
-		FAssetData PhysMatMaskAsset = AssetRegistryModule.Get().GetAssetByObjectPath(*Arguments[0]);
+		FAssetData PhysMatMaskAsset = AssetRegistryModule.Get().GetAssetByObjectPath(FSoftObjectPath(Arguments[0]));
 		if (PhysMatMaskAsset.IsValid() == false)
 		{
 			TArray<FAssetData> AssetsInPackage;
diff --git a/Engine/Source/Runtime/Engine/Private/Slate/SlateGameResources.cpp b/Engine/Source/Runtime/Engine/Private/Slate/SlateGameResources.cpp
--- a/Engine/Source/Runtime/Engine/Private/Slate/SlateGameResources.cpp
+++ b/Engine/Source/Runtime/Engine/Private/Slate/SlateGameResources.cpp
@@ -232,7 +232,7 @@ void FSlateGameResources::RemoveAsset(const FAssetData& InRemovedAssetData)
 
 bool FSlateGameResources::ShouldCache( const FAssetData& InAssetData )
 {
-	return InAssetData.ObjectPath.ToString().StartsWith( ContentRootDir, ESearchCase::CaseSensitive )
+	return InAssetData.PackageName.ToString().StartsWith( ContentRootDir, ESearchCase::CaseSensitive )
 		&& InAssetData.AssetClassPath == USlateWidgetStyleAsset::StaticClass()->GetClassPathName();
 }
 
diff --git a/Engine/Source/Runtime/Engine/Private/StreamableManager.cpp b/Engine/Source/Runtime/Engine/Private/StreamableManager.cpp
--- a/Engine/Source/Runtime/Engine/Private/StreamableManager.cpp
+++ b/Engine/Source/Runtime/Engine/Private/StreamableManager.cpp
@@ -1294,7 +1294,7 @@ TSharedPtr<FStreamableHandle> FStreamableManager::RequestAsyncLoad(TArray<FSoftO
 			--NumValidRequests;
 			continue;
 		}
-		else if (FPackageName::IsShortPackageName(TargetName.GetAssetPathName()))
+		else if (FPackageName::IsShortPackageName(TargetName.GetLongPackageFName()))
 		{
 			UE_LOG(LogStreamableManager, Error, TEXT("RequestAsyncLoad(%s) called with invalid package name %s"), *DebugName, *TargetName.ToString());
 			NewRequest->CancelHandle();
diff --git a/Engine/Source/Runtime/Experimental/Interchange/Engine/Private/InterchangeManager.cpp b/Engine/Source/Runtime/Experimental/Interchange/Engine/Private/InterchangeManager.cpp
--- a/Engine/Source/Runtime/Experimental/Interchange/Engine/Private/InterchangeManager.cpp
+++ b/Engine/Source/Runtime/Experimental/Interchange/Engine/Private/InterchangeManager.cpp
@@ -383,7 +383,7 @@ UInterchangePipelineBase* UE::Interchange::GeneratePipelineInstance(const FSoftO
 		else
 		{
 			//Log an error because we cannot load the python class, maybe the python script was not loaded
-			UE_LOG(LogInterchangeEngine, Error, TEXT("Cannot generate a pipeline instance because the blueprint %s do not have a valid generated class."), *PipelineInstance.GetAssetPathName().ToString());
+			UE_LOG(LogInterchangeEngine, Error, TEXT("Cannot generate a pipeline instance because the blueprint %s do not have a valid generated class."), *PipelineInstance.GetWithoutSubPath().ToString());
 		}
 	}
 	else if (const UInterchangePythonPipelineAsset* PythonPipeline = Cast<UInterchangePythonPipelineAsset>(ReferenceInstance))
@@ -395,7 +395,7 @@ UInterchangePipelineBase* UE::Interchange::GeneratePipelineInstance(const FSoftO
 		else
 		{
 			//Log an error because we cannot load the python class, maybe the python script was not loaded
-			UE_LOG(LogInterchangeEngine, Error, TEXT("Cannot generate a pipeline instance because the Python pipeline asset %s do not have a valid generated pipeline instance."), *PipelineInstance.GetAssetPathName().ToString());
+			UE_LOG(LogInterchangeEngine, Error, TEXT("Cannot generate a pipeline instance because the Python pipeline asset %s do not have a valid generated pipeline instance."), *PipelineInstance.GetWithoutSubPath().ToString());
 		}
 	}
 	else if (const UInterchangePipelineBase* DefaultPipeline = Cast<UInterchangePipelineBase>(ReferenceInstance))
@@ -404,7 +404,7 @@ UInterchangePipelineBase* UE::Interchange::GeneratePipelineInstance(const FSoftO
 	}
 	else
 	{
-		UE_LOG(LogInterchangeEngine, Error, TEXT("Cannot generate a pipeline instance because the pipeline asset %s type is unknown."), *PipelineInstance.GetAssetPathName().ToString());
+		UE_LOG(LogInterchangeEngine, Error, TEXT("Cannot generate a pipeline instance because the pipeline asset %s type is unknown."), *PipelineInstance.GetWithoutSubPath().ToString());
 	}
 	return GeneratedPipeline;
 }
