Add analytic to interchange

#rnx

diff --git a/Engine/Source/Runtime/Experimental/Interchange/Core/Private/InterchangePipelineBase.cpp b/Engine/Source/Runtime/Experimental/Interchange/Core/Private/InterchangePipelineBase.cpp
--- a/Engine/Source/Runtime/Experimental/Interchange/Core/Private/InterchangePipelineBase.cpp
+++ b/Engine/Source/Runtime/Experimental/Interchange/Core/Private/InterchangePipelineBase.cpp
@@ -93,6 +93,11 @@ FName UInterchangePipelineBase::GetPropertiesStatesPropertyName()
 	return GET_MEMBER_NAME_CHECKED(UInterchangePipelineBase, PropertiesStates);
 }
 
+FName UInterchangePipelineBase::GetResultsPropertyName()
+{
+	return GET_MEMBER_NAME_CHECKED(UInterchangePipelineBase, Results);
+}
+
 void UInterchangePipelineBase::LoadSettingsInternal(const FName PipelineStackName, const FString& ConfigFilename, TMap<FName, FInterchangePipelinePropertyStates>& ParentPropertiesStates)
 {
 	int32 PortFlags = 0;
diff --git a/Engine/Source/Runtime/Experimental/Interchange/Core/Public/InterchangePipelineBase.h b/Engine/Source/Runtime/Experimental/Interchange/Core/Public/InterchangePipelineBase.h
--- a/Engine/Source/Runtime/Experimental/Interchange/Core/Public/InterchangePipelineBase.h
+++ b/Engine/Source/Runtime/Experimental/Interchange/Core/Public/InterchangePipelineBase.h
@@ -267,6 +267,8 @@ public:
 	 */
 	static FName GetPropertiesStatesPropertyName();
 
+	static FName GetResultsPropertyName();
+
 	bool CanEditPropertiesStates() { return bAllowPropertyStatesEdition; }
 	bool IsReimportContext() { return bIsReimportContext; }
 protected:
diff --git a/Engine/Source/Runtime/Experimental/Interchange/Engine/Private/InterchangeManager.cpp b/Engine/Source/Runtime/Experimental/Interchange/Engine/Private/InterchangeManager.cpp
--- a/Engine/Source/Runtime/Experimental/Interchange/Engine/Private/InterchangeManager.cpp
+++ b/Engine/Source/Runtime/Experimental/Interchange/Engine/Private/InterchangeManager.cpp
@@ -7,6 +7,7 @@
 #include "Async/Async.h"
 #include "CoreMinimal.h"
 #include "Engine/Blueprint.h"
+#include "EngineAnalytics.h"
 #include "Framework/Notifications/NotificationManager.h"
 #include "InterchangeAssetImportData.h"
 #include "InterchangeBlueprintPipelineBase.h"
@@ -19,6 +20,7 @@
 #include "InterchangeWriterBase.h"
 #include "Internationalization/Internationalization.h"
 #include "Misc/App.h"
+#include "Misc/AssertionMacros.h"
 #include "Misc/AsyncTaskNotification.h"
 #include "Misc/MessageDialog.h"
 #include "Misc/ScopedSlowTask.h"
@@ -32,12 +34,12 @@
 #include "UObject/NameTypes.h"
 #include "UObject/Object.h"
 #include "UObject/ObjectMacros.h"
+#include "UObject/ReflectedTypeAccessors.h"
 #include "UObject/UObjectIterator.h"
 #include "UObject/WeakObjectPtrTemplates.h"
 #include "Widgets/Notifications/SNotificationList.h"
 
 
-
 namespace UE::Interchange::Private
 {
 	const FLogCategoryBase* GetLogInterchangePtr()
@@ -53,6 +55,106 @@ namespace UE::Interchange::Private
 	{
 		return const_cast<FInterchangeImportSettings&>(FInterchangeProjectSettingsUtils::GetImportSettings(InterchangeProjectSettings, ImportType == EImportType::ImportType_Scene));
 	}
+
+	void FillPipelineAnalyticData(UInterchangePipelineBase* Pipeline, TArray<FAnalyticsEventAttribute>& Attribs, const int32 UniqueId, const FString& ParentPipeline)
+	{
+		if (!FEngineAnalytics::IsAvailable())
+		{
+			return;
+		}
+		
+		int32 PortFlags = 0;
+		UClass* Class = Pipeline->GetClass();
+		FString PipelineChainName = ParentPipeline.IsEmpty() ? Pipeline->GetName() : ParentPipeline + TEXT(".") + Pipeline->GetName();
+		if (ParentPipeline.IsEmpty())
+		{
+			Attribs.Add(FAnalyticsEventAttribute(TEXT("Pipeline.Name"), PipelineChainName));
+		}
+		else
+		{
+			Attribs.Add(FAnalyticsEventAttribute(TEXT("SubPipeline.Name"), PipelineChainName));
+		}
+
+		TArray<FAnalyticsEventAttribute> PipelineAttribs;
+		PipelineAttribs.Add(FAnalyticsEventAttribute(TEXT("UniqueId"), UniqueId));
+		PipelineAttribs.Add(FAnalyticsEventAttribute(TEXT("Name"), PipelineChainName));
+		PipelineAttribs.Add(FAnalyticsEventAttribute(TEXT("Class"), Class->GetName()));
+
+		for (FProperty* Property = Class->PropertyLink; Property; Property = Property->PropertyLinkNext)
+		{
+			if (Property->HasAnyPropertyFlags(CPF_Transient))
+			{
+				continue;
+			}
+
+			if (Property->GetFName() == UInterchangePipelineBase::GetPropertiesStatesPropertyName())
+			{
+				continue;
+			}
+
+			if (Property->GetFName() == UInterchangePipelineBase::GetResultsPropertyName())
+			{
+				continue;
+			}
+
+			const FString PropertyName = Property->GetName();
+			FObjectProperty* SubObject = CastField<FObjectProperty>(Property);
+			UInterchangePipelineBase* SubPipeline = SubObject ? Cast<UInterchangePipelineBase>(SubObject->GetObjectPropertyValue_InContainer(Pipeline)) : nullptr;
+			//Add the category name to the key
+			FString CategoryName = FString();
+#if WITH_EDITORONLY_DATA
+			CategoryName = Property->GetMetaData("Category");
+			if (!SubPipeline && CategoryName.IsEmpty())
+			{
+				//In Editor do not add property with no category
+				continue;
+			}
+			CategoryName.ReplaceCharInline(TEXT('.'), TEXT('_'));
+			CategoryName.RemoveSpacesInline();
+			CategoryName = TEXT(".") + CategoryName;
+#endif
+
+			
+			if (FArrayProperty* Array = CastField<FArrayProperty>(Property))
+			{
+				FScriptArrayHelper_InContainer ArrayHelper(Array, Pipeline);
+				for (int32 i = 0; i < ArrayHelper.Num(); i++)
+				{
+					FString	Buffer;
+					Array->Inner->ExportTextItem_Direct(Buffer, ArrayHelper.GetRawPtr(i), ArrayHelper.GetRawPtr(i), Pipeline, PortFlags);
+					
+					PipelineAttribs.Add(FAnalyticsEventAttribute(FString::Printf(TEXT("Property%s.%s"), *CategoryName, *PropertyName), Buffer));
+				}
+			}
+			else if (SubPipeline)
+			{
+				// Save the settings if the referenced pipeline is a subobject of ours
+				if (SubPipeline->IsInOuter(Pipeline))
+				{
+					//Go recursive with subObject, like if they are part of the same object
+					FillPipelineAnalyticData(SubPipeline, Attribs, UniqueId, PipelineChainName);
+				}
+			}
+			else
+			{
+				for (int32 Index = 0; Index < Property->ArrayDim; Index++)
+				{
+					FString PropertyIndexName = PropertyName;
+					if (Property->ArrayDim != 1)
+					{
+						PropertyIndexName += TEXT("[") + FString::FromInt(Index) + TEXT("]");
+					}
+
+					FString	Value;
+					Property->ExportText_InContainer(Index, Value, Pipeline, Pipeline, Pipeline, PortFlags);
+					PipelineAttribs.Add(FAnalyticsEventAttribute(FString::Printf(TEXT("Property%s.%s"), *CategoryName, *PropertyIndexName), Value));
+				}
+			}
+		}
+		
+		FString EventString = FString::Printf(TEXT("Interchange.Usage.Import.Pipeline.%s"), *PipelineChainName);
+		FEngineAnalytics::GetProvider().RecordEvent(EventString, PipelineAttribs);
+	}
 }
 
 UE::Interchange::FScopedSourceData::FScopedSourceData(const FString& Filename)
@@ -93,6 +195,92 @@ void UE::Interchange::FImportAsyncHelper::AddReferencedObjects(FReferenceCollect
 	Collector.AddReferencedObjects(CreatedFactories);
 }
 
+void UE::Interchange::FImportAsyncHelper::SendAnalyticImportEndData()
+{
+	if (!FEngineAnalytics::IsAvailable())
+	{
+		return;
+	}
+
+	TArray<FAnalyticsEventAttribute> Attribs;
+	//Set the unique id of this import
+	Attribs.Add(FAnalyticsEventAttribute(TEXT("UniqueId"), UniqueId));
+	Attribs.Add(FAnalyticsEventAttribute(TEXT("IsCanceled"), bCancel));
+	if (bCancel)
+	{
+		return;
+	}
+
+	auto CollectObject = [&Attribs](bool bIsReimport, UObject* Object)
+	{
+		const int32 ObjectUniqueId = static_cast<int32>(Object->GetUniqueID());
+		const FString ObjectPathName = Object->GetPathName();
+		const FString ObjectClassName = Object->GetClass()->GetName();
+
+		Attribs.Add(FAnalyticsEventAttribute(FString::Printf(TEXT("Object.%d.Class"), ObjectUniqueId), ObjectClassName));
+		Attribs.Add(FAnalyticsEventAttribute(FString::Printf(TEXT("Object.%d.IsReimport"), ObjectUniqueId), bIsReimport));
+	};
+
+	for (const TPair<int32, TArray<FImportedObjectInfo>>& SourceIndexAndImportedAssets : ImportedAssetsPerSourceIndex)
+	{
+		const TArray<FImportedObjectInfo>& ImportedObjects = SourceIndexAndImportedAssets.Value;
+		for (const FImportedObjectInfo& ObjectInfo : ImportedObjects)
+		{
+			if (!ObjectInfo.ImportedObject)
+			{
+				continue;
+			}
+			CollectObject(ObjectInfo.bIsReimport, ObjectInfo.ImportedObject);
+		}
+	}
+
+	for (const TPair<int32, TArray<FImportedObjectInfo>>& SourceIndexAndImportedScneObjects : ImportedSceneObjectsPerSourceIndex)
+	{
+		const TArray<FImportedObjectInfo>& ImportedObjects = SourceIndexAndImportedScneObjects.Value;
+		for (const FImportedObjectInfo& ObjectInfo : ImportedObjects)
+		{
+			if (!ObjectInfo.ImportedObject)
+			{
+				continue;
+			}
+			CollectObject(ObjectInfo.bIsReimport, ObjectInfo.ImportedObject);
+		}
+	}
+
+	//Report any warning or error message
+	auto CollectResultContainer = [&Attribs](const UInterchangeResultsContainer* ResultContainer)
+	{
+		TArray<UInterchangeResult*> InterchangeResults = ResultContainer->GetResults();
+		for (const UInterchangeResult* InterchangeResult : InterchangeResults)
+		{
+			switch (InterchangeResult->GetResultType() )
+			{
+			case EInterchangeResultType::Success:
+				break;
+			case EInterchangeResultType::Warning:
+				Attribs.Add(FAnalyticsEventAttribute(TEXT("Message.Warning"), InterchangeResult->GetText().ToString()));
+				break;
+			case EInterchangeResultType::Error:
+				Attribs.Add(FAnalyticsEventAttribute(TEXT("Message.Error"), InterchangeResult->GetText().ToString()));
+				break;
+			}
+		}
+	};
+
+	if (const UInterchangeResultsContainer* ResultContainer = AssetImportResult->GetResults())
+	{
+		CollectResultContainer(ResultContainer);
+	}
+
+	if (const UInterchangeResultsContainer* ResultContainer = SceneImportResult->GetResults())
+	{
+		CollectResultContainer(ResultContainer);
+	}
+
+	FString EventString = TEXT("Interchange.Usage.ImportResult");
+	FEngineAnalytics::GetProvider().RecordEvent(EventString, Attribs);
+}
+
 void UE::Interchange::FImportAsyncHelper::ReleaseTranslatorsSource()
 {
 	for (UInterchangeTranslatorBase* BaseTranslator : Translators)
@@ -706,16 +894,56 @@ TTuple<UE::Interchange::FAssetImportResultRef, UE::Interchange::FSceneImportResu
 UInterchangeManager::ImportInternal(const FString& ContentPath, const UInterchangeSourceData* SourceData, const FImportAssetParameters& ImportAssetParameters, const UE::Interchange::EImportType ImportType)
 {
 	TRACE_CPUPROFILER_EVENT_SCOPE("UInterchangeManager::ImportInternal")
+	check(IsInGameThread());
+	static int32 GeneratedUniqueID = 0;
+	int32 UniqueId = ++GeneratedUniqueID;
+
+	TArray<FAnalyticsEventAttribute> Attribs;
+	auto PreReturn = [&Attribs]()
+	{
+		if (FEngineAnalytics::IsAvailable())
+		{
+			FString EventString = TEXT("Interchange.Usage.Import");
+			FEngineAnalytics::GetProvider().RecordEvent(EventString, Attribs);
+		}
+	};
+
+	auto EarlyExit = [&PreReturn]()
+	{
+		PreReturn();
+		return TTuple<UE::Interchange::FAssetImportResultRef, UE::Interchange::FSceneImportResultRef>{ MakeShared< UE::Interchange::FImportResult, ESPMode::ThreadSafe >(), MakeShared< UE::Interchange::FImportResult, ESPMode::ThreadSafe >() };
+	};
+
+	if (FEngineAnalytics::IsAvailable())
+	{
+		Attribs.Add(FAnalyticsEventAttribute(TEXT("UniqueId"), UniqueId));
+		Attribs.Add(FAnalyticsEventAttribute(TEXT("ImportType"), ImportType == UE::Interchange::EImportType::ImportType_Asset ? TEXT("Asset") : TEXT("Scene")));
+		Attribs.Add(FAnalyticsEventAttribute(TEXT("Parameters.IsAutomated"), ImportAssetParameters.bIsAutomated));
+		const bool bIsReimport = ImportAssetParameters.ReimportAsset != nullptr;
+		Attribs.Add(FAnalyticsEventAttribute(TEXT("Parameters.IsReimport"), bIsReimport));
+		if (ImportAssetParameters.ReimportAsset != nullptr)
+		{
+			Attribs.Add(FAnalyticsEventAttribute(TEXT("Parameters.ReimportSourceIndex"), ImportAssetParameters.ReimportSourceIndex));
+		}
+		const bool bIsPipelineOverride = ImportAssetParameters.OverridePipelines.Num() > 0;
+		Attribs.Add(FAnalyticsEventAttribute(TEXT("Parameters.IsPipelineOverrided"), bIsPipelineOverride));
+	}
+
 	if (!ensure(IsInGameThread()))
 	{
 		//Import process can be started only in the game thread
-		return TTuple<UE::Interchange::FAssetImportResultRef, UE::Interchange::FSceneImportResultRef>{ MakeShared< UE::Interchange::FImportResult, ESPMode::ThreadSafe >(), MakeShared< UE::Interchange::FImportResult, ESPMode::ThreadSafe >() };
+		return EarlyExit();
 	}
 
 	if (!SourceData)
 	{
 		UE_LOG(LogInterchangeEngine, Error, TEXT("Cannot import file, the source data is invalid"));
-		return TTuple<UE::Interchange::FAssetImportResultRef, UE::Interchange::FSceneImportResultRef>{ MakeShared< UE::Interchange::FImportResult, ESPMode::ThreadSafe >(), MakeShared< UE::Interchange::FImportResult, ESPMode::ThreadSafe >() };
+		return EarlyExit();
+	}
+	
+	if (FEngineAnalytics::IsAvailable())
+	{
+		Attribs.Add(FAnalyticsEventAttribute(TEXT("SourceExtension"), FPaths::GetExtension(SourceData->GetFilename())));
 	}
 
 	const bool bImportScene = ImportType == UE::Interchange::EImportType::ImportType_Scene;
@@ -724,7 +952,7 @@ UInterchangeManager::ImportInternal(const FString& ContentPath, const UInterchan
 	if (InterchangeImportSettings.PipelineStacks.Num() == 0)
 	{
 		UE_LOG(LogInterchangeEngine, Error, TEXT("Cannot import file, there is no pipeline stack define for %s import type. File path %s"), bImportScene ? TEXT("scene") : TEXT("content"));
-		return TTuple<UE::Interchange::FAssetImportResultRef, UE::Interchange::FSceneImportResultRef>{ MakeShared< UE::Interchange::FImportResult, ESPMode::ThreadSafe >(), MakeShared< UE::Interchange::FImportResult, ESPMode::ThreadSafe >() };
+		return EarlyExit();
 	}
 	
 	//Set a default pipeline stack if none is valid
@@ -765,6 +993,7 @@ UInterchangeManager::ImportInternal(const FString& ContentPath, const UInterchan
 	TaskData.ImportType = ImportType;
 	TaskData.ReimportObject = ImportAssetParameters.ReimportAsset;
 	TSharedRef<UE::Interchange::FImportAsyncHelper, ESPMode::ThreadSafe> AsyncHelper = CreateAsyncHelper(TaskData, ImportAssetParameters);
+	AsyncHelper->UniqueId = UniqueId;
 
 	//Create a duplicate of the source data, we need to be multithread safe so we copy it to control the life cycle. The async helper will hold it and delete it when the import task will be completed.
 	UInterchangeSourceData* DuplicateSourceData = Cast<UInterchangeSourceData>(StaticDuplicateObject(SourceData, GetTransientPackage()));
@@ -814,6 +1043,11 @@ UInterchangeManager::ImportInternal(const FString& ContentPath, const UInterchan
 		const bool bShowPipelineStacksConfigurationDialog = false;
 #endif
 
+		if (FEngineAnalytics::IsAvailable())
+		{
+			Attribs.Add(FAnalyticsEventAttribute(TEXT("ShowImportDialog"), bShowPipelineStacksConfigurationDialog));
+		}
+
 		const TMap<FName, FInterchangePipelineStack>& DefaultPipelineStacks = InterchangeImportSettings.PipelineStacks;
 
 		//If we reimport we want to load the original pipeline and the original pipeline settings
@@ -883,11 +1117,12 @@ UInterchangeManager::ImportInternal(const FString& ContentPath, const UInterchan
 				{
 					AsyncHelper->OriginalPipelines.Add(Pipeline);
 				}
+				UE::Interchange::Private::FillPipelineAnalyticData(Pipeline, Attribs, UniqueId, FString());
 			}
 		}
 		else
 		{
-			FName PipelineStackName = FInterchangeProjectSettingsUtils::GetDefaultPipelineStackName(ImportType == UE::Interchange::EImportType::ImportType_Scene, *SourceData);
+			FName PipelineStackName = FInterchangeProjectSettingsUtils::GetDefaultPipelineStackName(bImportScene, *SourceData);
 			if (RegisteredPipelineConfiguration && (bShowPipelineStacksConfigurationDialog || (!DefaultPipelineStacks.Contains(PipelineStackName) && !bIsUnattended)))
 			{
 				//Show the dialog, a plugin should have register this dialog. We use a plugin to be able to use editor code when doing UI
@@ -906,9 +1141,9 @@ UInterchangeManager::ImportInternal(const FString& ContentPath, const UInterchan
 			{
 				if (!DefaultPipelineStacks.Contains(PipelineStackName))
 				{
-					if (DefaultPipelineStacks.Contains(FInterchangeProjectSettingsUtils::GetDefaultPipelineStackName(ImportType == UE::Interchange::EImportType::ImportType_Scene, *SourceData)))
+					if (DefaultPipelineStacks.Contains(FInterchangeProjectSettingsUtils::GetDefaultPipelineStackName(bImportScene, *SourceData)))
 					{
-						PipelineStackName = FInterchangeProjectSettingsUtils::GetDefaultPipelineStackName(ImportType == UE::Interchange::EImportType::ImportType_Scene, *SourceData);
+						PipelineStackName = FInterchangeProjectSettingsUtils::GetDefaultPipelineStackName(bImportScene, *SourceData);
 					}
 					else
 					{
@@ -922,6 +1157,10 @@ UInterchangeManager::ImportInternal(const FString& ContentPath, const UInterchan
 
 				if (DefaultPipelineStacks.Contains(PipelineStackName))
 				{
+					if (FEngineAnalytics::IsAvailable())
+					{
+						Attribs.Add(FAnalyticsEventAttribute(TEXT("PipelineStackName"), PipelineStackName));
+					}
 					//use the default pipeline
 					const FInterchangePipelineStack& PipelineStack = DefaultPipelineStacks.FindChecked(PipelineStackName);
 					const TArray<FSoftObjectPath>& Pipelines = PipelineStack.Pipelines;
@@ -939,7 +1178,9 @@ UInterchangeManager::ImportInternal(const FString& ContentPath, const UInterchan
 									//The dialog is saving the settings for all default pipelines
 									GeneratedPipeline->LoadSettings(PipelineStackName);
 								}
-								
+					
+								UE::Interchange::Private::FillPipelineAnalyticData(GeneratedPipeline, Attribs, UniqueId, FString());
+
 								AsyncHelper->Pipelines.Add(GeneratedPipeline);
 
 								//We need to save the python pipeline asset because we cannot save an asset created with a python class
@@ -983,12 +1224,18 @@ UInterchangeManager::ImportInternal(const FString& ContentPath, const UInterchan
 			}
 			AsyncHelper->Pipelines.Add(GeneratedPipeline);
 			AsyncHelper->OriginalPipelines.Add(GeneratedPipeline);
+			UE::Interchange::Private::FillPipelineAnalyticData(GeneratedPipeline, Attribs, UniqueId, FString());
 		}
 	}
 
 	//Cancel the import do not queue task
 	if (bImportCancel)
 	{
+		if (FEngineAnalytics::IsAvailable())
+		{
+			Attribs.Add(FAnalyticsEventAttribute(TEXT("Canceled"), bImportCancel));
+		}
+
 		AsyncHelper->InitCancel();
 		AsyncHelper->CleanUp();
 	}
@@ -1002,6 +1249,7 @@ UInterchangeManager::ImportInternal(const FString& ContentPath, const UInterchan
 
 	StartQueuedTasks();
 
+	PreReturn();
 	return TTuple<UE::Interchange::FAssetImportResultRef, UE::Interchange::FSceneImportResultRef>{ AsyncHelper->AssetImportResult, AsyncHelper->SceneImportResult };
 }
 
diff --git a/Engine/Source/Runtime/Experimental/Interchange/Engine/Private/Tasks/InterchangeTaskCompletion.cpp b/Engine/Source/Runtime/Experimental/Interchange/Engine/Private/Tasks/InterchangeTaskCompletion.cpp
--- a/Engine/Source/Runtime/Experimental/Interchange/Engine/Private/Tasks/InterchangeTaskCompletion.cpp
+++ b/Engine/Source/Runtime/Experimental/Interchange/Engine/Private/Tasks/InterchangeTaskCompletion.cpp
@@ -169,6 +169,7 @@ void UE::Interchange::FTaskCompletion::DoTask(ENamedThreads::Type CurrentThread,
 	TSharedPtr<FImportAsyncHelper, ESPMode::ThreadSafe> AsyncHelper = WeakAsyncHelper.Pin();
 	check(AsyncHelper.IsValid());
 
+	AsyncHelper->SendAnalyticImportEndData();
 	//No need anymore of the translators sources
 	AsyncHelper->ReleaseTranslatorsSource();
 
diff --git a/Engine/Source/Runtime/Experimental/Interchange/Engine/Public/InterchangeManager.h b/Engine/Source/Runtime/Experimental/Interchange/Engine/Public/InterchangeManager.h
--- a/Engine/Source/Runtime/Experimental/Interchange/Engine/Public/InterchangeManager.h
+++ b/Engine/Source/Runtime/Experimental/Interchange/Engine/Public/InterchangeManager.h
@@ -175,6 +175,9 @@ namespace UE
 				return TEXT("UE::Interchange::FImportAsyncHelper");
 			}
 
+			/** Unique id for this asynchelper. */
+			int32 UniqueId;
+
 			//The following Arrays are per source data
 			TArray<TStrongObjectPtr<UInterchangeBaseNodeContainer>> BaseNodeContainers;
 			TArray<UInterchangeSourceData*> SourceDatas;
@@ -228,6 +231,7 @@ namespace UE
 			//If we cancel the tasks, we set this boolean to true
 			std::atomic<bool> bCancel;
 
+			void SendAnalyticImportEndData();
 			void ReleaseTranslatorsSource();
 			void InitCancel();
 			void CancelAndWaitUntilDoneSynchronously();
