Chaos - minor cleanup and some more comments

diff --git a/Engine/Source/Runtime/Experimental/Chaos/Private/Chaos/Island/IslandGroupManager.cpp b/Engine/Source/Runtime/Experimental/Chaos/Private/Chaos/Island/IslandGroupManager.cpp
--- a/Engine/Source/Runtime/Experimental/Chaos/Private/Chaos/Island/IslandGroupManager.cpp
+++ b/Engine/Source/Runtime/Experimental/Chaos/Private/Chaos/Island/IslandGroupManager.cpp
@@ -528,7 +528,8 @@ namespace Chaos
 #endif
 
 		// @todo(chaos): Remove SolveParallelFor when SolveParallelTasks has been thoroughly tested
-		const int32 ParallelMode = GSingleThreadedPhysics ? 0 : CVars::GIslandGroupsParallelMode;
+		const bool bSingleThreaded = GSingleThreadedPhysics || (IslandGroups.Num() == 1);
+		const int32 ParallelMode = bSingleThreaded ? 0 : CVars::GIslandGroupsParallelMode;
 		switch (ParallelMode)
 		{
 		case 0:
diff --git a/Engine/Source/Runtime/Experimental/Chaos/Public/Chaos/Collision/CollisionKeys.h b/Engine/Source/Runtime/Experimental/Chaos/Public/Chaos/Collision/CollisionKeys.h
--- a/Engine/Source/Runtime/Experimental/Chaos/Public/Chaos/Collision/CollisionKeys.h
+++ b/Engine/Source/Runtime/Experimental/Chaos/Public/Chaos/Collision/CollisionKeys.h
@@ -7,6 +7,9 @@
 
 namespace Chaos
 {
+	/**
+	 * Combine two hashes in an order-independent way so that Hash(A, B) == Hash(B, A)
+	 */
 	inline uint32 OrderIndependentHashCombine(const uint32 A, const uint32 B)
 	{
 		if (A < B)
@@ -20,20 +23,31 @@ namespace Chaos
 	}
 
 	/**
-	 * @brief Order particles in a consistent way for use by Broadphase and Resim
-	*/
+	 * Return true if the particles are in the preferred order (the first one has the lower ID)
+	 */
+	inline bool AreParticlesInPreferredOrder(const FGeometryParticleHandle* Particle0, const FGeometryParticleHandle* Particle1)
+	{
+		return (Particle1->ParticleID() < Particle0->ParticleID());
+	}
+
+	/**
+	 * Used to order particles in a consistent way in Broadphase and Resim. Returns whether the partcile order should be reversed.
+	 * We want the particle with the lower ID first, unless only one is dynamic in which case that one is first.
+	 */
+	inline bool ShouldSwapParticleOrder(const FGeometryParticleHandle* Particle0, const FGeometryParticleHandle* Particle1, const bool bIsParticle0Preferred)
+	{
+		return !FConstGenericParticleHandle(Particle0)->IsDynamic() || !bIsParticle0Preferred;
+	}
 	inline bool ShouldSwapParticleOrder(const FGeometryParticleHandle* Particle0, const FGeometryParticleHandle* Particle1)
 	{
-		const bool bIsParticle1Preferred = (Particle1->ParticleID() < Particle0->ParticleID());
-		const bool bSwapOrder = !FConstGenericParticleHandle(Particle0)->IsDynamic() || !bIsParticle1Preferred;
-		return bSwapOrder;
+		return ShouldSwapParticleOrder(Particle0, Particle1, AreParticlesInPreferredOrder(Particle0, Particle1));
 	}
 
 	/**
-	 * @brief A key which uniquely identifes a particle pair for use by the collision detection system
+	 * A key which uniquely identifes a particle pair for use by the collision detection system.
 	 * This key will be the same if particles order is reversed.
 	 * @note This uses ParticleID and truncates it from 32 to 31 bits
-	*/
+	 */
 	class FCollisionParticlePairKey
 	{
 	public:
@@ -104,14 +118,14 @@ namespace Chaos
 	};
 
 	/**
-	 * @brief A key which uniquely identifes a collision constraint within a particle pair
+	 * A key which uniquely identifes a collision constraint within a particle pair
 	 * 
 	 * This key only needs to be uinque within the context of a particle pair. There is no
 	 * guarantee of global uniqueness. This key is only used by the FMultiShapePairCollisionDetector
 	 * class which is used for colliding shape pairs where each shape is actually a hierarchy
 	 * of shapes. 
 	 * 
-	*/
+	 */
 	class FCollisionParticlePairConstraintKey
 	{
 	public:
diff --git a/Engine/Source/Runtime/Experimental/Chaos/Public/Chaos/Collision/SpatialAccelerationBroadPhase.h b/Engine/Source/Runtime/Experimental/Chaos/Public/Chaos/Collision/SpatialAccelerationBroadPhase.h
--- a/Engine/Source/Runtime/Experimental/Chaos/Public/Chaos/Collision/SpatialAccelerationBroadPhase.h
+++ b/Engine/Source/Runtime/Experimental/Chaos/Public/Chaos/Collision/SpatialAccelerationBroadPhase.h
@@ -297,8 +297,8 @@ namespace Chaos
 
 					// In both cases (resim or not) we will generate (1) dynamic-(sleeping,kinematic(steady+moving),static) pairs + (2) sleeping-moving kinematic ones
 					// Sleeping particles could collide with dynamic ones but these collisions are already handled in case 1
-					// Sleeping particles won't collide with static or steady kinematic particles since both are not supposed to move but will
-					// collide against moving kinematic particles
+					// Sleeping particles won't collide with static or steady kinematic particles since neither are moving
+					// Sleeping particles will collide against moving kinematic particles
 					if ((Particle1.ObjectState() == EObjectStateType::Dynamic && Particle2.ObjectState() != EObjectStateType::Dynamic) ||
 						(Particle1.ObjectState() == EObjectStateType::Sleeping && Particle2.ObjectState() == EObjectStateType::Kinematic))
 					{
@@ -306,7 +306,7 @@ namespace Chaos
 					}
 					
 					// Used to determine a winner in cases where we will visit particle pairs in both orders
-					const bool bIsParticle1Preferred = (Particle2.ParticleID() < Particle1.ParticleID());
+					const bool bIsParticle1Preferred = AreParticlesInPreferredOrder(Particle1.Handle(), Particle2.Handle());
 					
 					if(!bIsResimming)
 					{
@@ -375,7 +375,7 @@ namespace Chaos
 					// requirement but some downstream systems assume this is true (e.g., CCD, TriMesh collision).
 					TGeometryParticleHandle<FReal, 3>* ParticleA = Particle1.Handle();
 					TGeometryParticleHandle<FReal, 3>* ParticleB = Particle2.Handle();
-					const bool bSwapOrder = ShouldSwapParticleOrder(ParticleA, ParticleB);
+					const bool bSwapOrder = ShouldSwapParticleOrder(ParticleA, ParticleB, bIsParticle1Preferred);
 					if (bSwapOrder)
 					{
 						Swap(ParticleA, ParticleB);
