Interchange: Fix create package crash, and task deadlock.

#rnx

diff --git a/Engine/Plugins/Experimental/Interchange/Runtime/Source/Import/Private/Mesh/InterchangeStaticMeshFactory.cpp b/Engine/Plugins/Experimental/Interchange/Runtime/Source/Import/Private/Mesh/InterchangeStaticMeshFactory.cpp
--- a/Engine/Plugins/Experimental/Interchange/Runtime/Source/Import/Private/Mesh/InterchangeStaticMeshFactory.cpp
+++ b/Engine/Plugins/Experimental/Interchange/Runtime/Source/Import/Private/Mesh/InterchangeStaticMeshFactory.cpp
@@ -1278,7 +1278,11 @@ bool UInterchangeStaticMeshFactory::GenerateKDopCollision(const FCreateAssetPara
 	}
 
 	// Construct temporary UModel for kdop creation.
-	UModel* TempModel = NewObject<UModel>();
+	UModel* TempModel = nullptr;
+	{
+		FGCScopeGuard GCScopeGuard;
+		TempModel = NewObject<UModel>();
+	}
 	TempModel->RootOutside = true;
 	TempModel->EmptyModel(true, true);
 	TempModel->Polys->ClearFlags(RF_Transactional);
diff --git a/Engine/Source/Runtime/Experimental/Interchange/Engine/Private/Tasks/InterchangeTaskParsing.cpp b/Engine/Source/Runtime/Experimental/Interchange/Engine/Private/Tasks/InterchangeTaskParsing.cpp
--- a/Engine/Source/Runtime/Experimental/Interchange/Engine/Private/Tasks/InterchangeTaskParsing.cpp
+++ b/Engine/Source/Runtime/Experimental/Interchange/Engine/Private/Tasks/InterchangeTaskParsing.cpp
@@ -313,6 +313,10 @@ void UE::Interchange::FTaskParsing::DoTask(ENamedThreads::Type CurrentThread, co
 	};
 
 	FGraphEventArray CompletionPrerequistes;
+	const int32 PoolWorkerThreadCount = FTaskGraphInterface::Get().GetNumWorkerThreads() / 2;
+	const int32 MaxNumWorker = FMath::Max(PoolWorkerThreadCount, 1);
+	FGraphEventArray GroupPrerequistes[2];
+	int32 CurrentGroup = 0;
 	for (int32 TaskIndex = 0; TaskIndex < TaskDatas.Num(); ++TaskIndex)
 	{
 		FTaskData& TaskData = TaskDatas[TaskIndex];
@@ -326,13 +330,23 @@ void UE::Interchange::FTaskParsing::DoTask(ENamedThreads::Type CurrentThread, co
 				{
 					//Add has prerequisite
 					TaskData.Prerequisites.Add(TaskDatas[DepTaskIndex].GraphEventRef);
+					if (GroupPrerequistes[(CurrentGroup + 1) % 2].Num() >= MaxNumWorker)
+					{
+						TaskData.Prerequisites.Append(GroupPrerequistes[(CurrentGroup + 1) % 2]);
+					}
 				}
 			}
 		}
 
 		TaskData.GraphEventRef = CreateTasksFromData(TaskData);
-		CompletionPrerequistes.Add(TaskData.GraphEventRef);
+		if (GroupPrerequistes[CurrentGroup%2].Num() >= MaxNumWorker)
+		{
+			CurrentGroup++;
+			GroupPrerequistes[CurrentGroup % 2].Reset();
+		}
+		GroupPrerequistes[CurrentGroup%2].Add(TaskData.GraphEventRef);
 	}
+	CompletionPrerequistes.Append(GroupPrerequistes[CurrentGroup % 2]);
 
 	if (!RenameAssets.IsEmpty())
 	{
