Iris - Don't call removed NetworkRemapPath method.

#rnx

diff --git a/Engine/Plugins/Runtime/ReplicationSystemTestPlugin/Source/Private/Tests/Serialization/TestFieldPathNetSerializer.cpp b/Engine/Plugins/Runtime/ReplicationSystemTestPlugin/Source/Private/Tests/Serialization/TestFieldPathNetSerializer.cpp
--- a/Engine/Plugins/Runtime/ReplicationSystemTestPlugin/Source/Private/Tests/Serialization/TestFieldPathNetSerializer.cpp
+++ b/Engine/Plugins/Runtime/ReplicationSystemTestPlugin/Source/Private/Tests/Serialization/TestFieldPathNetSerializer.cpp
@@ -99,18 +99,39 @@ void FTestFieldPathNetSerializer::SetUp()
 		ServerNetSerializationContext = FNetSerializationContext(&Reader, &Writer);
 		ServerNetSerializationContext.SetInternalContext(&ServerInternalContext);
 		ServerNetSerializationContext.SetLocalConnectionId(Client->ConnectionIdOnServer);
-		ServerInternalContext = FInternalNetSerializationContext(Server->GetReplicationSystem(), Server->GetConnectionInfo(ServerNetSerializationContext.GetLocalConnectionId()).RemoteNetTokenStoreState);
 
-		// Since we do explicit serialization involving object references we set the context to allow for inlined exports
-		ServerInternalContext.bInlineObjectReferenceExports = 1U;
+		{
+			const FReplicationSystemTestNode::FConnectionInfo& ConnectionInfo = Server->GetConnectionInfo(Client->ConnectionIdOnServer);
+
+			FInternalNetSerializationContext TempServerInternalContext;
+			FInternalNetSerializationContext::FInitParameters TempServerInternalContextInitParams;
+			TempServerInternalContextInitParams.ReplicationSystem = Server->GetReplicationSystem();
+			TempServerInternalContextInitParams.ObjectResolveContext.RemoteNetTokenStoreState = ConnectionInfo.RemoteNetTokenStoreState;
+			TempServerInternalContextInitParams.ObjectResolveContext.ConnectionId = ConnectionInfo.ConnectionId;
+			TempServerInternalContext.Init(TempServerInternalContextInitParams);
+			// Since we do explicit serialization involving object references we set the context to allow for inlined exports
+			TempServerInternalContext.bInlineObjectReferenceExports = 1U;
+			ServerInternalContext = MoveTemp(TempServerInternalContext);
+		}
+
+		{
+			const FReplicationSystemTestNode::FConnectionInfo& ConnectionInfo = Client->GetConnectionInfo(Client->LocalConnectionId);
+
+			FInternalNetSerializationContext TempClientInternalContext;
+			FInternalNetSerializationContext::FInitParameters TempClientInternalContextInitParams;
+			TempClientInternalContextInitParams.ReplicationSystem = Client->GetReplicationSystem();
+			TempClientInternalContextInitParams.ObjectResolveContext.RemoteNetTokenStoreState = ConnectionInfo.RemoteNetTokenStoreState;
+			TempClientInternalContextInitParams.ObjectResolveContext.ConnectionId = ConnectionInfo.ConnectionId;
+			TempClientInternalContext.Init(TempClientInternalContextInitParams);
+			// Since we do explicit serialization involving object references we set the context to allow for inlined exports
+			TempClientInternalContext.bInlineObjectReferenceExports = 1U;
+
+			ClientInternalContext = MoveTemp(TempClientInternalContext);
+		}
 
 		ClientNetSerializationContext = FNetSerializationContext(&Reader, &Writer);
 		ClientNetSerializationContext.SetInternalContext(&ClientInternalContext);
 		ClientNetSerializationContext.SetLocalConnectionId(Client->LocalConnectionId);
-		ClientInternalContext = FInternalNetSerializationContext(Client->GetReplicationSystem(), Client->GetConnectionInfo(ClientNetSerializationContext.GetLocalConnectionId()).RemoteNetTokenStoreState);
-
-		// Since we do explicit serialization involving object references we set the context to allow for inlined exports
-		ClientInternalContext.bInlineObjectReferenceExports = 1U;
 	}
 
 	if (TestValues.Num() == 0)
diff --git a/Engine/Plugins/Runtime/ReplicationSystemTestPlugin/Source/Private/Tests/Serialization/TestPolymorphicStructNetSerializer.cpp b/Engine/Plugins/Runtime/ReplicationSystemTestPlugin/Source/Private/Tests/Serialization/TestPolymorphicStructNetSerializer.cpp
--- a/Engine/Plugins/Runtime/ReplicationSystemTestPlugin/Source/Private/Tests/Serialization/TestPolymorphicStructNetSerializer.cpp
+++ b/Engine/Plugins/Runtime/ReplicationSystemTestPlugin/Source/Private/Tests/Serialization/TestPolymorphicStructNetSerializer.cpp
@@ -323,7 +323,13 @@ void FTestPolymorphicArrayStructNetSerializerFixture::SetUp()
 
 	FReplicationSystemInternal* ReplicationSystemInternal = Server->GetReplicationSystem()->GetReplicationSystemInternal();
 
-	InternalNetSerializationContext = FInternalNetSerializationContext(Server->ReplicationSystem, ReplicationSystemInternal->GetNetTokenStore().GetLocalNetTokenStoreState());
+	FInternalNetSerializationContext TempInternalNetSerializationContext;
+	FInternalNetSerializationContext::FInitParameters TempInternalNetSerializationContextInitParams;
+	TempInternalNetSerializationContextInitParams.ReplicationSystem = Server->ReplicationSystem;
+	TempInternalNetSerializationContextInitParams.ObjectResolveContext.RemoteNetTokenStoreState = ReplicationSystemInternal->GetNetTokenStore().GetLocalNetTokenStoreState();
+	TempInternalNetSerializationContext.Init(TempInternalNetSerializationContextInitParams);
+
+	InternalNetSerializationContext = MoveTemp(TempInternalNetSerializationContext);
 	NetSerializationContext.SetInternalContext(&InternalNetSerializationContext);
 
 	FMemory::Memzero(QuantizedBuffer, 0);
diff --git a/Engine/Source/Runtime/Engine/Private/Net/Iris/ReplicationSystem/ActorReplicationBridge.cpp b/Engine/Source/Runtime/Engine/Private/Net/Iris/ReplicationSystem/ActorReplicationBridge.cpp
--- a/Engine/Source/Runtime/Engine/Private/Net/Iris/ReplicationSystem/ActorReplicationBridge.cpp
+++ b/Engine/Source/Runtime/Engine/Private/Net/Iris/ReplicationSystem/ActorReplicationBridge.cpp
@@ -5,6 +5,7 @@
 
 #include "Net/Iris/ReplicationSystem/ActorReplicationBridgeInternal.h"
 #include "Iris/IrisConfig.h"
+#include "Iris/IrisConstants.h"
 #include "Iris/Core/IrisLog.h"
 #include "Iris/Core/IrisProfiler.h"
 #include "Iris/Core/NetObjectReference.h"
@@ -20,6 +21,7 @@
 #include "Iris/Serialization/IrisObjectReferencePackageMap.h"
 #include "Engine/Engine.h"
 #include "Engine/EngineTypes.h"
+#include "Engine/NetConnection.h"
 #include "Engine/NetDriver.h"
 #include "Engine/World.h"
 #include "GameFramework/Actor.h"
@@ -943,11 +945,18 @@ void UActorReplicationBridge::GetSubObjectCreationHeader(const UObject* Object,
 	Header.ObjectReference = ObjectRef;
 }
 
-bool UActorReplicationBridge::RemapPathForPIE(FString &Str, bool bReading) const
+bool UActorReplicationBridge::RemapPathForPIE(uint32 ConnectionId, FString& Str, bool bReading) const
 {
-	PRAGMA_DISABLE_DEPRECATION_WARNINGS
-	return GEngine->NetworkRemapPath(NetDriver, Str, bReading);
-	PRAGMA_ENABLE_DEPRECATION_WARNINGS
+	if (ConnectionId == UE::Net::InvalidConnectionId)
+	{
+		return GEngine->NetworkRemapPath(static_cast<UNetConnection*>(nullptr), Str, bReading);
+	}
+	else
+	{
+		UObject* UserData = GetReplicationSystem()->GetConnectionUserData(ConnectionId);
+		UNetConnection* NetConnection = Cast<UNetConnection>(UserData);
+		return GEngine->NetworkRemapPath(NetConnection, Str, bReading);
+	}
 }
 
 bool UActorReplicationBridge::ObjectLevelHasFinishedLoading(UObject* Object) const
diff --git a/Engine/Source/Runtime/Engine/Public/Net/Iris/ReplicationSystem/ActorReplicationBridge.h b/Engine/Source/Runtime/Engine/Public/Net/Iris/ReplicationSystem/ActorReplicationBridge.h
--- a/Engine/Source/Runtime/Engine/Public/Net/Iris/ReplicationSystem/ActorReplicationBridge.h
+++ b/Engine/Source/Runtime/Engine/Public/Net/Iris/ReplicationSystem/ActorReplicationBridge.h
@@ -86,7 +86,7 @@ protected:
 	virtual void EndInstantiateFromRemote(FNetHandle Handle) override;
 	virtual void DestroyInstanceFromRemote(UObject* Instance, bool bTearOff) override;
 	virtual void GetInitialDependencies(FNetHandle Handle, FNetDependencyInfoArray& OutDependencies) const override;
-	virtual bool RemapPathForPIE(FString& Path, bool bReading) const override;
+	virtual bool RemapPathForPIE(uint32 ConnectionId, FString& Path, bool bReading) const override;
 	virtual bool ObjectLevelHasFinishedLoading(UObject* Object) const override;
 
 private:
diff --git a/Engine/Source/Runtime/Experimental/Iris/Core/Private/Iris/ReplicationSystem/ObjectReferenceCache.cpp b/Engine/Source/Runtime/Experimental/Iris/Core/Private/Iris/ReplicationSystem/ObjectReferenceCache.cpp
--- a/Engine/Source/Runtime/Experimental/Iris/Core/Private/Iris/ReplicationSystem/ObjectReferenceCache.cpp
+++ b/Engine/Source/Runtime/Experimental/Iris/Core/Private/Iris/ReplicationSystem/ObjectReferenceCache.cpp
@@ -1,6 +1,7 @@
 // Copyright Epic Games, Inc. All Rights Reserved.
 
 #include "ObjectReferenceCache.h"
+#include "Iris/IrisConstants.h"
 #include "Iris/Core/IrisLog.h"
 #include "Iris/Core/IrisMemoryTracker.h"
 #include "Iris/ReplicationSystem/ObjectReplicationBridge.h"
@@ -164,9 +165,9 @@ bool FObjectReferenceCache::ShouldIgnoreWhenMissing(FNetHandle RefHandle) const
 	return CacheObject->bIgnoreWhenMissing;
 }
 
-bool FObjectReferenceCache::RenamePathForPie(FString &Str, bool bReading)
+bool FObjectReferenceCache::RenamePathForPie(uint32 ConnectionId, FString& Str, bool bReading)
 {
-	return GPlayInEditorID != -1 ? ReplicationBridge->RemapPathForPIE(Str, bReading) : false;
+	return GPlayInEditorID != -1 ? ReplicationBridge->RemapPathForPIE(ConnectionId, Str, bReading) : false;
 }
 
 FNetHandle FObjectReferenceCache::CreateObjectReferenceHandle(const UObject* Object)
@@ -239,9 +240,10 @@ bool FObjectReferenceCache::CreateObjectReferenceInternal(const UObject* Object,
 			if (FNetHandle* ReplicatedOuter = ObjectToNetReferenceHandle.Find(Outer))
 			{
 				FString RelativeObjectPath = Object->GetPathName(Outer);
-				RenamePathForPie(RelativeObjectPath, false);		
+				constexpr bool bReading = false;
+				// The connection ID isn't used unless reading.
+				RenamePathForPie(UE::Net::InvalidConnectionId, RelativeObjectPath, bReading);
 				OutReference = MakeNetObjectReference(*ReplicatedOuter, StringTokenStore->GetOrCreateToken(RelativeObjectPath));
-
 				return true;
 			}
 		}
@@ -255,9 +257,10 @@ bool FObjectReferenceCache::CreateObjectReferenceInternal(const UObject* Object,
 		{
 			// For static objects we can always fall back on the full pathname
 			FString ObjectPath = Object->GetPathName();
-			RenamePathForPie(ObjectPath, false);		
-			OutReference = MakeNetObjectReference(FNetHandle(), StringTokenStore->GetOrCreateToken(ObjectPath));
-		
+			constexpr bool bReading = false;
+			// The connection ID isn't used unless reading.
+			RenamePathForPie(UE::Net::InvalidConnectionId, ObjectPath, bReading);
+			OutReference = MakeNetObjectReference(FNetHandle(), StringTokenStore->GetOrCreateToken(ObjectPath));	
 			return true;
 		}
 	}
@@ -272,7 +275,9 @@ bool FObjectReferenceCache::CreateObjectReferenceInternal(const UObject* Object,
 	if (bIsFullNameStableForNetworking || Object->IsNameStableForNetworking())
 	{
 		FString ObjectPath = Object->GetName();
-		RenamePathForPie(ObjectPath, false);
+		constexpr bool bReading = false;
+		// The connection ID isn't used unless reading.
+		RenamePathForPie(UE::Net::InvalidConnectionId, ObjectPath, bReading);
 		PathToken = StringTokenStore->GetOrCreateToken(ObjectPath);
 	}
 
@@ -520,7 +525,8 @@ UObject* FObjectReferenceCache::ResolveObjectReferenceHandleInternal(FNetHandle
 #if UE_NET_VALIDATE_REFERENCECACHE
 		{			
 			FString ObjectPath(StringTokenStore->ResolveRemoteToken(CacheObjectPtr->RelativePath, *ResolveContext.RemoteNetTokenStoreState));
-			RenamePathForPie(ObjectPath, true);
+			constexpr bool bReading = true;
+			RenamePathForPie(ResolveContext.ConnectionId, ObjectPath, bReading);
 			check(ObjectPath == Object->GetFName().ToString() || RefHandle.IsDynamic() || IsAuthority());
 
 			FNetHandle* ExistingRefHandle = ObjectToNetReferenceHandle.Find(Object);
@@ -615,7 +621,8 @@ UObject* FObjectReferenceCache::ResolveObjectReferenceHandleInternal(FNetHandle
 	}
 
 	FString ObjectPath(ResolvedToken);
-	RenamePathForPie(ObjectPath, true);
+	constexpr bool bReading = true;
+	RenamePathForPie(ResolveContext.ConnectionId, ObjectPath, bReading);
 
 	// See if this object is in memory
 	Object = StaticFindObject(UObject::StaticClass(), ObjOuter, *ObjectPath, false);
@@ -847,7 +854,8 @@ ENetObjectReferenceResolveResult FObjectReferenceCache::ResolveObjectReference(c
 		if (ResolvedToken)
 		{
 			FString ObjectPath(ResolvedToken);
-			RenamePathForPie(ObjectPath, true);
+			constexpr bool bReading = true;
+			RenamePathForPie(ResolveContext.ConnectionId, ObjectPath, bReading);
 
 			// If both path and handle is valid this is a reference to a dynamic object with a relative path
 			if (Reference.GetRefHandle().IsValid())
diff --git a/Engine/Source/Runtime/Experimental/Iris/Core/Private/Iris/ReplicationSystem/ObjectReferenceCache.h b/Engine/Source/Runtime/Experimental/Iris/Core/Private/Iris/ReplicationSystem/ObjectReferenceCache.h
--- a/Engine/Source/Runtime/Experimental/Iris/Core/Private/Iris/ReplicationSystem/ObjectReferenceCache.h
+++ b/Engine/Source/Runtime/Experimental/Iris/Core/Private/Iris/ReplicationSystem/ObjectReferenceCache.h
@@ -126,7 +126,7 @@ private:
 	bool SupportsObjectInternal(const UObject* Object) const;
 	bool CanClientLoadObjectInternal(const UObject* Object, bool bIsDynamic) const;
 	bool ShouldIgnoreWhenMissing(FNetHandle RefHandle) const;
-	bool RenamePathForPie(FString &Str, bool bReading);
+	bool RenamePathForPie(uint32 ConnectionId, FString& Str, bool bReading);
 
 	// Get the string path of RefHandle
 	FString FullPath(FNetHandle RefHandle, const FNetObjectResolveContext& ResolveContext) const;
diff --git a/Engine/Source/Runtime/Experimental/Iris/Core/Private/Iris/ReplicationSystem/ObjectReferenceCacheFwd.h b/Engine/Source/Runtime/Experimental/Iris/Core/Private/Iris/ReplicationSystem/ObjectReferenceCacheFwd.h
--- a/Engine/Source/Runtime/Experimental/Iris/Core/Private/Iris/ReplicationSystem/ObjectReferenceCacheFwd.h
+++ b/Engine/Source/Runtime/Experimental/Iris/Core/Private/Iris/ReplicationSystem/ObjectReferenceCacheFwd.h
@@ -2,6 +2,8 @@
 
 #pragma once
 
+#include "Iris/IrisConstants.h"
+
 class UReplicationSystem;
 class UObjectReplicationBridge;
 
@@ -12,11 +14,12 @@ namespace UE::Net
 
 	struct FNetObjectResolveContext
 	{
-		FNetTokenStoreState* RemoteNetTokenStoreState;
+		FNetTokenStoreState* RemoteNetTokenStoreState = nullptr;
+		uint32 ConnectionId = InvalidConnectionId;
 	};
 
 	namespace Private
 	{
 		class FNetExportContext;
 	}
-}
\ No newline at end of file
+}
diff --git a/Engine/Source/Runtime/Experimental/Iris/Core/Private/Iris/ReplicationSystem/ReplicationConnections.cpp b/Engine/Source/Runtime/Experimental/Iris/Core/Private/Iris/ReplicationSystem/ReplicationConnections.cpp
--- a/Engine/Source/Runtime/Experimental/Iris/Core/Private/Iris/ReplicationSystem/ReplicationConnections.cpp
+++ b/Engine/Source/Runtime/Experimental/Iris/Core/Private/Iris/ReplicationSystem/ReplicationConnections.cpp
@@ -38,7 +38,6 @@ void FReplicationConnections::InitDataStreams(uint32 ReplicationSystemId, uint32
 
 	// Init ReplicationDataStream
 	Connection->ReplicationWriter->SetNetExports(DataStreamManager->GetNetExports());
-	Connection->ReplicationWriter->SetRemoteNetTokenStoreState(&RemoteNetTokenStoreStates[ConnectionId]);
 	Connection->ReplicationReader->SetRemoteNetTokenStoreState(&RemoteNetTokenStoreStates[ConnectionId]);
 
 	UReplicationDataStream* ReplicationDataStream = StaticCast<UReplicationDataStream*>(DataStreamManager->GetStream(FName("Replication")));
diff --git a/Engine/Source/Runtime/Experimental/Iris/Core/Private/Iris/ReplicationSystem/ReplicationReader.cpp b/Engine/Source/Runtime/Experimental/Iris/Core/Private/Iris/ReplicationSystem/ReplicationReader.cpp
--- a/Engine/Source/Runtime/Experimental/Iris/Core/Private/Iris/ReplicationSystem/ReplicationReader.cpp
+++ b/Engine/Source/Runtime/Experimental/Iris/Core/Private/Iris/ReplicationSystem/ReplicationReader.cpp
@@ -140,6 +140,8 @@ void FReplicationReader::Init(const FReplicationParameters& InParameters)
 	// Store copy of parameters
 	Parameters = InParameters;
 
+	ResolveContext.ConnectionId = InParameters.ConnectionId;
+
 	// Cache internal systems
 	ReplicationSystemInternal = Parameters.ReplicationSystem->GetReplicationSystemInternal();
 	NetHandleManager = &ReplicationSystemInternal->GetNetHandleManager();
@@ -1160,9 +1162,14 @@ void FReplicationReader::ResolveAndDispatchUnresolvedReferences()
 	IRIS_PROFILER_SCOPE(FReplicationReader_ResolveAndDispatchUnresolvedReferences);
 
 	// Setup context for dispatch
+	FInternalNetSerializationContext InternalContext;
+	FInternalNetSerializationContext::FInitParameters InternalContextInitParams;
+	InternalContextInitParams.ReplicationSystem = Parameters.ReplicationSystem;
+	InternalContextInitParams.ObjectResolveContext = ResolveContext;
+	InternalContext.Init(InternalContextInitParams);
+
 	FNetSerializationContext Context;
-	FInternalNetSerializationContext InternalContext(Parameters.ReplicationSystem, ResolveContext.RemoteNetTokenStoreState);
-	Context.SetLocalConnectionId(Parameters.ConnectionId);
+	Context.SetLocalConnectionId(ResolveContext.ConnectionId);
 	Context.SetInternalContext(&InternalContext);
 
 	// Currently we brute force this by iterating over all handles pending resolve and update all objects pending resolve
@@ -1387,7 +1394,12 @@ void FReplicationReader::Read(FNetSerializationContext& Context)
 	FNetBitStreamReader& Reader = *Context.GetBitStreamReader();
 
 	// Setup internal context
-	FInternalNetSerializationContext InternalContext(Parameters.ReplicationSystem, ResolveContext.RemoteNetTokenStoreState);
+	FInternalNetSerializationContext InternalContext;
+	FInternalNetSerializationContext::FInitParameters InternalContextInitParams;
+	InternalContextInitParams.ReplicationSystem = Parameters.ReplicationSystem;
+	InternalContextInitParams.ObjectResolveContext = ResolveContext;
+	InternalContext.Init(InternalContextInitParams);
+	
 	Context.SetLocalConnectionId(Parameters.ConnectionId);
 	Context.SetInternalContext(&InternalContext);
 	Context.SetNetBlobReceiver(&ReplicationSystemInternal->GetNetBlobHandlerManager());
diff --git a/Engine/Source/Runtime/Experimental/Iris/Core/Private/Iris/ReplicationSystem/ReplicationSystem.cpp b/Engine/Source/Runtime/Experimental/Iris/Core/Private/Iris/ReplicationSystem/ReplicationSystem.cpp
--- a/Engine/Source/Runtime/Experimental/Iris/Core/Private/Iris/ReplicationSystem/ReplicationSystem.cpp
+++ b/Engine/Source/Runtime/Experimental/Iris/Core/Private/Iris/ReplicationSystem/ReplicationSystem.cpp
@@ -746,6 +746,11 @@ void UReplicationSystem::SetConnectionUserData(uint32 ConnectionId, UObject* InU
 UObject* UReplicationSystem::GetConnectionUserData(uint32 ConnectionId) const
 {
 	UE::Net::Private::FReplicationConnections& Connections = Impl->ReplicationSystemInternal.GetConnections();
+	if (!ensureMsgf(Connections.IsValidConnection(ConnectionId), TEXT("Invalid ConnectionId %u passed to UReplicationSystem::GetConnectionUserData."), ConnectionId))
+	{
+		return nullptr;
+	}
+
 	if (UE::Net::Private::FReplicationConnection* Connection = Connections.GetConnection(ConnectionId))
 	{
 		return Connection->UserData.Get();
diff --git a/Engine/Source/Runtime/Experimental/Iris/Core/Private/Iris/ReplicationSystem/ReplicationWriter.cpp b/Engine/Source/Runtime/Experimental/Iris/Core/Private/Iris/ReplicationSystem/ReplicationWriter.cpp
--- a/Engine/Source/Runtime/Experimental/Iris/Core/Private/Iris/ReplicationSystem/ReplicationWriter.cpp
+++ b/Engine/Source/Runtime/Experimental/Iris/Core/Private/Iris/ReplicationSystem/ReplicationWriter.cpp
@@ -18,7 +18,6 @@
 #include "Iris/ReplicationSystem/NetBlob/NetObjectBlobHandler.h"
 #include "Iris/ReplicationSystem/NetBlob/PartialNetObjectAttachmentHandler.h"
 #include "Iris/Serialization/NetReferenceCollector.h"
-#include "Iris/ReplicationSystem/NetTokenStoreState.h"
 #include "Iris/Serialization/InternalNetSerializationContext.h"
 #include "Iris/Serialization/NetBitStreamUtil.h"
 #include "Iris/Serialization/NetSerializer.h"
@@ -2728,11 +2727,6 @@ bool FReplicationWriter::PatchupObjectChangeMaskWithInflightChanges(uint32 Inter
 	return bInFlightChangesAdded;
 }
 
-void FReplicationWriter::SetRemoteNetTokenStoreState(FNetTokenStoreState* InRemoteTokenStoreState)
-{
-	ResolveContext.RemoteNetTokenStoreState = InRemoteTokenStoreState;
-}
-
 void FReplicationWriter::SetNetExports(FNetExports& InNetExports)
 {
 	NetExports = &InNetExports;
diff --git a/Engine/Source/Runtime/Experimental/Iris/Core/Private/Iris/ReplicationSystem/ReplicationWriter.h b/Engine/Source/Runtime/Experimental/Iris/Core/Private/Iris/ReplicationSystem/ReplicationWriter.h
--- a/Engine/Source/Runtime/Experimental/Iris/Core/Private/Iris/ReplicationSystem/ReplicationWriter.h
+++ b/Engine/Source/Runtime/Experimental/Iris/Core/Private/Iris/ReplicationSystem/ReplicationWriter.h
@@ -168,7 +168,6 @@ public:
 	void SetReplicationEnabled(bool bInReplicationEnabled);
 	bool IsReplicationEnabled() const;
 
-	void SetRemoteNetTokenStoreState(FNetTokenStoreState* RemoteTokenStoreState);
 	void SetNetExports(FNetExports& InNetExports);
 
 	// Attachments
@@ -463,7 +462,6 @@ private:
 	FNetHandleManager* NetHandleManager = nullptr;
 	UReplicationBridge* ReplicationBridge = nullptr;
 	FDeltaCompressionBaselineManager* BaselineManager = nullptr;
-	FNetObjectResolveContext ResolveContext;
 	FObjectReferenceCache* ObjectReferenceCache = nullptr;
 	const FReplicationFiltering* ReplicationFiltering = nullptr;
 	FReplicationConditionals* ReplicationConditionals = nullptr;
diff --git a/Engine/Source/Runtime/Experimental/Iris/Core/Private/Iris/Serialization/InternalNetSerializationContext.cpp b/Engine/Source/Runtime/Experimental/Iris/Core/Private/Iris/Serialization/InternalNetSerializationContext.cpp
--- a/Engine/Source/Runtime/Experimental/Iris/Core/Private/Iris/Serialization/InternalNetSerializationContext.cpp
+++ b/Engine/Source/Runtime/Experimental/Iris/Core/Private/Iris/Serialization/InternalNetSerializationContext.cpp
@@ -8,6 +8,14 @@
 namespace UE::Net::Private
 {
 
+FInternalNetSerializationContext::FInternalNetSerializationContext(UReplicationSystem* InReplicationSystem)
+: ReplicationSystem(InReplicationSystem)
+, ObjectReferenceCache(&InReplicationSystem->GetReplicationSystemInternal()->GetObjectReferenceCache())
+, bDowngradeAutonomousProxyRole(0)
+, bInlineObjectReferenceExports(0)
+{
+}
+
 void* FInternalNetSerializationContext::Alloc(SIZE_T Size, SIZE_T Alignment)
 {
 	return GMalloc->Malloc(Size, Alignment);
@@ -23,13 +31,11 @@ void* FInternalNetSerializationContext::Realloc(void* PrevAddress, SIZE_T NewSiz
 	return GMalloc->Realloc(PrevAddress, NewSize, Alignment);
 }
 
-FInternalNetSerializationContext::FInternalNetSerializationContext(UReplicationSystem* ReplicationSystem, FNetTokenStoreState* RemoteTokenStoreState)
-: ReplicationSystem(ReplicationSystem)
-, ObjectReferenceCache(&ReplicationSystem->GetReplicationSystemInternal()->GetObjectReferenceCache())
-, bDowngradeAutonomousProxyRole(0)
-, bInlineObjectReferenceExports(0)
+void FInternalNetSerializationContext::Init(const FInitParameters& InitParams)
 {
-	ResolveContext.RemoteNetTokenStoreState = RemoteTokenStoreState;
+	ReplicationSystem = InitParams.ReplicationSystem;
+	ObjectReferenceCache = &InitParams.ReplicationSystem->GetReplicationSystemInternal()->GetObjectReferenceCache();
+	ResolveContext = InitParams.ObjectResolveContext;
 }
 
 }
diff --git a/Engine/Source/Runtime/Experimental/Iris/Core/Private/Iris/Serialization/InternalNetSerializationContext.h b/Engine/Source/Runtime/Experimental/Iris/Core/Private/Iris/Serialization/InternalNetSerializationContext.h
--- a/Engine/Source/Runtime/Experimental/Iris/Core/Private/Iris/Serialization/InternalNetSerializationContext.h
+++ b/Engine/Source/Runtime/Experimental/Iris/Core/Private/Iris/Serialization/InternalNetSerializationContext.h
@@ -2,7 +2,8 @@
 
 #pragma once
 
-#include "CoreTypes.h"
+#include "HAL/Platform.h"
+#include "Iris/IrisConstants.h"
 #include "Iris/ReplicationSystem/ObjectReferenceCacheFwd.h"
 
 class UReplicationSystem;
@@ -19,8 +20,16 @@ namespace UE::Net::Private
 class FInternalNetSerializationContext final
 {
 public:
+	struct FInitParameters
+	{
+		UReplicationSystem* ReplicationSystem = nullptr;
+		FNetObjectResolveContext ObjectResolveContext;
+	};
+
 	FInternalNetSerializationContext();
-	FInternalNetSerializationContext(UReplicationSystem* ReplicationSystem, FNetTokenStoreState* RemoteTokenStoreState = nullptr);
+	explicit FInternalNetSerializationContext(UReplicationSystem* InReplicationSystem);
+
+	void Init(const FInitParameters& Params);
 
 	// We allow memory allocations for dynamic states
 	void* Alloc(SIZE_T Size, SIZE_T Alignment);
diff --git a/Engine/Source/Runtime/Experimental/Iris/Core/Public/Iris/ReplicationSystem/ObjectReplicationBridge.h b/Engine/Source/Runtime/Experimental/Iris/Core/Public/Iris/ReplicationSystem/ObjectReplicationBridge.h
--- a/Engine/Source/Runtime/Experimental/Iris/Core/Public/Iris/ReplicationSystem/ObjectReplicationBridge.h
+++ b/Engine/Source/Runtime/Experimental/Iris/Core/Public/Iris/ReplicationSystem/ObjectReplicationBridge.h
@@ -102,7 +102,7 @@ public:
 	IRISCORE_API void AddStaticDestructionInfo(const FString& ObjectPath, const UObject* Outer, const FEndReplicationParameters& Parameters);
 	
 	/** PIE package name remapping support. */
-	virtual bool RemapPathForPIE(FString& Path, bool bReading) const { return false; }
+	virtual bool RemapPathForPIE(uint32 ConnectionId, FString& Path, bool bReading) const { return false; }
 
 	/** Returns true of the level that the Object belongs to has finished loading. */
 	virtual bool ObjectLevelHasFinishedLoading(UObject* Object) const { return true; }
