make fracture editor more robust to invalid selections

diff --git a/Engine/Plugins/Experimental/ChaosEditor/Source/FractureEditor/Private/FractureEditorMode.cpp b/Engine/Plugins/Experimental/ChaosEditor/Source/FractureEditor/Private/FractureEditorMode.cpp
--- a/Engine/Plugins/Experimental/ChaosEditor/Source/FractureEditor/Private/FractureEditorMode.cpp
+++ b/Engine/Plugins/Experimental/ChaosEditor/Source/FractureEditor/Private/FractureEditorMode.cpp
@@ -505,6 +505,7 @@ void UFractureEditorMode::OnActorSelectionChanged(const TArray<UObject*>& NewSel
 			FScopedColorEdit ShowBoneColorsEdit(GeometryCollectionComponent);
 			ShowBoneColorsEdit.SetEnableBoneSelection(true);
 			// ShowBoneColorsEdit.SetLevelViewMode(ViewLevel);
+			ShowBoneColorsEdit.Sanitize(); // Clean any stale data (e.g. due to the geometry being edited via a different component)
 
 			NewGeomSelection.Add(GeometryCollectionComponent);
 		}
diff --git a/Engine/Plugins/Experimental/ChaosEditor/Source/FractureEditor/Private/FractureEditorModeToolkit.cpp b/Engine/Plugins/Experimental/ChaosEditor/Source/FractureEditor/Private/FractureEditorModeToolkit.cpp
--- a/Engine/Plugins/Experimental/ChaosEditor/Source/FractureEditor/Private/FractureEditorModeToolkit.cpp
+++ b/Engine/Plugins/Experimental/ChaosEditor/Source/FractureEditor/Private/FractureEditorModeToolkit.cpp
@@ -983,6 +983,12 @@ void FFractureEditorModeToolkit::SetHideForUnselected(UGeometryCollectionCompone
 
 				for (int32 SelectedBone : SelectedBones)
 				{
+					if (!ensure(SelectedBone >= 0 && SelectedBone < Hide.Num()))
+					{
+						// Invalid selection, don't hide anything
+						Hide.Fill(false);
+						break;
+					}
 					Hide[SelectedBone] = false;
 					if (Children[SelectedBone].Num() > 0)
 					{ 
diff --git a/Engine/Source/Runtime/Experimental/GeometryCollectionEngine/Private/GeometryCollection/GeometryCollectionComponent.cpp b/Engine/Source/Runtime/Experimental/GeometryCollectionEngine/Private/GeometryCollection/GeometryCollectionComponent.cpp
--- a/Engine/Source/Runtime/Experimental/GeometryCollectionEngine/Private/GeometryCollection/GeometryCollectionComponent.cpp
+++ b/Engine/Source/Runtime/Experimental/GeometryCollectionEngine/Private/GeometryCollection/GeometryCollectionComponent.cpp
@@ -2798,6 +2798,26 @@ bool FScopedColorEdit::IsBoneSelected(int BoneIndex) const
 	return Component->SelectedBones.Contains(BoneIndex);
 }
 
+void FScopedColorEdit::Sanitize()
+{
+	const UGeometryCollection* GeometryCollection = Component->GetRestCollection();
+	if (GeometryCollection)
+	{
+		TSharedPtr<FGeometryCollection, ESPMode::ThreadSafe> GeometryCollectionPtr = GeometryCollection->GetGeometryCollection();
+		if (GeometryCollectionPtr)
+		{
+			const int32 NumTransforms = GeometryCollectionPtr->NumElements(FGeometryCollection::TransformGroup);
+			const int32 NumSelectionRemoved = Component->SelectedBones.RemoveAll([this, NumTransforms](int32 Index) {
+				return Index < 0 || Index >= NumTransforms;
+				});
+			const int32 NumHighlightRemoved = Component->HighlightedBones.RemoveAll([this, NumTransforms](int32 Index) {
+				return Index < 0 || Index >= NumTransforms;
+				});
+			bUpdated = NumSelectionRemoved || NumHighlightRemoved;
+		}
+	}
+}
+
 void FScopedColorEdit::SetSelectedBones(const TArray<int32>& SelectedBonesIn)
 {
 	bUpdated = true;
diff --git a/Engine/Source/Runtime/Experimental/GeometryCollectionEngine/Public/GeometryCollection/GeometryCollectionComponent.h b/Engine/Source/Runtime/Experimental/GeometryCollectionEngine/Public/GeometryCollection/GeometryCollectionComponent.h
--- a/Engine/Source/Runtime/Experimental/GeometryCollectionEngine/Public/GeometryCollection/GeometryCollectionComponent.h
+++ b/Engine/Source/Runtime/Experimental/GeometryCollectionEngine/Public/GeometryCollection/GeometryCollectionComponent.h
@@ -216,6 +216,7 @@ public:
 	bool GetEnableBoneSelection() const;
 
 	bool IsBoneSelected(int BoneIndex) const;
+	void Sanitize();
 	void SetSelectedBones(const TArray<int32>& SelectedBonesIn);
 	void AppendSelectedBones(const TArray<int32>& SelectedBonesIn);
 	void ToggleSelectedBones(const TArray<int32>& SelectedBonesIn, bool bAdd, bool bSnapToLevel = true);
