ChaosFlesh : Caching and playback.

diff --git a/Engine/Source/Runtime/Experimental/Chaos/Private/Chaos/Deformable/ChaosDeformableSolver.cpp b/Engine/Source/Runtime/Experimental/Chaos/Private/Chaos/Deformable/ChaosDeformableSolver.cpp
--- a/Engine/Source/Runtime/Experimental/Chaos/Private/Chaos/Deformable/ChaosDeformableSolver.cpp
+++ b/Engine/Source/Runtime/Experimental/Chaos/Private/Chaos/Deformable/ChaosDeformableSolver.cpp
@@ -39,23 +39,25 @@ namespace Chaos::Softs
 	FCriticalSection FDeformableSolver::RemovalMutex;
 	FCriticalSection FDeformableSolver::PackageOutputMutex;
 	FCriticalSection FDeformableSolver::PackageInputMutex;
+	FCriticalSection FDeformableSolver::SolverEnabledMutex;
 
 	FDeformableSolver::FDeformableSolver(FDeformableSolverProperties InProp)
 		: CurrentInputPackage(TUniquePtr < FDeformablePackage >(nullptr))
 		, PreviousInputPackage(TUniquePtr < FDeformablePackage >(nullptr))
 		, Property(InProp)
-
 	{
 		Reset(Property);
 	}
 
 	FDeformableSolver::~FDeformableSolver()
 	{
+		FScopeLock Lock(&InitializationMutex);
 		for (FThreadingProxy* Proxy : UninitializedProxys_Internal)
 		{
 			delete Proxy;
 		}
 		UninitializedProxys_Internal.Empty();
+		EventTeardown.Broadcast();
 	}
 
 
@@ -99,28 +101,6 @@ namespace Chaos::Softs
 		}
 	}
 
-	bool FDeformableSolver::AdvanceDt(FSolverReal DeltaTime)
-	{
-		SCOPE_CYCLE_COUNTER(STAT_DeformableSolver_Advance);
-		TRACE_CPUPROFILER_EVENT_SCOPE(DeformableSolver_Advance);
-
-		int32 NumIterations = FMath::Clamp<int32>(Property.NumSolverSubSteps, 0, INT_MAX);
-		if (NumIterations)
-		{
-			FSolverReal SubDeltaTime = DeltaTime / (FSolverReal)NumIterations;
-			if (!FMath::IsNearlyZero(SubDeltaTime))
-			{
-				for (int i = 0; i < NumIterations; ++i)
-				{
-					TickSimulation(SubDeltaTime);
-				}
-				Frame++;
-				return true;
-			}
-		}
-		return false;
-	}
-
 	void FDeformableSolver::InitializeSimulationObjects()
 	{
 		TRACE_CPUPROFILER_EVENT_SCOPE(DeformableSolver_InitializeSimulationObjects);
@@ -417,39 +397,6 @@ namespace Chaos::Softs
 	
 	}
 
-
-	void FDeformableSolver::PushInputPackage(int32 InFrame, FDeformableDataMap&& InPackage)
-	{
-		FScopeLock Lock(&PackageInputMutex);
-		TRACE_CPUPROFILER_EVENT_SCOPE(DeformableSolver_PushInputPackage);
-		BufferedInputPackages.Push(TUniquePtr< FDeformablePackage >(new FDeformablePackage(InFrame, MoveTemp(InPackage))));
-	}
-
-	TUniquePtr<FDeformablePackage> FDeformableSolver::PullInputPackage()
-	{
-		FScopeLock Lock(&PackageInputMutex);
-		TRACE_CPUPROFILER_EVENT_SCOPE(DeformableSolver_PullInputPackage);
-		if (BufferedInputPackages.Num())
-			return BufferedInputPackages.Pop();
-		return TUniquePtr<FDeformablePackage>(nullptr);
-	}
-
-	void FDeformableSolver::UpdateProxyInputPackages()
-	{
-		if (CurrentInputPackage)
-		{
-			PreviousInputPackage = TUniquePtr < FDeformablePackage >(CurrentInputPackage.Release());
-			CurrentInputPackage = TUniquePtr < FDeformablePackage >(nullptr);
-		}
-
-		TUniquePtr < FDeformablePackage > TailPackage = PullInputPackage();
-		while(TailPackage)
-		{
-			CurrentInputPackage = TUniquePtr < FDeformablePackage >(TailPackage.Release());
-			TailPackage = PullInputPackage();
-		}
-	}
-
 	void FDeformableSolver::RemoveSimulationObjects()
 	{
 		TRACE_CPUPROFILER_EVENT_SCOPE(DeformableSolver_RemoveSimulationObjects);
@@ -511,12 +458,94 @@ namespace Chaos::Softs
 		}
 	}
 
-	void FDeformableSolver::TickSimulation(FSolverReal DeltaTime)
+	void FDeformableSolver::AdvanceDt(FSolverReal DeltaTime)
+	{
+		SCOPE_CYCLE_COUNTER(STAT_DeformableSolver_Advance);
+		TRACE_CPUPROFILER_EVENT_SCOPE(DeformableSolver_Advance);
+
+		EventPreSolve.Broadcast(DeltaTime);
+
+		int32 NumIterations = FMath::Clamp<int32>(Property.NumSolverSubSteps, 0, INT_MAX);
+		if (bEnableSolver && NumIterations)
+		{
+			FSolverReal SubDeltaTime = DeltaTime / (FSolverReal)NumIterations;
+			if (!FMath::IsNearlyZero(SubDeltaTime))
+			{
+				for (int i = 0; i < NumIterations; ++i)
+				{
+					Update(SubDeltaTime);
+				}
+
+				Frame++;
+				EventPostSolve.Broadcast(DeltaTime);
+			}
+		}
+
+
+
+		{
+			// Update client state
+			FDeformableDataMap OutputBuffers;
+			for (TPair< FThreadingProxy::FKey, TUniquePtr<FThreadingProxy> >& BaseProxyPair : Proxies)
+			{
+				UpdateOutputState(*BaseProxyPair.Value);
+				if (FFleshThreadingProxy* Proxy = BaseProxyPair.Value->As<FFleshThreadingProxy>())
+				{
+					OutputBuffers.Add(Proxy->GetOwner(), TSharedPtr<FThreadingProxy::FBuffer>(new FFleshThreadingProxy::FFleshOutputBuffer(*Proxy)));
+
+					if (Property.CacheToFile)
+					{
+						WriteFrame(*Proxy, DeltaTime);
+					}
+				}
+			}
+			PushOutputPackage(Frame, MoveTemp(OutputBuffers));
+		}
+
+
+
+		EventPreBuffer.Broadcast(DeltaTime);
+	}
+
+	void FDeformableSolver::PushInputPackage(int32 InFrame, FDeformableDataMap&& InPackage)
+	{
+		FScopeLock Lock(&PackageInputMutex);
+		TRACE_CPUPROFILER_EVENT_SCOPE(DeformableSolver_PushInputPackage);
+		BufferedInputPackages.Push(TUniquePtr< FDeformablePackage >(new FDeformablePackage(InFrame, MoveTemp(InPackage))));
+	}
+
+	TUniquePtr<FDeformablePackage> FDeformableSolver::PullInputPackage()
+	{
+		FScopeLock Lock(&PackageInputMutex);
+		TRACE_CPUPROFILER_EVENT_SCOPE(DeformableSolver_PullInputPackage);
+		if (BufferedInputPackages.Num())
+			return BufferedInputPackages.Pop();
+		return TUniquePtr<FDeformablePackage>(nullptr);
+	}
+
+	void FDeformableSolver::UpdateProxyInputPackages()
+	{
+		if (CurrentInputPackage)
+		{
+			PreviousInputPackage = TUniquePtr < FDeformablePackage >(CurrentInputPackage.Release());
+			CurrentInputPackage = TUniquePtr < FDeformablePackage >(nullptr);
+		}
+
+		TUniquePtr < FDeformablePackage > TailPackage = PullInputPackage();
+		while (TailPackage)
+		{
+			CurrentInputPackage = TUniquePtr < FDeformablePackage >(TailPackage.Release());
+			TailPackage = PullInputPackage();
+		}
+	}
+
+	void FDeformableSolver::Update(FSolverReal DeltaTime)
 	{
-		TRACE_CPUPROFILER_EVENT_SCOPE(DeformableSolver_TickSimulation);
+		TRACE_CPUPROFILER_EVENT_SCOPE(DeformableSolver_Update);
 
 		if (!Proxies.Num()) return;
 
+
 		if (!Property.FixTimeStep)
 		{
 			Evolution->AdvanceOneTimeStep(DeltaTime);
@@ -528,23 +557,6 @@ namespace Chaos::Softs
 			Time += Property.TimeStepSize;
 		}
 
-
-		FDeformableDataMap OutputBuffers;
-		for (TPair< FThreadingProxy::FKey, TUniquePtr<FThreadingProxy> > & BaseProxyPair : Proxies)
-		{
-			UpdateOutputState(*BaseProxyPair.Value);
-			if (FFleshThreadingProxy* Proxy = BaseProxyPair.Value->As<FFleshThreadingProxy>())
-			{
-				OutputBuffers.Add(Proxy->GetOwner(), TSharedPtr<FThreadingProxy::FBuffer>(new FFleshThreadingProxy::FFleshOutputBuffer(*Proxy)));
-
-				if (Property.CacheToFile)
-				{
-					WriteFrame(*Proxy, DeltaTime);
-				}
-			}
-		}
-
-		PushOutputPackage(Frame, MoveTemp(OutputBuffers));
 	}
 
 	void FDeformableSolver::PushOutputPackage(int32 InFrame, FDeformableDataMap&& InPackage)
@@ -572,7 +584,8 @@ namespace Chaos::Softs
 
 	void FDeformableSolver::RemoveProxy(FThreadingProxy* InProxy)
 	{
-		FScopeLock Lock(&RemovalMutex);
+		FScopeLock LockA(&RemovalMutex);
+		FScopeLock LockB(&InitializationMutex);
 
 		InitializedObjects_External.Remove(InProxy->GetOwner());
 
diff --git a/Engine/Source/Runtime/Experimental/Chaos/Private/Chaos/Deformable/ChaosDeformableSolverThreadingWrapper.cpp b/Engine/Source/Runtime/Experimental/Chaos/Private/Chaos/Deformable/ChaosDeformableSolverThreadingWrapper.cpp
--- a/Engine/Source/Runtime/Experimental/Chaos/Private/Chaos/Deformable/ChaosDeformableSolverThreadingWrapper.cpp
+++ b/Engine/Source/Runtime/Experimental/Chaos/Private/Chaos/Deformable/ChaosDeformableSolverThreadingWrapper.cpp
@@ -39,11 +39,10 @@ namespace Chaos::Softs
 			Solver->Reset(InProps);
 	}
 
-	bool FDeformableSolver::FPhysicsThreadAccess::AdvanceDt(FSolverReal DeltaTime)
+	void FDeformableSolver::FPhysicsThreadAccess::AdvanceDt(FSolverReal DeltaTime)
 	{
 		if (Solver)
-			return Solver->AdvanceDt(DeltaTime);
-		return false;
+			Solver->AdvanceDt(DeltaTime);
 	}
 
 	void FDeformableSolver::FPhysicsThreadAccess::Simulate(FSolverReal DeltaTime)
@@ -113,10 +112,10 @@ namespace Chaos::Softs
 			Solver->RemoveSimulationObjects();
 	}
 
-	void FDeformableSolver::FPhysicsThreadAccess::TickSimulation(FSolverReal DeltaTime)
+	void FDeformableSolver::FPhysicsThreadAccess::Update(FSolverReal DeltaTime)
 	{
 		if(Solver)
-			Solver->TickSimulation(DeltaTime);
+			Solver->Update(DeltaTime);
 	}
 
 	void FDeformableSolver::FPhysicsThreadAccess::PushOutputPackage(int32 InFrame, FDeformableDataMap&& InPackage)
@@ -144,6 +143,19 @@ namespace Chaos::Softs
 			Solver->RemoveProxy(InObject);
 	}
 
+	void FDeformableSolver::FGameThreadAccess::SetEnableSolver(bool InbEnableSolver)
+	{
+		if (Solver)
+			Solver->SetEnableSolver(InbEnableSolver);
+	}
+
+	bool FDeformableSolver::FGameThreadAccess::GetEnableSolver()
+	{
+		if (Solver)
+			return Solver->GetEnableSolver();
+		return false;
+	}
+
 	void  FDeformableSolver::FPhysicsThreadAccess::UpdateOutputState(FThreadingProxy& InProxy)
 	{
 		if(Solver) 
diff --git a/Engine/Source/Runtime/Experimental/Chaos/Public/Chaos/Deformable/ChaosDeformableSolver.h b/Engine/Source/Runtime/Experimental/Chaos/Public/Chaos/Deformable/ChaosDeformableSolver.h
--- a/Engine/Source/Runtime/Experimental/Chaos/Public/Chaos/Deformable/ChaosDeformableSolver.h
+++ b/Engine/Source/Runtime/Experimental/Chaos/Public/Chaos/Deformable/ChaosDeformableSolver.h
@@ -20,11 +20,10 @@
 
 namespace Chaos::Softs
 {
-	class CHAOS_API FDeformableSolver
+	class CHAOS_API FDeformableSolver : public FPhysicsSolverEvents
 	{
 		friend class FGameThreadAccess;
 		friend class FPhysicsThreadAccess;
-		friend class FDeformableFun;
 
 	public:
 
@@ -41,9 +40,9 @@ namespace Chaos::Softs
 			/* Simulation Advance */
 			void UpdateProxyInputPackages();
 			void Simulate(FSolverReal DeltaTime);
-			bool AdvanceDt(FSolverReal DeltaTime);
+			void AdvanceDt(FSolverReal DeltaTime);
 			void Reset(const FDeformableSolverProperties&);
-			void TickSimulation(FSolverReal DeltaTime);
+			void Update(FSolverReal DeltaTime);
 			void UpdateOutputState(FThreadingProxy&);
 			TUniquePtr<FDeformablePackage> PullInputPackage();
 			void PushOutputPackage(int32 Frame, FDeformableDataMap&& Package);
@@ -87,6 +86,9 @@ namespace Chaos::Softs
 			void RemoveProxy(FThreadingProxy* InObject);
 			void PushInputPackage(int32 Frame, FDeformableDataMap&& InPackage);
 
+			void SetEnableSolver(bool InbEnableSolver);
+			bool GetEnableSolver();
+
 			TUniquePtr<FDeformablePackage> PullOutputPackage();
 
 		private:
@@ -95,13 +97,16 @@ namespace Chaos::Softs
 
 	protected:
 
+		void SetEnableSolver(bool InbEnableSolver) {FScopeLock Lock(&SolverEnabledMutex); bEnableSolver = InbEnableSolver; }
+		bool GetEnableSolver() const { return bEnableSolver; }
+
 		/* Simulation Advance */
 		int32 GetFrame() const { return Frame; }
 		void UpdateProxyInputPackages();
 		void Simulate(FSolverReal DeltaTime);
-		bool AdvanceDt(FSolverReal DeltaTime);
+		void AdvanceDt(FSolverReal DeltaTime);
 		void Reset(const FDeformableSolverProperties&);
-		void TickSimulation(FSolverReal DeltaTime);
+		void Update(FSolverReal DeltaTime);
 		void UpdateOutputState(FThreadingProxy&);
 		void PushOutputPackage(int32 Frame, FDeformableDataMap&& Package);
 		TUniquePtr<FDeformablePackage> PullInputPackage( );
@@ -140,6 +145,8 @@ namespace Chaos::Softs
 		static FCriticalSection	RemovalMutex; // @todo(flesh) : change to threaded commands to prevent the lock. 
 		static FCriticalSection	PackageOutputMutex;
 		static FCriticalSection	PackageInputMutex;
+		static FCriticalSection	SolverEnabledMutex;
+
 		TArray< FThreadingProxy* > RemovedProxys_Internal;
 		TArray< FThreadingProxy* > UninitializedProxys_Internal;
 		TArray< TUniquePtr<FDeformablePackage>  > BufferedInputPackages;
@@ -166,6 +173,7 @@ namespace Chaos::Softs
 		TUniquePtr <TArray<Chaos::TVec4<int32>>> AllElements;
 		TUniquePtr <FTriangleMesh> SurfaceTriangleMesh;
 
+		bool bEnableSolver = true;
 		FSolverReal Time = 0.f;
 		int32 Frame = 0;
 		bool bSimulationInitialized = false;
diff --git a/Engine/Source/Runtime/Experimental/Chaos/Public/Chaos/Deformable/ChaosDeformableSolverTypes.h b/Engine/Source/Runtime/Experimental/Chaos/Public/Chaos/Deformable/ChaosDeformableSolverTypes.h
--- a/Engine/Source/Runtime/Experimental/Chaos/Public/Chaos/Deformable/ChaosDeformableSolverTypes.h
+++ b/Engine/Source/Runtime/Experimental/Chaos/Public/Chaos/Deformable/ChaosDeformableSolverTypes.h
@@ -7,6 +7,7 @@
 #include "CoreMinimal.h"
 
 class UDeformableSolverComponent;
+class FFleshCacheAdapter;
 
 namespace Chaos::Softs
 {
@@ -82,10 +83,11 @@ namespace Chaos::Softs
 	class CHAOS_API FGameThreadAccessor
 	{
 	public:
-		friend class UDeformableSolverComponent;
-#if PLATFORM_WINDOWS
-	protected:
-#endif
+//		friend class UDeformableSolverComponent;
+//		friend class FFleshCacheAdapter;
+//#if PLATFORM_WINDOWS
+//	protected:
+//#endif
 		FGameThreadAccessor() {}
 	};
 
@@ -94,10 +96,11 @@ namespace Chaos::Softs
 	class CHAOS_API FPhysicsThreadAccessor
 	{
 	public:
-		friend class UDeformableSolverComponent;
-#if PLATFORM_WINDOWS
-	protected:
-#endif
+//		friend class UDeformableSolverComponent;
+//		friend class FFleshCacheAdapter;
+//#if PLATFORM_WINDOWS
+//	protected:
+//#endif
 		FPhysicsThreadAccessor() {}
 	};
 
