Iris - InvalidConnectionId constant.

#rnx

diff --git a/Engine/Source/Runtime/Experimental/Iris/Core/Public/Iris/IrisConstants.h b/Engine/Source/Runtime/Experimental/Iris/Core/Public/Iris/IrisConstants.h
--- a/Engine/Source/Runtime/Experimental/Iris/Core/Public/Iris/IrisConstants.h
+++ b/Engine/Source/Runtime/Experimental/Iris/Core/Public/Iris/IrisConstants.h
@@ -8,5 +8,6 @@ namespace UE::Net
 {
 
 constexpr uint32 InvalidReplicationSystemId = 0xFFFFFFFFU;
+constexpr uint32 InvalidConnectionId = 0;
 
 }
