Iris - FInternalNetSerializationContext adjustments.

#rnx

diff --git a/Engine/Source/Runtime/Experimental/Iris/Core/Private/Iris/Core/IrisDebugging.cpp b/Engine/Source/Runtime/Experimental/Iris/Core/Private/Iris/Core/IrisDebugging.cpp
--- a/Engine/Source/Runtime/Experimental/Iris/Core/Private/Iris/Core/IrisDebugging.cpp
+++ b/Engine/Source/Runtime/Experimental/Iris/Core/Private/Iris/Core/IrisDebugging.cpp
@@ -189,9 +189,16 @@ void NetObjectStateToString(FStringBuilderBase& StringBuilder, FNetHandle NetHan
 	FNetTokenStoreState* TokenStoreState = (bIsServer || (FirstValidConnectionId == FNetBitArray::InvalidIndex)) ? ReplicationSystemInternal->GetNetTokenStore().GetLocalNetTokenStoreState() : &Connections.GetRemoteNetTokenStoreState(FirstValidConnectionId);
 
 	// Setup Context
-	FInternalNetSerializationContext InternalContext(ReplicationSystem, TokenStoreState);
+	FInternalNetSerializationContext InternalContext;
+	FInternalNetSerializationContext::FInitParameters InternalContextInitParams;
+	InternalContextInitParams.ReplicationSystem = ReplicationSystem;
+	InternalContextInitParams.ObjectResolveContext.RemoteNetTokenStoreState = TokenStoreState;
+	InternalContextInitParams.ObjectResolveContext.ConnectionId = (FirstValidConnectionId == FNetBitArray::InvalidIndex ? InvalidConnectionId : FirstValidConnectionId);
+	InternalContext.Init(InternalContextInitParams);
+
 	FNetSerializationContext NetSerializationContext;
 	NetSerializationContext.SetInternalContext(&InternalContext);
+	NetSerializationContext.SetLocalConnectionId(InternalContextInitParams.ObjectResolveContext.ConnectionId);
 
 	StringBuilder << TEXT("State for ") << NetHandle;
 	StringBuilder.Appendf(TEXT(" of type %s\n"), Protocol->DebugName->Name);
