Iris
- Implement support for flushing state before the network object representation is destroyed.
- Makes sure to flush all reliable attachements + optionally replicated state as well, if replicated state should be flushed an explicit copy/quantize will be done on final state. Once final state is obtained network code will no longer touch source data eliminating the need to keep the actor/source around so that it can finish replication.

diff --git a/Engine/Plugins/Runtime/ReplicationSystemTestPlugin/Source/Private/Tests/ReplicationSystem/ReplicationSystemServerClientTestFixture.cpp b/Engine/Plugins/Runtime/ReplicationSystemTestPlugin/Source/Private/Tests/ReplicationSystem/ReplicationSystemServerClientTestFixture.cpp
--- a/Engine/Plugins/Runtime/ReplicationSystemTestPlugin/Source/Private/Tests/ReplicationSystem/ReplicationSystemServerClientTestFixture.cpp
+++ b/Engine/Plugins/Runtime/ReplicationSystemTestPlugin/Source/Private/Tests/ReplicationSystem/ReplicationSystemServerClientTestFixture.cpp
@@ -150,13 +150,13 @@ UTestReplicatedIrisObject* FReplicationSystemTestNode::CreateObjectWithDynamicSt
 	return CreatedObject;
 }
 
-void FReplicationSystemTestNode::DestroyObject(UReplicatedTestObject* Object)
+void FReplicationSystemTestNode::DestroyObject(UReplicatedTestObject* Object, EEndReplicationFlags EndReplicationFlags)
 {
 	// Destroy handle
 	check(Object && Object->NetHandle.IsValid());
 
-	// Destroy the handle
-	ReplicationBridge->EndReplication(Object);
+	// End replication for the handle
+	ReplicationBridge->EndReplication(Object, EndReplicationFlags);
 
 	// Release ref
 	CreatedObjects.Remove(TStrongObjectPtr<UObject>(Object));
@@ -277,7 +277,7 @@ void FReplicationSystemTestNode::DeliverTo(FReplicationSystemTestNode& Dest, uin
 	else
 	{
 		UE_NET_TRACE_PACKET_DROPPED(Dest.GetNetTraceId(), RemoteConnectionId, Packet.PacketId, ENetTracePacketType::Incoming);
-		UE_NET_TRACE_PACKET_DROPPED(GetNetTraceId(), RemoteConnectionId, Packet.PacketId, ENetTracePacketType::Outgoing);
+		UE_NET_TRACE_PACKET_DROPPED(GetNetTraceId(), LocalConnectionId, Packet.PacketId, ENetTracePacketType::Outgoing);
 	}
 
 	// If this triggers an assert, ensure that SendTo() actually wrote any packets before.
diff --git a/Engine/Plugins/Runtime/ReplicationSystemTestPlugin/Source/Private/Tests/ReplicationSystem/ReplicationSystemServerClientTestFixture.h b/Engine/Plugins/Runtime/ReplicationSystemTestPlugin/Source/Private/Tests/ReplicationSystem/ReplicationSystemServerClientTestFixture.h
--- a/Engine/Plugins/Runtime/ReplicationSystemTestPlugin/Source/Private/Tests/ReplicationSystem/ReplicationSystemServerClientTestFixture.h
+++ b/Engine/Plugins/Runtime/ReplicationSystemTestPlugin/Source/Private/Tests/ReplicationSystem/ReplicationSystemServerClientTestFixture.h
@@ -119,7 +119,7 @@ public:
 	UTestReplicatedIrisObject* CreateObject(const UTestReplicatedIrisObject::FComponents& Components);
 	UTestReplicatedIrisObject* CreateSubObject(FNetHandle Owner, const UTestReplicatedIrisObject::FComponents& Components);
 	UTestReplicatedIrisObject* CreateObjectWithDynamicState(uint32 NumComponents, uint32 NumIrisComponents, uint32 NumDynamicStateComponents);
-	void DestroyObject(UReplicatedTestObject*);
+	void DestroyObject(UReplicatedTestObject*, EEndReplicationFlags EndReplicationFlags = EEndReplicationFlags::Destroy);
 
 	// Connection
 	uint32 AddConnection();
diff --git a/Engine/Plugins/Runtime/ReplicationSystemTestPlugin/Source/Private/Tests/ReplicationSystem/TestFlushBeforeDestroy.cpp b/Engine/Plugins/Runtime/ReplicationSystemTestPlugin/Source/Private/Tests/ReplicationSystem/TestFlushBeforeDestroy.cpp
new file mode 100644
--- /dev/null
+++ b/Engine/Plugins/Runtime/ReplicationSystemTestPlugin/Source/Private/Tests/ReplicationSystem/TestFlushBeforeDestroy.cpp
@@ -0,0 +1,565 @@
+// Copyright Epic Games, Inc. All Rights Reserved.
+
+#include "ReplicationSystemServerClientTestFixture.h"
+#include "Iris/ReplicationSystem/ReplicationSystem.h"
+#include "Iris/ReplicationSystem/Filtering/NetObjectFilter.h"
+#include "NetBlob/PartialNetBlobTestFixture.h"
+#include "NetBlob/MockNetBlob.h"
+
+namespace UE::Net::Private
+{
+
+class FTestFlushBeforeDestroyFixture : public FPartialNetBlobTestFixture
+{
+};
+
+UE_NET_TEST_FIXTURE(FTestFlushBeforeDestroyFixture, TestReliableAttachmentFlushedBeforeDestroy)
+{
+	FReplicationSystemTestClient* Client = CreateClient();
+	RegisterNetBlobHandlers(Client);
+
+	UReplicatedTestObject* ServerObject = Server->CreateObject(0, 0);
+	FNetHandle ObjectHandle = ServerObject->NetHandle;
+
+	Server->PreSendUpdate();
+	Server->SendAndDeliverTo(Client, DeliverPacket);
+	Server->PostSendUpdate();
+
+	UE_NET_ASSERT_NE(Client->GetReplicationBridge()->GetReplicatedObject(ServerObject->NetHandle), nullptr);
+
+	// Create attachment
+	{
+		constexpr uint32 PayloadBitCount = 24;
+		const TRefCountPtr<FNetObjectAttachment>& Attachment = MockNetObjectAttachmentHandler->CreateReliableNetObjectAttachment(PayloadBitCount);
+		FNetObjectReference AttachmentTarget = FObjectReferenceCache::MakeNetObjectReference(ServerObject->NetHandle);
+		Server->GetReplicationSystem()->QueueNetObjectAttachment(Client->ConnectionIdOnServer, AttachmentTarget, Attachment);
+	}
+
+	// Destroy object, on server
+	Server->DestroyObject(ServerObject);
+
+	// Deliver a packet, this should flush the object and deliver the attachment
+	Server->PreSendUpdate();
+	Server->SendAndDeliverTo(Client, DeliverPacket);
+	Server->PostSendUpdate();
+
+	// Verify that the attachment has been received
+	UE_NET_ASSERT_EQ(ClientMockNetObjectAttachmentHandler->GetFunctionCallCounts().OnNetBlobReceived, 1U);
+
+	// Deliver a packet. Should destroy the object on the client unless that was done
+	Server->PreSendUpdate();
+	Server->SendAndDeliverTo(Client, DeliverPacket);
+	Server->PostSendUpdate();
+
+	// Verify that object is destroyed
+	UE_NET_ASSERT_TRUE(Client->GetReplicationBridge()->GetReplicatedObject(ObjectHandle) == nullptr);
+}
+
+UE_NET_TEST_FIXTURE(FTestFlushBeforeDestroyFixture, TestReliableAttachmentFlushedWithDataInflightBeforeDestroy)
+{
+
+	FReplicationSystemTestClient* Client = CreateClient();
+	RegisterNetBlobHandlers(Client);
+
+	UReplicatedTestObject* ServerObject = Server->CreateObject(0, 0);
+	FNetHandle ObjectHandle = ServerObject->NetHandle;
+
+	Server->PreSendUpdate();
+	Server->SendAndDeliverTo(Client, DeliverPacket);
+	Server->PostSendUpdate();
+
+	UE_NET_ASSERT_NE(Client->GetReplicationBridge()->GetReplicatedObject(ServerObject->NetHandle), nullptr);
+
+	// Create attachment
+	{
+		constexpr uint32 PayloadBitCount = 24;
+		const TRefCountPtr<FNetObjectAttachment>& Attachment = MockNetObjectAttachmentHandler->CreateReliableNetObjectAttachment(PayloadBitCount);
+		FNetObjectReference AttachmentTarget = FObjectReferenceCache::MakeNetObjectReference(ServerObject->NetHandle);
+		Server->GetReplicationSystem()->QueueNetObjectAttachment(Client->ConnectionIdOnServer, AttachmentTarget, Attachment);
+	}
+
+	// Setup a situation where we have reliable data in flight when the object is destroyed
+	Server->PreSendUpdate();
+	Server->SendTo(Client);
+	Server->PostSendUpdate();
+
+	// Destroy object, on server
+	Server->DestroyObject(ServerObject);
+
+	// Drop the data and notify server
+	Server->DeliverTo(Client, false);
+
+	// Deliver a packet, this should flush the object and deliver the attachment
+	Server->PreSendUpdate();
+	Server->SendAndDeliverTo(Client, DeliverPacket);
+	Server->PostSendUpdate();
+
+	// Verify that the attachment has been received
+	UE_NET_ASSERT_EQ(ClientMockNetObjectAttachmentHandler->GetFunctionCallCounts().OnNetBlobReceived, 1U);
+
+	// Deliver a packet. Should destroy the object on the client unless that was done
+	Server->PreSendUpdate();
+	Server->SendAndDeliverTo(Client, DeliverPacket);
+	Server->PostSendUpdate();
+
+	// Verify that object is destroyed
+	UE_NET_ASSERT_TRUE(Client->GetReplicationBridge()->GetReplicatedObject(ObjectHandle) == nullptr);
+}
+
+UE_NET_TEST_FIXTURE(FTestFlushBeforeDestroyFixture, TestReliableAttachmentSubObjectFlushedBeforeDestroy)
+{
+	FReplicationSystemTestClient* Client = CreateClient();
+	RegisterNetBlobHandlers(Client);
+
+	UReplicatedTestObject* ServerObject = Server->CreateObject(0, 0);
+	FNetHandle ObjectHandle = ServerObject->NetHandle;
+	UReplicatedTestObject* ServerSubObject = Server->CreateSubObject(ObjectHandle, 0, 0);
+	FNetHandle SubObjectHandle = ServerSubObject->NetHandle;
+	
+	Server->PreSendUpdate();
+	Server->SendAndDeliverTo(Client, DeliverPacket);
+	Server->PostSendUpdate();
+
+	UE_NET_ASSERT_NE(Client->GetReplicationBridge()->GetReplicatedObject(ServerObject->NetHandle), nullptr);
+
+	// Create attachment
+	{
+		constexpr uint32 PayloadBitCount = 24;
+		const TRefCountPtr<FNetObjectAttachment>& Attachment = MockNetObjectAttachmentHandler->CreateReliableNetObjectAttachment(PayloadBitCount);
+		FNetObjectReference AttachmentTarget = FObjectReferenceCache::MakeNetObjectReference(ServerSubObject->NetHandle);
+		Server->GetReplicationSystem()->QueueNetObjectAttachment(Client->ConnectionIdOnServer, AttachmentTarget, Attachment);
+	}
+
+	// Destroy sub object, on server
+	Server->DestroyObject(ServerSubObject);
+
+	// Deliver a packet, this should flush the object and deliver the attachment
+	Server->PreSendUpdate();
+	Server->SendAndDeliverTo(Client, DeliverPacket);
+	Server->PostSendUpdate();
+
+	// Verify that the attachment has been received
+	UE_NET_ASSERT_EQ(ClientMockNetObjectAttachmentHandler->GetFunctionCallCounts().OnNetBlobReceived, 1U);
+
+	// Deliver a packet. Should destroy the object on the client unless that was done
+	Server->PreSendUpdate();
+	Server->SendAndDeliverTo(Client, DeliverPacket);
+	Server->PostSendUpdate();
+
+	// Verify that object is destroyed
+	UE_NET_ASSERT_TRUE(Client->GetReplicationBridge()->GetReplicatedObject(SubObjectHandle) == nullptr);
+}
+
+UE_NET_TEST_FIXTURE(FTestFlushBeforeDestroyFixture, TestReliableAttachmentSubObjectFlushedBeforeDestroyIfOwnerIsDestroyed)
+{
+	FReplicationSystemTestClient* Client = CreateClient();
+	RegisterNetBlobHandlers(Client);
+
+	UReplicatedTestObject* ServerObject = Server->CreateObject(0, 0);
+	FNetHandle ObjectHandle = ServerObject->NetHandle;
+	UReplicatedTestObject* ServerSubObject = Server->CreateSubObject(ObjectHandle, 0, 0);
+	FNetHandle SubObjectHandle = ServerSubObject->NetHandle;
+	
+	Server->PreSendUpdate();
+	Server->SendAndDeliverTo(Client, DeliverPacket);
+	Server->PostSendUpdate();
+
+	UE_NET_ASSERT_NE(Client->GetReplicationBridge()->GetReplicatedObject(ServerObject->NetHandle), nullptr);
+
+	// Create attachment
+	{
+		constexpr uint32 PayloadBitCount = 24;
+		const TRefCountPtr<FNetObjectAttachment>& Attachment = MockNetObjectAttachmentHandler->CreateReliableNetObjectAttachment(PayloadBitCount);
+		FNetObjectReference AttachmentTarget = FObjectReferenceCache::MakeNetObjectReference(ServerSubObject->NetHandle);
+		Server->GetReplicationSystem()->QueueNetObjectAttachment(Client->ConnectionIdOnServer, AttachmentTarget, Attachment);
+	}
+
+	// Destroy object which should flush subobject and then destroy both subobject and object
+	Server->DestroyObject(ServerObject);
+
+	// Deliver a packet, this should flush the object and deliver the attachment
+	Server->PreSendUpdate();
+	Server->SendAndDeliverTo(Client, DeliverPacket);
+	Server->PostSendUpdate();
+
+	// Verify that the attachment has been received
+	UE_NET_ASSERT_EQ(ClientMockNetObjectAttachmentHandler->GetFunctionCallCounts().OnNetBlobReceived, 1U);
+
+	// Deliver a packet. Should destroy the object on the client unless that was done
+	Server->PreSendUpdate();
+	Server->SendAndDeliverTo(Client, DeliverPacket);
+	Server->PostSendUpdate();
+
+	// Verify that both object and subobject are destroyed
+	UE_NET_ASSERT_TRUE(Client->GetReplicationBridge()->GetReplicatedObject(SubObjectHandle) == nullptr);
+	UE_NET_ASSERT_TRUE(Client->GetReplicationBridge()->GetReplicatedObject(ObjectHandle) == nullptr);
+}
+
+UE_NET_TEST_FIXTURE(FTestFlushBeforeDestroyFixture, TestStateFlushedBeforeDestroy)
+{
+	FReplicationSystemTestClient* Client = CreateClient();
+	RegisterNetBlobHandlers(Client);
+
+	UTestReplicatedIrisObject* ServerObject = Server->CreateObject(0, 0);
+	FNetHandle ObjectHandle = ServerObject->NetHandle;
+	
+	Server->PreSendUpdate();
+	Server->SendAndDeliverTo(Client, DeliverPacket);
+	Server->PostSendUpdate();
+
+	// Verify that object is created
+	UE_NET_ASSERT_NE(Client->GetReplicationBridge()->GetReplicatedObject(ObjectHandle), nullptr);
+
+	// Modify state
+	ServerObject->IntA = 3;
+
+	// Destroy object with flush flag which should flush the state before destroying the object
+	Server->DestroyObject(ServerObject, EEndReplicationFlags::Destroy | EEndReplicationFlags::Flush);
+
+	// Deliver a packet, this should flush the object and deliver the last state
+	Server->PreSendUpdate();
+	Server->SendAndDeliverTo(Client, DeliverPacket);
+	Server->PostSendUpdate();
+
+	UTestReplicatedIrisObject* ClientObject = Cast<UTestReplicatedIrisObject>(Client->GetReplicationBridge()->GetReplicatedObject(ObjectHandle));
+
+	// Verify that object is created
+	UE_NET_ASSERT_NE(ClientObject, nullptr);
+
+	// Verify that we got the expected state
+	UE_NET_ASSERT_EQ(ClientObject->IntA, 3);
+	
+	// Deliver a packet.
+	Server->PreSendUpdate();
+	Server->SendAndDeliverTo(Client, DeliverPacket);
+	Server->PostSendUpdate();
+
+	// Verify that object is destroyed
+	UE_NET_ASSERT_TRUE(Client->GetReplicationBridge()->GetReplicatedObject(ObjectHandle) == nullptr);
+}
+
+UE_NET_TEST_FIXTURE(FTestFlushBeforeDestroyFixture, TestStateInFlightFlushedBeforeDestroy)
+{
+	FReplicationSystemTestClient* Client = CreateClient();
+	RegisterNetBlobHandlers(Client);
+
+	UTestReplicatedIrisObject* ServerObject = Server->CreateObject(0, 0);
+	FNetHandle ObjectHandle = ServerObject->NetHandle;
+	
+	Server->PreSendUpdate();
+	Server->SendAndDeliverTo(Client, DeliverPacket);
+	Server->PostSendUpdate();
+
+	// Verify that object is created
+	UE_NET_ASSERT_NE(Client->GetReplicationBridge()->GetReplicatedObject(ObjectHandle), nullptr);
+
+	// Modify state
+	ServerObject->IntA = 3;
+
+	Server->PreSendUpdate();
+	Server->SendTo(Client);
+	Server->PostSendUpdate();
+
+	// Modify state
+	ServerObject->IntB = 4;
+
+	// Destroy object with flush flag which should flush the state before destroying the object
+	Server->DestroyObject(ServerObject, EEndReplicationFlags::Destroy | EEndReplicationFlags::Flush);
+
+	// Drop the data we had in flight and notify server
+	Server->DeliverTo(Client, false);
+
+	// Deliver a packet, this should flush the object and deliver the complete last state
+	Server->PreSendUpdate();
+	Server->SendAndDeliverTo(Client, DeliverPacket);
+	Server->PostSendUpdate();
+
+	UTestReplicatedIrisObject* ClientObject = Cast<UTestReplicatedIrisObject>(Client->GetReplicationBridge()->GetReplicatedObject(ObjectHandle));
+
+	// Verify that object is created
+	UE_NET_ASSERT_NE(ClientObject, nullptr);
+
+	// Verify that we got the expected state
+	UE_NET_ASSERT_EQ(ClientObject->IntA, 3);
+	UE_NET_ASSERT_EQ(ClientObject->IntB, 4);
+
+	// Deliver a packet. Should destroy the object on the client unless that was done
+	Server->PreSendUpdate();
+	Server->SendAndDeliverTo(Client, DeliverPacket);
+	Server->PostSendUpdate();
+
+	// Verify that object is destroyed
+	UE_NET_ASSERT_TRUE(Client->GetReplicationBridge()->GetReplicatedObject(ObjectHandle) == nullptr);
+}
+
+UE_NET_TEST_FIXTURE(FTestFlushBeforeDestroyFixture, TestSubObjectStateFlushedBeforeOwnerDestroy)
+{
+	FReplicationSystemTestClient* Client = CreateClient();
+	RegisterNetBlobHandlers(Client);
+
+	UTestReplicatedIrisObject* ServerObject = Server->CreateObject(0, 0);
+	FNetHandle ObjectHandle = ServerObject->NetHandle;
+	UTestReplicatedIrisObject* ServerSubObject = Server->CreateSubObject(ObjectHandle, 0, 0);
+	FNetHandle SubObjectHandle = ServerSubObject->NetHandle;
+
+	Server->PreSendUpdate();
+	Server->SendAndDeliverTo(Client, DeliverPacket);
+	Server->PostSendUpdate();
+
+	// Verify that objects is created
+	UE_NET_ASSERT_NE(Client->GetReplicationBridge()->GetReplicatedObject(ObjectHandle), nullptr);
+	UE_NET_ASSERT_NE(Client->GetReplicationBridge()->GetReplicatedObject(SubObjectHandle), nullptr);
+
+	// Modify state
+	ServerSubObject->IntA = 3;
+
+	// Destroy object with flush flag which should flush the state including before destroying the object
+	Server->DestroyObject(ServerObject, EEndReplicationFlags::Destroy | EEndReplicationFlags::Flush);
+
+	// Deliver a packet, this should flush the object and deliver the last state
+	Server->PreSendUpdate();
+	Server->SendAndDeliverTo(Client, DeliverPacket);
+	Server->PostSendUpdate();
+
+	UTestReplicatedIrisObject* ClientObject = Cast<UTestReplicatedIrisObject>(Client->GetReplicationBridge()->GetReplicatedObject(ObjectHandle));
+	UTestReplicatedIrisObject* ClientSubObject = Cast<UTestReplicatedIrisObject>(Client->GetReplicationBridge()->GetReplicatedObject(SubObjectHandle));
+
+	// Verify that objects is created
+	UE_NET_ASSERT_NE(ClientObject, nullptr);
+	UE_NET_ASSERT_NE(ClientSubObject, nullptr);
+
+	// Verify that we got the expected state
+	UE_NET_ASSERT_EQ(ClientSubObject->IntA, 3);
+	
+	// Deliver a packet.
+	Server->PreSendUpdate();
+	Server->SendAndDeliverTo(Client, DeliverPacket);
+	Server->PostSendUpdate();
+
+	// Verify that both objects are destroyed
+	UE_NET_ASSERT_TRUE(Client->GetReplicationBridge()->GetReplicatedObject(ObjectHandle) == nullptr);
+	UE_NET_ASSERT_TRUE(Client->GetReplicationBridge()->GetReplicatedObject(SubObjectHandle) == nullptr);
+}
+
+UE_NET_TEST_FIXTURE(FTestFlushBeforeDestroyFixture, TestSubObjectStateFlushedBeforeSubObjectDestroy)
+{
+	FReplicationSystemTestClient* Client = CreateClient();
+	RegisterNetBlobHandlers(Client);
+
+	UTestReplicatedIrisObject* ServerObject = Server->CreateObject(0, 0);
+	FNetHandle ObjectHandle = ServerObject->NetHandle;
+	UTestReplicatedIrisObject* ServerSubObject = Server->CreateSubObject(ObjectHandle, 0, 0);
+	FNetHandle SubObjectHandle = ServerSubObject->NetHandle;
+
+	Server->PreSendUpdate();
+	Server->SendAndDeliverTo(Client, DeliverPacket);
+	Server->PostSendUpdate();
+
+	// Verify that objects is created
+	UE_NET_ASSERT_NE(Client->GetReplicationBridge()->GetReplicatedObject(ObjectHandle), nullptr);
+	UE_NET_ASSERT_NE(Client->GetReplicationBridge()->GetReplicatedObject(SubObjectHandle), nullptr);
+
+	// Modify state on SubObject
+	ServerSubObject->IntA = 3;
+
+	// Destroy object with flush flag which should flush the state including before destroying the object
+	Server->DestroyObject(ServerSubObject, EEndReplicationFlags::Destroy | EEndReplicationFlags::Flush);
+
+	// Deliver a packet, this should flush the object and deliver the last state
+	Server->PreSendUpdate();
+	Server->SendAndDeliverTo(Client, DeliverPacket);
+	Server->PostSendUpdate();
+
+	UTestReplicatedIrisObject* ClientObject = Cast<UTestReplicatedIrisObject>(Client->GetReplicationBridge()->GetReplicatedObject(ObjectHandle));
+	UTestReplicatedIrisObject* ClientSubObject = Cast<UTestReplicatedIrisObject>(Client->GetReplicationBridge()->GetReplicatedObject(SubObjectHandle));
+
+	// Verify that objects are created
+	UE_NET_ASSERT_NE(ClientObject, nullptr);
+	UE_NET_ASSERT_NE(ClientSubObject, nullptr);
+
+	// Verify that we got the expected state
+	UE_NET_ASSERT_EQ(ClientSubObject->IntA, 3);
+	
+	// Deliver a packet.
+	Server->PreSendUpdate();
+	Server->SendAndDeliverTo(Client, DeliverPacket);
+	Server->PostSendUpdate();
+
+	// Verify subobject is destroyed now that last state was confirmed flushed while the main object still is around
+	UE_NET_ASSERT_FALSE(Client->GetReplicationBridge()->GetReplicatedObject(ObjectHandle) == nullptr);
+	UE_NET_ASSERT_TRUE(Client->GetReplicationBridge()->GetReplicatedObject(SubObjectHandle) == nullptr);
+}
+
+// Test TearOff for existing confirmed object
+UE_NET_TEST_FIXTURE(FTestFlushBeforeDestroyFixture, TestTearOffNewObjectWithReliableAttachment)
+{
+	// Add a client
+	FReplicationSystemTestClient* Client = CreateClient();
+	RegisterNetBlobHandlers(Client);
+
+	// Spawn object on server
+	UTestReplicatedIrisObject* ServerObject = Server->CreateObject(0,0);
+
+	// Trigger replication
+	ServerObject->IntA = 1;
+
+	// Create attachment
+	{
+		constexpr uint32 PayloadBitCount = 24;
+		const TRefCountPtr<FNetObjectAttachment>& Attachment = MockNetObjectAttachmentHandler->CreateReliableNetObjectAttachment(PayloadBitCount);
+		FNetObjectReference AttachmentTarget = FObjectReferenceCache::MakeNetObjectReference(ServerObject->NetHandle);
+		Server->GetReplicationSystem()->QueueNetObjectAttachment(Client->ConnectionIdOnServer, AttachmentTarget, Attachment);
+	}
+
+	// TearOff the object
+	Server->ReplicationBridge->EndReplication(ServerObject, EEndReplicationFlags::TearOff);
+
+	// Send and deliver packet
+	Server->PreSendUpdate();
+	Server->SendAndDeliverTo(Client, true);
+	Server->PostSendUpdate();
+
+	// Verify that object got created
+	UTestReplicatedIrisObject* ClientObjectThatWillBeTornOff = Cast<UTestReplicatedIrisObject>(Client->GetReplicationBridge()->GetReplicatedObject(ServerObject->NetHandle));
+
+	UE_NET_ASSERT_TRUE(ClientObjectThatWillBeTornOff != nullptr);
+
+	// Verify that ClientObject got final state and that the attachement was received
+	UE_NET_ASSERT_EQ(ServerObject->IntA, ClientObjectThatWillBeTornOff->IntA);
+	UE_NET_ASSERT_EQ(ClientMockNetObjectAttachmentHandler->GetFunctionCallCounts().OnNetBlobReceived, 1U);
+
+	// Send and deliver packet
+	Server->PreSendUpdate();
+	Server->SendAndDeliverTo(Client, true);
+	Server->PostSendUpdate();
+
+	// Verify that ClientObject now has been teared off
+	UE_NET_ASSERT_TRUE(Client->GetReplicationBridge()->GetReplicatedObject(ServerObject->NetHandle) == nullptr);
+}
+
+// Test TearOff for existing confirmed object
+UE_NET_TEST_FIXTURE(FTestFlushBeforeDestroyFixture, TestTearOffExistingObjectWithReliableAttachment)
+{
+	// Add a client
+	FReplicationSystemTestClient* Client = CreateClient();
+	RegisterNetBlobHandlers(Client);
+
+	// Spawn object on server
+	UTestReplicatedIrisObject* ServerObject = Server->CreateObject(0,0);
+
+	// Trigger replication
+	ServerObject->IntA = 1;
+
+	// Send and deliver packet
+	Server->PreSendUpdate();
+	Server->SendAndDeliverTo(Client, true);
+	Server->PostSendUpdate();
+
+	// Create attachment
+	{
+		constexpr uint32 PayloadBitCount = 24;
+		const TRefCountPtr<FNetObjectAttachment>& Attachment = MockNetObjectAttachmentHandler->CreateReliableNetObjectAttachment(PayloadBitCount);
+		FNetObjectReference AttachmentTarget = FObjectReferenceCache::MakeNetObjectReference(ServerObject->NetHandle);
+		Server->GetReplicationSystem()->QueueNetObjectAttachment(Client->ConnectionIdOnServer, AttachmentTarget, Attachment);
+	}
+
+	// Store Pointer to object 
+	UTestReplicatedIrisObject* ClientObjectThatWillBeTornOff = Cast<UTestReplicatedIrisObject>(Client->GetReplicationBridge()->GetReplicatedObject(ServerObject->NetHandle));
+
+	UE_NET_ASSERT_TRUE(ClientObjectThatWillBeTornOff != nullptr);
+	UE_NET_ASSERT_EQ(ServerObject->IntA, ClientObjectThatWillBeTornOff->IntA);
+
+	// Modify the value
+	ServerObject->IntA = 2;
+
+	// TearOff the object
+	Server->ReplicationBridge->EndReplication(ServerObject, EEndReplicationFlags::TearOff);
+
+	// Send and deliver packet
+	Server->PreSendUpdate();
+	Server->SendAndDeliverTo(Client, true);
+	Server->PostSendUpdate();
+
+	// Verify that ClientObject got final state and that the attachement was received
+	UE_NET_ASSERT_EQ(ServerObject->IntA, ClientObjectThatWillBeTornOff->IntA);
+	UE_NET_ASSERT_EQ(ClientMockNetObjectAttachmentHandler->GetFunctionCallCounts().OnNetBlobReceived, 1U);
+
+	// Verify that ClientObject still is around (from a network perspective)
+	UE_NET_ASSERT_TRUE(Cast<UTestReplicatedIrisObject>(Client->GetReplicationBridge()->GetReplicatedObject(ServerObject->NetHandle)) != nullptr);
+
+	// Send and deliver packet
+	Server->PreSendUpdate();
+	Server->SendAndDeliverTo(Client, true);
+	Server->PostSendUpdate();
+
+	// Verify that ClientObject now has been teared off
+	UE_NET_ASSERT_TRUE(Cast<UTestReplicatedIrisObject>(Client->GetReplicationBridge()->GetReplicatedObject(ServerObject->NetHandle)) == nullptr);
+}
+
+// Test TearOff and SubObjects, SubObjects must apply state?
+UE_NET_TEST_FIXTURE(FTestFlushBeforeDestroyFixture, TestImmediateTearOffExistingObjectWithSubObjectWithReliableAttachment)
+{
+	UReplicationSystem* ReplicationSystem = Server->ReplicationSystem;
+
+	// Add a client
+	FReplicationSystemTestClient* Client = CreateClient();
+	RegisterNetBlobHandlers(Client);
+
+	// Spawn object on server
+	UTestReplicatedIrisObject* ServerObject = Server->CreateObject(0,0);
+
+	// Spawn second object on server as a subobject
+	UTestReplicatedIrisObject* ServerSubObject = Server->CreateSubObject(ServerObject->NetHandle, 0, 0);
+
+	// Trigger replication
+	ServerObject->IntA = 1;
+	ServerSubObject->IntA = 1;
+
+	// Send and deliver packet
+	Server->PreSendUpdate();
+	Server->SendAndDeliverTo(Client, true);
+	Server->PostSendUpdate();
+
+	// Store Pointer to objects
+	UTestReplicatedIrisObject* ClientObjectThatWillBeTornOff = Cast<UTestReplicatedIrisObject>(Client->GetReplicationBridge()->GetReplicatedObject(ServerObject->NetHandle));
+	UE_NET_ASSERT_TRUE(ClientObjectThatWillBeTornOff != nullptr);
+	UE_NET_ASSERT_EQ(ServerObject->IntA, ClientObjectThatWillBeTornOff->IntA);
+
+	UTestReplicatedIrisObject* ClientSubObjectThatWillBeTornOff = Cast<UTestReplicatedIrisObject>(Client->GetReplicationBridge()->GetReplicatedObject(ServerSubObject->NetHandle));
+	UE_NET_ASSERT_TRUE(ClientSubObjectThatWillBeTornOff != nullptr);
+	UE_NET_ASSERT_EQ(ServerSubObject->IntA, ClientSubObjectThatWillBeTornOff->IntA);
+
+	// Modify the value of subobject only
+	ServerSubObject->IntA = 2;
+
+	// Create attachment
+	{
+		constexpr uint32 PayloadBitCount = 24;
+		const TRefCountPtr<FNetObjectAttachment>& Attachment = MockNetObjectAttachmentHandler->CreateReliableNetObjectAttachment(PayloadBitCount);
+		FNetObjectReference AttachmentTarget = FObjectReferenceCache::MakeNetObjectReference(ServerSubObject->NetHandle);
+		Server->GetReplicationSystem()->QueueNetObjectAttachment(Client->ConnectionIdOnServer, AttachmentTarget, Attachment);
+	}
+
+	// TearOff the object using immediate tear-off
+	Server->ReplicationBridge->EndReplication(ServerObject, EEndReplicationFlags::TearOff);
+
+	// Send and deliver packet
+	Server->PreSendUpdate();
+	Server->SendAndDeliverTo(Client, true);
+	Server->PostSendUpdate();
+
+	// Verify that ClientObject got final state and that the attachement was received
+	UE_NET_ASSERT_EQ(ClientMockNetObjectAttachmentHandler->GetFunctionCallCounts().OnNetBlobReceived, 1U);
+	UE_NET_ASSERT_EQ(ServerSubObject->IntA, ClientSubObjectThatWillBeTornOff->IntA);
+
+	// Send and deliver packet
+	Server->PreSendUpdate();
+	Server->SendAndDeliverTo(Client, true);
+	Server->PostSendUpdate();
+
+	// Verify that ClientObject is torn-off
+	UE_NET_ASSERT_TRUE(Cast<UTestReplicatedIrisObject>(Client->GetReplicationBridge()->GetReplicatedObject(ServerSubObject->NetHandle)) == nullptr);
+}
+
+
+}
diff --git a/Engine/Source/Runtime/Experimental/Iris/Core/Private/Iris/ReplicationSystem/AttachmentReplication.cpp b/Engine/Source/Runtime/Experimental/Iris/Core/Private/Iris/ReplicationSystem/AttachmentReplication.cpp
--- a/Engine/Source/Runtime/Experimental/Iris/Core/Private/Iris/ReplicationSystem/AttachmentReplication.cpp
+++ b/Engine/Source/Runtime/Experimental/Iris/Core/Private/Iris/ReplicationSystem/AttachmentReplication.cpp
@@ -225,6 +225,11 @@ bool FNetObjectAttachmentSendQueue::IsAllSentAndAcked() const
 	return UnreliableQueue.IsEmpty() && (ReliableQueue == nullptr || ReliableQueue->IsAllSentAndAcked());
 }
 
+bool FNetObjectAttachmentSendQueue::IsAllReliableSentAndAcked() const
+{
+	return ReliableQueue == nullptr || ReliableQueue->IsAllSentAndAcked();
+}
+
 bool FNetObjectAttachmentSendQueue::IsSafeToDestroy() const
 {
 	return UnreliableQueue.IsEmpty() && (ReliableQueue == nullptr || ReliableQueue->IsSafeToDestroy());
@@ -469,7 +474,18 @@ bool FNetObjectAttachmentsWriter::IsAllSentAndAcked(ENetObjectAttachmentType Typ
 	const FNetObjectAttachmentSendQueue* Queue = GetQueue(Type, ObjectIndex);
 	if (Queue == nullptr)
 	{
-		return false;
+		return true;
+	}
+
+	return Queue->IsAllSentAndAcked();
+}
+
+bool FNetObjectAttachmentsWriter::IsAllReliableSentAndAcked(ENetObjectAttachmentType Type, uint32 ObjectIndex) const
+{
+	const FNetObjectAttachmentSendQueue* Queue = GetQueue(Type, ObjectIndex);
+	if (Queue == nullptr)
+	{
+		return true;
 	}
 
 	return Queue->IsAllSentAndAcked();
diff --git a/Engine/Source/Runtime/Experimental/Iris/Core/Private/Iris/ReplicationSystem/AttachmentReplication.h b/Engine/Source/Runtime/Experimental/Iris/Core/Private/Iris/ReplicationSystem/AttachmentReplication.h
--- a/Engine/Source/Runtime/Experimental/Iris/Core/Private/Iris/ReplicationSystem/AttachmentReplication.h
+++ b/Engine/Source/Runtime/Experimental/Iris/Core/Private/Iris/ReplicationSystem/AttachmentReplication.h
@@ -63,6 +63,8 @@ public:
 
 	bool IsAllSentAndAcked() const;
 
+	bool IsAllReliableSentAndAcked() const;
+
 	void SetUnreliableQueueCapacity(uint32 QueueCapacity);
 
 private:
@@ -109,6 +111,10 @@ public:
 	bool HasUnsentAttachments(ENetObjectAttachmentType Type, uint32 ObjectIndex) const;
 	// Whether all queued attachments have been sent and that all reliable ones have been acked.
 	bool IsAllSentAndAcked(ENetObjectAttachmentType Type, uint32 ObjectIndex) const;
+
+	// Whether all queued reliable attachments have been sent and acked
+	bool IsAllReliableSentAndAcked(ENetObjectAttachmentType Type, uint32 ObjectIndex) const;
+
 	// Whether the queue can be destroyed without causing issues if more attachments are queued to this instance.
 	bool IsSafeToDestroy(ENetObjectAttachmentType Type, uint32 ObjectIndex) const;
 
diff --git a/Engine/Source/Runtime/Experimental/Iris/Core/Private/Iris/ReplicationSystem/NetBlob/NetBlobManager.cpp b/Engine/Source/Runtime/Experimental/Iris/Core/Private/Iris/ReplicationSystem/NetBlob/NetBlobManager.cpp
--- a/Engine/Source/Runtime/Experimental/Iris/Core/Private/Iris/ReplicationSystem/NetBlob/NetBlobManager.cpp
+++ b/Engine/Source/Runtime/Experimental/Iris/Core/Private/Iris/ReplicationSystem/NetBlob/NetBlobManager.cpp
@@ -241,9 +241,15 @@ bool FNetBlobManager::GetOwnerAndSubObjectIndicesFromHandle(FNetHandle NetHandle
 	return true;
 }
 
-void FNetBlobManager::ProcessNetObjectAttachmentSendQueue()
+void FNetBlobManager::ProcessNetObjectAttachmentSendQueue(EProcessMode ProcessMode)
 {
-	AttachmentSendQueue.ProcessQueue(Connections);
+	AttachmentSendQueue.PrepareProcessQueue(Connections, NetHandleManager);
+	AttachmentSendQueue.ProcessQueue(ProcessMode);
+}
+
+void FNetBlobManager::ResetNetObjectAttachmentSendQueue()
+{
+	AttachmentSendQueue.ResetProcessQueue();
 }
 
 void FNetBlobManager::AddConnection(uint32 ConnectionId)
@@ -333,29 +339,77 @@ void FNetBlobManager::FNetObjectAttachmentSendQueue::Enqueue(FInternalNetHandle
 	bHasMulticastAttachments = true;
 }
 
-void FNetBlobManager::FNetObjectAttachmentSendQueue::ProcessQueue(FReplicationConnections* InConnections)
+void FNetBlobManager::FNetObjectAttachmentSendQueue::PrepareProcessQueue(FReplicationConnections* InConnections, const FNetHandleManager* NetHandleManager)
 {
-	if (AttachmentQueue.Num() <= 0)
+	if (ProcessContext.IsValid())
 	{
 		return;
 	}
 
-	const FNetBitArray& ValidConnections = InConnections->GetValidConnections();
-	const FNetBitArrayView ReplicatingConnections = MakeNetBitArrayView(ValidConnections);
+	ProcessContext.Connections = InConnections;
+	ProcessContext.NetHandleManager = NetHandleManager;
+
+	ProcessContext.AttachmentsToObjectsGoingOutOfScope.Init(AttachmentQueue.Num());
+	ProcessContext.AttachmentsToObjectsInScope.Init(AttachmentQueue.Num());
 
-	uint32* ConnectionIds = nullptr;
-	uint32 ConnectionCount = 0;
+	if (AttachmentQueue.Num() <= 0)
+	{
+		return;
+	}
 
 	if (bHasMulticastAttachments)
 	{
-		ConnectionIds = static_cast<uint32*>(FMemory_Alloca(ReplicatingConnections.GetNumBits() * 4));
-		CA_ASSUME(ConnectionIds != nullptr);
-		ReplicatingConnections.ForAllSetBits([ConnectionIds, &ConnectionCount](uint32 Bit) { ConnectionIds[ConnectionCount++] = Bit; });
+		const FNetBitArray& ValidConnections = InConnections->GetValidConnections();
+		const FNetBitArrayView ReplicatingConnections = MakeNetBitArrayView(ValidConnections);
+
+		ProcessContext.ConnectionIds.SetNum(ReplicatingConnections.CountSetBits());
+		ReplicatingConnections.GetSetBitIndices(0, ~0, ProcessContext.ConnectionIds.GetData(), ProcessContext.ConnectionIds.Num());
 	}
 
-	TArray<TRefCountPtr<FNetBlob>> PartialNetBlobs;
+	// Figure out if we have any attachments to objects going out of scope.
+	const FNetBitArray& ScopableObjects = NetHandleManager->GetScopableInternalIndices();
+	const FNetBitArray& PrevScopableObjects = NetHandleManager->GetPrevFrameScopableInternalIndices();
+	
+	uint32 CurrentEntryIndex = 0U;
 	for (const FNetObjectAttachmentQueueEntry& Entry : MakeArrayView(AttachmentQueue))
 	{
+		const uint32 TargetInternalObjectIndex = Entry.SubObjectIndex != FNetHandleManager::InvalidInternalIndex ? Entry.SubObjectIndex : Entry.OwnerIndex;
+
+		const bool bIsAttachmentToObjectGoingOutOfScope = !ScopableObjects.GetBit(TargetInternalObjectIndex) && PrevScopableObjects.GetBit(TargetInternalObjectIndex);
+		if (bIsAttachmentToObjectGoingOutOfScope)
+		{
+			ProcessContext.AttachmentsToObjectsGoingOutOfScope.SetBit(CurrentEntryIndex);
+		}
+		else
+		{
+			ProcessContext.AttachmentsToObjectsInScope.SetBit(CurrentEntryIndex);
+		}
+		++CurrentEntryIndex;
+	}
+}
+
+void FNetBlobManager::FNetObjectAttachmentSendQueue::ResetProcessQueue()
+{
+	// Clear queue
+	AttachmentQueue.Reset();
+	bHasMulticastAttachments = false;	
+	ProcessContext.Reset();
+}
+
+void FNetBlobManager::FNetObjectAttachmentSendQueue::ProcessQueue(EProcessMode ProcessMode)
+{
+	const FNetBitArray& IndicesToProcess = (ProcessMode == EProcessMode::ProcessObjectsGoingOutOfScope) ? ProcessContext.AttachmentsToObjectsGoingOutOfScope : ProcessContext.AttachmentsToObjectsInScope;
+
+	TArray<TRefCountPtr<FNetBlob>> PartialNetBlobs;
+	const FNetObjectAttachmentQueueEntry* Entries = AttachmentQueue.GetData();
+
+	// Verify that we have not missed to prepare the process context
+	check(ProcessContext.IsValid() && IndicesToProcess.GetNumBits() == AttachmentQueue.Num());
+
+	IndicesToProcess.ForAllSetBits([this, Entries, &PartialNetBlobs](uint32 Index)
+	{
+		const FNetObjectAttachmentQueueEntry& Entry = Entries[Index];
+
 		PartialNetBlobs.Reset();
 
 		const TRefCountPtr<FNetObjectAttachment>& Attachment = Entry.Attachment;
@@ -367,7 +421,7 @@ void FNetBlobManager::FNetObjectAttachmentSendQueue::ProcessQueue(FReplicationCo
 		if (!(bMulticast && bHasConnectionSpecificSerialization) && !PreSerializeAndSplitNetBlob(Entry.ConnectionId, Attachment, PartialNetBlobs))
 		{
 			checkf(false, TEXT("Unable to split %s NetObjectAttachment."), (EnumHasAnyFlags(Attachment->GetCreationInfo().Flags, ENetBlobFlags::Reliable) ? TEXT("reliable") : TEXT("unreliable")));
-			continue;
+			return;
 		}
 
 		const bool bIsMultiPartAttachment = (PartialNetBlobs.Num() > 1);
@@ -381,10 +435,10 @@ void FNetBlobManager::FNetObjectAttachmentSendQueue::ProcessQueue(FReplicationCo
 
 			const FNetBlobCreationInfo& BlobCreationInfo = Attachment->GetCreationInfo();
 			const bool bIsReliableRPC = EnumHasAnyFlags(BlobCreationInfo.Flags, ENetBlobFlags::Reliable);
-			for (uint32 ConnectionId : MakeArrayView(ConnectionIds, ConnectionCount))
+			for (uint32 ConnectionId : MakeArrayView(ProcessContext.ConnectionIds))
 			{
 				// Objects won't be prioritized until there's a view so let's avoid queuing multicast attachments.
-				if (!bIsReliableRPC && InConnections->GetReplicationView(ConnectionId).Views.Num() <= 0)
+				if (!bIsReliableRPC && ProcessContext.Connections->GetReplicationView(ConnectionId).Views.Num() <= 0)
 				{
 					continue;
 				}
@@ -401,23 +455,19 @@ void FNetBlobManager::FNetObjectAttachmentSendQueue::ProcessQueue(FReplicationCo
 					AttachmentsView = MakeArrayView(PartialNetBlobs);
 				}
 
-				FReplicationConnection* Connection = InConnections->GetConnection(ConnectionId);
+				FReplicationConnection* Connection = ProcessContext.Connections->GetConnection(ConnectionId);
 				// We're only iterating over valid connections so the Connection pointer must be valid.
 				Connection->ReplicationWriter->QueueNetObjectAttachments(Entry.OwnerIndex, Entry.SubObjectIndex, AttachmentsView);
 			}
 		}
 		else
 		{
-			if (FReplicationConnection* Connection = InConnections->GetConnection(Entry.ConnectionId))
+			if (FReplicationConnection* Connection = ProcessContext.Connections->GetConnection(Entry.ConnectionId))
 			{
 				Connection->ReplicationWriter->QueueNetObjectAttachments(Entry.OwnerIndex, Entry.SubObjectIndex, AttachmentsView);
 			}
 		}
-	}
-
-	// Clear queue
-	AttachmentQueue.Reset();
-	bHasMulticastAttachments = false;
+	});
 }
 
 bool FNetBlobManager::FNetObjectAttachmentSendQueue::PreSerializeAndSplitNetBlob(uint32 ConnectionId, const TRefCountPtr<FNetObjectAttachment>& Attachment, TArray<TRefCountPtr<FNetBlob>>& OutPartialNetBlobs) const
diff --git a/Engine/Source/Runtime/Experimental/Iris/Core/Private/Iris/ReplicationSystem/NetBlob/NetBlobManager.h b/Engine/Source/Runtime/Experimental/Iris/Core/Private/Iris/ReplicationSystem/NetBlob/NetBlobManager.h
--- a/Engine/Source/Runtime/Experimental/Iris/Core/Private/Iris/ReplicationSystem/NetBlob/NetBlobManager.h
+++ b/Engine/Source/Runtime/Experimental/Iris/Core/Private/Iris/ReplicationSystem/NetBlob/NetBlobManager.h
@@ -46,7 +46,14 @@ public:
 	bool SendRPC(const UObject* Object, const UObject* SubObject, const UFunction* Function, const void* Parameters);
 	bool SendRPC(uint32 ConnectionId, const UObject* Object, const UObject* SubObject, const UFunction* Function, const void* Parameters);
 
-	void ProcessNetObjectAttachmentSendQueue();
+	enum class EProcessMode 
+	{
+		ProcessObjectsGoingOutOfScope,
+		ProcessObjectsInScope,
+	};
+
+	void ProcessNetObjectAttachmentSendQueue(EProcessMode ProcessMode);
+	void ResetNetObjectAttachmentSendQueue();
 
 	FNetBlobHandlerManager& GetNetBlobHandlerManager() { return BlobHandlerManager; }
 	const FNetBlobHandlerManager& GetNetBlobHandlerManager() const { return BlobHandlerManager; }
@@ -76,8 +83,10 @@ private:
 		// Multicast
 		void Enqueue(FInternalNetHandle OwnerIndex, FInternalNetHandle SubObjectIndex, const TRefCountPtr<FNetObjectAttachment>& Attachment);
 
-		void ProcessQueue(FReplicationConnections* Connections);
-
+		void PrepareProcessQueue(FReplicationConnections* InConnections, const FNetHandleManager* NetHandleManager);
+		void ProcessQueue(EProcessMode ProcessMode);
+		void ResetProcessQueue();
+	
 	private:
 		struct FNetObjectAttachmentQueueEntry
 		{
@@ -93,6 +102,25 @@ private:
 		FNetBlobManager* Manager;
 		FQueue AttachmentQueue;
 		bool bHasMulticastAttachments;
+
+		struct FProcessQueueContext
+		{
+			FNetBitArray AttachmentsToObjectsGoingOutOfScope;
+			FNetBitArray AttachmentsToObjectsInScope;
+
+			TArray<uint32> ConnectionIds;
+			FReplicationConnections* Connections = nullptr;
+			const FNetHandleManager* NetHandleManager = nullptr;
+
+			void Reset()
+			{
+				Connections = nullptr;
+				NetHandleManager = nullptr;
+			}
+
+			bool IsValid() const { return NetHandleManager != nullptr; }
+		};
+		FProcessQueueContext ProcessContext;
 	};
 
 
diff --git a/Engine/Source/Runtime/Experimental/Iris/Core/Private/Iris/ReplicationSystem/ReplicationBridge.cpp b/Engine/Source/Runtime/Experimental/Iris/Core/Private/Iris/ReplicationSystem/ReplicationBridge.cpp
--- a/Engine/Source/Runtime/Experimental/Iris/Core/Private/Iris/ReplicationSystem/ReplicationBridge.cpp
+++ b/Engine/Source/Runtime/Experimental/Iris/Core/Private/Iris/ReplicationSystem/ReplicationBridge.cpp
@@ -515,6 +515,12 @@ void UReplicationBridge::EndReplication(FNetHandle Handle, EEndReplicationFlags
 				InternalAddDestructionInfo(Handle, *Parameters);
 			}
 
+			// If the flush flag is set, we have to explicitly copy the final state and propagate pending changes as we do not expect the source object to be kept alive after the call to EndReplication
+			if (EnumHasAnyFlags(EndReplicationFlags, EEndReplicationFlags::Flush))
+			{
+				InternalFlushStateData(Handle);
+			}
+
 			DestroyLocalNetHandle(Handle);	
 		}
 	}
@@ -632,6 +638,79 @@ void UReplicationBridge::TearOff(FNetHandle Handle, bool bIsImmediate)
 	}	
 }
 
+void UReplicationBridge::InternalFlushStateData(UE::Net::FNetSerializationContext& SerializationContext, UE::Net::Private::FChangeMaskCache& ChangeMaskCache, UE::Net::FNetBitStreamWriter& ChangeMaskWriter, uint32 InternalObjectIndex)
+{
+	using namespace UE::Net;
+	using namespace UE::Net::Private;
+
+	if (InternalObjectIndex == FNetHandleManager::InvalidInternalIndex)
+	{
+		return;
+	}
+
+	FNetHandleManager::FReplicatedObjectData& ObjectData = NetHandleManager->GetReplicatedObjectDataNoCheck(InternalObjectIndex);
+
+	// Copy state data, if object already is torn off there is nothing to do
+	if (ObjectData.bTearOff)
+	{
+		return;
+	}
+
+	for (FInternalNetHandle SubObjectInternalIndex : NetHandleManager->GetChildSubObjects(InternalObjectIndex))
+	{
+		InternalFlushStateData(SerializationContext, ChangeMaskCache, ChangeMaskWriter, SubObjectInternalIndex);
+	}
+
+	if (EnumHasAnyFlags(ObjectData.InstanceProtocol->InstanceTraits, EReplicationInstanceProtocolTraits::NeedsPoll | EReplicationInstanceProtocolTraits::NeedsPreSendUpdate))
+	{
+		CallPreSendUpdateSingleHandle(ObjectData.Handle);
+	} 
+
+	FReplicationInstanceOperationsInternal::CopyObjectStateData(ChangeMaskWriter, ChangeMaskCache, *NetHandleManager, SerializationContext, InternalObjectIndex);
+}
+
+void UReplicationBridge::InternalFlushStateData(FNetHandle Handle)
+{
+	using namespace UE::Net;
+	using namespace UE::Net::Private;
+
+	IRIS_PROFILER_SCOPE(InternalFlushStateData);
+
+	const uint32 InternalObjectIndex = NetHandleManager->GetInternalIndex(Handle);
+	if (InternalObjectIndex == FNetHandleManager::InvalidInternalIndex)
+	{
+		return;
+	}
+
+	FNetHandleManager::FReplicatedObjectData& ObjectData = NetHandleManager->GetReplicatedObjectDataNoCheck(InternalObjectIndex);
+
+	FChangeMaskCache ChangeMaskCache;
+	FNetBitStreamWriter ChangeMaskWriter;
+
+	// Setup context
+	FNetSerializationContext SerializationContext;
+	FInternalNetSerializationContext InternalContext(ReplicationSystem);
+	SerializationContext.SetInternalContext(&InternalContext);
+
+	InternalFlushStateData(SerializationContext, ChangeMaskCache, ChangeMaskWriter, InternalObjectIndex);
+
+	// Iterate over connections and propagate dirty changemasks to all connections already scoping this object
+	if (ChangeMaskCache.Indices.Num() > 0)
+	{
+		FReplicationConnections& Connections = ReplicationSystem->GetReplicationSystemInternal()->GetConnections();
+
+		auto&& UpdateDirtyChangeMasks = [&Connections, &ChangeMaskCache](uint32 ConnectionId)
+		{
+			FReplicationConnection* Conn = Connections.GetConnection(ConnectionId);			
+
+			const bool bMarkForTearOff = false;
+			Conn->ReplicationWriter->ForceUpdateDirtyChangeMasks(ChangeMaskCache, FReplicationWriter::FlushFlags_FlushState, bMarkForTearOff);
+		};
+		const FNetBitArray& ValidConnections = Connections.GetValidConnections();
+		ValidConnections.ForAllSetBits(UpdateDirtyChangeMasks);		
+	}
+}
+
 void UReplicationBridge::InternalTearOff(FNetHandle Handle)
 {
 	using namespace UE::Net;
@@ -641,7 +720,7 @@ void UReplicationBridge::InternalTearOff(FNetHandle Handle)
 
 	const uint32 InternalObjectIndex = NetHandleManager->GetInternalIndex(Handle);
 
-	if (!InternalObjectIndex)
+	if (InternalObjectIndex == FNetHandleManager::InvalidInternalIndex)
 	{
 		return;
 	}
@@ -681,8 +760,9 @@ void UReplicationBridge::InternalTearOff(FNetHandle Handle)
 			// Iterate over connections and propagate dirty changemasks to all connections already scoping this object
 			auto UpdateDirtyChangeMasks = [&Connections, &ChangeMaskCache](uint32 ConnectionId)
 			{
-				FReplicationConnection* Conn = Connections.GetConnection(ConnectionId);			
-				Conn->ReplicationWriter->TearOffAndUpdateDirtyChangeMasks(ChangeMaskCache);
+				FReplicationConnection* Conn = Connections.GetConnection(ConnectionId);
+				const bool bMarkForTearOff = true;
+				Conn->ReplicationWriter->ForceUpdateDirtyChangeMasks(ChangeMaskCache, FReplicationWriter::FlushFlags_None, bMarkForTearOff);
 			};
 			const FNetBitArray& ValidConnections = Connections.GetValidConnections();
 			ValidConnections.ForAllSetBits(UpdateDirtyChangeMasks);		
diff --git a/Engine/Source/Runtime/Experimental/Iris/Core/Private/Iris/ReplicationSystem/ReplicationReader.cpp b/Engine/Source/Runtime/Experimental/Iris/Core/Private/Iris/ReplicationSystem/ReplicationReader.cpp
--- a/Engine/Source/Runtime/Experimental/Iris/Core/Private/Iris/ReplicationSystem/ReplicationReader.cpp
+++ b/Engine/Source/Runtime/Experimental/Iris/Core/Private/Iris/ReplicationSystem/ReplicationReader.cpp
@@ -376,15 +376,10 @@ void FReplicationReader::ReadObject(FNetSerializationContext& Context)
 #endif
 
 	const FNetHandle IncompleteHandle = ReadNetHandleId(Reader);
-	const bool bTearOff = Reader.ReadBool();
-
-	bool bSubObjectPendingDestroy = false;
-	const bool bSubObjectPendingEndReplication = Reader.ReadBool();
-	if (bSubObjectPendingEndReplication)
-	{
-		UE_NET_TRACE_SCOPE(SubObjectPendingDestroy, Reader, Context.GetTraceCollector(), ENetTraceVerbosity::Trace);
-		bSubObjectPendingDestroy = Reader.ReadBool();
-	}
+	
+	// Read replicated destroy header if necessary
+	const bool bReadReplicatedDestroyHeader = !IsObjectIndexForOOBAttachment(IncompleteHandle.GetId());
+	const uint32 ReplicatedDestroyHeaderFlags = bReadReplicatedDestroyHeader ? Reader.ReadBits(ReplicatedDestroyHeaderFlags_BitCount) : ReplicatedDestroyHeaderFlags_None;
 
 	const bool bHasState = Reader.ReadBool();
 	uint32 NewBaselineIndex = FDeltaCompressionBaselineManager::InvalidBaselineIndex;
@@ -489,7 +484,7 @@ void FReplicationReader::ReadObject(FNetSerializationContext& Context)
 			InternalIndex = NetHandleManager->GetInternalIndex(IncompleteHandle);
 
 			// If this is a subobject that is being destroyed this was no error as we send destroy info for unconfirmed objects
-			if (!bSubObjectPendingEndReplication)
+			if (!!(ReplicatedDestroyHeaderFlags & ReplicatedDestroyHeaderFlags_EndReplication))
 			{
 				bHasErrors = InternalIndex == FNetHandleManager::InvalidInternalIndex;
 			}
@@ -510,9 +505,10 @@ void FReplicationReader::ReadObject(FNetSerializationContext& Context)
 		FDispatchObjectInfo& Info = ObjectsToDispatch[ObjectsToDispatchCount];
 		Info = FDispatchObjectInfo();
 
-		Info.bDestroy = bTearOff || bSubObjectPendingDestroy;
-		Info.bTearOff = bTearOff;
-		Info.bDeferredEndReplication = bTearOff || bSubObjectPendingEndReplication;
+		// Update info based on ReplicatedDestroyHeader
+		Info.bDestroy = !!(ReplicatedDestroyHeaderFlags & (ReplicatedDestroyHeaderFlags_TearOff | ReplicatedDestroyHeaderFlags_DestroyInstance));
+		Info.bTearOff = !!(ReplicatedDestroyHeaderFlags & ReplicatedDestroyHeaderFlags_TearOff);
+		Info.bDeferredEndReplication = !!(ReplicatedDestroyHeaderFlags & (ReplicatedDestroyHeaderFlags_TearOff | ReplicatedDestroyHeaderFlags_EndReplication));
 
 		if (bHasState)
 		{
diff --git a/Engine/Source/Runtime/Experimental/Iris/Core/Private/Iris/ReplicationSystem/ReplicationRecord.h b/Engine/Source/Runtime/Experimental/Iris/Core/Private/Iris/ReplicationSystem/ReplicationRecord.h
--- a/Engine/Source/Runtime/Experimental/Iris/Core/Private/Iris/ReplicationSystem/ReplicationRecord.h
+++ b/Engine/Source/Runtime/Experimental/Iris/Core/Private/Iris/ReplicationSystem/ReplicationRecord.h
@@ -34,11 +34,11 @@ public:
 		uint32 Index : 20;								// Index in the ReplicationInfo array, effectively identifying the object
 		uint32 ReplicatedObjectState : 5;				// Encode EReplicatedObjectState using as few bits as we can
 		uint32 HasChangeMask : 1;						// Do we have a changeMask? 
-		uint32 HasDirtySubObjects : 1;
 		uint32 HasAttachments : 1;						// If this flag is set there's an associated AttachmentRecord
 		uint32 WroteTearOff : 1;						// If this flag is set, we wrote TearOff
 		uint32 WroteDestroySubObject : 1;				// If this flag is set, we wrote DestroySubObject
 		uint32 NewBaselineIndex : 2;					// This is a new Baseline pending ack
+		uint32 UnusedFlag : 1;
 		ReplicationRecordIndex NextIndex;				// Points to next older record index
 		uint16 Padding : 16;
 	};
@@ -70,6 +70,7 @@ public:
 
 	// Get the info for the provided index
 	inline FRecordInfo* GetInfoForIndex(ReplicationRecordIndex Index);
+	inline const FRecordInfo* GetInfoForIndex(ReplicationRecordIndex Index) const;
 
 	// PeekInfo . If the info indicates there's an attachment record one should call DequeueAttachmentRecord() as well.
 	inline const FRecordInfo& PeekInfo() const { return RecordInfos.Peek(); }
@@ -125,7 +126,12 @@ FReplicationRecord::FReplicationRecord()
 
 FReplicationRecord::FRecordInfo* FReplicationRecord::GetInfoForIndex(ReplicationRecordIndex Index)
 {
-	return (Index != InvalidReplicationRecordIndex) ? &RecordInfos.PokeAtOffset((Index - FrontIndex) & ReplicationRecordIndexMask) : nullptr;
+	return (Index != InvalidReplicationRecordIndex) ? &RecordInfos.PokeAtOffsetNoCheck((Index - FrontIndex) & ReplicationRecordIndexMask) : nullptr;
+}
+
+const FReplicationRecord::FRecordInfo* FReplicationRecord::GetInfoForIndex(ReplicationRecordIndex Index) const
+{
+	return (Index != InvalidReplicationRecordIndex) ? &RecordInfos.PeekAtOffsetNoCheck((Index - FrontIndex) & ReplicationRecordIndexMask) : nullptr;
 }
 
 // $IRIS TODO. To simplify (?) we could have one public facing struct and one or more internal structs.
diff --git a/Engine/Source/Runtime/Experimental/Iris/Core/Private/Iris/ReplicationSystem/ReplicationSystem.cpp b/Engine/Source/Runtime/Experimental/Iris/Core/Private/Iris/ReplicationSystem/ReplicationSystem.cpp
--- a/Engine/Source/Runtime/Experimental/Iris/Core/Private/Iris/ReplicationSystem/ReplicationSystem.cpp
+++ b/Engine/Source/Runtime/Experimental/Iris/Core/Private/Iris/ReplicationSystem/ReplicationSystem.cpp
@@ -318,12 +318,18 @@ public:
 		UE_NET_TRACE_FRAME_STATSCOUNTER(ReplicationSystemId, ReplicationSystem.CopiedObjectCount, CopiedObjectCount, ENetTraceVerbosity::Trace);
 	}
 
-	void ProcessNetObjectAttachmentSendQueue()
+	void ProcessNetObjectAttachmentSendQueue(FNetBlobManager::EProcessMode ProcessMode)
 	{
 		IRIS_PROFILER_SCOPE(FReplicationSystem_ProcessNetObjectAttachmentSendQueue);
 
 		FNetBlobManager& NetBlobManager = ReplicationSystemInternal.GetNetBlobManager();
-		NetBlobManager.ProcessNetObjectAttachmentSendQueue();
+		NetBlobManager.ProcessNetObjectAttachmentSendQueue(ProcessMode);
+	}
+
+	void ResetNetObjectAttachmentSendQueue()
+	{
+		FNetBlobManager& NetBlobManager = ReplicationSystemInternal.GetNetBlobManager();
+		NetBlobManager.ResetNetObjectAttachmentSendQueue();
 	}
 
 	// send data
@@ -480,6 +486,9 @@ const UE::Net::Private::FReplicationSystemInternal* UReplicationSystem::GetRepli
 
 void UReplicationSystem::PreSendUpdate(float DeltaSeconds)
 {
+	using namespace UE::Net;
+	using namespace UE::Net::Private;
+
 	IRIS_PROFILER_SCOPE(FReplicationSystem_PreSendUpdate);
 
 #if !UE_BUILD_SHIPPING
@@ -490,14 +499,14 @@ void UReplicationSystem::PreSendUpdate(float DeltaSeconds)
 	}
 #endif
 
-	UE::Net::Private::FReplicationSystemInternal& InternalSys = Impl->ReplicationSystemInternal;
+	FReplicationSystemInternal& InternalSys = Impl->ReplicationSystemInternal;
 
 	// $IRIS TODO. There may be some throttling of connections to tick that we should take into account.
-	const UE::Net::FNetBitArrayView& ReplicatingConnections = MakeNetBitArrayView(Impl->ReplicationSystemInternal.GetConnections().GetValidConnections());
+	const FNetBitArrayView& ReplicatingConnections = MakeNetBitArrayView(Impl->ReplicationSystemInternal.GetConnections().GetValidConnections());
 
 #if UE_NET_IRIS_CSV_STATS && CSV_PROFILER
 	{
-		UE::Net::FNetSendStats& SendStats = InternalSys.GetSendStats();
+		FNetSendStats& SendStats = InternalSys.GetSendStats();
 		SendStats.Reset();
 		SendStats.SetNumberOfReplicatingConnections(ReplicatingConnections.CountSetBits());
 	}
@@ -522,9 +531,12 @@ void UReplicationSystem::PreSendUpdate(float DeltaSeconds)
 		// Copy dirty state data. We need this to happen before both filtering and prioritization
 		Impl->CopyDirtyStateData();
 
+		// We must process all attachments to objects going out of scope before we update the scope
+		Impl->ProcessNetObjectAttachmentSendQueue(FNetBlobManager::EProcessMode::ProcessObjectsGoingOutOfScope);
+
 		// Update filtering and scope for all connections
 		{
-			UE::Net::FNetBitArrayView DirtyObjects = InternalSys.GetDirtyNetObjectTracker().GetDirtyNetObjects();
+			FNetBitArrayView DirtyObjects = InternalSys.GetDirtyNetObjectTracker().GetDirtyNetObjects();
 			Impl->UpdateFiltering(DirtyObjects);
 		}
 
@@ -532,18 +544,19 @@ void UReplicationSystem::PreSendUpdate(float DeltaSeconds)
 		Impl->PropagateDirtyChanges();
 	}
 
-	// Forward attachments to the connections now that we know which objects are in scope
-	Impl->ProcessNetObjectAttachmentSendQueue();
+	// Forward attachments to the connections after scope update
+	Impl->ProcessNetObjectAttachmentSendQueue(FNetBlobManager::EProcessMode::ProcessObjectsInScope);
+	Impl->ResetNetObjectAttachmentSendQueue();
 
 	if (bAllowObjectReplication)
 	{
 		// Update object priorities
-		UE::Net::FNetBitArrayView DirtyObjects = InternalSys.GetDirtyNetObjectTracker().GetDirtyNetObjects();
+		FNetBitArrayView DirtyObjects = InternalSys.GetDirtyNetObjectTracker().GetDirtyNetObjects();
 		Impl->UpdatePrioritization(ReplicatingConnections, DirtyObjects);
 
 		// Delta compression preparations before send
 		{
-			UE::Net::Private::FDeltaCompressionBaselineManagerPreSendUpdateParams UpdateParams;
+			FDeltaCompressionBaselineManagerPreSendUpdateParams UpdateParams;
 			UpdateParams.ChangeMaskCache = &InternalSys.GetChangeMaskCache();
 			InternalSys.GetDeltaCompressionBaselineManager().PreSendUpdate(UpdateParams);
 		}
diff --git a/Engine/Source/Runtime/Experimental/Iris/Core/Private/Iris/ReplicationSystem/ReplicationTypes.h b/Engine/Source/Runtime/Experimental/Iris/Core/Private/Iris/ReplicationSystem/ReplicationTypes.h
--- a/Engine/Source/Runtime/Experimental/Iris/Core/Private/Iris/ReplicationSystem/ReplicationTypes.h
+++ b/Engine/Source/Runtime/Experimental/Iris/Core/Private/Iris/ReplicationSystem/ReplicationTypes.h
@@ -27,4 +27,13 @@ struct FReplicationParameters
 #endif
 };
 
+enum EReplicatedDestroyHeaderFlags : uint32
+{
+	ReplicatedDestroyHeaderFlags_None						= 0U,
+	ReplicatedDestroyHeaderFlags_TearOff					= 1U << 0U,
+	ReplicatedDestroyHeaderFlags_EndReplication				= ReplicatedDestroyHeaderFlags_TearOff << 1U,
+	ReplicatedDestroyHeaderFlags_DestroyInstance			= ReplicatedDestroyHeaderFlags_EndReplication << 1U,
+	ReplicatedDestroyHeaderFlags_BitCount					= 3U
+};
+
 }
diff --git a/Engine/Source/Runtime/Experimental/Iris/Core/Private/Iris/ReplicationSystem/ReplicationWriter.cpp b/Engine/Source/Runtime/Experimental/Iris/Core/Private/Iris/ReplicationSystem/ReplicationWriter.cpp
--- a/Engine/Source/Runtime/Experimental/Iris/Core/Private/Iris/ReplicationSystem/ReplicationWriter.cpp
+++ b/Engine/Source/Runtime/Experimental/Iris/Core/Private/Iris/ReplicationSystem/ReplicationWriter.cpp
@@ -77,10 +77,7 @@ const TCHAR* FReplicationWriter::LexToString(const EReplicatedObjectState State)
 		TEXT("PendingCreate"),
 		TEXT("WaitOnCreateConfirmation"),
 		TEXT("Created"),
-		TEXT("PendingFlush"),
 		TEXT("WaitOnFlush"),
-		TEXT("Flushed"),
-		TEXT("Hibernating"),
 		TEXT("PendingTearOff"),
 		TEXT("SubObjectPendingDestroy"),
 		TEXT("CancelPendingDestroy"),
@@ -116,12 +113,17 @@ void FReplicationWriter::FReplicationInfo::SetState(EReplicatedObjectState NewSt
 		break;
 		case EReplicatedObjectState::PendingTearOff:
 		{
-			check(CurrentState == EReplicatedObjectState::PendingTearOff || CurrentState == EReplicatedObjectState::WaitOnCreateConfirmation || CurrentState == EReplicatedObjectState::Created || CurrentState == EReplicatedObjectState::WaitOnDestroyConfirmation);
+			check(CurrentState == EReplicatedObjectState::PendingTearOff || CurrentState == EReplicatedObjectState::WaitOnFlush || CurrentState == EReplicatedObjectState::WaitOnCreateConfirmation || CurrentState == EReplicatedObjectState::Created || CurrentState == EReplicatedObjectState::WaitOnDestroyConfirmation);
 		}
 		break;
 		case EReplicatedObjectState::SubObjectPendingDestroy:
 		{
-			check(CurrentState == EReplicatedObjectState::PendingDestroy || CurrentState == EReplicatedObjectState::SubObjectPendingDestroy || CurrentState == EReplicatedObjectState::WaitOnCreateConfirmation || CurrentState == EReplicatedObjectState::Created || CurrentState == EReplicatedObjectState::WaitOnDestroyConfirmation);
+			check(CurrentState == EReplicatedObjectState::PendingDestroy || CurrentState == EReplicatedObjectState::SubObjectPendingDestroy || CurrentState == EReplicatedObjectState::WaitOnCreateConfirmation || CurrentState == EReplicatedObjectState::Created || CurrentState == EReplicatedObjectState::WaitOnFlush || CurrentState == EReplicatedObjectState::WaitOnDestroyConfirmation);
+		}
+		break;
+		case EReplicatedObjectState::WaitOnFlush:
+		{
+			check(CurrentState != EReplicatedObjectState::Invalid);
 		}
 		break;
 		case EReplicatedObjectState::PendingDestroy:
@@ -432,6 +434,7 @@ void FReplicationWriter::StartReplication(uint32 InternalIndex)
 	Info.IsDestructionInfo = bIsDestructionInfo;
 	Info.IsCreationConfirmed = 0U;
 	Info.TearOff = Data.bTearOff;
+	Info.FlushFlags = GetDefaultFlushFlags();
 	Info.SubObjectPendingDestroy = 0U;
 	Info.IsDeltaCompressionEnabled = BaselineManager->GetDeltaCompressionStatus(InternalIndex) == ENetObjectDeltaCompressionStatus::Allow ? 1U : 0U;
 	Info.LastAckedBaselineIndex = FDeltaCompressionBaselineManager::InvalidBaselineIndex;
@@ -510,6 +513,92 @@ void FReplicationWriter::WriteNetHandleId(FNetBitStreamWriter& Writer, FNetHandl
 	Writer.WriteBits(Handle.GetId(), FNetHandle::IdBits);
 }
 
+uint32 FReplicationWriter::GetDefaultFlushFlags() const
+{
+	// By default we currently always flush if we have pending reliable attachments when EndReplication is called for a NetObject.
+	return EFlushFlags::FlushFlags_FlushReliable;
+}
+
+uint32 FReplicationWriter::GetFlushStatus(uint32 InternalIndex, const FReplicationInfo& Info, uint32 InFlushFlags) const
+{
+	uint32 FlushFlags = EFlushFlags::FlushFlags_None;
+
+	if (InFlushFlags == EFlushFlags::FlushFlags_None)
+	{
+		return FlushFlags;
+	}
+
+	if (!!(InFlushFlags & EFlushFlags::FlushFlags_FlushState) && (Info.HasDirtyChangeMask || HasInFlightStateChanges(InternalIndex, Info) || IsObjectPartOfActiveHugeObject(InternalIndex, Info)))
+	{
+		FlushFlags |= EFlushFlags::FlushFlags_FlushState;
+	}
+
+	if (!!(InFlushFlags & EFlushFlags::FlushFlags_FlushReliable) && !Attachments.IsAllReliableSentAndAcked(ENetObjectAttachmentType::Normal, InternalIndex))
+	{
+		FlushFlags |= EFlushFlags::FlushFlags_FlushReliable;
+	}
+
+	if (!Info.IsSubObject && FlushFlags != FlushFlags_All)
+	{
+		// Check status of SubObjects as well.
+		for (uint32 SubObjectIndex : NetHandleManager->GetSubObjects(InternalIndex))
+		{
+			const FReplicationInfo& SubObjectInfo = GetReplicationInfo(SubObjectIndex);
+			FlushFlags |= GetFlushStatus(SubObjectIndex, SubObjectInfo, InFlushFlags);
+
+			if (FlushFlags == FlushFlags_All)
+			{
+				break;
+			}
+		}
+	}
+
+	return FlushFlags;
+}
+
+void FReplicationWriter::SetPendingDestroyOrSubObjectPendingDestroyState(uint32 InternalIndex, FReplicationInfo& Info)
+{
+	if (Info.IsSubObject)
+	{
+		// Subobject destroyed before its owner is explicitly replicated as state data.
+		FNetHandleManager::FReplicatedObjectData& ObjectData = NetHandleManager->GetReplicatedObjectDataNoCheck(InternalIndex);
+		if (ObjectData.IsSubObject())
+		{
+			// If owner is not pending destroy we mark the state of the SubObject to SubObjectPendingDestroy and mark owner as having dirty subobjects which will 
+			// destroy the subobject using the replicated state path of the owner.
+			FReplicationInfo& OwnerInfo = ReplicatedObjects[ObjectData.SubObjectRootIndex];
+			if (OwnerInfo.GetState() < EReplicatedObjectState::PendingDestroy)
+			{
+				ObjectsWithDirtyChanges.SetBit(ObjectData.SubObjectRootIndex);
+				OwnerInfo.HasDirtySubObjects = 1U;
+
+				SetState(InternalIndex, EReplicatedObjectState::SubObjectPendingDestroy);
+				ObjectsPendingDestroy.SetBit(InternalIndex);
+				Info.SubObjectPendingDestroy = 1U;
+				return;
+			}
+		}
+	}
+	else if (Info.HasDirtySubObjects)
+	{
+		// If the owner is destroyed, all subobjects in the EReplicatedObjectState::SubObjectPendingDestroy state must also be marked as PendingDestroy as owner no longer will be replicated
+		for (uint32 SubObjectIndex : NetHandleManager->GetSubObjects(InternalIndex))
+		{
+			FReplicationInfo& SubObjectInfo = GetReplicationInfo(SubObjectIndex);
+			if (SubObjectInfo.GetState() == EReplicatedObjectState::SubObjectPendingDestroy)
+			{
+				SubObjectInfo.SetState(EReplicatedObjectState::PendingDestroy);
+				SubObjectInfo.SubObjectPendingDestroy = 0U;
+			}
+		}
+	}
+				
+	ObjectsPendingDestroy.SetBit(InternalIndex);
+	ObjectsWithDirtyChanges.ClearBit(InternalIndex);
+	SetState(InternalIndex, EReplicatedObjectState::PendingDestroy);				
+	Info.HasDirtyChangeMask = 0U;
+}
+
 void FReplicationWriter::UpdateScope(const FNetBitArrayView& UpdatedScope)
 {
 	IRIS_PROFILER_SCOPE(FReplicationWriter_ScopeUpdate);
@@ -517,7 +606,6 @@ void FReplicationWriter::UpdateScope(const FNetBitArrayView& UpdatedScope)
 	auto NewObjectFunctor = [this](uint32 Index)
 	{
 		// We can only start replicating an object that is not currently replicated
-		// Later on we will support "cancelling" destroy objects that are flushing state data
 		FReplicationInfo& Info = GetReplicationInfo(Index);
 		const EReplicatedObjectState State = Info.GetState();
 
@@ -525,6 +613,15 @@ void FReplicationWriter::UpdateScope(const FNetBitArrayView& UpdatedScope)
 		{
 			StartReplication(Index);
 		}
+		else if (State == EReplicatedObjectState::WaitOnFlush)
+		{			
+			if (ensureAlwaysMsgf(!!Info.TearOff, TEXT("We cannot cancel flush for an object pending tearoff.")))
+			{
+				// If we are waiting on flush but are re-added to scope we reset flush flags to default.
+				Info.FlushFlags = GetDefaultFlushFlags();
+				SetState(Index, EReplicatedObjectState::Created);
+			}
+		}
 		else if (State == EReplicatedObjectState::WaitOnDestroyConfirmation || State == EReplicatedObjectState::CancelPendingDestroy)
 		{
 			// Need to clear the pending destroy bit or else the object will be masked out of ObjectsInScope.
@@ -591,50 +688,24 @@ void FReplicationWriter::UpdateScope(const FNetBitArrayView& UpdatedScope)
 				ObjectsPendingDestroy.SetBit(Index);
 				SetState(Index, EReplicatedObjectState::WaitOnDestroyConfirmation);
 			}
-			else if (State < EReplicatedObjectState::PendingDestroy)
+			else if (const uint32 FlushFlags = GetFlushStatus(Index, Info, Info.FlushFlags))
 			{
-				// Handle special inlined destroy path for subobjects
-				if (Info.IsSubObject)
-				{
-					FNetHandleManager::FReplicatedObjectData& ObjectData = NetHandleManager->GetReplicatedObjectDataNoCheck(Index);
-					if (ObjectData.IsSubObject())
-					{
-						// If owner is not pending destroy we mark it as dirty so that we can replicate subobject destruction properly
-						// We might get away with not doing this if owner or subobject does not have any unconfirmed changes in flight.
-						FReplicationInfo& OwnerInfo = ReplicatedObjects[ObjectData.SubObjectRootIndex];
-						if (OwnerInfo.GetState() < EReplicatedObjectState::PendingDestroy)
-						{
-							ObjectsWithDirtyChanges.SetBit(ObjectData.SubObjectRootIndex);
-							OwnerInfo.HasDirtySubObjects = 1U;
+				// Store info about what we need to flush
+				Info.FlushFlags = FlushFlags;
+				SetState(Index, EReplicatedObjectState::WaitOnFlush);
 
-							SetState(Index, EReplicatedObjectState::SubObjectPendingDestroy);
-							ObjectsPendingDestroy.SetBit(Index);
-							Info.SubObjectPendingDestroy = 1U;			
-			
-							return;
-						}
-					}
-				}
-				// Subobject owner should try to fixup any objects in SubObjectPendingDestroyState
-				else if (Info.HasDirtySubObjects)
+				// If we do not have any state data to flush we can clear the has dirty states flag
+				if ((FlushFlags & FlushFlags_FlushState) == 0U)
 				{
-					for (uint32 SubObjectIndex : NetHandleManager->GetSubObjects(Index))
-					{
-						FReplicationInfo& SubObjectInfo = GetReplicationInfo(SubObjectIndex);
-						if (SubObjectInfo.GetState() == EReplicatedObjectState::SubObjectPendingDestroy)
-						{
-							SubObjectInfo.SetState(EReplicatedObjectState::PendingDestroy);
-							SubObjectInfo.SubObjectPendingDestroy = 0U;
-						}
-					}
+					Info.HasDirtyChangeMask = 0U;
 				}
-				
 
-				SetState(Index, EReplicatedObjectState::PendingDestroy);
+				// Mark object as pending destroy so that we can poll the flush status in WriteObjectPendingDestroy
 				ObjectsPendingDestroy.SetBit(Index);
-				ObjectsWithDirtyChanges.ClearBit(Index);
-
-				Info.HasDirtyChangeMask = 0U;
+			}
+			else
+			{
+				SetPendingDestroyOrSubObjectPendingDestroyState(Index, Info);
 			}
 		}
 		else if (State == EReplicatedObjectState::PermanentlyDestroyed)
@@ -651,11 +722,13 @@ void FReplicationWriter::UpdateScope(const FNetBitArrayView& UpdatedScope)
 	ObjectsInScope.Combine(ObjectsPendingDestroy, FNetBitArrayBase::AndNotOp);
 }
 
-void FReplicationWriter::InternalUpdateDirtyChangeMasks(const FChangeMaskCache& CachedChangeMasks, uint32 bTearOff)
+void FReplicationWriter::InternalUpdateDirtyChangeMasks(const FChangeMaskCache& CachedChangeMasks, EFlushFlags ExtraFlushFlags, bool bMarkForTearOff)
 {
 	IRIS_PROFILER_SCOPE(FReplicationWriter_UpdateDirtyChangeMasks);
 
+	const uint32 MarkForTearOff = bMarkForTearOff ? 1U : 0U;
 	const ChangeMaskStorageType* StoragePtr = CachedChangeMasks.Storage.GetData();
+
 	for (const auto& Entry : CachedChangeMasks.Indices)
 	{
 		if (!ObjectsInScope.GetBit(Entry.InternalIndex))
@@ -669,7 +742,10 @@ void FReplicationWriter::InternalUpdateDirtyChangeMasks(const FChangeMaskCache&
 		if (Entry.bMarkSubObjectOwnerDirty == 0U)
 		{		
 			// Mark object for TearOff, that is that we will stop replication as soon as the tear-off is acknowledged
-			Info.TearOff = bTearOff;
+			Info.TearOff = MarkForTearOff;
+
+			// Update flush flags
+			Info.FlushFlags = Info.FlushFlags | ExtraFlushFlags;
 
 			// Merge in dirty changes
 			if (Entry.bHasDirtyChangeMask)
@@ -685,10 +761,6 @@ void FReplicationWriter::InternalUpdateDirtyChangeMasks(const FChangeMaskCache&
 				// Mark changemask as dirty
 				Info.HasDirtyChangeMask = 1U;
 			}
-
-			// In order to support flush behavior of tear off objects we should make sure to include everything in the final state https://jira.it.epicgames.com/browse/UENET-1079
-			// We should walk the list of any in-flight changes and include them to the changemask to ensure that
-			// all states are included with the tear-off
 		}
 		else
 		{
@@ -827,6 +899,12 @@ void FReplicationWriter::HandleDeliveredRecord(const FReplicationRecord::FRecord
 				else
 				{
 					SetState(InternalIndex, EReplicatedObjectState::Created);
+
+					// Tear-off is marked as a flush
+					if (Info.TearOff)
+					{
+						SetState(InternalIndex, EReplicatedObjectState::WaitOnFlush);
+					}
 				}
 			}
 			Info.IsCreationConfirmed = 1U;
@@ -891,6 +969,61 @@ void FReplicationWriter::HandleDeliveredRecord(const FReplicationRecord::FRecord
 	{
 		Attachments.OnPacketDelivered(ENetObjectAttachmentType::Normal, InternalIndex, AttachmentRecord);
 	}
+
+	// Must process WaitOnflush after attachments in order to correctly evaluate flush-status if needed
+	if (Info.GetState() == EReplicatedObjectState::WaitOnFlush)
+	{
+		bool bStillPendingFlush = false;
+		if (RecordInfo.HasChangeMask && !!(Info.FlushFlags & EFlushFlags::FlushFlags_FlushState))
+		{
+			bStillPendingFlush |= (Info.HasDirtyChangeMask || HasInFlightStateChanges(ReplicationRecord.GetInfoForIndex(RecordInfo.NextIndex)) || IsObjectPartOfActiveHugeObject(InternalIndex, Info));
+		}
+
+		if (RecordInfo.HasAttachments && !!(Info.FlushFlags & FlushFlags_FlushReliable))
+		{
+			bStillPendingFlush |= !Attachments.IsAllReliableSentAndAcked(ENetObjectAttachmentType::Normal, InternalIndex);
+		}
+
+		// This is a bit blunt as subobjects might be "acked" later but in this case it will be captured in WriteObjectsPendingDestroy
+		if (!bStillPendingFlush && !Info.IsSubObject)
+		{
+			// Check status of SubObjects as well.
+			for (uint32 SubObjectIndex : NetHandleManager->GetSubObjects(InternalIndex))
+			{
+				const FReplicationInfo& SubObjectInfo = GetReplicationInfo(SubObjectIndex);
+				if (GetFlushStatus(SubObjectIndex, SubObjectInfo, Info.FlushFlags) != FlushFlags_None)
+				{
+					bStillPendingFlush = true;
+					break;
+				}
+			}
+		}
+
+		if (!bStillPendingFlush)
+		{
+			if (Info.TearOff)
+			{
+				SetState(InternalIndex, EReplicatedObjectState::PendingTearOff);
+				ObjectsWithDirtyChanges.SetBit(InternalIndex);
+
+				// Must also mark owner dirty to make sure that we send the tearoff
+				if (Info.IsSubObject)
+				{
+					FNetHandleManager::FReplicatedObjectData& ObjectData = NetHandleManager->GetReplicatedObjectDataNoCheck(InternalIndex);
+					if (ObjectData.IsSubObject())
+					{
+						FReplicationInfo& OwnerInfo = ReplicatedObjects[ObjectData.SubObjectRootIndex];
+						ObjectsWithDirtyChanges.SetBit(ObjectData.SubObjectRootIndex);
+						OwnerInfo.HasDirtySubObjects = 1U;
+					}
+				}
+			}
+			else
+			{
+				SetPendingDestroyOrSubObjectPendingDestroyState(InternalIndex, Info);
+			}
+		}
+	}
 }
 
 void FReplicationWriter::HandleDiscardedRecord(const FReplicationRecord::FRecordInfo& RecordInfo, FReplicationInfo& Info, FNetObjectAttachmentsWriter::ReplicationRecord AttachmentRecord)
@@ -1024,11 +1157,9 @@ void FReplicationWriter::HandleDroppedRecord<FReplicationWriter::EReplicatedObje
 			// Give slight priority bump
 			SchedulingPriorities[InternalIndex] += FReplicationWriter::LostStatePriorityBump;
 
-			// Check uncached flag in order to catch if object was added as a subobject or dependent object after data was sent
-			const bool bIsSubObjectOrDependentObject = NetHandleManager->GetSubObjectInternalIndices().GetBit(InternalIndex);
-			if (bIsSubObjectOrDependentObject)
+			if (Info.IsSubObject)
 			{
-				// If this object is a subobject, mark owner dirty as well as subobjects always should be scheduled together with parent
+				// Mark owner dirty as well as subobjects only are scheduled together with owner
 				const FNetHandleManager::FReplicatedObjectData& ObjectData = NetHandleManager->GetReplicatedObjectData(InternalIndex);
 				uint32 SubObjectOwnerInternalIndex = ObjectData.SubObjectRootIndex;
 
@@ -1188,7 +1319,6 @@ void FReplicationWriter::HandleDroppedRecord(const FReplicationRecord::FRecordIn
 
 		// Object is created, update lost state data unless object are currently being flushed/teared-off or destroyed
 		case EReplicatedObjectState::Created:
-		case EReplicatedObjectState::PendingFlush:
 		case EReplicatedObjectState::WaitOnFlush:
 		{
 			HandleDroppedRecord<EReplicatedObjectState::Created>(CurrentState, RecordInfo, Info, AttachmentRecord);
@@ -1288,7 +1418,7 @@ void FReplicationWriter::CreateObjectRecord(const FNetBitArrayView* ChangeMask,
 	RecordInfo.ReplicatedObjectState = ObjectInfo.AttachmentType == ENetObjectAttachmentType::HugeObject ? uint8(EReplicatedObjectState::HugeObject) : (uint8)Info.GetState();
 	RecordInfo.HasChangeMask = ChangeMask ? 1U : 0U;
 	RecordInfo.HasAttachments = (ObjectInfo.AttachmentRecord != 0 ? 1U : 0U);
-	RecordInfo.WroteTearOff = Info.TearOff;
+	RecordInfo.WroteTearOff = ObjectInfo.bSentTearOff;
 	RecordInfo.WroteDestroySubObject = Info.SubObjectPendingDestroy;
 	
 	// If we wrote a new baseline we need to store it in the record
@@ -1354,6 +1484,17 @@ uint32 FReplicationWriter::WriteObjectsPendingDestroy(FNetSerializationContext&
 			continue;
 		}
 
+		if (Info.GetState() == EReplicatedObjectState::WaitOnFlush)
+		{
+			if (GetFlushStatus(InternalIndex, Info, Info.FlushFlags) != EFlushFlags::FlushFlags_None)
+			{
+				continue;
+			}
+			
+			// Object and subobjects are now flushed and can be destroyed
+			SetPendingDestroyOrSubObjectPendingDestroyState(InternalIndex, Info);
+		}
+		
 		if (ObjectData.IsSubObject())
 		{
 			if (Info.SubObjectPendingDestroy)
@@ -1396,7 +1537,7 @@ uint32 FReplicationWriter::WriteObjectsPendingDestroy(FNetSerializationContext&
 		check(Info.GetState() == EReplicatedObjectState::PendingDestroy);
 
 		// We do not support destroying an object that is currently being sent as a huge object.
-		if (IsDestroyObjectPartOfActiveHugeObject(InternalIndex, Info))
+		if (IsObjectPartOfActiveHugeObject(InternalIndex, Info))
 		{
 			UE_LOG(LogIris, Warning, TEXT("Skipping writing destroy for object ( InternalIndex: %u ) which is part of active huge object."), InternalIndex);
 			bWroteAllDestroyedObjects = false;
@@ -1465,8 +1606,7 @@ bool FReplicationWriter::CanSendObject(uint32 InternalIndex) const
 	// Don't send more recent state that could arrive before the huge state. We only need to check the parent.
 	if (IsActiveHugeObject(InternalIndex))
 	{
-		const bool bIsSubObjectOrDependentObject = NetHandleManager->GetSubObjectInternalIndices().GetBit(InternalIndex);
-		if (!bIsSubObjectOrDependentObject)
+		if (!Info.IsSubObject)
 		{
 			return false;
 		}
@@ -1584,12 +1724,15 @@ FReplicationWriter::EWriteObjectStatus FReplicationWriter::WriteObjectInBatch(FN
 	const bool bWriteAttachments = bHasAttachments & !!(WriteObjectFlags & EWriteObjectFlag::WriteObjectFlag_Attachments);
 	BatchEntry.bHasUnsentAttachments = bHasAttachments;
 
+	// Check if we must defer tearoff until after flush
+	const bool bSentTearOff = Info.TearOff && (GetFlushStatus(InternalIndex, Info, Info.FlushFlags) == EFlushFlags::FlushFlags_None);
+
 	Context.SetIsInitState(bIsInitialState);
 
 	// Update flag since dependent objects might change their status after creation
 	Info.IsSubObject = ObjectData.IsSubObject();
 
-	if (bHasState | bWriteAttachments | Info.TearOff | Info.SubObjectPendingDestroy)
+	if (bHasState | bWriteAttachments | bSentTearOff | Info.SubObjectPendingDestroy)
 	{
 		const FNetHandle NetHandle = ObjectData.Handle;
 #if UE_NET_TRACE_ENABLED
@@ -1612,16 +1755,15 @@ FReplicationWriter::EWriteObjectStatus FReplicationWriter::WriteObjectInBatch(FN
 			WriteNetHandleId(Writer, NetHandle);
 		}
 
-		// TearOff
-		Writer.WriteBool(Info.TearOff);
+		// Store position of header bits
+		const uint32 ReplicatedDestroyHeaderBitPos = Writer.GetPosBits();
 
-		// SubObject pending destroy
-		if (Writer.WriteBool(Info.SubObjectPendingDestroy))
+		// We only need to write this for actual replicated objects
+		const bool bWriteReplicatedDestroyHeader = !bIsObjectIndexForAttachment;
+		if (bWriteReplicatedDestroyHeader)
 		{
-			UE_NET_TRACE_SCOPE(SubObjectPendingDestroy, Writer, Context.GetTraceCollector(), ENetTraceVerbosity::Trace);
-			// Write bit indicating if the static instance should be destroyed or not (could skip the bit for dynamic objects)
-			const bool bShouldDestroyInstance = ObjectData.Handle.IsDynamic() || NetHandleManager->GetIsDestroyedStartupObject(InternalIndex);
-			Writer.WriteBool(bShouldDestroyInstance);
+			// Write destroy header bits, we always want to write the same number of bits to be able to update the header afterwards we know what data made it into the packet
+			Writer.WriteBits(0U, ReplicatedDestroyHeaderFlags_BitCount);
 		}
 
 		if (Writer.WriteBool(bHasState))
@@ -1825,6 +1967,29 @@ FReplicationWriter::EWriteObjectStatus FReplicationWriter::WriteObjectInBatch(FN
 		{
 			return EWriteObjectStatus::BitStreamOverflow;
 		}
+
+		if (bWriteReplicatedDestroyHeader)
+		{
+			// Rewrite destroy header if necessary	
+			if (bSentTearOff || Info.SubObjectPendingDestroy)
+			{
+				uint32 ReplicatedDestroyHeaderFlags = 0U;
+
+				// TearOff			
+				ReplicatedDestroyHeaderFlags |= bSentTearOff ? ReplicatedDestroyHeaderFlags_TearOff : ReplicatedDestroyHeaderFlags_None;
+
+				// Write SubObject destroy
+				if (Info.SubObjectPendingDestroy)
+				{
+					ReplicatedDestroyHeaderFlags |= ReplicatedDestroyHeaderFlags_EndReplication;
+					const bool bShouldDestroyInstance = ObjectData.Handle.IsDynamic() || NetHandleManager->GetIsDestroyedStartupObject(InternalIndex);					
+					ReplicatedDestroyHeaderFlags |= bShouldDestroyInstance ? ReplicatedDestroyHeaderFlags_DestroyInstance : ReplicatedDestroyHeaderFlags_None;
+				}
+
+				FNetBitStreamWriteScope WriteScope(Writer, ReplicatedDestroyHeaderBitPos);
+				Writer.WriteBits(ReplicatedDestroyHeaderFlags, ReplicatedDestroyHeaderFlags_BitCount);
+			}
+		}
 	}
 
 	// Success so far. Fill in batch entry. Keep index to update info later as the array can resize.
@@ -1836,7 +2001,7 @@ FReplicationWriter::EWriteObjectStatus FReplicationWriter::WriteObjectInBatch(FN
 		FinalBatchEntry.bIsInitialState = bIsInitialState;
 		FinalBatchEntry.InternalIndex = InternalIndex;
 		FinalBatchEntry.bHasDirtySubObjects = false;
-		FinalBatchEntry.bSentTearOff = Info.TearOff;
+		FinalBatchEntry.bSentTearOff = bSentTearOff;
 		FinalBatchEntry.bSentDestroySubObject = Info.SubObjectPendingDestroy;
 		FinalBatchEntry.NewBaselineIndex = CreatedBaselineIndex;
 	}
@@ -2327,10 +2492,17 @@ int FReplicationWriter::HandleObjectBatchSuccess(const FBatchInfo& BatchInfo, FR
 		{
 			SetState(BatchObjectInfo.InternalIndex, EReplicatedObjectState::WaitOnCreateConfirmation);
 		}
-		else if (BatchObjectInfo.bSentTearOff)
+		else if (Info.TearOff)
 		{
-			SetState(BatchObjectInfo.InternalIndex, EReplicatedObjectState::PendingTearOff);
-			SetState(BatchObjectInfo.InternalIndex, EReplicatedObjectState::WaitOnDestroyConfirmation);
+			if (BatchObjectInfo.bSentTearOff)
+			{
+				SetState(BatchObjectInfo.InternalIndex, EReplicatedObjectState::PendingTearOff);
+				SetState(BatchObjectInfo.InternalIndex, EReplicatedObjectState::WaitOnDestroyConfirmation);
+			}
+			else
+			{
+				SetState(BatchObjectInfo.InternalIndex, EReplicatedObjectState::WaitOnFlush);
+			}
 		}
 		else if (BatchObjectInfo.bSentDestroySubObject)
 		{			
@@ -2701,6 +2873,28 @@ void FReplicationWriter::InvalidateBaseline(uint32 InternalIndex, FReplicationIn
 	Info.PendingBaselineIndex = FDeltaCompressionBaselineManager::InvalidBaselineIndex;
 }
 
+bool FReplicationWriter::HasInFlightStateChanges(const FReplicationRecord::FRecordInfo* RecordInfo) const
+{
+	while (RecordInfo)
+	{
+		if (RecordInfo->HasChangeMask)
+		{
+			return true;
+		}
+		RecordInfo = ReplicationRecord.GetInfoForIndex(RecordInfo->NextIndex);
+	};
+
+	return false;
+}
+
+bool FReplicationWriter::HasInFlightStateChanges(uint32 InternalIndex, const FReplicationInfo& Info) const
+{
+	const FReplicationRecord::FRecordInfoList& RecordInfoList = ReplicatedObjectsRecordInfoLists[InternalIndex];
+	const FReplicationRecord::FRecordInfo* CurrentRecordInfo = ReplicationRecord.GetInfoForIndex(RecordInfoList.FirstRecordIndex);
+
+	return HasInFlightStateChanges(CurrentRecordInfo);
+}
+
 bool FReplicationWriter::PatchupObjectChangeMaskWithInflightChanges(uint32 InternalIndex, FReplicationInfo& Info)
 {
 	bool bInFlightChangesAdded = false;
@@ -2732,7 +2926,7 @@ void FReplicationWriter::SetNetExports(FNetExports& InNetExports)
 	NetExports = &InNetExports;
 }
 
-bool FReplicationWriter::IsDestroyObjectPartOfActiveHugeObject(uint32 InternalIndex, const FReplicationInfo& Info) const
+bool FReplicationWriter::IsObjectPartOfActiveHugeObject(uint32 InternalIndex, const FReplicationInfo& Info) const
 {
 	if (HugeObjectContext.SendStatus == EHugeObjectSendStatus::Idle)
 	{
diff --git a/Engine/Source/Runtime/Experimental/Iris/Core/Private/Iris/ReplicationSystem/ReplicationWriter.h b/Engine/Source/Runtime/Experimental/Iris/Core/Private/Iris/ReplicationSystem/ReplicationWriter.h
--- a/Engine/Source/Runtime/Experimental/Iris/Core/Private/Iris/ReplicationSystem/ReplicationWriter.h
+++ b/Engine/Source/Runtime/Experimental/Iris/Core/Private/Iris/ReplicationSystem/ReplicationWriter.h
@@ -15,12 +15,14 @@
 #include "Iris/ReplicationSystem/ObjectReferenceCache.h"
 #include "Iris/Stats/NetStats.h"
 #include "Containers/Array.h"
+#include "Misc/EnumClassFlags.h"
 
 // Forward declaration
 class UReplicationSystem;
 class UNetObjectBlobHandler;
 class UPartialNetObjectAttachmentHandler;
 class UReplicationBridge;
+
 namespace UE::Net
 {
 	class FNetBitStreamReader;
@@ -28,6 +30,7 @@ namespace UE::Net
 	class FNetObjectAttachment;
 	class FNetSerializationContext;
 	struct FReplicationProtocol;
+	
 	namespace Private
 	{
 		struct FChangeMaskCache;
@@ -67,14 +70,11 @@ public:
 		AttachmentToObjectNotInScope,	// Special state for object index used for sending attachments to remote owned objects
 		HugeObject,						// Special state for object index used for sending parts of huge object payloads
 
-		// 
+		// Normal states
 		PendingCreate,					// Not yet created, no data in flight
 		WaitOnCreateConfirmation,		// Waiting for confirmation from remote, we can send state data, but if we do we must also include creation info until the object is Created
-		Created,						// Confirmed by remote
-		PendingFlush,					// This object should be flushed
+		Created,						// Confirmed by remote, this is the normal replicating state
 		WaitOnFlush,					// Pending flush, we are waiting for all in-flight state data to be acknowledged
-		Flushed,						// No more data in flight we can now hibernate or destroy
-		Hibernating,					// Hibernating, if an object is flushed it can be set to hibernate/be dormant
 		PendingTearOff,					// TearOff should be sent
 		SubObjectPendingDestroy,		// SubObject destroy should be sent
 		CancelPendingDestroy,			// Destroy was sent but object wants to start replicating again
@@ -90,6 +90,15 @@ public:
 
 	const TCHAR* LexToString(const EReplicatedObjectState State);
 
+	enum EFlushFlags : uint32
+	{
+		FlushFlags_None				= 0U,
+		FlushFlags_FlushState		= 1U << 0U,											// Make sure that all current state data is acknowledged before we stop replicating the object
+		FlushFlags_FlushReliable	= FlushFlags_FlushState << 1U,						// Make sure that all enqueued Reliable RPCs are delivered before we stop replicating the object
+		FlushFlags_All				= FlushFlags_FlushState | FlushFlags_FlushState,
+		FlushFlags_Default			= FlushFlags_FlushReliable,
+	};
+
 	// Bookkeeping info required for a replicated object
 	// Keep as small as possible since there is one per replicated object per connection
 	// Changemask can and will most likely be replaced by a smaller index into a pool to reduce overhead
@@ -114,10 +123,11 @@ public:
 				uint64 IsCreationConfirmed : 1;							// We know that this object has been created on the receiving end
 				uint64 TearOff : 1;										// This object should be torn off
 				uint64 SubObjectPendingDestroy : 1;						// This object is a subobject that should be destroyed when we replicate owner
-				uint64 IsDeltaCompressionEnabled : 1;
-				uint64 LastAckedBaselineIndex : 2;
-				uint64 PendingBaselineIndex : 2;
-				uint64 Padding : 29;
+				uint64 IsDeltaCompressionEnabled : 1;					// Set to 1 if deltacompression is enabled for this object
+				uint64 LastAckedBaselineIndex : 2;						// Last acknowledged baseline index which we can use for deltacompresion
+				uint64 PendingBaselineIndex : 2;						// Baseline index pending acknowledgment from client
+				uint64 FlushFlags : 2;									// Flags indicating what we are waiting for when flushing
+				uint64 Padding : 27;
 			};
 		};
 
@@ -141,13 +151,14 @@ public:
 	// Update new or existing/destroyed 
 	void UpdateScope(const FNetBitArrayView& ScopedObjects);
 
-	// Mark objects as tearoff and propagate dirty changemasks
-	void TearOffAndUpdateDirtyChangeMasks(const FChangeMaskCache& CachedChangeMasks) { InternalUpdateDirtyChangeMasks(CachedChangeMasks, 1U); }
+	// Force update DirtyChangeMasks and mark objects for flush and/or tearoff depending on flags
+	void ForceUpdateDirtyChangeMasks(const FChangeMaskCache& CachedChangeMasks, EFlushFlags ExtraFlushFlags, bool bMarkForTearOff) { InternalUpdateDirtyChangeMasks(CachedChangeMasks, ExtraFlushFlags, bMarkForTearOff); }
 
+	// Called if an object first being teared off and then explicitly destroyed before it has been removed from scope
 	void NotifyDestroyedObjectPendingTearOff(FInternalNetHandle ObjectInternalIndex);
 
 	// Propagate dirty changemasks
-	void UpdateDirtyChangeMasks(const FChangeMaskCache& CachedChangeMasks) { InternalUpdateDirtyChangeMasks(CachedChangeMasks, 0U); }
+	void UpdateDirtyChangeMasks(const FChangeMaskCache& CachedChangeMasks) { InternalUpdateDirtyChangeMasks(CachedChangeMasks, EFlushFlags::FlushFlags_None, false); }
 
 	// Returns objects that are in need of a priority update. It could be dirty, new or objects in need of resending.
 	const FNetBitArray& GetObjectsRequiringPriorityUpdate() const;
@@ -182,7 +193,7 @@ private:
 	};
 
 	// Propagate dirty changemasks
-	void InternalUpdateDirtyChangeMasks(const FChangeMaskCache& CachedChangeMasks, uint32 bTearOff);
+	void InternalUpdateDirtyChangeMasks(const FChangeMaskCache& CachedChangeMasks, EFlushFlags ExtraFlushFlags, bool bTearOff);
 
 	struct FScheduleObjectInfo
 	{
@@ -326,6 +337,10 @@ private:
 
 private:
 
+	uint32 GetDefaultFlushFlags() const;
+	uint32 GetFlushStatus(uint32 InternalIndex, const FReplicationInfo& Info, uint32 InFlushFlags = EFlushFlags::FlushFlags_Default) const;
+	void SetPendingDestroyOrSubObjectPendingDestroyState(uint32 Index, FReplicationInfo& Info);
+
 	bool IsObjectIndexForOOBAttachment(uint32 InternalIndex) const { return InternalIndex == ObjectIndexForOOBAttachment; }
 
 	void GetInitialChangeMask(ChangeMaskStorageType* ChangeMaskStorage, const FReplicationProtocol* Protocol);
@@ -346,7 +361,8 @@ private:
 	// Write index part of handle
 	void WriteNetHandleId(FNetBitStreamWriter& Writer, FNetHandle Handle);
 		
-	// 
+	// Create new ObjectRecord
+	// Note: be aware that it will allocate a copy of the ChangeMask that needs to be handled if the record is not Committed
 	void CreateObjectRecord(const FNetBitArrayView* ChangeMask, const FReplicationInfo& Info, const FBatchObjectInfo& ObjectInfo, FObjectRecord& OutRecord);
 	
 	// Push and link info for written object to ReplicationRecord
@@ -408,13 +424,21 @@ private:
 	// Invalidate all inflight baseline information
 	void InvalidateBaseline(uint32 InternalIndex, FReplicationInfo& Info);
 
+	// Returns true if the record chain starting from the provided RecordInfo contains any records with statedata
+	// Note: Does not check if it is part of a huge object
+	bool HasInFlightStateChanges(uint32 InternalIndex, const FReplicationInfo& Info) const;
+
+	// Returns true if object has pending state changes in flight
+	// Note: Does not check if it is part of a huge object
+	bool HasInFlightStateChanges(const FReplicationRecord::FRecordInfo* RecordInfo) const;
+
 	// Returns true object and subobjects can be created on remote
 	bool CanSendObject(uint32 InternalIndex) const;
 
 	inline bool IsInitialState(const EReplicatedObjectState State) const { return State == EReplicatedObjectState::PendingCreate || (bHighPrioCreate && State == EReplicatedObjectState::WaitOnCreateConfirmation); }
 
 	bool IsActiveHugeObject(uint32 InternalIndex) const { return HugeObjectContext.InternalIndex == InternalIndex && HugeObjectContext.SendStatus != EHugeObjectSendStatus::Idle; }
-	bool IsDestroyObjectPartOfActiveHugeObject(uint32 InternalIndex, const FReplicationInfo& Info) const;
+	bool IsObjectPartOfActiveHugeObject(uint32 InternalIndex, const FReplicationInfo& Info) const;
 
 	void ClearHugeObjectContext(FHugeObjectContext& Context) const;
 
diff --git a/Engine/Source/Runtime/Experimental/Iris/Core/Public/Iris/ReplicationSystem/ReplicationBridge.h b/Engine/Source/Runtime/Experimental/Iris/Core/Public/Iris/ReplicationSystem/ReplicationBridge.h
--- a/Engine/Source/Runtime/Experimental/Iris/Core/Public/Iris/ReplicationSystem/ReplicationBridge.h
+++ b/Engine/Source/Runtime/Experimental/Iris/Core/Public/Iris/ReplicationSystem/ReplicationBridge.h
@@ -38,6 +38,7 @@ namespace UE::Net
 		class FReplicationSystemImpl;
 		class FReplicationSystemInternal;
 		class FReplicationWriter;
+		struct FChangeMaskCache;
 	}
 
 	typedef TArray<FNetDependencyInfo, TInlineAllocator<32> > FNetDependencyInfoArray;
@@ -224,6 +225,12 @@ private:
 	void CallDetachInstanceFromRemote(FNetHandle Handle, bool bTearOff, bool bShouldDestroyInstance);
 
 private:
+
+	void InternalFlushStateData(UE::Net::FNetSerializationContext& SerializationContext, UE::Net::Private::FChangeMaskCache& ChangeMaskCache, UE::Net::FNetBitStreamWriter& ChangeMaskWriter, uint32 InternalObjectIndex);
+
+	// Internal method to copy state data for Handle
+	void InternalFlushStateData(FNetHandle Handle);
+
 	// Internal method to copy state data for Handle and any SubObjects and mark them as being torn-off
 	void InternalTearOff(FNetHandle OwnerHandle);
 
