Optimization : AABBTree:  Stack allocation to navigate the tree

diff --git a/Engine/Source/Runtime/Experimental/Chaos/Public/Chaos/AABBTree.h b/Engine/Source/Runtime/Experimental/Chaos/Public/Chaos/AABBTree.h
--- a/Engine/Source/Runtime/Experimental/Chaos/Public/Chaos/AABBTree.h
+++ b/Engine/Source/Runtime/Experimental/Chaos/Public/Chaos/AABBTree.h
@@ -2503,7 +2503,7 @@ private:
 				{
 					TSpatialVisitorData<TPayloadType> VisitData(Elem.Payload, true, InstanceBounds);
 					bool bContinue;
-					if (Query == EAABBQueryType::Overlap)
+					if constexpr (Query == EAABBQueryType::Overlap)
 					{
 						bContinue = Visitor.VisitOverlap(VisitData);
 					}
@@ -2527,15 +2527,15 @@ private:
 
 				if (DirtyElementGridEnabled() && CellHashToFlatArray.Num() > 0)
 				{
-					if (Query == EAABBQueryType::Overlap)
+					if constexpr (Query == EAABBQueryType::Overlap)
 					{
 						bUseGrid = !TooManyOverlapQueryCells(QueryBounds, DirtyElementGridCellSizeInv, DirtyElementMaxGridCellQueryCount);
 					}
-					else if (Query == EAABBQueryType::Raycast)
+					else if constexpr (Query == EAABBQueryType::Raycast)
 					{
 						bUseGrid = !TooManyRaycastQueryCells(Start, CurData.Dir, CurData.CurrentLength, DirtyElementGridCellSizeInv, DirtyElementMaxGridCellQueryCount);
 					}
-					else if (Query == EAABBQueryType::Sweep)
+					else if constexpr (Query == EAABBQueryType::Sweep)
 					{
 						bUseGrid = !TooManySweepQueryCells(QueryHalfExtents, Start, CurData.Dir, CurData.CurrentLength, DirtyElementGridCellSizeInv, DirtyElementMaxGridCellQueryCount);
 					}
@@ -2556,15 +2556,15 @@ private:
 						return true;
 					};
 
-					if (Query == EAABBQueryType::Overlap)
+					if constexpr (Query == EAABBQueryType::Overlap)
 					{
 						DoForOverlappedCells(QueryBounds, DirtyElementGridCellSize, DirtyElementGridCellSizeInv, AddHashEntry);
 					}
-					else if (Query == EAABBQueryType::Raycast)
+					else if constexpr (Query == EAABBQueryType::Raycast)
 					{
 						DoForRaycastIntersectCells(Start, CurData.Dir, CurData.CurrentLength, DirtyElementGridCellSize, DirtyElementGridCellSizeInv, AddHashEntry);
 					}
-					else if (Query == EAABBQueryType::Sweep)
+					else if constexpr (Query == EAABBQueryType::Sweep)
 					{
 						DoForSweepIntersectCells(QueryHalfExtents, Start, CurData.Dir, CurData.CurrentLength, DirtyElementGridCellSize, DirtyElementGridCellSizeInv,
 							[&](FReal X, FReal Y)
@@ -2615,18 +2615,21 @@ private:
 			}
 		}
 		
-		TArray<FNodeQueueEntry> NodeStack;
+		constexpr int32 MaxNodeStackNum = 255;
+		check(MaxTreeDepth + 2 <= MaxNodeStackNum);
+		FNodeQueueEntry NodeStack[MaxNodeStackNum];
+		int32 NodeStackNum = 0;
 		if (bDynamicTree)
 		{
 
 			if (RootNode != INDEX_NONE)
 			{
-				NodeStack.Add(FNodeQueueEntry{ RootNode, 0 });
+				NodeStack[NodeStackNum++] = FNodeQueueEntry{ RootNode, 0 };
 			}
 		}
 		else if (Nodes.Num())
 		{
-			NodeStack.Add(FNodeQueueEntry{ 0, 0 });
+			NodeStack[NodeStackNum++] = FNodeQueueEntry{ 0, 0 };
 		}
 
 // Slow debug code
@@ -2638,7 +2641,7 @@ private:
 //		}
 //#endif
 
-		while (NodeStack.Num())
+		while (NodeStackNum)
 		{
 			PHYSICS_CSV_SCOPED_VERY_EXPENSIVE(PhysicsVerbose, QueryImp_NodeTraverse);
 
@@ -2649,7 +2652,7 @@ private:
 //			}
 //#endif
 
-			const FNodeQueueEntry NodeEntry = NodeStack.Pop(false);
+			const FNodeQueueEntry NodeEntry = NodeStack[--NodeStackNum];
 			if (Query != EAABBQueryType::Overlap)
 			{
 				if (NodeEntry.TOI > CurData.CurrentLength)
@@ -2700,22 +2703,22 @@ private:
 					{
 						if (TOI1 > TOI0)
 						{
-							NodeStack.Add(FNodeQueueEntry{ Node.ChildrenNodes[1], TOI1 });
-							NodeStack.Add(FNodeQueueEntry{ Node.ChildrenNodes[0], TOI0 });
+							NodeStack[NodeStackNum++] = FNodeQueueEntry{ Node.ChildrenNodes[1], TOI1 };
+							NodeStack[NodeStackNum++] = FNodeQueueEntry{ Node.ChildrenNodes[0], TOI0 };
 						}
 						else
 						{
-							NodeStack.Add(FNodeQueueEntry{ Node.ChildrenNodes[0], TOI0 });
-							NodeStack.Add(FNodeQueueEntry{ Node.ChildrenNodes[1], TOI1 });
+							NodeStack[NodeStackNum++] = FNodeQueueEntry{ Node.ChildrenNodes[0], TOI0 };
+							NodeStack[NodeStackNum++] = FNodeQueueEntry{ Node.ChildrenNodes[1], TOI1 };
 						}
 					}
 					else if (bIntersect0)
 					{
-						NodeStack.Add(FNodeQueueEntry{ Node.ChildrenNodes[0], TOI0 });
+						NodeStack[NodeStackNum++] = FNodeQueueEntry{ Node.ChildrenNodes[0], TOI0 };
 					}
 					else if (bIntersect1)
 					{
-						NodeStack.Add(FNodeQueueEntry{ Node.ChildrenNodes[1], TOI1 });
+						NodeStack[NodeStackNum++] = FNodeQueueEntry{ Node.ChildrenNodes[1], TOI1 };
 					}
 				}
 				else
@@ -2724,7 +2727,7 @@ private:
 					{
 						if (TAABBTreeIntersectionHelper<TQueryFastData, Query>::Intersects(Start, CurData, TOI, FAABB3(AABB.Min(), AABB.Max()), QueryBounds, QueryHalfExtents, Dir, InvDir, bParallel))
 						{
-							NodeStack.Add(FNodeQueueEntry{ Node.ChildrenNodes[Idx], TOI });
+							NodeStack[NodeStackNum++] = FNodeQueueEntry{ Node.ChildrenNodes[Idx], TOI };
 						}
 						++Idx;
 					}
