Chaos - reduce collision constraint struct size
- remove some unused properties
- use floats for all local-space or normalized vectors, and elsewhere here possible

diff --git a/Engine/Source/Programs/HeadlessChaos/Private/HeadlessChaosTestDetectCollision.cpp b/Engine/Source/Programs/HeadlessChaos/Private/HeadlessChaosTestDetectCollision.cpp
--- a/Engine/Source/Programs/HeadlessChaos/Private/HeadlessChaosTestDetectCollision.cpp
+++ b/Engine/Source/Programs/HeadlessChaos/Private/HeadlessChaosTestDetectCollision.cpp
@@ -259,21 +259,11 @@ namespace ChaosTest
 		const FRigidTransform3 CapsuleTransform = FRigidTransform3(CapsulePos, FRotation3::FromIdentity());
 		const FRigidTransform3 ConvexTransform = FRigidTransform3(ConvexPos, FRotation3::FromIdentity());
 
+		FContactPoint ContactPoint = CapsuleConvexContactPoint(Capsule, CapsuleTransform, Convex, ConvexTransform);
+		EXPECT_TRUE(ContactPoint.IsSet());
+		if (ContactPoint.IsSet())
 		{
-			FContactPoint ContactPoint = CapsuleConvexContactPoint(Capsule, CapsuleTransform, Convex, ConvexTransform);
-			EXPECT_TRUE(ContactPoint.IsSet());
-			if (ContactPoint.IsSet())
-			{
-				CheckContactPoint(ContactPoint, FVec3(0, 0, -CapsuleRadius), FVec3(0, 0, 0.5f * BoxSize.Z), FVec3(0, 0, 1), 1, FVec3(0, 0, 1), CapsulePos.Z - CapsuleRadius);
-			}
-		}
-		{
-			FContactPoint ContactPoint = ConvexCapsuleContactPoint(Convex, ConvexTransform, Capsule, CapsuleTransform);
-			EXPECT_TRUE(ContactPoint.IsSet());
-			if (ContactPoint.IsSet())
-			{
-				CheckContactPoint(ContactPoint, FVec3(0, 0, 0.5f * BoxSize.Z), FVec3(0, 0, -CapsuleRadius), FVec3(0, 0, -1), 0, FVec3(0, 0, -1), CapsulePos.Z - CapsuleRadius);
-			}
+			CheckContactPoint(ContactPoint, FVec3(0, 0, -CapsuleRadius), FVec3(0, 0, 0.5f * BoxSize.Z), FVec3(0, 0, 1), 1, FVec3(0, 0, 1), CapsulePos.Z - CapsuleRadius);
 		}
 	}
 
@@ -293,16 +283,14 @@ namespace ChaosTest
 		const FRigidTransform3 CapsuleTransform = FRigidTransform3(CapsulePos, FRotation3::FromAxisAngle(FVec3(0,1,0), FMath::DegreesToRadians(45)));
 		const FRigidTransform3 ConvexTransform = FRigidTransform3(ConvexPos, FRotation3::FromIdentity());
 
+		FContactPoint ContactPoint = CapsuleConvexContactPoint(Capsule, CapsuleTransform, Convex, ConvexTransform);
+		EXPECT_TRUE(ContactPoint.IsSet());
+		if (ContactPoint.IsSet())
 		{
-			FContactPoint ContactPoint = CapsuleConvexContactPoint(Capsule, CapsuleTransform, Convex, ConvexTransform);
-			EXPECT_TRUE(ContactPoint.IsSet());
-			if (ContactPoint.IsSet())
-			{
-				const FVec3 CapsuleOffset = CapsulePos - FVec3(0.5f * BoxSize.X, 0.0f, 0.0f);
-				const FVec3 ExpectedNormal = CapsuleOffset.GetSafeNormal();
-				const FReal ExpectedPhi = CapsuleOffset.Size() - CapsuleRadius;
-				CheckContactPoint(ContactPoint, FVec3(0, 0, -CapsuleRadius), FVec3(0.5f * BoxSize.X, 0, 0.5f * BoxSize.Z), ExpectedNormal, 1, ExpectedNormal, ExpectedPhi);
-			}
+			const FVec3 CapsuleOffset = CapsulePos - FVec3(0.5f * BoxSize.X, 0.0f, 0.0f);
+			const FVec3 ExpectedNormal = CapsuleOffset.GetSafeNormal();
+			const FReal ExpectedPhi = CapsuleOffset.Size() - CapsuleRadius;
+			CheckContactPoint(ContactPoint, FVec3(0, 0, -CapsuleRadius), FVec3(0.5f * BoxSize.X, 0, 0.5f * BoxSize.Z), ExpectedNormal, 1, ExpectedNormal, ExpectedPhi);
 		}
 	}
 
@@ -323,16 +311,14 @@ namespace ChaosTest
 		const FRigidTransform3 CapsuleTransform = FRigidTransform3(CapsulePos, FRotation3::FromAxisAngle(FVec3(0, 1, 0), FMath::DegreesToRadians(45)));
 		const FRigidTransform3 ConvexTransform = FRigidTransform3(ConvexPos, FRotation3::FromIdentity());
 
+		FContactPoint ContactPoint = CapsuleConvexContactPoint(Capsule, CapsuleTransform, Convex, ConvexTransform);
+		EXPECT_TRUE(ContactPoint.IsSet());
+		if (ContactPoint.IsSet())
 		{
-			FContactPoint ContactPoint = CapsuleConvexContactPoint(Capsule, CapsuleTransform, Convex, ConvexTransform);
-			EXPECT_TRUE(ContactPoint.IsSet());
-			if (ContactPoint.IsSet())
-			{
-				const FVec3 CapsuleOffset = CapsulePos - FVec3(0.5f * BoxSize.X, 0.0f, 0.0f);
-				const FVec3 ExpectedNormal = -CapsuleOffset.GetSafeNormal();
-				const FReal ExpectedPhi = -(CapsuleOffset.Size() + CapsuleRadius);
-				CheckContactPoint(ContactPoint, FVec3(0, 0, -CapsuleRadius), FVec3(0.5f * BoxSize.X, 0, 0.5f * BoxSize.Z), ExpectedNormal, 1, ExpectedNormal, ExpectedPhi);
-			}
+			const FVec3 CapsuleOffset = CapsulePos - FVec3(0.5f * BoxSize.X, 0.0f, 0.0f);
+			const FVec3 ExpectedNormal = -CapsuleOffset.GetSafeNormal();
+			const FReal ExpectedPhi = -(CapsuleOffset.Size() + CapsuleRadius);
+			CheckContactPoint(ContactPoint, FVec3(0, 0, -CapsuleRadius), FVec3(0.5f * BoxSize.X, 0, 0.5f * BoxSize.Z), ExpectedNormal, 1, ExpectedNormal, ExpectedPhi);
 		}
 	}
 
diff --git a/Engine/Source/Programs/HeadlessChaos/Private/HeadlessChaosTestOneShotManifolds.cpp b/Engine/Source/Programs/HeadlessChaos/Private/HeadlessChaosTestOneShotManifolds.cpp
--- a/Engine/Source/Programs/HeadlessChaos/Private/HeadlessChaosTestOneShotManifolds.cpp
+++ b/Engine/Source/Programs/HeadlessChaos/Private/HeadlessChaosTestOneShotManifolds.cpp
@@ -192,8 +192,8 @@ namespace ChaosTest
 			for (int ConstraintIndex = 0; ConstraintIndex < ContactCount; ConstraintIndex++)
 			{
 				EXPECT_NEAR(10.0f, Constraint.GetManifoldPoints()[ConstraintIndex].ContactPoint.Phi, 0.01);
-				FVec3 Location1 = Box1Transform.TransformPosition(Constraint.GetManifoldPoints()[ConstraintIndex].ContactPoint.ShapeContactPoints[0]);
-				FVec3 Location2 = Box2Transform.TransformPosition(Constraint.GetManifoldPoints()[ConstraintIndex].ContactPoint.ShapeContactPoints[1]);
+				FVec3 Location1 = Box1Transform.TransformPosition(FVec3(Constraint.GetManifoldPoints()[ConstraintIndex].ContactPoint.ShapeContactPoints[0]));
+				FVec3 Location2 = Box2Transform.TransformPosition(FVec3(Constraint.GetManifoldPoints()[ConstraintIndex].ContactPoint.ShapeContactPoints[1]));
 				EXPECT_NEAR(Location1.Z, 110.0f, 0.01);
 				EXPECT_NEAR(Location2.Z, 100.0f, 0.01);
 			}
@@ -213,7 +213,7 @@ namespace ChaosTest
 			for (int ConstraintIndex = 0; ConstraintIndex < ContactCount; ConstraintIndex++)
 			{
 				EXPECT_NEAR(10.0f, Constraint.GetManifoldPoints()[ConstraintIndex].ContactPoint.Phi, 0.01);
-				EXPECT_NEAR(Box2Transform.TransformPosition(Constraint.GetManifoldPoints()[ConstraintIndex].ContactPoint.ShapeContactPoints[1]).Z, 100.0f, 0.01);
+				EXPECT_NEAR(Box2Transform.TransformPosition(FVec3(Constraint.GetManifoldPoints()[ConstraintIndex].ContactPoint.ShapeContactPoints[1])).Z, 100.0f, 0.01);
 			}
 		}
 
@@ -255,7 +255,7 @@ namespace ChaosTest
 		if (ContactCount > 0)
 		{
 			EXPECT_NEAR(Constraint.GetManifoldPoints()[0].ContactPoint.Phi, 0.0f, 0.01f);
-			EXPECT_NEAR(Convex1Transform.TransformPosition(Constraint.GetManifoldPoints()[0].ContactPoint.ShapeContactPoints[0]).X, 0.5f * Offset, 0.01f);
+			EXPECT_NEAR(Convex1Transform.TransformPosition(FVec3(Constraint.GetManifoldPoints()[0].ContactPoint.ShapeContactPoints[0])).X, 0.5f * Offset, 0.01f);
 		}
 	}
 
diff --git a/Engine/Source/Programs/HeadlessChaos/Private/HeadlessChaosTestProbes.cpp b/Engine/Source/Programs/HeadlessChaos/Private/HeadlessChaosTestProbes.cpp
--- a/Engine/Source/Programs/HeadlessChaos/Private/HeadlessChaosTestProbes.cpp
+++ b/Engine/Source/Programs/HeadlessChaos/Private/HeadlessChaosTestProbes.cpp
@@ -119,7 +119,7 @@ namespace ChaosTest
 		// with CCD, make sure the constraint doesn't actually think it's CCD
 		for (FPBDCollisionConstraint* const Constraint : Evolution.GetCollisionConstraints().GetConstraints())
 		{
-			EXPECT_EQ(Constraint->GetCCDType(), ECollisionCCDType::Disabled);
+			EXPECT_EQ(Constraint->GetCCDEnabled(), false);
 		}
 
 		// Make sure that the velocity of the dynamic body wasn't affected
diff --git a/Engine/Source/Programs/HeadlessChaos/Public/HeadlessChaosCollisionConstraints.h b/Engine/Source/Programs/HeadlessChaos/Public/HeadlessChaosCollisionConstraints.h
--- a/Engine/Source/Programs/HeadlessChaos/Public/HeadlessChaosCollisionConstraints.h
+++ b/Engine/Source/Programs/HeadlessChaos/Public/HeadlessChaosCollisionConstraints.h
@@ -81,7 +81,7 @@ public:
 
 	void Update(FPBDCollisionConstraint& Constraint)
 	{
-		if (Constraint.GetCCDType() == ECollisionCCDType::Disabled)
+		if (!Constraint.GetCCDEnabled())
 		{
 			// Dt is not important for the tests that use this function
 			const FReal Dt = FReal(1) / FReal(30);
diff --git a/Engine/Source/Runtime/Experimental/Chaos/Private/Chaos/CCDUtilities.cpp b/Engine/Source/Runtime/Experimental/Chaos/Private/Chaos/CCDUtilities.cpp
--- a/Engine/Source/Runtime/Experimental/Chaos/Private/Chaos/CCDUtilities.cpp
+++ b/Engine/Source/Runtime/Experimental/Chaos/Private/Chaos/CCDUtilities.cpp
@@ -825,7 +825,7 @@ namespace Chaos
 				continue;
 			}
 			
-			const FVec3 Normal = ShapeWorldTransform1.TransformVectorNoScale(ManifoldPoint.ContactPoint.ShapeContactNormal);
+			const FVec3 Normal = ShapeWorldTransform1.TransformVectorNoScale(FVec3(ManifoldPoint.ContactPoint.ShapeContactNormal));
 			const FVec3 V0 = Rigid0 != nullptr ? Rigid0->V() : FVec3(0.f);
 			const FVec3 V1 = Rigid1 != nullptr ? Rigid1->V() : FVec3(0.f);
 			const FReal NormalV = FVec3::DotProduct(V0 - V1, Normal);
diff --git a/Engine/Source/Runtime/Experimental/Chaos/Private/Chaos/ChaosDebugDraw.cpp b/Engine/Source/Runtime/Experimental/Chaos/Private/Chaos/ChaosDebugDraw.cpp
--- a/Engine/Source/Runtime/Experimental/Chaos/Private/Chaos/ChaosDebugDraw.cpp
+++ b/Engine/Source/Runtime/Experimental/Chaos/Private/Chaos/ChaosDebugDraw.cpp
@@ -889,9 +889,9 @@ namespace Chaos
 					const FRigidTransform3& PlaneTransform = (ContactPlaneOwner == 0) ? Contact.GetShapeRelativeTransform0() * WorldActorTransform0 : Contact.GetShapeRelativeTransform1() * WorldActorTransform1;
 					const FRigidTransform3& PointTransform = (ContactPlaneOwner == 0) ? Contact.GetShapeRelativeTransform1() * WorldActorTransform1 : Contact.GetShapeRelativeTransform0() * WorldActorTransform0;
 					const FConstGenericParticleHandle PlaneParticle = (ContactPlaneOwner == 0) ? Particle0 : Particle1;
-					const FVec3 PlaneNormal = PlaneTransform.TransformVectorNoScale(ManifoldPoint.ContactPoint.ShapeContactNormal);
-					const FVec3 PointLocation = PointTransform.TransformPosition(ManifoldPoint.ContactPoint.ShapeContactPoints[ContactPointOwner]);
-					const FVec3 PlaneLocation = PlaneTransform.TransformPosition(ManifoldPoint.ContactPoint.ShapeContactPoints[ContactPlaneOwner]);
+					const FVec3 PlaneNormal = PlaneTransform.TransformVectorNoScale(FVec3(ManifoldPoint.ContactPoint.ShapeContactNormal));
+					const FVec3 PointLocation = PointTransform.TransformPosition(FVec3(ManifoldPoint.ContactPoint.ShapeContactPoints[ContactPointOwner]));
+					const FVec3 PlaneLocation = PlaneTransform.TransformPosition(FVec3(ManifoldPoint.ContactPoint.ShapeContactPoints[ContactPlaneOwner]));
 					const FVec3 PointPlaneLocation = PointLocation - FVec3::DotProduct(PointLocation - PlaneLocation, PlaneNormal) * PlaneNormal;
 
 					// Dynamic friction, restitution = red
@@ -935,12 +935,12 @@ namespace Chaos
 					if ((Settings.PushOutScale > 0) && !ManifoldPoint.NetPushOut.IsNearlyZero())
 					{
 						FColor Color = (ColorScale * PushOutImpusleColor).ToFColor(false);
-						FDebugDrawQueue::GetInstance().DrawDebugLine(WorldPointPlaneLocation, WorldPointPlaneLocation + Settings.DrawScale * Settings.PushOutScale * SpaceTransform.TransformVectorNoScale(ManifoldPoint.NetPushOut), Color, false, Duration, Settings.DrawPriority, Settings.LineThickness);
+						FDebugDrawQueue::GetInstance().DrawDebugLine(WorldPointPlaneLocation, WorldPointPlaneLocation + Settings.DrawScale * Settings.PushOutScale * SpaceTransform.TransformVectorNoScale(FVec3(ManifoldPoint.NetPushOut)), Color, false, Duration, Settings.DrawPriority, Settings.LineThickness);
 					}
 					if ((Settings.ImpulseScale > 0) && !ManifoldPoint.NetImpulse.IsNearlyZero())
 					{
 						FColor Color = (ColorScale * ImpulseColor).ToFColor(false);
-						FDebugDrawQueue::GetInstance().DrawDebugLine(WorldPointPlaneLocation, WorldPointPlaneLocation + Settings.DrawScale * Settings.ImpulseScale * SpaceTransform.TransformVectorNoScale(ManifoldPoint.NetImpulse), Color, false, Duration, Settings.DrawPriority, Settings.LineThickness);
+						FDebugDrawQueue::GetInstance().DrawDebugLine(WorldPointPlaneLocation, WorldPointPlaneLocation + Settings.DrawScale * Settings.ImpulseScale * SpaceTransform.TransformVectorNoScale(FVec3(ManifoldPoint.NetImpulse)), Color, false, Duration, Settings.DrawPriority, Settings.LineThickness);
 					}
 
 					// Manifold plane and normal
@@ -967,8 +967,8 @@ namespace Chaos
 					// Previous points
 					if (bIsActive && ManifoldPoint.Flags.bWasFrictionRestored)
 					{
-						const FVec3 WorldPrevPointLocation = SpaceTransform.TransformPosition(PointTransform.TransformPosition(ManifoldPoint.ShapeAnchorPoints[ContactPointOwner]));
-						const FVec3 WorldPrevPlaneLocation = SpaceTransform.TransformPosition(PlaneTransform.TransformPosition(ManifoldPoint.ShapeAnchorPoints[ContactPlaneOwner]));
+						const FVec3 WorldPrevPointLocation = SpaceTransform.TransformPosition(PointTransform.TransformPosition(FVec3(ManifoldPoint.ShapeAnchorPoints[ContactPointOwner])));
+						const FVec3 WorldPrevPlaneLocation = SpaceTransform.TransformPosition(PlaneTransform.TransformPosition(FVec3(ManifoldPoint.ShapeAnchorPoints[ContactPlaneOwner])));
 						FDebugDrawQueue::GetInstance().DrawDebugLine(WorldPrevPointLocation, WorldPointLocation, FColor::White, false, Duration, Settings.DrawPriority, Settings.LineThickness);
 						FDebugDrawQueue::GetInstance().DrawDebugLine(WorldPrevPlaneLocation, WorldPlaneLocation, FColor::White, false, Duration, Settings.DrawPriority, Settings.LineThickness);
 					}
@@ -1311,8 +1311,8 @@ namespace Chaos
 					ContactPos = FVec3(0);
 					for (const FManifoldPoint& ManifoldPoint : Constraint->GetManifoldPoints())
 					{
-						ContactPos += SpaceTransform.TransformPosition(ShapeTransform0.TransformPositionNoScale(ManifoldPoint.ContactPoint.ShapeContactPoints[0]));
-						ContactPos += SpaceTransform.TransformPosition(ShapeTransform1.TransformPositionNoScale(ManifoldPoint.ContactPoint.ShapeContactPoints[1]));
+						ContactPos += SpaceTransform.TransformPosition(ShapeTransform0.TransformPositionNoScale(FVec3(ManifoldPoint.ContactPoint.ShapeContactPoints[0])));
+						ContactPos += SpaceTransform.TransformPosition(ShapeTransform1.TransformPositionNoScale(FVec3(ManifoldPoint.ContactPoint.ShapeContactPoints[1])));
 					}
 					ContactPos /= (FReal)(2 * Constraint->GetManifoldPoints().Num());
 				}
diff --git a/Engine/Source/Runtime/Experimental/Chaos/Private/Chaos/Collision/CapsuleConvexContactPoint.cpp b/Engine/Source/Runtime/Experimental/Chaos/Private/Chaos/Collision/CapsuleConvexContactPoint.cpp
--- a/Engine/Source/Runtime/Experimental/Chaos/Private/Chaos/Collision/CapsuleConvexContactPoint.cpp
+++ b/Engine/Source/Runtime/Experimental/Chaos/Private/Chaos/Collision/CapsuleConvexContactPoint.cpp
@@ -9,10 +9,8 @@
 
 namespace Chaos
 {
-	// Use GJK (point to convex) to calculate separation.
-	// Fall back to plane testing if penetrating by more than Radius.
 	template<typename T_CONVEX>
-	FContactPoint ConvexCapsuleContactPointImpl(const T_CONVEX& Convex, const FRigidTransform3& ConvexTransform, const FImplicitCapsule3& Capsule, const FRigidTransform3& CapsuleTransform)
+	FContactPoint CapsuleConvexContactPointImpl(const FImplicitCapsule3& Capsule, const FRigidTransform3& CapsuleTransform, const T_CONVEX& Convex, const FRigidTransform3& ConvexTransform)
 	{
 		FContactPoint ContactPoint;
 		if (Convex.NumPlanes() > 0)
@@ -26,16 +24,16 @@ namespace Chaos
 			// NOTE: Capsule is treated as a line (its core shape), Convex margin is ignored so we are using the outer non-shrunken hull.
 			// @todo(chaos): use GJKDistance and SAT when that fails rather that GJK/EPA (but this requires an edge list in the convex)
 			bool bHaveResult = GJKPenetration<true>(MakeGJKShape(Convex), MakeGJKCoreShape(Capsule), CapsuleToConvexTransform, Penetration, PosConvex, PosCapsuleInConvex, NormalConvex, VertexIndexA, VertexIndexB);
-			
+
 			// Build the contact point
 			if (bHaveResult)
 			{
 				const FVec3 PosCapsule = CapsuleToConvexTransform.InverseTransformPosition(PosCapsuleInConvex);
 				const FReal Phi = -Penetration;
 
-				ContactPoint.ShapeContactPoints[0] = PosConvex;
-				ContactPoint.ShapeContactPoints[1] = PosCapsule;
-				ContactPoint.ShapeContactNormal = -NormalConvex;
+				ContactPoint.ShapeContactPoints[0] = PosCapsule;
+				ContactPoint.ShapeContactPoints[1] = PosConvex;
+				ContactPoint.ShapeContactNormal = NormalConvex;
 				ContactPoint.Phi = Phi;
 			}
 		}
@@ -43,61 +41,19 @@ namespace Chaos
 		return ContactPoint;
 	}
 
-	template<typename T_CONVEX>
-	FContactPoint CapsuleConvexContactPointImpl(const FImplicitCapsule3& Capsule, const FRigidTransform3& CapsuleTransform, const T_CONVEX& Convex, const FRigidTransform3& ConvexTransform)
-	{
-		return ConvexCapsuleContactPointImpl(Convex, ConvexTransform, Capsule, CapsuleTransform).SwapShapes();
-	}
-
-
-
-	FContactPoint CapsuleConvexContactPoint(const FImplicitCapsule3& Capsule, const FRigidTransform3& CapsuleTransform, const FImplicitConvex3& Convex, const FRigidTransform3& ConvexTransform)
-	{
-		return CapsuleConvexContactPointImpl(Capsule, CapsuleTransform, Convex, ConvexTransform);
-	}
-
-	FContactPoint CapsuleConvexContactPoint(const FImplicitCapsule3& Capsule, const FRigidTransform3& CapsuleTransform, const TImplicitObjectInstanced<FImplicitConvex3>& Convex, const FRigidTransform3& ConvexTransform)
-	{
-		return CapsuleConvexContactPointImpl(Capsule, CapsuleTransform, Convex, ConvexTransform);
-	}
-
-	FContactPoint CapsuleConvexContactPoint(const FImplicitCapsule3& Capsule, const FRigidTransform3& CapsuleTransform, const TImplicitObjectScaled<FImplicitConvex3>& Convex, const FRigidTransform3& ConvexTransform)
-	{
-		return CapsuleConvexContactPointImpl(Capsule, CapsuleTransform, Convex, ConvexTransform);
-	}
-
-
-
-	FContactPoint ConvexCapsuleContactPoint(const FImplicitConvex3& Convex, const FRigidTransform3& ConvexTransform, const FImplicitCapsule3& Capsule, const FRigidTransform3& CapsuleTransform)
-	{
-		return ConvexCapsuleContactPointImpl(Convex, ConvexTransform, Capsule, CapsuleTransform);
-	}
-
-	FContactPoint ConvexCapsuleContactPoint(const TImplicitObjectInstanced<FImplicitConvex3>& Convex, const FRigidTransform3& ConvexTransform, const FImplicitCapsule3& Capsule, const FRigidTransform3& CapsuleTransform)
-	{
-		return ConvexCapsuleContactPointImpl(Convex, ConvexTransform, Capsule, CapsuleTransform);
-	}
-
-	FContactPoint ConvexCapsuleContactPoint(const TImplicitObjectScaled<FImplicitConvex3>& Convex, const FRigidTransform3& ConvexTransform, const FImplicitCapsule3& Capsule, const FRigidTransform3& CapsuleTransform)
-	{
-		return ConvexCapsuleContactPointImpl(Convex, ConvexTransform, Capsule, CapsuleTransform);
-	}
-
-
-
 	FContactPoint CapsuleConvexContactPoint(const FImplicitCapsule3& Capsule, const FRigidTransform3& CapsuleTransform, const FImplicitObject& Object, const FRigidTransform3& ConvexTransform)
 	{
 		if (const TImplicitObjectInstanced<FImplicitConvex3>* InstancedConvex = Object.template GetObject<const TImplicitObjectInstanced<FImplicitConvex3>>())
 		{
-			return CapsuleConvexContactPoint(Capsule, CapsuleTransform, *InstancedConvex, ConvexTransform);
+			return CapsuleConvexContactPointImpl(Capsule, CapsuleTransform, *InstancedConvex, ConvexTransform);
 		}
 		else if (const TImplicitObjectScaled<FImplicitConvex3>* ScaledConvex = Object.template GetObject<const TImplicitObjectScaled<FImplicitConvex3>>())
 		{
-			return CapsuleConvexContactPoint(Capsule, CapsuleTransform, *ScaledConvex, ConvexTransform);
+			return CapsuleConvexContactPointImpl(Capsule, CapsuleTransform, *ScaledConvex, ConvexTransform);
 		}
 		else if (const FImplicitConvex3* Convex = Object.template GetObject<const FImplicitConvex3>())
 		{
-			return CapsuleConvexContactPoint(Capsule, CapsuleTransform, *Convex, ConvexTransform);
+			return CapsuleConvexContactPointImpl(Capsule, CapsuleTransform, *Convex, ConvexTransform);
 		}
 		return FContactPoint();
 	}
diff --git a/Engine/Source/Runtime/Experimental/Chaos/Private/Chaos/Collision/CollisionPruning.cpp b/Engine/Source/Runtime/Experimental/Chaos/Private/Chaos/Collision/CollisionPruning.cpp
--- a/Engine/Source/Runtime/Experimental/Chaos/Private/Chaos/Collision/CollisionPruning.cpp
+++ b/Engine/Source/Runtime/Experimental/Chaos/Private/Chaos/Collision/CollisionPruning.cpp
@@ -41,7 +41,7 @@ namespace Chaos
 						if (bIsEdgeContact)
 						{
 							const FRigidTransform3& EdgeTransform = (EdgeOtherShapeIndex == 0) ? EdgeCollision.GetShapeWorldTransform0() : EdgeCollision.GetShapeWorldTransform1();
-							const FVec3 EdgePos = EdgeTransform.TransformPositionNoScale(EdgeManifoldPoint.ContactPoint.ShapeContactPoints[EdgeOtherShapeIndex]);
+							const FVec3 EdgePos = EdgeTransform.TransformPositionNoScale(FVec3(EdgeManifoldPoint.ContactPoint.ShapeContactPoints[EdgeOtherShapeIndex]));
 
 							// Loop over plane collisions
 							ECollisionVisitorResult PlaneResult = ParticleCollisions.VisitConstCollisions(
@@ -58,8 +58,8 @@ namespace Chaos
 											if (PlaneManifoldPoint.ContactPoint.ContactType == PlaneContactType)
 											{
 												// If the edge position is in the plane, disable it
-												const FVec3 PlanePos = PlaneTransform.TransformPositionNoScale(PlaneManifoldPoint.ContactPoint.ShapeContactPoints[PlaneOtherShapeIndex]);
-												const FVec3 PlaneNormal = PlaneCollision.GetShapeWorldTransform1().TransformVectorNoScale(PlaneManifoldPoint.ContactPoint.ShapeContactNormal);
+												const FVec3 PlanePos = PlaneTransform.TransformPositionNoScale(FVec3(PlaneManifoldPoint.ContactPoint.ShapeContactPoints[PlaneOtherShapeIndex]));
+												const FVec3 PlaneNormal = PlaneCollision.GetShapeWorldTransform1().TransformVectorNoScale(FVec3(PlaneManifoldPoint.ContactPoint.ShapeContactNormal));
 
 												const FVec3 EdgePlaneDelta = EdgePos - PlanePos;
 												const FReal EdgePlaneDistance = FVec3::DotProduct(EdgePlaneDelta, PlaneNormal);
@@ -223,8 +223,8 @@ namespace Chaos
 						{
 							const FManifoldPoint& ManifoldPoint = Collision.GetManifoldPoint(ManifoldPointIndex);
 
-							const FVec3 ContactPos = ContactTransform.TransformPositionNoScale(ManifoldPoint.ContactPoint.ShapeContactPoints[OtherShapeIndex]);
-							const FVec3 ContactNormal = Collision.GetShapeWorldTransform1().TransformVectorNoScale(ManifoldPoint.ContactPoint.ShapeContactNormal);
+							const FVec3 ContactPos = ContactTransform.TransformPositionNoScale(FVec3(ManifoldPoint.ContactPoint.ShapeContactPoints[OtherShapeIndex]));
+							const FVec3 ContactNormal = Collision.GetShapeWorldTransform1().TransformVectorNoScale(FVec3(ManifoldPoint.ContactPoint.ShapeContactNormal));
 
 							ContactGroupContainer.Add(Collision, ManifoldPointIndex, ContactPos, ContactNormal);
 						}
diff --git a/Engine/Source/Runtime/Experimental/Chaos/Private/Chaos/Collision/ContactPointsMiscShapes.cpp b/Engine/Source/Runtime/Experimental/Chaos/Private/Chaos/Collision/ContactPointsMiscShapes.cpp
--- a/Engine/Source/Runtime/Experimental/Chaos/Private/Chaos/Collision/ContactPointsMiscShapes.cpp
+++ b/Engine/Source/Runtime/Experimental/Chaos/Private/Chaos/Collision/ContactPointsMiscShapes.cpp
@@ -187,7 +187,7 @@ namespace Chaos
 	}
 
 	template <typename GeometryA, typename GeometryB>
-	FContactPoint GJKImplicitContactPoint(const FImplicitObject& A, const FRigidTransform3& ATransform, const GeometryB& B, const FRigidTransform3& BTransform, const FReal CullDistance, const FReal ShapePadding)
+	FContactPoint GJKImplicitContactPoint(const FImplicitObject& A, const FRigidTransform3& ATransform, const GeometryB& B, const FRigidTransform3& BTransform, const FReal CullDistance)
 	{
 		FContactPoint Contact;
 		const FRigidTransform3 AToBTM = ATransform.GetRelativeTransform(BTransform);
@@ -269,12 +269,12 @@ namespace Chaos
 	}
 
 
-	FContactPoint SphereSphereContactPoint(const TSphere<FReal, 3>& Sphere1, const FRigidTransform3& Sphere1Transform, const TSphere<FReal, 3>& Sphere2, const FRigidTransform3& Sphere2Transform, const FReal CullDistance, const FReal ShapePadding)
+	FContactPoint SphereSphereContactPoint(const TSphere<FReal, 3>& Sphere1, const FRigidTransform3& Sphere1Transform, const TSphere<FReal, 3>& Sphere2, const FRigidTransform3& Sphere2Transform, const FReal CullDistance)
 	{
 		FContactPoint Result;
 
-		const FReal R1 = Sphere1.GetRadius() + 0.5f * ShapePadding;
-		const FReal R2 = Sphere2.GetRadius() + 0.5f * ShapePadding;
+		const FReal R1 = Sphere1.GetRadius();
+		const FReal R2 = Sphere2.GetRadius();
 
 		// World-space contact
 		const FVec3 Center1 = Sphere1Transform.TransformPosition(Sphere1.GetCenter());
@@ -296,17 +296,17 @@ namespace Chaos
 		return Result;
 	}
 
-	FContactPoint SpherePlaneContactPoint(const TSphere<FReal, 3>& Sphere, const FRigidTransform3& SphereTransform, const TPlane<FReal, 3>& Plane, const FRigidTransform3& PlaneTransform, const FReal ShapePadding)
+	FContactPoint SpherePlaneContactPoint(const TSphere<FReal, 3>& Sphere, const FRigidTransform3& SphereTransform, const TPlane<FReal, 3>& Plane, const FRigidTransform3& PlaneTransform)
 	{
 		FContactPoint Result;
 
-		FReal SphereRadius = Sphere.GetRadius() + 0.5f * ShapePadding;
+		FReal SphereRadius = Sphere.GetRadius();
 
 		FVec3 SpherePosWorld = SphereTransform.TransformPosition(Sphere.GetCenter());
 		FVec3 SpherePosPlane = PlaneTransform.InverseTransformPosition(SpherePosWorld);
 
 		FVec3 NormalPlane;
-		FReal Phi = Plane.PhiWithNormal(SpherePosPlane, NormalPlane) - SphereRadius - 0.5f * ShapePadding;	// Adding plane's share of padding
+		FReal Phi = Plane.PhiWithNormal(SpherePosPlane, NormalPlane) - SphereRadius;	// Adding plane's share of padding
 		FVec3 NormalWorld = PlaneTransform.TransformVector(NormalPlane);
 		FVec3 Location = SpherePosWorld - SphereRadius * NormalWorld;
 
@@ -318,7 +318,7 @@ namespace Chaos
 		return Result;
 	}
 
-	FContactPoint SphereBoxContactPoint(const TSphere<FReal, 3>& Sphere, const FRigidTransform3& SphereTransform, const FImplicitBox3& Box, const FRigidTransform3& BoxTransform, const FReal ShapePadding)
+	FContactPoint SphereBoxContactPoint(const TSphere<FReal, 3>& Sphere, const FRigidTransform3& SphereTransform, const FImplicitBox3& Box, const FRigidTransform3& BoxTransform)
 	{
 		FContactPoint Result;
 
@@ -327,10 +327,10 @@ namespace Chaos
 
 		FVec3 NormalBox;																	// Box-space normal
 		FReal PhiToSphereCenter = Box.PhiWithNormal(SphereBox, NormalBox);
-		FReal Phi = PhiToSphereCenter - Sphere.GetRadius() - ShapePadding;
+		FReal Phi = PhiToSphereCenter - Sphere.GetRadius();
 
 		FVec3 NormalWorld = BoxTransform.TransformVectorNoScale(NormalBox);
-		FVec3 LocationWorld = SphereWorld - (Sphere.GetRadius() + 0.5f * ShapePadding) * NormalWorld;
+		FVec3 LocationWorld = SphereWorld - (Sphere.GetRadius()) * NormalWorld;
 
 		Result.ShapeContactPoints[0] = SphereTransform.InverseTransformPosition(LocationWorld);
 		Result.ShapeContactPoints[1] = BoxTransform.InverseTransformPosition(LocationWorld - Phi * NormalWorld);
@@ -339,7 +339,7 @@ namespace Chaos
 		return Result;
 	}
 
-	FContactPoint SphereCapsuleContactPoint(const TSphere<FReal, 3>& A, const FRigidTransform3& ATransform, const FCapsule& B, const FRigidTransform3& BTransform, const FReal ShapePadding)
+	FContactPoint SphereCapsuleContactPoint(const TSphere<FReal, 3>& A, const FRigidTransform3& ATransform, const FCapsule& B, const FRigidTransform3& BTransform)
 	{
 		FContactPoint Result;
 
@@ -352,7 +352,7 @@ namespace Chaos
 		FReal DeltaLen = Delta.Size();
 		if (DeltaLen > UE_KINDA_SMALL_NUMBER)
 		{
-			FReal NewPhi = DeltaLen - (A.GetRadius() + B.GetRadius()) - ShapePadding;
+			FReal NewPhi = DeltaLen - (A.GetRadius() + B.GetRadius());
 			FVec3 Dir = Delta / DeltaLen;
 			FVec3 LocationA = A1 + Dir * A.GetRadius();
 			FVec3 LocationB = P2 - Dir * B.GetRadius();
@@ -368,9 +368,9 @@ namespace Chaos
 	}
 
 	template <typename TriMeshType>
-	FContactPoint SphereTriangleMeshContactPoint(const TSphere<FReal, 3>& A, const FRigidTransform3& ATransform, const TriMeshType& B, const FRigidTransform3& BTransform, const FReal CullDistance, const FReal ShapePadding)
+	FContactPoint SphereTriangleMeshContactPoint(const TSphere<FReal, 3>& A, const FRigidTransform3& ATransform, const TriMeshType& B, const FRigidTransform3& BTransform, const FReal CullDistance)
 	{
-		return GJKImplicitContactPoint<TSphere<FReal, 3>>(TSphere<FReal, 3>(A), ATransform, B, BTransform, CullDistance, ShapePadding);
+		return GJKImplicitContactPoint<TSphere<FReal, 3>>(TSphere<FReal, 3>(A), ATransform, B, BTransform, CullDistance);
 	}
 
 	template<typename TriMeshType>
@@ -389,33 +389,33 @@ namespace Chaos
 		return FContactPoint();
 	}
 
-	FContactPoint BoxHeightFieldContactPoint(const FImplicitBox3& A, const FRigidTransform3& ATransform, const FHeightField& B, const FRigidTransform3& BTransform, const FReal CullDistance, const FReal ShapePadding)
+	FContactPoint BoxHeightFieldContactPoint(const FImplicitBox3& A, const FRigidTransform3& ATransform, const FHeightField& B, const FRigidTransform3& BTransform, const FReal CullDistance)
 	{
-		return GJKImplicitContactPoint<FImplicitBox3>(A, ATransform, B, BTransform, CullDistance, ShapePadding);
+		return GJKImplicitContactPoint<FImplicitBox3>(A, ATransform, B, BTransform, CullDistance);
 	}
 
 	template <typename TriMeshType>
-	FContactPoint BoxTriangleMeshContactPoint(const FImplicitBox3& A, const FRigidTransform3& ATransform, const TriMeshType& B, const FRigidTransform3& BTransform, const FReal CullDistance, const FReal ShapePadding)
+	FContactPoint BoxTriangleMeshContactPoint(const FImplicitBox3& A, const FRigidTransform3& ATransform, const TriMeshType& B, const FRigidTransform3& BTransform, const FReal CullDistance)
 	{
-		return GJKImplicitContactPoint<TBox<FReal, 3>>(A, ATransform, B, BTransform, CullDistance, ShapePadding);
+		return GJKImplicitContactPoint<TBox<FReal, 3>>(A, ATransform, B, BTransform, CullDistance);
 	}
 
-	FContactPoint SphereHeightFieldContactPoint(const TSphere<FReal, 3>& A, const FRigidTransform3& ATransform, const FHeightField& B, const FRigidTransform3& BTransform, const FReal CullDistance, const FReal ShapePadding)
+	FContactPoint SphereHeightFieldContactPoint(const TSphere<FReal, 3>& A, const FRigidTransform3& ATransform, const FHeightField& B, const FRigidTransform3& BTransform, const FReal CullDistance)
 	{
-		return GJKImplicitContactPoint<TSphere<FReal, 3>>(TSphere<FReal, 3>(A), ATransform, B, BTransform, CullDistance, ShapePadding);
+		return GJKImplicitContactPoint<TSphere<FReal, 3>>(TSphere<FReal, 3>(A), ATransform, B, BTransform, CullDistance);
 	}
 
-	FContactPoint CapsuleHeightFieldContactPoint(const FCapsule& A, const FRigidTransform3& ATransform, const FHeightField& B, const FRigidTransform3& BTransform, const FReal CullDistance, const FReal ShapePadding)
+	FContactPoint CapsuleHeightFieldContactPoint(const FCapsule& A, const FRigidTransform3& ATransform, const FHeightField& B, const FRigidTransform3& BTransform, const FReal CullDistance)
 	{
 		CONDITIONAL_SCOPE_CYCLE_COUNTER(STAT_Collisions_CapsuleHeightFieldContactPoint, ConstraintsDetailedStats);
-		return GJKImplicitContactPoint<FCapsule>(A, ATransform, B, BTransform, CullDistance, ShapePadding);
+		return GJKImplicitContactPoint<FCapsule>(A, ATransform, B, BTransform, CullDistance);
 	}
 
 	template <typename TriMeshType>
-	FContactPoint CapsuleTriangleMeshContactPoint(const FCapsule& A, const FRigidTransform3& ATransform, const TriMeshType& B, const FRigidTransform3& BTransform, const FReal CullDistance, const FReal ShapePadding)
+	FContactPoint CapsuleTriangleMeshContactPoint(const FCapsule& A, const FRigidTransform3& ATransform, const TriMeshType& B, const FRigidTransform3& BTransform, const FReal CullDistance)
 	{
 		CONDITIONAL_SCOPE_CYCLE_COUNTER(STAT_Collisions_CapsuleTriangleMeshContactPoint, ConstraintsDetailedStats);
-		return GJKImplicitContactPoint<FCapsule>(A, ATransform, B, BTransform, CullDistance, ShapePadding);
+		return GJKImplicitContactPoint<FCapsule>(A, ATransform, B, BTransform, CullDistance);
 	}
 
 	template <typename TriMeshType>
@@ -435,13 +435,13 @@ namespace Chaos
 		return FContactPoint();
 	}
 
-	FContactPoint ConvexHeightFieldContactPoint(const FImplicitObject& A, const FRigidTransform3& ATransform, const FHeightField& B, const FRigidTransform3& BTransform, const FReal CullDistance, const FReal ShapePadding)
+	FContactPoint ConvexHeightFieldContactPoint(const FImplicitObject& A, const FRigidTransform3& ATransform, const FHeightField& B, const FRigidTransform3& BTransform, const FReal CullDistance)
 	{
 		CONDITIONAL_SCOPE_CYCLE_COUNTER(STAT_Collisions_ConvexHeightFieldContactPoint, ConstraintsDetailedStats);
-		return GJKImplicitContactPoint<FConvex>(A, ATransform, B, BTransform, CullDistance, ShapePadding);
+		return GJKImplicitContactPoint<FConvex>(A, ATransform, B, BTransform, CullDistance);
 	}
 
-	FContactPoint ConvexTriangleMeshContactPoint(const FImplicitObject& A, const FRigidTransform3& ATransform, const FImplicitObject& B, const FRigidTransform3& BTransform, const FReal CullDistance, const FReal ShapePadding)
+	FContactPoint ConvexTriangleMeshContactPoint(const FImplicitObject& A, const FRigidTransform3& ATransform, const FImplicitObject& B, const FRigidTransform3& BTransform, const FReal CullDistance)
 	{
 		CONDITIONAL_SCOPE_CYCLE_COUNTER(STAT_Collisions_ConvexTriangleMeshContactPoint, ConstraintsDetailedStats);
 
@@ -450,7 +450,7 @@ namespace Chaos
 			[&](auto BConcretePtr)
 			{
 				check(BConcretePtr != nullptr);
-				return GJKImplicitContactPoint<FConvex>(A, ATransform, *BConcretePtr, BTransform, CullDistance, ShapePadding);
+				return GJKImplicitContactPoint<FConvex>(A, ATransform, *BConcretePtr, BTransform, CullDistance);
 			});
 	}
 
@@ -471,7 +471,7 @@ namespace Chaos
 		return FContactPoint();
 	}
 
-	FContactPoint CapsuleCapsuleContactPoint(const FCapsule& A, const FRigidTransform3& ATransform, const FCapsule& B, const FRigidTransform3& BTransform, const FReal ShapePadding)
+	FContactPoint CapsuleCapsuleContactPoint(const FCapsule& A, const FRigidTransform3& ATransform, const FCapsule& B, const FRigidTransform3& BTransform)
 	{
 		FContactPoint Result;
 
@@ -486,7 +486,7 @@ namespace Chaos
 		FReal DeltaLen = Delta.Size();
 		if (DeltaLen > UE_KINDA_SMALL_NUMBER)
 		{
-			FReal NewPhi = DeltaLen - (A.GetRadius() + B.GetRadius()) - ShapePadding;
+			FReal NewPhi = DeltaLen - (A.GetRadius() + B.GetRadius());
 			FVec3 Dir = Delta / DeltaLen;
 			FVec3 Normal = -Dir;
 			FVec3 LocationA = P1 + Dir * A.GetRadius();
@@ -501,9 +501,9 @@ namespace Chaos
 		return Result;
 	}
 
-	FContactPoint CapsuleBoxContactPoint(const FCapsule& A, const FRigidTransform3& ATransform, const FImplicitBox3& B, const FRigidTransform3& BTransform, const FVec3& InitialDir, const FReal ShapePadding)
+	FContactPoint CapsuleBoxContactPoint(const FCapsule& A, const FRigidTransform3& ATransform, const FImplicitBox3& B, const FRigidTransform3& BTransform, const FVec3& InitialDir)
 	{
-		return GJKContactPoint(A, ATransform, B, BTransform, InitialDir, ShapePadding);
+		return GJKContactPoint(A, ATransform, B, BTransform, InitialDir);
 	}
 
 
@@ -516,10 +516,10 @@ namespace Chaos
 	template FContactPoint CapsuleTriangleMeshSweptContactPoint<FTriangleMeshImplicitObject>(const FCapsule& A, const FRigidTransform3& ATransform, const FTriangleMeshImplicitObject& B, const FRigidTransform3& BStartTransform, const FVec3& Dir, const FReal Length, const FReal IgnorePenetration, const FReal TargetPenetration, FReal& TOI);
 	template FContactPoint ConvexTriangleMeshSweptContactPoint<TImplicitObjectScaled<class Chaos::FTriangleMeshImplicitObject, 1>>(const FImplicitObject& A, const FRigidTransform3& ATransform, const TImplicitObjectScaled<class Chaos::FTriangleMeshImplicitObject, 1>& B, const FRigidTransform3& BStartTransform, const FVec3& Dir, const FReal Length, const FReal IgnorePenetration, const FReal TargetPenetration, FReal& TOI);
 	template FContactPoint ConvexTriangleMeshSweptContactPoint<FTriangleMeshImplicitObject>(const FImplicitObject& A, const FRigidTransform3& ATransform, const FTriangleMeshImplicitObject& B, const FRigidTransform3& BStartTransform, const FVec3& Dir, const FReal Length, const FReal IgnorePenetration, const FReal TargetPenetration, FReal& TOI);
-	template FContactPoint BoxTriangleMeshContactPoint<TImplicitObjectScaled<class Chaos::FTriangleMeshImplicitObject, 1>>(const FImplicitBox3& A, const FRigidTransform3& ATransform, const TImplicitObjectScaled<class Chaos::FTriangleMeshImplicitObject, 1>& B, const FRigidTransform3& BTransform, const FReal CullDistance, const FReal ShapePadding);
-	template FContactPoint BoxTriangleMeshContactPoint<FTriangleMeshImplicitObject>(const FImplicitBox3& A, const FRigidTransform3& ATransform, const FTriangleMeshImplicitObject& B, const FRigidTransform3& BTransform, const FReal CullDistance, const FReal ShapePadding);
-	template FContactPoint SphereTriangleMeshContactPoint<TImplicitObjectScaled<class Chaos::FTriangleMeshImplicitObject, 1>>(const TSphere<FReal, 3>& A, const FRigidTransform3& ATransform, const TImplicitObjectScaled<class Chaos::FTriangleMeshImplicitObject, 1>& B, const FRigidTransform3& BTransform, const FReal CullDistance, const FReal ShapePadding);
-	template FContactPoint SphereTriangleMeshContactPoint<FTriangleMeshImplicitObject>(const TSphere<FReal, 3>& A, const FRigidTransform3& ATransform, const FTriangleMeshImplicitObject& B, const FRigidTransform3& BTransform, const FReal CullDistance, const FReal ShapePadding);
-	template FContactPoint CapsuleTriangleMeshContactPoint<TImplicitObjectScaled<class Chaos::FTriangleMeshImplicitObject, 1>>(const FCapsule& A, const FRigidTransform3& ATransform, const TImplicitObjectScaled<class Chaos::FTriangleMeshImplicitObject, 1>& B, const FRigidTransform3& BTransform, const FReal CullDistance, const FReal ShapePadding);
-	template FContactPoint CapsuleTriangleMeshContactPoint<FTriangleMeshImplicitObject>(const FCapsule& A, const FRigidTransform3& ATransform, const FTriangleMeshImplicitObject& B, const FRigidTransform3& BTransform, const FReal CullDistance, const FReal ShapePadding);
+	template FContactPoint BoxTriangleMeshContactPoint<TImplicitObjectScaled<class Chaos::FTriangleMeshImplicitObject, 1>>(const FImplicitBox3& A, const FRigidTransform3& ATransform, const TImplicitObjectScaled<class Chaos::FTriangleMeshImplicitObject, 1>& B, const FRigidTransform3& BTransform, const FReal CullDistance);
+	template FContactPoint BoxTriangleMeshContactPoint<FTriangleMeshImplicitObject>(const FImplicitBox3& A, const FRigidTransform3& ATransform, const FTriangleMeshImplicitObject& B, const FRigidTransform3& BTransform, const FReal CullDistance);
+	template FContactPoint SphereTriangleMeshContactPoint<TImplicitObjectScaled<class Chaos::FTriangleMeshImplicitObject, 1>>(const TSphere<FReal, 3>& A, const FRigidTransform3& ATransform, const TImplicitObjectScaled<class Chaos::FTriangleMeshImplicitObject, 1>& B, const FRigidTransform3& BTransform, const FReal CullDistance);
+	template FContactPoint SphereTriangleMeshContactPoint<FTriangleMeshImplicitObject>(const TSphere<FReal, 3>& A, const FRigidTransform3& ATransform, const FTriangleMeshImplicitObject& B, const FRigidTransform3& BTransform, const FReal CullDistance);
+	template FContactPoint CapsuleTriangleMeshContactPoint<TImplicitObjectScaled<class Chaos::FTriangleMeshImplicitObject, 1>>(const FCapsule& A, const FRigidTransform3& ATransform, const TImplicitObjectScaled<class Chaos::FTriangleMeshImplicitObject, 1>& B, const FRigidTransform3& BTransform, const FReal CullDistance);
+	template FContactPoint CapsuleTriangleMeshContactPoint<FTriangleMeshImplicitObject>(const FCapsule& A, const FRigidTransform3& ATransform, const FTriangleMeshImplicitObject& B, const FRigidTransform3& BTransform, const FReal CullDistance);
 }
diff --git a/Engine/Source/Runtime/Experimental/Chaos/Private/Chaos/Collision/GJKContactPoint.h b/Engine/Source/Runtime/Experimental/Chaos/Private/Chaos/Collision/GJKContactPoint.h
--- a/Engine/Source/Runtime/Experimental/Chaos/Private/Chaos/Collision/GJKContactPoint.h
+++ b/Engine/Source/Runtime/Experimental/Chaos/Private/Chaos/Collision/GJKContactPoint.h
@@ -12,7 +12,7 @@
 namespace Chaos
 {
 	template <typename GeometryA, typename GeometryB>
-	FContactPoint GJKContactPoint2(const GeometryA& A, const GeometryB& B, const FRigidTransform3& ATM, const FRigidTransform3& BToATM, const FVec3& InitialDir, const FReal ShapePadding)
+	FContactPoint GJKContactPoint2(const GeometryA& A, const GeometryB& B, const FRigidTransform3& ATM, const FRigidTransform3& BToATM, const FVec3& InitialDir)
 	{
 		FContactPoint Contact;
 
@@ -23,9 +23,7 @@ namespace Chaos
 		// Slightly increased epsilon to reduce error in normal for almost touching objects.
 		const FReal Epsilon = 3.e-3f;
 
-		const FReal ThicknessA = 0.5f * ShapePadding;
-		const FReal ThicknessB = 0.5f * ShapePadding;
-		if (GJKPenetration<true>(A, B, BToATM, Penetration, ClosestA, ClosestBInA, Normal, ClosestVertexIndexA, ClosestVertexIndexB, ThicknessA, ThicknessB, InitialDir, Epsilon))
+		if (GJKPenetration<true>(A, B, BToATM, Penetration, ClosestA, ClosestBInA, Normal, ClosestVertexIndexA, ClosestVertexIndexB, FReal(0), FReal(0), InitialDir, Epsilon))
 		{
 			// GJK output is all in the local space of A. We need to transform the B-relative position and the normal in to B-space
 			Contact.ShapeContactPoints[0] = ClosestA;
@@ -38,9 +36,9 @@ namespace Chaos
 	}
 
 	template <typename GeometryA, typename GeometryB>
-	FContactPoint GJKContactPoint(const GeometryA& A, const FRigidTransform3& ATM, const GeometryB& B, const FRigidTransform3& BTM, const FVec3& InitialDir, const FReal ShapePadding)
+	FContactPoint GJKContactPoint(const GeometryA& A, const FRigidTransform3& ATM, const GeometryB& B, const FRigidTransform3& BTM, const FVec3& InitialDir)
 	{
 		const FRigidTransform3 BToATM = BTM.GetRelativeTransform(ATM);
-		return GJKContactPoint2(A, B, ATM, BToATM, InitialDir, ShapePadding);
+		return GJKContactPoint2(A, B, ATM, BToATM, InitialDir);
 	}
 }
\ No newline at end of file
diff --git a/Engine/Source/Runtime/Experimental/Chaos/Private/Chaos/Collision/GJKContactPointSwept.h b/Engine/Source/Runtime/Experimental/Chaos/Private/Chaos/Collision/GJKContactPointSwept.h
--- a/Engine/Source/Runtime/Experimental/Chaos/Private/Chaos/Collision/GJKContactPointSwept.h
+++ b/Engine/Source/Runtime/Experimental/Chaos/Private/Chaos/Collision/GJKContactPointSwept.h
@@ -1,13 +1,14 @@
 // Copyright Epic Games, Inc. All Rights Reserved.
 #pragma once
 
-#include "CoreMinimal.h"
 #include "Chaos/Core.h"
+#include "Chaos/CollisionResolutionTypes.h"
 #include "Chaos/ImplicitFwd.h"
 
 namespace Chaos
 {
-	class FContactPoint;
+	template<typename T> class TContactPoint;
+	using FContactPoint = TContactPoint<FReal>;
 
 	FContactPoint GenericConvexConvexContactPointSwept(const FImplicitObject& A, const FRigidTransform3& AStartTM, const FRigidTransform3& AEndTM, const FImplicitObject& B, const FRigidTransform3& BStartTM, const FRigidTransform3& BEndTM, const FVec3& Dir, const FReal Length, const FReal IgnorePenetration, const FReal TargetPenetration, FReal& TOI);
 }
diff --git a/Engine/Source/Runtime/Experimental/Chaos/Private/Chaos/Collision/PBDCollisionConstraint.cpp b/Engine/Source/Runtime/Experimental/Chaos/Private/Chaos/Collision/PBDCollisionConstraint.cpp
--- a/Engine/Source/Runtime/Experimental/Chaos/Private/Chaos/Collision/PBDCollisionConstraint.cpp
+++ b/Engine/Source/Runtime/Experimental/Chaos/Private/Chaos/Collision/PBDCollisionConstraint.cpp
@@ -86,9 +86,9 @@ namespace Chaos
 		//sort constraints by the smallest particle idx in them first
 		//if the smallest particle idx is the same for both, use the other idx
 
-		if (L.GetCCDType() != R.GetCCDType())
+		if (L.GetCCDEnabled() != R.GetCCDEnabled())
 		{
-			return L.GetCCDType() < R.GetCCDType();
+			return L.GetCCDEnabled();
 		}
 
 		const FParticleID ParticleIdxs[] = { L.Particle[0]->ParticleID(), L.Particle[1]->ParticleID() };
@@ -140,7 +140,7 @@ namespace Chaos
 	FPBDCollisionConstraint FPBDCollisionConstraint::MakeTriangle(const FImplicitObject* Implicit0)
 	{
 		FPBDCollisionConstraint Constraint;
-		Constraint.InitMarginsAndTolerances(Implicit0->GetCollisionType(), ImplicitObjectType::Triangle, Implicit0->GetMargin(), FReal(0));
+		Constraint.InitMarginsAndTolerances(Implicit0->GetCollisionType(), ImplicitObjectType::Triangle, FRealSingle(Implicit0->GetMargin()), FRealSingle(0));
 		return Constraint;
 	}
 
@@ -185,15 +185,10 @@ namespace Chaos
 		, Material()
 		, Stiffness(1)
 		, AccumulatedImpulse(0)
-		, CCDTimeOfImpact(0)
-		, CCDEnablePenetration(0)
-		, CCDTargetPenetration(0)
 		, ContainerCookie()
 		, ShapesType(EContactShapesType::Unknown)
-		, CCDType(ECollisionCCDType::Disabled)
-		, ShapeWorldTransform0()
-		, ShapeWorldTransform1()
-		, CullDistance(TNumericLimits<FReal>::Max())
+		, ShapeWorldTransforms{ FRigidTransform3(), FRigidTransform3() }
+		, CullDistance(TNumericLimits<FRealSingle>::Max())
 		, CollisionMargins{ 0, 0 }
 		, CollisionTolerance(0)
 		, ClosestManifoldPointIndex(INDEX_NONE)
@@ -205,6 +200,9 @@ namespace Chaos
 		, GJKWarmStartData()
 		, SavedManifoldPoints()
 		, ManifoldPoints()
+		, CCDTimeOfImpact(0)
+		, CCDEnablePenetration(0)
+		, CCDTargetPenetration(0)
 	{
 	}
 
@@ -225,15 +223,10 @@ namespace Chaos
 		, Material()
 		, Stiffness(1)
 		, AccumulatedImpulse(0)
-		, CCDTimeOfImpact(0)
-		, CCDEnablePenetration(0)
-		, CCDTargetPenetration(0)
 		, ContainerCookie()
 		, ShapesType(EContactShapesType::Unknown)
-		, CCDType(ECollisionCCDType::Disabled)
-		, ShapeWorldTransform0()
-		, ShapeWorldTransform1()
-		, CullDistance(TNumericLimits<FReal>::Max())
+		, ShapeWorldTransforms{ FRigidTransform3(), FRigidTransform3() }
+		, CullDistance(TNumericLimits<FRealSingle>::Max())
 		, CollisionMargins{ 0, 0 }
 		, CollisionTolerance(0)
 		, ClosestManifoldPointIndex(INDEX_NONE)
@@ -245,6 +238,9 @@ namespace Chaos
 		, GJKWarmStartData()
 		, SavedManifoldPoints()
 		, ManifoldPoints()
+		, CCDTimeOfImpact(0)
+		, CCDEnablePenetration(0)
+		, CCDTargetPenetration(0)
 	{
 	}
 
@@ -264,14 +260,12 @@ namespace Chaos
 		const FReal InCullDistance,
 		const bool bInUseManifold)
 	{
-		CCDType = InCCDType;
-
 		ShapesType = InShapesType;
 
 		ImplicitTransform[0] = InImplicitLocalTransform0;
 		ImplicitTransform[1] = InImplicitLocalTransform1;
 
-		CullDistance = InCullDistance;
+		CullDistance = FRealSingle(InCullDistance);
 
 		// Initialize the is-probe and is-probe-unmodified flags to the same value.
 		// Contact modification may change bIsProbe but we want to store the unmodified value
@@ -279,8 +273,12 @@ namespace Chaos
 		Flags.bIsProbe = (Shape[0] && Shape[0]->GetIsProbe()) || (Shape[1] && Shape[1]->GetIsProbe());
 		Flags.bIsProbeUnmodified = Flags.bIsProbe;
 
-		const FReal Margin0 = GetImplicit0()->GetMargin();
-		const FReal Margin1 = GetImplicit1()->GetMargin();
+		// If this currently a CCD contact (this may change on later ticks)
+		Flags.bCCDEnabled = (InCCDType == ECollisionCCDType::Enabled);
+
+		// Initialize tolerances that depend on shape type etc, also the bIsQuadratic flags
+		const FRealSingle Margin0 = FRealSingle(GetImplicit0()->GetMargin());
+		const FRealSingle Margin1 = FRealSingle(GetImplicit1()->GetMargin());
 		const EImplicitObjectType ImplicitType0 = GetInnerType(GetImplicit0()->GetCollisionType());
 		const EImplicitObjectType ImplicitType1 = GetInnerType(GetImplicit1()->GetCollisionType());
 		InitMarginsAndTolerances(ImplicitType0, ImplicitType1, Margin0, Margin1);
@@ -302,7 +300,7 @@ namespace Chaos
 		}
 	}
 
-	void FPBDCollisionConstraint::InitMarginsAndTolerances(const EImplicitObjectType ImplicitType0, const EImplicitObjectType ImplicitType1, const FReal Margin0, const FReal Margin1)
+	void FPBDCollisionConstraint::InitMarginsAndTolerances(const EImplicitObjectType ImplicitType0, const EImplicitObjectType ImplicitType1, const FRealSingle Margin0, const FRealSingle Margin1)
 	{
 		// Set up the margins and tolerances to be used during the narrow phase.
 		// When convex margins are enabled, at least one shape in a collision will always have a margin.
@@ -323,26 +321,26 @@ namespace Chaos
 		
 		// @todo(chaos): should probably be tunable. Used to use the same settings as the margin scale (for convex), but we want to support zero
 		// margins, but still have a non-zero CollisionTolerance (it is used for matching contact points for friction and manifold reuse)
-		const FReal ToleranceScale = 0.1f;
-		const FReal QuadraticToleranceScale = 0.05f;
+		const FRealSingle ToleranceScale = 0.1f;
+		const FRealSingle QuadraticToleranceScale = 0.05f;
 		
 		if (!bIsQuadratic0 && !bIsQuadratic1)
 		{
-			const FReal MaxSize0 = ((Implicit[0] != nullptr) && Implicit[0]->HasBoundingBox()) ? Implicit[0]->BoundingBox().Extents().GetAbsMax() : FReal(0);
-			const FReal MaxSize1 = ((Implicit[1] != nullptr) && Implicit[1]->HasBoundingBox()) ? Implicit[1]->BoundingBox().Extents().GetAbsMax() : FReal(0);
-			const FReal MaxSize = FMath::Min(MaxSize0, MaxSize1);
+			const FRealSingle MaxSize0 = ((Implicit[0] != nullptr) && Implicit[0]->HasBoundingBox()) ? FRealSingle(Implicit[0]->BoundingBox().Extents().GetAbsMax()) : FRealSingle(0);
+			const FRealSingle MaxSize1 = ((Implicit[1] != nullptr) && Implicit[1]->HasBoundingBox()) ? FRealSingle(Implicit[1]->BoundingBox().Extents().GetAbsMax()) : FRealSingle(0);
+			const FRealSingle MaxSize = FMath::Min(MaxSize0, MaxSize1);
 			CollisionTolerance = ToleranceScale * MaxSize;
 
 			// If one shape has a zero margin, enforce a minimum margin to avoid the EPA degenerate case. E.g., Box-Triangle
 			// If both shapes have a margin, use the smaller margin on both shapes. E.g., Box-Box
 			// We should never see both shapes with zero margin, but if we did we'd end up with a zero margin
-			const FReal MinMargin = Chaos_Collision_ConvexZeroMargin;
-			if (Margin0 == FReal(0))
+			const FRealSingle MinMargin = Chaos_Collision_ConvexZeroMargin;
+			if (Margin0 == FRealSingle(0))
 			{
 				CollisionMargins[0] = 0;
 				CollisionMargins[1] = FMath::Max(MinMargin, Margin1);
 			}
-			else if (Margin1 == FReal(0))
+			else if (Margin1 == FRealSingle(0))
 			{
 				CollisionMargins[0] = FMath::Max(MinMargin, Margin0);
 				CollisionMargins[1] = 0;
@@ -379,9 +377,9 @@ namespace Chaos
 	void FPBDCollisionConstraint::InitCCDThreshold()
 	{
 		// Calculate the depth at which we enforce CCD constraints, and the amount of penetration to leave when resolving a CCD constraint.
-		const FReal MinBounds0 = FConstGenericParticleHandle(Particle[0])->CCDEnabled() ? Implicit[0]->BoundingBox().Extents().GetAbsMin() : FReal(0);
-		const FReal MinBounds1 = FConstGenericParticleHandle(Particle[1])->CCDEnabled() ? Implicit[1]->BoundingBox().Extents().GetAbsMin() : FReal(0);
-		const FReal MinBounds = FMath::Max(MinBounds0, MinBounds1);
+		const FRealSingle MinBounds0 = FConstGenericParticleHandle(Particle[0])->CCDEnabled() ? FRealSingle(Implicit[0]->BoundingBox().Extents().GetAbsMin()) : FRealSingle(0);
+		const FRealSingle MinBounds1 = FConstGenericParticleHandle(Particle[1])->CCDEnabled() ? FRealSingle(Implicit[1]->BoundingBox().Extents().GetAbsMin()) : FRealSingle(0);
+		const FRealSingle MinBounds = FMath::Max(MinBounds0, MinBounds1);
 		CCDEnablePenetration = MinBounds * CVars::CCDEnableThresholdBoundsScale;
 		CCDTargetPenetration = FMath::Min(MinBounds * CVars::CCDAllowedDepthBoundsScale, CCDEnablePenetration);
 	}
@@ -408,8 +406,8 @@ namespace Chaos
 	{
 		if (ClosestManifoldPointIndex != INDEX_NONE)
 		{
-			const FVec3 WorldContact0 = GetShapeWorldTransform0().TransformPositionNoScale(ManifoldPoints[ClosestManifoldPointIndex].ContactPoint.ShapeContactPoints[0]);
-			const FVec3 WorldContact1 = GetShapeWorldTransform1().TransformPositionNoScale(ManifoldPoints[ClosestManifoldPointIndex].ContactPoint.ShapeContactPoints[1]);
+			const FVec3 WorldContact0 = GetShapeWorldTransform0().TransformPositionNoScale(FVec3(ManifoldPoints[ClosestManifoldPointIndex].ContactPoint.ShapeContactPoints[0]));
+			const FVec3 WorldContact1 = GetShapeWorldTransform1().TransformPositionNoScale(FVec3(ManifoldPoints[ClosestManifoldPointIndex].ContactPoint.ShapeContactPoints[1]));
 			return FReal(0.5) * (WorldContact0 + WorldContact1);
 		}
 		return FVec3(0);
@@ -419,7 +417,7 @@ namespace Chaos
 	{
 		if (ClosestManifoldPointIndex != INDEX_NONE)
 		{
-			return GetShapeWorldTransform1().TransformVectorNoScale(ManifoldPoints[ClosestManifoldPointIndex].ContactPoint.ShapeContactNormal);
+			return GetShapeWorldTransform1().TransformVectorNoScale(FVec3(ManifoldPoints[ClosestManifoldPointIndex].ContactPoint.ShapeContactNormal));
 		}
 		return FVec3(0, 0, 1);
 	}
@@ -516,10 +514,10 @@ namespace Chaos
 	void FPBDCollisionConstraint::UpdateManifoldPointPhi(const int32 ManifoldPointIndex)
 	{
 		FManifoldPoint& ManifoldPoint = ManifoldPoints[ManifoldPointIndex];
-		const FVec3 WorldContact0 = GetShapeWorldTransform0().TransformPositionNoScale(ManifoldPoint.ContactPoint.ShapeContactPoints[0]);
-		const FVec3 WorldContact1 = GetShapeWorldTransform1().TransformPositionNoScale(ManifoldPoint.ContactPoint.ShapeContactPoints[1]);
-		const FVec3 WorldContactNormal = GetShapeWorldTransform1().TransformVectorNoScale(ManifoldPoint.ContactPoint.ShapeContactNormal);
-		ManifoldPoint.ContactPoint.Phi = FVec3::DotProduct(WorldContact0 - WorldContact1, WorldContactNormal);
+		const FVec3 WorldContact0 = GetShapeWorldTransform0().TransformPositionNoScale(FVec3(ManifoldPoint.ContactPoint.ShapeContactPoints[0]));
+		const FVec3 WorldContact1 = GetShapeWorldTransform1().TransformPositionNoScale(FVec3(ManifoldPoint.ContactPoint.ShapeContactPoints[1]));
+		const FVec3 WorldContactNormal = GetShapeWorldTransform1().TransformVectorNoScale(FVec3(ManifoldPoint.ContactPoint.ShapeContactNormal));
+		ManifoldPoint.ContactPoint.Phi = FRealSingle(FVec3::DotProduct(WorldContact0 - WorldContact1, WorldContactNormal));
 	}
 
 	void FPBDCollisionConstraint::UpdateManifoldContacts()
@@ -550,12 +548,11 @@ namespace Chaos
 				P1->RotationOfMass().UnrotateVector(ImplicitTransform[1].GetLocation() - P1->CenterOfMass()),
 				P1->RotationOfMass().Inverse() * ImplicitTransform[1].GetRotation());
 
-			ShapeWorldTransform0 = ShapeCoMRelativeTransform0 * ParticleCoMTransform0;
-			ShapeWorldTransform1 = ShapeCoMRelativeTransform1 * ParticleCoMTransform1;
+			ShapeWorldTransforms[0] = ShapeCoMRelativeTransform0 * ParticleCoMTransform0;
+			ShapeWorldTransforms[1] = ShapeCoMRelativeTransform1 * ParticleCoMTransform1;
 		}
 
 		ClosestManifoldPointIndex = INDEX_NONE;
-		Material.Reset();
 
 		for (int32 ManifoldPointIndex = 0; ManifoldPointIndex < ManifoldPoints.Num(); ManifoldPointIndex++)
 		{
@@ -622,7 +619,6 @@ namespace Chaos
 	void FPBDCollisionConstraint::ResetActiveManifoldContacts()
 	{
 		ClosestManifoldPointIndex = INDEX_NONE;
-		Material.Reset();
 		ManifoldPoints.Reset();
 		ExpectedNumManifoldPoints = 0;
 		Flags.bWasManifoldRestored = false;
@@ -638,7 +634,6 @@ namespace Chaos
 
 		// Reset current closest point
 		ClosestManifoldPointIndex = INDEX_NONE;
-		Material.Reset();
 
 		// How many manifold points we expect. E.g., for Box-box this will be 4 or 1 depending on whether
 		// we have a face or edge contact. We don't reuse the manifold if we lose points after culling
@@ -653,11 +648,11 @@ namespace Chaos
 			// The transform check is necessary regardless of how many points we have left in the manifold because
 			// as a body moves/rotates we may have to change which faces/edges are colliding. We can't know if the face/edge
 			// will change until we run the closest-point checks (GJK) in the narrow phase.
-			const FVec3 Shape1ToShape0Translation = ShapeWorldTransform0.GetTranslation() - ShapeWorldTransform1.GetTranslation();
+			const FVec3 Shape1ToShape0Translation = ShapeWorldTransforms[0].GetTranslation() - ShapeWorldTransforms[1].GetTranslation();
 			const FVec3 TranslationDelta = Shape1ToShape0Translation - LastShapeWorldPositionDelta;
 			if (TranslationDelta.IsNearlyZero(ShapePositionTolerance))
 			{
-				const FRotation3 Shape1toShape0Rotation = ShapeWorldTransform0.GetRotation().Inverse() * ShapeWorldTransform1.GetRotation();
+				const FRotation3 Shape1toShape0Rotation = ShapeWorldTransforms[0].GetRotation().Inverse() * ShapeWorldTransforms[1].GetRotation();
 				const FRotation3 OriginalShape1toShape0Rotation = LastShapeWorldRotationDelta;
 				const FReal RotationOverlap = FRotation3::DotProduct(Shape1toShape0Rotation, LastShapeWorldRotationDelta);
 				if (RotationOverlap > ShapeRotationThreshold)
@@ -678,7 +673,7 @@ namespace Chaos
 		int32 ManifoldPointToRemove = INDEX_NONE;
 		if (ManifoldPoints.Num() > 0)
 		{
-			const FRigidTransform3 Shape0ToShape1Transform = ShapeWorldTransform0.GetRelativeTransformNoScale(ShapeWorldTransform1);
+			const FRigidTransform3 Shape0ToShape1Transform = ShapeWorldTransforms[0].GetRelativeTransformNoScale(ShapeWorldTransforms[1]);
 			
 			// Update or prune manifold points. If we would end up removing more than 1 point, we just throw the 
 			// whole manifold away because it will get rebuilt in the narrow phasee anyway.
@@ -688,7 +683,7 @@ namespace Chaos
 
 				// Calculate the world-space contact location and separation at the current shape transforms
 				// @todo(chaos): this should use the normal owner. Currently we assume body 1 is the owner
-				const FVec3 Contact0In1 = Shape0ToShape1Transform.TransformPositionNoScale(ManifoldPoint.InitialShapeContactPoints[0]);
+				const FVec3 Contact0In1 = Shape0ToShape1Transform.TransformPositionNoScale(FVec3(ManifoldPoint.InitialShapeContactPoints[0]));
 				const FVec3& Contact1In1 = ManifoldPoint.InitialShapeContactPoints[1];
 				const FVec3& ContactNormalIn1 = ManifoldPoint.ContactPoint.ShapeContactNormal;
 
@@ -703,8 +698,8 @@ namespace Chaos
 					// Recalculate the contact points at the new location
 					// @todo(chaos): we should reproject the contact on the plane owner
 					const FVec3 ShapeContactPoint1 = Contact0In1 - ContactPhi * ContactNormalIn1;
-					ManifoldPoint.ContactPoint.ShapeContactPoints[1] = ShapeContactPoint1;
-					ManifoldPoint.ContactPoint.Phi = ContactPhi;
+					ManifoldPoint.ContactPoint.ShapeContactPoints[1] = FVec3f(ShapeContactPoint1);
+					ManifoldPoint.ContactPoint.Phi = FRealSingle(ContactPhi);
 					ManifoldPoint.Flags.bWasRestored = true;
 					ManifoldPoint.Flags.bWasReplaced = false;
 					if (ManifoldPoint.ContactPoint.Phi < GetPhi())
diff --git a/Engine/Source/Runtime/Experimental/Chaos/Private/Chaos/Collision/ParticlePairMidPhase.cpp b/Engine/Source/Runtime/Experimental/Chaos/Private/Chaos/Collision/ParticlePairMidPhase.cpp
--- a/Engine/Source/Runtime/Experimental/Chaos/Private/Chaos/Collision/ParticlePairMidPhase.cpp
+++ b/Engine/Source/Runtime/Experimental/Chaos/Private/Chaos/Collision/ParticlePairMidPhase.cpp
@@ -120,9 +120,9 @@ namespace Chaos
 		, Particle1(InParticle1)
 		, Shape0(InShape0)
 		, Shape1(InShape1)
-		, ShapePairType(InShapePairType)
 		, SphereBoundsCheckSize(0)
 		, LastUsedEpoch(-1)
+		, ShapePairType(InShapePairType)
 		, Flags()
 	{
 		const FImplicitObject* Implicit0 = Shape0->GetLeafGeometry();
@@ -146,7 +146,7 @@ namespace Chaos
 
 		if (bAllowBoundsChecked && bIsSphere0 && bIsSphere1)
 		{
-			SphereBoundsCheckSize = Implicit0->GetMargin() + Implicit1->GetMargin();	// Sphere-Sphere bounds test
+			SphereBoundsCheckSize = FRealSingle(Implicit0->GetMargin() + Implicit1->GetMargin());	// Sphere-Sphere bounds test
 		}
 
 		// Do not try to reuse manifold points for capsules or spheres (against anything)
@@ -168,8 +168,8 @@ namespace Chaos
 		, Particle1(R.Particle1)
 		, Shape0(R.Shape0)
 		, Shape1(R.Shape1)
-		, ShapePairType(R.ShapePairType)
 		, SphereBoundsCheckSize(R.SphereBoundsCheckSize)
+		, ShapePairType(R.ShapePairType)
 		, Flags(R.Flags)
 	{
 	}
@@ -191,11 +191,11 @@ namespace Chaos
 		}
 
 		// World-space sphere bounds check
-		if (SphereBoundsCheckSize > FReal(0))
+		if (SphereBoundsCheckSize > FRealSingle(0))
 		{
 			const FVec3 Separation = ShapeWorldBounds0.GetCenter() - ShapeWorldBounds1.GetCenter();
 			const FReal SeparationSq = Separation.SizeSquared();
-			const FReal CullDistanceSq = FMath::Square(CullDistance + SphereBoundsCheckSize);
+			const FReal CullDistanceSq = FMath::Square(CullDistance + FReal(SphereBoundsCheckSize));
 			if (SeparationSq > CullDistanceSq)
 			{
 				return false;
@@ -212,8 +212,8 @@ namespace Chaos
 		const bool bCollidedLastTick = IsUsedSince(LastEpoch);
 		if ((Flags.bEnableOBBCheck0 || Flags.bEnableOBBCheck1) && !bCollidedLastTick)
 		{
-			const FRigidTransform3& ShapeWorldTransform0 = Shape0->GetLeafWorldTransform(Particle0);
-			const FRigidTransform3& ShapeWorldTransform1 = Shape1->GetLeafWorldTransform(Particle1);
+			const FRigidTransform3& ShapeWorldTransform0 = Shape0->GetLeafWorldTransform(GetParticle0());
+			const FRigidTransform3& ShapeWorldTransform1 = Shape1->GetLeafWorldTransform(GetParticle1());
 			const FImplicitObject* Implicit0 = Shape0->GetLeafGeometry();
 			const FImplicitObject* Implicit1 = Shape1->GetLeafGeometry();
 
@@ -262,7 +262,7 @@ namespace Chaos
 		PHYSICS_CSV_SCOPED_EXPENSIVE(PhysicsVerbose, NarrowPhase_CreateConstraint);
 		check(!Constraint.IsValid());
 
-		Constraint = CreateShapePairConstraint(Particle0, Shape0, Particle1, Shape1, CullDistance, ShapePairType, Context.GetSettings().bAllowManifolds, MidPhase.GetCollisionAllocator());
+		Constraint = CreateShapePairConstraint(GetParticle0(), Shape0, GetParticle1(), Shape1, CullDistance, ShapePairType, Context.GetSettings().bAllowManifolds, MidPhase.GetCollisionAllocator());
 
 		Constraint->GetContainerCookie().MidPhase = &MidPhase;
 		Constraint->GetContainerCookie().bIsMultiShapePair = false;
@@ -295,8 +295,8 @@ namespace Chaos
 
 			const FImplicitObject* Implicit0 = Shape0->GetLeafGeometry();
 			const FImplicitObject* Implicit1 = Shape1->GetLeafGeometry();
-			const FRigidTransform3& ShapeWorldTransform0 = Shape0->GetLeafWorldTransform(Particle0);
-			const FRigidTransform3& ShapeWorldTransform1 = Shape1->GetLeafWorldTransform(Particle1);
+			const FRigidTransform3& ShapeWorldTransform0 = Shape0->GetLeafWorldTransform(GetParticle0());
+			const FRigidTransform3& ShapeWorldTransform1 = Shape1->GetLeafWorldTransform(GetParticle1());
 			const int32 CurrentEpoch = MidPhase.GetCollisionAllocator().GetCurrentEpoch();
 			const int32 LastEpoch = CurrentEpoch - 1;
 			const bool bWasUpdatedLastTick = IsUsedSince(LastEpoch);
@@ -380,8 +380,8 @@ namespace Chaos
 
 			const FImplicitObject* Implicit0 = Shape0->GetLeafGeometry();
 			const FImplicitObject* Implicit1 = Shape1->GetLeafGeometry();
-			const FRigidTransform3& ShapeWorldTransform0 = Shape0->GetLeafWorldTransform(Particle0);
-			const FRigidTransform3& ShapeWorldTransform1 = Shape1->GetLeafWorldTransform(Particle1);
+			const FRigidTransform3& ShapeWorldTransform0 = Shape0->GetLeafWorldTransform(GetParticle0());
+			const FRigidTransform3& ShapeWorldTransform1 = Shape1->GetLeafWorldTransform(GetParticle1());
 
 			// Update the world shape transforms on the constraint (we cannot just give it the PerShapeData 
 			// pointer because of Unions - see FMultiShapePairCollisionDetector)
@@ -394,8 +394,8 @@ namespace Chaos
 			Constraint->ResetManifold();
 			Constraint->ResetActiveManifoldContacts();
 
-			FConstGenericParticleHandle P0 = Particle0;
-			FConstGenericParticleHandle P1 = Particle1;
+			FConstGenericParticleHandle P0 = GetParticle0();
+			FConstGenericParticleHandle P1 = GetParticle1();
 			// For kinematic particles, X = P (at TOI=1), we need to compute P-V*dt to get position at TOI=0. 
 			const FVec3 StartX0 = P0->ObjectState() == EObjectStateType::Kinematic ? P0->P() - P0->V() * Dt : P0->X();
 			const FVec3 StartX1 = P1->ObjectState() == EObjectStateType::Kinematic ? P1->P() - P1->V() * Dt : P1->X();
@@ -467,7 +467,7 @@ namespace Chaos
 			// We have skipped collision detection for this particle because it was asleep, so we need to update the transforms...
 			// NOTE: this relies on the shape world transforms being up-to-date. They are usually updated in Integarte which
 			// is also skipped for sleeping particles, so they must be updated manually when waking partciles (see IslandManager)
-			Constraint->SetShapeWorldTransforms(Shape0->GetLeafWorldTransform(Particle0), Shape1->GetLeafWorldTransform(Particle1));
+			Constraint->SetShapeWorldTransforms(Shape0->GetLeafWorldTransform(GetParticle0()), Shape1->GetLeafWorldTransform(GetParticle1()));
 		}
 	}
 
diff --git a/Engine/Source/Runtime/Experimental/Chaos/Private/Chaos/Collision/SolverCollisionContainer.cpp b/Engine/Source/Runtime/Experimental/Chaos/Private/Chaos/Collision/SolverCollisionContainer.cpp
--- a/Engine/Source/Runtime/Experimental/Chaos/Private/Chaos/Collision/SolverCollisionContainer.cpp
+++ b/Engine/Source/Runtime/Experimental/Chaos/Private/Chaos/Collision/SolverCollisionContainer.cpp
@@ -251,11 +251,11 @@ namespace Chaos
 					continue;
 				}
 
-				const FVec3 WorldContactPoint0 = ShapeWorldTransform0.TransformPositionNoScale(ManifoldPoint.ContactPoint.ShapeContactPoints[0]);
-				const FVec3 WorldContactPoint1 = ShapeWorldTransform1.TransformPositionNoScale(ManifoldPoint.ContactPoint.ShapeContactPoints[1]);
+				const FVec3 WorldContactPoint0 = ShapeWorldTransform0.TransformPositionNoScale(FVec3(ManifoldPoint.ContactPoint.ShapeContactPoints[0]));
+				const FVec3 WorldContactPoint1 = ShapeWorldTransform1.TransformPositionNoScale(FVec3(ManifoldPoint.ContactPoint.ShapeContactPoints[1]));
 				const FVec3 WorldContactPoint = FReal(0.5) * (WorldContactPoint0 + WorldContactPoint1);
 
-				const FSolverVec3 WorldContactNormal = FSolverVec3(ShapeWorldTransform1.TransformVectorNoScale(ManifoldPoint.ContactPoint.ShapeContactNormal));
+				const FSolverVec3 WorldContactNormal = FSolverVec3(ShapeWorldTransform1.TransformVectorNoScale(FVec3(ManifoldPoint.ContactPoint.ShapeContactNormal)));
 				const FSolverVec3 RelativeContactPosition0 = FSolverVec3(WorldContactPoint - Body0.P());
 				const FSolverVec3 RelativeContactPosition1 = FSolverVec3(WorldContactPoint - Body1.P());
 				const FSolverReal TargetPhi = FSolverReal(ManifoldPoint.TargetPhi);
@@ -333,8 +333,6 @@ namespace Chaos
 		*/
 		void ScatterOutput(const FReal Dt)
 		{
-			FVec3 AccumulatedImpulse = FVec3(0);
-
 			Constraint->ResetSolverResults();
 
 			// NOTE: We only put the non-pruned manifold points into the solver so the ManifoldPointIndex and
@@ -342,9 +340,9 @@ namespace Chaos
 			int32 SolverManifoldPointIndex = 0;
 			for (int32 ManifoldPointIndex = 0; ManifoldPointIndex < Constraint->NumManifoldPoints(); ++ManifoldPointIndex)
 			{
-				FSolverVec3 NetPushOut = FVec3(0);
-				FSolverVec3 NetImpulse = FVec3(0);
-				FReal StaticFrictionRatio = FReal(0);
+				FSolverVec3 NetPushOut = FSolverVec3(0);
+				FSolverVec3 NetImpulse = FSolverVec3(0);
+				FRealSingle StaticFrictionRatio = FRealSingle(0);
 
 				if (!Constraint->GetManifoldPoint(ManifoldPointIndex).Flags.bDisabled)
 				{
@@ -370,7 +368,7 @@ namespace Chaos
 					NetPushOut, 
 					NetImpulse, 
 					StaticFrictionRatio,
-					Dt);
+					FRealSingle(Dt));
 			}
 
 			Unbind();
diff --git a/Engine/Source/Runtime/Experimental/Chaos/Private/Chaos/Collision/SphereConvexContactPoint.cpp b/Engine/Source/Runtime/Experimental/Chaos/Private/Chaos/Collision/SphereConvexContactPoint.cpp
--- a/Engine/Source/Runtime/Experimental/Chaos/Private/Chaos/Collision/SphereConvexContactPoint.cpp
+++ b/Engine/Source/Runtime/Experimental/Chaos/Private/Chaos/Collision/SphereConvexContactPoint.cpp
@@ -55,7 +55,7 @@ namespace Chaos
 			// Run GJK to find separating distance if available
 			// NOTE: Sphere is treated as a point (its core shape), Convex margin is ignored so we are using the outer non-shrunken hull.
 			// Sphere is also created in the space of the convex to eliminate the transform in the Support call (per GJK iteration).
-			const FVec3 SpherePosConvex = SphereToConvexTransform.TransformPosition(Sphere.GetCenter());
+			const FVec3 SpherePosConvex = SphereToConvexTransform.TransformPositionNoScale(Sphere.GetCenter());
 			const FReal SphereRadius = Sphere.GetRadius();
 
 			// A point in the Minkowski Sum (A-B) in a likely direction
@@ -83,7 +83,7 @@ namespace Chaos
 			{
 				// Results so far are all in convex-space and for a point rather than a sphere. Convert them.
 				PosConvex = SpherePosConvex - Phi * NormalConvex;
-				PosSphere = Sphere.GetCenter() + SphereToConvexTransform.InverseTransformVector(-SphereRadius * NormalConvex);
+				PosSphere = Sphere.GetCenter() + SphereToConvexTransform.InverseTransformVectorNoScale(-SphereRadius * NormalConvex);
 				Phi = Phi - SphereRadius;
 
 				FContactPoint ContactPoint;
@@ -98,35 +98,19 @@ namespace Chaos
 		return FContactPoint();
 	}
 
-	FContactPoint SphereConvexContactPoint(const FImplicitSphere3& Sphere, const FImplicitConvex3& Convex, const FRigidTransform3& SphereToConvexTransform)
-	{
-		return SphereConvexContactPointImpl(Sphere, Convex, SphereToConvexTransform);
-	}
-
-	FContactPoint SphereConvexContactPoint(const FImplicitSphere3& Sphere, const TImplicitObjectInstanced<FImplicitConvex3>& Convex, const FRigidTransform3& SphereToConvexTransform)
-	{
-		return SphereConvexContactPointImpl(Sphere, Convex, SphereToConvexTransform);
-	}
-
-	FContactPoint SphereConvexContactPoint(const FImplicitSphere3& Sphere, const TImplicitObjectScaled<FImplicitConvex3>& Convex, const FRigidTransform3& SphereToConvexTransform)
-	{
-		return SphereConvexContactPointImpl(Sphere, Convex, SphereToConvexTransform);
-	}
-
-
 	FContactPoint SphereConvexContactPoint(const FImplicitSphere3& Sphere, const FImplicitObject& Object, const FRigidTransform3& SphereToConvexTransform)
 	{
 		if (const TImplicitObjectInstanced<FImplicitConvex3>* InstancedConvex = Object.template GetObject<const TImplicitObjectInstanced<FImplicitConvex3>>())
 		{
-			return SphereConvexContactPoint(Sphere, *InstancedConvex, SphereToConvexTransform);
+			return SphereConvexContactPointImpl(Sphere, *InstancedConvex, SphereToConvexTransform);
 		}
 		else if (const TImplicitObjectScaled<FImplicitConvex3>* ScaledConvex = Object.template GetObject<const TImplicitObjectScaled<FImplicitConvex3>>())
 		{
-			return SphereConvexContactPoint(Sphere, *ScaledConvex, SphereToConvexTransform);
+			return SphereConvexContactPointImpl(Sphere, *ScaledConvex, SphereToConvexTransform);
 		}
 		else if (const FImplicitConvex3* Convex = Object.template GetObject<const FImplicitConvex3>())
 		{
-			return SphereConvexContactPoint(Sphere, *Convex, SphereToConvexTransform);
+			return SphereConvexContactPointImpl(Sphere, *Convex, SphereToConvexTransform);
 		}
 		return FContactPoint();
 	}
diff --git a/Engine/Source/Runtime/Experimental/Chaos/Private/Chaos/CollisionOneShotManifolds.cpp b/Engine/Source/Runtime/Experimental/Chaos/Private/Chaos/CollisionOneShotManifolds.cpp
--- a/Engine/Source/Runtime/Experimental/Chaos/Private/Chaos/CollisionOneShotManifolds.cpp
+++ b/Engine/Source/Runtime/Experimental/Chaos/Private/Chaos/CollisionOneShotManifolds.cpp
@@ -93,7 +93,7 @@ namespace Chaos
 	{
 		// Forward delarations we need from CollisionRestitution.cpp
 
-		FContactPoint BoxBoxContactPoint(const FImplicitBox3& Box1, const FImplicitBox3& Box2, const FRigidTransform3& Box1TM, const FRigidTransform3& Box2TM, const FReal ShapePadding);
+		FContactPoint BoxBoxContactPoint(const FImplicitBox3& Box1, const FImplicitBox3& Box2, const FRigidTransform3& Box1TM, const FRigidTransform3& Box2TM);
 
 		//////////////////////////
 		// Box Box
diff --git a/Engine/Source/Runtime/Experimental/Chaos/Private/Chaos/CollisionOneShotManifoldsMiscShapes.cpp b/Engine/Source/Runtime/Experimental/Chaos/Private/Chaos/CollisionOneShotManifoldsMiscShapes.cpp
--- a/Engine/Source/Runtime/Experimental/Chaos/Private/Chaos/CollisionOneShotManifoldsMiscShapes.cpp
+++ b/Engine/Source/Runtime/Experimental/Chaos/Private/Chaos/CollisionOneShotManifoldsMiscShapes.cpp
@@ -60,7 +60,7 @@ namespace Chaos
 			// @todo(chaos): support manifold maintenance
 			Constraint.ResetActiveManifoldContacts();
 
-			FContactPoint ContactPoint = SphereSphereContactPoint(SphereA, SphereATransform, SphereB, SphereBTransform, Constraint.GetCullDistance(), Constraint.GetRestitutionPadding());
+			FContactPoint ContactPoint = SphereSphereContactPoint(SphereA, SphereATransform, SphereB, SphereBTransform, Constraint.GetCullDistance());
 			if (ContactPoint.Phi < Constraint.GetCullDistance())
 			{
 				Constraint.AddOneshotManifoldContact(ContactPoint);
@@ -83,7 +83,7 @@ namespace Chaos
 			// @todo(chaos): support manifold maintenance
 			Constraint.ResetActiveManifoldContacts();
 
-			FContactPoint ContactPoint = SpherePlaneContactPoint(Sphere, SphereTransform, Plane, PlaneTransform, Constraint.GetRestitutionPadding());
+			FContactPoint ContactPoint = SpherePlaneContactPoint(Sphere, SphereTransform, Plane, PlaneTransform);
 			if (ContactPoint.Phi < Constraint.GetCullDistance())
 			{
 				Constraint.AddOneshotManifoldContact(ContactPoint);
@@ -100,7 +100,7 @@ namespace Chaos
 			// @todo(chaos): support manifold maintenance
 			Constraint.ResetActiveManifoldContacts();
 
-			FContactPoint ContactPoint = SphereBoxContactPoint(Sphere, SphereTransform, Box, BoxTransform, Constraint.GetRestitutionPadding());
+			FContactPoint ContactPoint = SphereBoxContactPoint(Sphere, SphereTransform, Box, BoxTransform);
 			if (ContactPoint.Phi < Constraint.GetCullDistance())
 			{
 				Constraint.AddOneshotManifoldContact(ContactPoint);
@@ -123,7 +123,7 @@ namespace Chaos
 			Constraint.ResetActiveManifoldContacts();
 
 			// Build a multi-point manifold
-			const FReal NetCullDistance = Sphere.GetRadius() + Capsule.GetRadius() + Constraint.GetCullDistance() + Constraint.GetRestitutionPadding();
+			const FReal NetCullDistance = Sphere.GetRadius() + Capsule.GetRadius() + Constraint.GetCullDistance();
 			const FReal NetCullDistanceSq = FMath::Square(NetCullDistance);
 
 			// Transform the sphere into capsule space and find the closest point on the capsule line segment
@@ -142,7 +142,7 @@ namespace Chaos
 				{
 					const FReal NearPosDistance = FMath::Sqrt(NearPosDistanceSq);
 					const FVec3 NearPosDir = NearPosDelta / NearPosDistance;
-					const FReal NearPhi = NearPosDistance - Sphere.GetRadius() - Capsule.GetRadius() - Constraint.GetRestitutionPadding();
+					const FReal NearPhi = NearPosDistance - Sphere.GetRadius() - Capsule.GetRadius();
 
 					FContactPoint NearContactPoint;
 					NearContactPoint.ShapeContactPoints[0] = SphereToCapsuleTransform.InverseTransformPositionNoScale(SpherePos - Sphere.GetRadius() * NearPosDir);
@@ -197,7 +197,7 @@ namespace Chaos
 						{
 							const FVec3 EndCapPos0 = Capsule.GetX1() + CapsuleOrthogonal * Capsule.GetRadius();
 							const FReal EndCapDistance0 = (SpherePos - EndCapPos0).Size();
-							const FReal EndCapPhi0 = EndCapDistance0 - Sphere.GetRadius() - Constraint.GetRestitutionPadding();
+							const FReal EndCapPhi0 = EndCapDistance0 - Sphere.GetRadius();
 							
 							if (EndCapPhi0 < Constraint.GetCullDistance())
 							{
@@ -219,7 +219,7 @@ namespace Chaos
 						{
 							const FVec3 EndCapPos1 = Capsule.GetX2() + CapsuleOrthogonal * Capsule.GetRadius();
 							const FReal EndCapDistance1 = (SpherePos - EndCapPos1).Size();
-							const FReal EndCapPhi1 = EndCapDistance1 - Sphere.GetRadius() - Constraint.GetRestitutionPadding();
+							const FReal EndCapPhi1 = EndCapDistance1 - Sphere.GetRadius();
 
 							if (EndCapPhi1 < Constraint.GetCullDistance())
 							{
@@ -350,7 +350,7 @@ namespace Chaos
 			// @todo(chaos): support manifold maintenance
 			Constraint.ResetActiveManifoldContacts();
 
-			FContactPoint ContactPoint = SphereTriangleMeshContactPoint(Sphere, SphereWorldTransform, TriangleMesh, TriMeshWorldTransform, Constraint.GetCullDistance(), 0.0f);
+			FContactPoint ContactPoint = SphereTriangleMeshContactPoint(Sphere, SphereWorldTransform, TriangleMesh, TriMeshWorldTransform, Constraint.GetCullDistance());
 			if (ContactPoint.Phi < Constraint.GetCullDistance())
 			{
 				Constraint.AddOneshotManifoldContact(ContactPoint);
@@ -364,7 +364,7 @@ namespace Chaos
 			ensure(SphereTransform.GetScale3D() == FVec3(1.0f, 1.0f, 1.0f));
 			ensure(HeightfieldTransform.GetScale3D() == FVec3(1.0f, 1.0f, 1.0f));
 
-			FContactPoint ContactPoint = SphereHeightFieldContactPoint(Sphere, SphereTransform, Heightfield, HeightfieldTransform, Constraint.GetCullDistance(), 0.0f);
+			FContactPoint ContactPoint = SphereHeightFieldContactPoint(Sphere, SphereTransform, Heightfield, HeightfieldTransform, Constraint.GetCullDistance());
 			if (ContactPoint.Phi < Constraint.GetCullDistance())
 			{
 				Constraint.AddOneshotManifoldContact(ContactPoint);
diff --git a/Engine/Source/Runtime/Experimental/Chaos/Private/Chaos/CollisionResolution.cpp b/Engine/Source/Runtime/Experimental/Chaos/Private/Chaos/CollisionResolution.cpp
--- a/Engine/Source/Runtime/Experimental/Chaos/Private/Chaos/CollisionResolution.cpp
+++ b/Engine/Source/Runtime/Experimental/Chaos/Private/Chaos/CollisionResolution.cpp
@@ -329,9 +329,9 @@ namespace Chaos
 		struct TConvexConvexUpdater
 		{
 			template<typename T_ImplicitA, typename T_ImplicitB>
-			static void UpdateConvexConvexConstraint(const T_ImplicitA& A, const FRigidTransform3& ATM, const T_ImplicitB& B, const FRigidTransform3& BTM, const FReal Dt, const FReal ShapePadding, FPBDCollisionConstraint& Constraint)
+			static void UpdateConvexConvexConstraint(const T_ImplicitA& A, const FRigidTransform3& ATM, const T_ImplicitB& B, const FRigidTransform3& BTM, const FReal Dt, FPBDCollisionConstraint& Constraint)
 			{
-				UpdateContactPoint(Constraint, GJKContactPoint(A, ATM, B, BTM, FVec3(1, 0, 0), ShapePadding), Dt);
+				UpdateContactPoint(Constraint, GJKContactPoint(A, ATM, B, BTM, FVec3(1, 0, 0)), Dt);
 			}
 		};
 
@@ -340,7 +340,7 @@ namespace Chaos
 		struct TConvexConvexUpdater<true>
 		{
 			template<typename T_ImplicitA, typename T_ImplicitB>
-			static void UpdateConvexConvexConstraint(const T_ImplicitA& A, const FRigidTransform3& ATM, const T_ImplicitB& B, const FRigidTransform3& BTM, const FReal Dt, const FReal ShapePadding, FPBDCollisionConstraint& Constraint)
+			static void UpdateConvexConvexConstraint(const T_ImplicitA& A, const FRigidTransform3& ATM, const T_ImplicitB& B, const FRigidTransform3& BTM, const FReal Dt, FPBDCollisionConstraint& Constraint)
 			{
 				if (Constraint.GetUseManifold())
 				{
@@ -348,7 +348,7 @@ namespace Chaos
 				}
 				else
 				{
-					UpdateContactPoint(Constraint, GJKContactPoint(A, ATM, B, BTM, FVec3(1, 0, 0), ShapePadding), Dt);
+					UpdateContactPoint(Constraint, GJKContactPoint(A, ATM, B, BTM, FVec3(1, 0, 0)), Dt);
 				}
 			}
 		};
@@ -358,23 +358,23 @@ namespace Chaos
 		struct FConvexConvexUpdaterCaller
 		{
 			template<typename T_ImplicitA, typename T_ImplicitB>
-			static void Update(const T_ImplicitA& A, const FRigidTransform3& ATM, const T_ImplicitB& B, const FRigidTransform3& BTM, const FReal Dt, const FReal ShapePadding, FPBDCollisionConstraint& Constraint)
+			static void Update(const T_ImplicitA& A, const FRigidTransform3& ATM, const T_ImplicitB& B, const FRigidTransform3& BTM, const FReal Dt, FPBDCollisionConstraint& Constraint)
 			{
 				using FConvexConvexUpdater = TConvexConvexUpdater<TConvexImplicitPairTraits<T_ImplicitA, T_ImplicitB>::bSupportsOneShotManifold>;
-				FConvexConvexUpdater::UpdateConvexConvexConstraint(A, ATM, B, BTM, Dt, ShapePadding, Constraint);
+				FConvexConvexUpdater::UpdateConvexConvexConstraint(A, ATM, B, BTM, Dt, Constraint);
 			}
 		};
 
 
 		// Unwrap the many convex types, including scaled, and call the appropriate update which depends on the concrete types
-		void UpdateGenericConvexConvexConstraintHelper(const FImplicitObject& A, const FRigidTransform3& ATM, const FImplicitObject& B, const FRigidTransform3& BTM, const FReal Dt, const FReal ShapePadding, FPBDCollisionConstraint& Constraint)
+		void UpdateGenericConvexConvexConstraintHelper(const FImplicitObject& A, const FRigidTransform3& ATM, const FImplicitObject& B, const FRigidTransform3& BTM, const FReal Dt, FPBDCollisionConstraint& Constraint)
 		{
 			// This expands to a switch of switches that calls the inner function with the appropriate concrete implicit types
 			Utilities::CastHelperNoUnwrap(A, ATM, [&](const auto& ADowncast, const FRigidTransform3& AFullTM)
 			{
 				Utilities::CastHelperNoUnwrap(B, BTM, [&](const auto& BDowncast, const FRigidTransform3& BFullTM)
 				{
-					FConvexConvexUpdaterCaller::Update(ADowncast, AFullTM, BDowncast, BFullTM, Dt, ShapePadding, Constraint);
+					FConvexConvexUpdaterCaller::Update(ADowncast, AFullTM, BDowncast, BFullTM, Dt, Constraint);
 				});
 			});
 		}
@@ -397,9 +397,9 @@ namespace Chaos
 		// Box - Box
 		//
 
-		FContactPoint BoxBoxContactPoint(const FImplicitBox3& Box1, const FImplicitBox3& Box2, const FRigidTransform3& Box1TM, const FRigidTransform3& Box2TM, const FReal ShapePadding)
+		FContactPoint BoxBoxContactPoint(const FImplicitBox3& Box1, const FImplicitBox3& Box2, const FRigidTransform3& Box1TM, const FRigidTransform3& Box2TM)
 		{
-			return GJKContactPoint(Box1, Box1TM, Box2, Box2TM, FVec3(1, 0, 0), ShapePadding);
+			return GJKContactPoint(Box1, Box1TM, Box2, Box2TM, FVec3(1, 0, 0));
 		}
 
 		void UpdateBoxBoxConstraint(const FImplicitBox3& Box1, const FRigidTransform3& Box1Transform, const FImplicitBox3& Box2, const FRigidTransform3& Box2Transform, const FReal Dt, FPBDCollisionConstraint& Constraint)
@@ -410,7 +410,7 @@ namespace Chaos
 			}
 			else
 			{
-				UpdateContactPoint(Constraint, BoxBoxContactPoint(Box1, Box2, Box1Transform, Box2Transform, Constraint.GetRestitutionPadding()), Dt);
+				UpdateContactPoint(Constraint, BoxBoxContactPoint(Box1, Box2, Box1Transform, Box2Transform), Dt);
 			}
 		}
 
@@ -445,8 +445,7 @@ namespace Chaos
 			}
 			else
 			{
-				// @todo(chaos): restitutionpadding
-				UpdateContactPoint(Constraint, BoxHeightFieldContactPoint(A, ATransform, B, BTransform, Constraint.GetCullDistance(), 0.0f), Dt);
+				UpdateContactPoint(Constraint, BoxHeightFieldContactPoint(A, ATransform, B, BTransform, Constraint.GetCullDistance()), Dt);
 			}
 		}
 
@@ -475,7 +474,7 @@ namespace Chaos
 		// Box-Plane
 		//
 
-		FContactPoint BoxPlaneContactPoint(const FImplicitBox3& Box, const FImplicitPlane3& Plane, const FRigidTransform3& BoxTransform, const FRigidTransform3& PlaneTransform, const FReal ShapePadding)
+		FContactPoint BoxPlaneContactPoint(const FImplicitBox3& Box, const FImplicitPlane3& Plane, const FRigidTransform3& BoxTransform, const FRigidTransform3& PlaneTransform)
 		{
 			FContactPoint Contact;
 
@@ -540,8 +539,7 @@ namespace Chaos
 		void UpdateBoxPlaneConstraint(const FImplicitBox3& A, const FRigidTransform3& ATransform, const TPlane<FReal, 3>& B, const FRigidTransform3& BTransform, const FReal Dt, FPBDCollisionConstraint& Constraint)
 		{
 			// @todo(chaos): one-shot manifold
-			// @todo(chaos): restitutionpadding
-			UpdateContactPoint(Constraint, BoxPlaneContactPoint(A, B, ATransform, BTransform, 0.0f), Dt);
+			UpdateContactPoint(Constraint, BoxPlaneContactPoint(A, B, ATransform, BTransform), Dt);
 		}
 
 
@@ -576,7 +574,7 @@ namespace Chaos
 			else
 			{
 				// @toto(chaos): restitutionpadding
-				UpdateContactPoint(Constraint, BoxTriangleMeshContactPoint(Box0, WorldTransform0, TriangleMesh1, WorldTransform1, Constraint.GetCullDistance(), 0.0f), Dt);
+				UpdateContactPoint(Constraint, BoxTriangleMeshContactPoint(Box0, WorldTransform0, TriangleMesh1, WorldTransform1, Constraint.GetCullDistance()), Dt);
 			}
 		}
 
@@ -627,7 +625,7 @@ namespace Chaos
 			}
 			else
 			{
-				UpdateContactPoint(Constraint, SphereSphereContactPoint(Sphere1, Sphere1Transform, Sphere2, Sphere2Transform, Constraint.GetCullDistance(), Constraint.GetRestitutionPadding()), Dt);
+				UpdateContactPoint(Constraint, SphereSphereContactPoint(Sphere1, Sphere1Transform, Sphere2, Sphere2Transform, Constraint.GetCullDistance()), Dt);
 			}
 		}
 
@@ -660,8 +658,7 @@ namespace Chaos
 			}
 			else
 			{
-				// @todo(chaos): restitution padding
-				UpdateContactPoint(Constraint, SphereHeightFieldContactPoint(A, ATransform, B, BTransform, Constraint.GetCullDistance(), 0.0f), Dt);
+				UpdateContactPoint(Constraint, SphereHeightFieldContactPoint(A, ATransform, B, BTransform, Constraint.GetCullDistance()), Dt);
 			}
 		}
 
@@ -715,7 +712,7 @@ namespace Chaos
 			}
 			else
 			{
-				UpdateContactPoint(Constraint, SpherePlaneContactPoint(Sphere, SphereTransform, Plane, PlaneTransform, Constraint.GetRestitutionPadding()), Dt);
+				UpdateContactPoint(Constraint, SpherePlaneContactPoint(Sphere, SphereTransform, Plane, PlaneTransform), Dt);
 			}
 		}
 
@@ -750,7 +747,7 @@ namespace Chaos
 			}
 			else
 			{
-				UpdateContactPoint(Constraint, SphereBoxContactPoint(Sphere, SphereTransform, Box, BoxTransform, Constraint.GetRestitutionPadding()), Dt);
+				UpdateContactPoint(Constraint, SphereBoxContactPoint(Sphere, SphereTransform, Box, BoxTransform), Dt);
 			}
 		}
 
@@ -786,7 +783,7 @@ namespace Chaos
 			}
 			else
 			{
-				UpdateContactPoint(Constraint, SphereCapsuleContactPoint(A, ATransform, B, BTransform, Constraint.GetRestitutionPadding()), Dt);
+				UpdateContactPoint(Constraint, SphereCapsuleContactPoint(A, ATransform, B, BTransform), Dt);
 			}
 		}
 
@@ -854,7 +851,7 @@ namespace Chaos
 			else
 			{
 				// @todo(chaos): restitutionpadding
-				UpdateContactPoint(Constraint, SphereTriangleMeshContactPoint(Sphere0, WorldTransform0, TriangleMesh1, WorldTransform1, Constraint.GetCullDistance(), 0.0f), Dt);
+				UpdateContactPoint(Constraint, SphereTriangleMeshContactPoint(Sphere0, WorldTransform0, TriangleMesh1, WorldTransform1, Constraint.GetCullDistance()), Dt);
 			}
 		}
 
@@ -941,7 +938,7 @@ namespace Chaos
 			}
 			else
 			{
-				UpdateContactPoint(Constraint, CapsuleCapsuleContactPoint(A, ATransform, B, BTransform, Constraint.GetRestitutionPadding()), Dt);
+				UpdateContactPoint(Constraint, CapsuleCapsuleContactPoint(A, ATransform, B, BTransform), Dt);
 			}
 		}
 
@@ -982,7 +979,7 @@ namespace Chaos
 			else
 			{
 				const FVec3 InitialDir = FVec3(1,0,0);
-				UpdateContactPoint(Constraint, CapsuleBoxContactPoint(A, ATransform, B, BTransform, InitialDir, Constraint.GetRestitutionPadding()), Dt);
+				UpdateContactPoint(Constraint, CapsuleBoxContactPoint(A, ATransform, B, BTransform, InitialDir), Dt);
 			}
 		}
 
@@ -1092,7 +1089,6 @@ namespace Chaos
 
 		void UpdateCapsuleHeightFieldConstraint(const FCapsule& A, const FRigidTransform3& ATransform, const FHeightField& B, const FRigidTransform3& BTransform, const FReal Dt, FPBDCollisionConstraint& Constraint)
 		{
-			// @todo(chaos): restitution padding
 			CONDITIONAL_SCOPE_CYCLE_COUNTER(STAT_Collisions_UpdateCapsuleHeightFieldConstraint, ConstraintsDetailedStats);
 			if (Constraint.GetUseManifold())
 			{
@@ -1100,7 +1096,7 @@ namespace Chaos
 			}
 			else
 			{
-				UpdateContactPoint(Constraint, CapsuleHeightFieldContactPoint(A, ATransform, B, BTransform, Constraint.GetCullDistance(), 0.0f), Dt);
+				UpdateContactPoint(Constraint, CapsuleHeightFieldContactPoint(A, ATransform, B, BTransform, Constraint.GetCullDistance()), Dt);
 			}
 		}
 
@@ -1163,7 +1159,7 @@ namespace Chaos
 			}
 			else
 			{
-				UpdateContactPoint(Constraint, CapsuleTriangleMeshContactPoint(Capsule0, WorldTransform0, TriangleMesh1, WorldTransform1, Constraint.GetCullDistance(), 0.0f), Dt);
+				UpdateContactPoint(Constraint, CapsuleTriangleMeshContactPoint(Capsule0, WorldTransform0, TriangleMesh1, WorldTransform1, Constraint.GetCullDistance()), Dt);
 			}
 		}
 
@@ -1247,7 +1243,7 @@ namespace Chaos
 		{
 			CONDITIONAL_SCOPE_CYCLE_COUNTER(STAT_Collisions_UpdateGenericConvexConvexConstraint, ConstraintsDetailedStats);
 
-			UpdateGenericConvexConvexConstraintHelper(Implicit0, WorldTransform0, Implicit1, WorldTransform1, Dt, Constraint.GetRestitutionPadding(), Constraint);
+			UpdateGenericConvexConvexConstraintHelper(Implicit0, WorldTransform0, Implicit1, WorldTransform1, Dt, Constraint);
 		}
 
 		void UpdateGenericConvexConvexConstraintSwept(TGeometryParticleHandle<FReal, 3>* Particle0, const FImplicitObject& Implicit0, const FRigidTransform3& StartWorldTransform0, TGeometryParticleHandle<FReal, 3>* Particle1, const FImplicitObject& Implicit1, const FRigidTransform3& StartWorldTransform1, const FVec3& Dir, const FReal Length, const FReal Dt, FPBDCollisionConstraint& Constraint)
@@ -1302,7 +1298,6 @@ namespace Chaos
 
 		void UpdateConvexHeightFieldConstraint(const FImplicitObject& A, const FRigidTransform3& ATransform, const FHeightField& B, const FRigidTransform3& BTransform, const FReal Dt, FPBDCollisionConstraint& Constraint)
 		{
-			// @todo(chaos): restitution padding
 			CONDITIONAL_SCOPE_CYCLE_COUNTER(STAT_Collisions_UpdateConvexHeightFieldConstraint, ConstraintsDetailedStats);
 			if (Constraint.GetUseManifold())
 			{
@@ -1310,7 +1305,7 @@ namespace Chaos
 			}
 			else
 			{
-				UpdateContactPoint(Constraint, ConvexHeightFieldContactPoint(A, ATransform, B, BTransform, Constraint.GetCullDistance(), 0.0f), Dt);
+				UpdateContactPoint(Constraint, ConvexHeightFieldContactPoint(A, ATransform, B, BTransform, Constraint.GetCullDistance()), Dt);
 			}
 		}
 
@@ -1366,7 +1361,7 @@ namespace Chaos
 			}
 			else
 			{
-				UpdateContactPoint(Constraint, ConvexTriangleMeshContactPoint(Convex0, WorldTransform0, TriangleMesh1, WorldTransform1, Constraint.GetCullDistance(), 0.0f), Dt);
+				UpdateContactPoint(Constraint, ConvexTriangleMeshContactPoint(Convex0, WorldTransform0, TriangleMesh1, WorldTransform1, Constraint.GetCullDistance()), Dt);
 			}
 		}
 
diff --git a/Engine/Source/Runtime/Experimental/Chaos/Private/Chaos/ContactModification.cpp b/Engine/Source/Runtime/Experimental/Chaos/Private/Chaos/ContactModification.cpp
--- a/Engine/Source/Runtime/Experimental/Chaos/Private/Chaos/ContactModification.cpp
+++ b/Engine/Source/Runtime/Experimental/Chaos/Private/Chaos/ContactModification.cpp
@@ -87,7 +87,7 @@ namespace Chaos
 	void FContactPairModifier::ModifyTargetSeparation(FReal TargetSeparation, int32 ContactPointIdx)
 	{	
 		FManifoldPoint& ManifoldPoint = Constraint->GetManifoldPoint(ContactPointIdx);
-		ManifoldPoint.TargetPhi = TargetSeparation;
+		ManifoldPoint.TargetPhi = FRealSingle(TargetSeparation);
 	}
 
 	FVec3 FContactPairModifier::GetWorldNormal(int32 ContactPointIdx) const
@@ -97,7 +97,7 @@ namespace Chaos
 		TArrayView<FManifoldPoint> ManifoldPoints = Constraint->GetManifoldPoints();
 		FManifoldPoint& ManifoldPoint = ManifoldPoints[ContactPointIdx];
 
-		return ShapeTransform1.TransformVectorNoScale(ManifoldPoint.ContactPoint.ShapeContactNormal);
+		return ShapeTransform1.TransformVectorNoScale(FVec3(ManifoldPoint.ContactPoint.ShapeContactNormal));
 	}
 
 	void FContactPairModifier::ModifyWorldNormal(const FVec3& Normal, int32 ContactPointIdx)
@@ -112,7 +112,7 @@ namespace Chaos
 		const FVec3 ShapeNormal = ShapeTransform1.InverseTransformVectorNoScale(Normal);
 
 		ManifoldPoint.ContactPoint.ShapeContactNormal = ShapeNormal;
-		ManifoldPoint.ContactPoint.Phi = FVec3::DotProduct(WorldContactPoint0 - WorldContactPoint1, Normal);
+		ManifoldPoint.ContactPoint.Phi = FRealSingle(FVec3::DotProduct(WorldContactPoint0 - WorldContactPoint1, Normal));
 
 		Modifier->MarkConstraintForManifoldUpdate(*Constraint);
 	}
@@ -124,8 +124,8 @@ namespace Chaos
 
 		TArrayView<FManifoldPoint> ManifoldPoints = Constraint->GetManifoldPoints();
 		FManifoldPoint& ManifoldPoint = ManifoldPoints[ContactPointIdx];
-		OutLocation0 = ShapeTransform0.TransformPositionNoScale(ManifoldPoint.ContactPoint.ShapeContactPoints[0]);
-		OutLocation1 = ShapeTransform1.TransformPositionNoScale(ManifoldPoint.ContactPoint.ShapeContactPoints[1]);
+		OutLocation0 = ShapeTransform0.TransformPositionNoScale(FVec3(ManifoldPoint.ContactPoint.ShapeContactPoints[0]));
+		OutLocation1 = ShapeTransform1.TransformPositionNoScale(FVec3(ManifoldPoint.ContactPoint.ShapeContactPoints[1]));
 	}
 
 	FVec3 FContactPairModifier::GetWorldContactLocation(int32 ContactPointIdx) const
diff --git a/Engine/Source/Runtime/Experimental/Chaos/Private/Chaos/PBDCollisionConstraints.cpp b/Engine/Source/Runtime/Experimental/Chaos/Private/Chaos/PBDCollisionConstraints.cpp
--- a/Engine/Source/Runtime/Experimental/Chaos/Private/Chaos/PBDCollisionConstraints.cpp
+++ b/Engine/Source/Runtime/Experimental/Chaos/Private/Chaos/PBDCollisionConstraints.cpp
@@ -192,61 +192,74 @@ namespace Chaos
 		const FChaosPhysicsMaterial* PhysicsMaterial0 = GetPhysicsMaterial(Constraint.Particle[0], Constraint.Implicit[0], MPhysicsMaterials, MPerParticlePhysicsMaterials, SimMaterials);
 		const FChaosPhysicsMaterial* PhysicsMaterial1 = GetPhysicsMaterial(Constraint.Particle[1], Constraint.Implicit[1], MPhysicsMaterials, MPerParticlePhysicsMaterials, SimMaterials);
 
-		FPBDCollisionConstraintMaterial& CollisionMaterial = Constraint.Material;
+		FReal MaterialRestitution = 0;
+		FReal MaterialRestitutionThreshold = 0;
+		FReal MaterialStaticFriction = 0;
+		FReal MaterialDynamicFriction = 0;
+
 		if (PhysicsMaterial0 && PhysicsMaterial1)
 		{
 			const FChaosPhysicsMaterial::ECombineMode RestitutionCombineMode = FChaosPhysicsMaterial::ChooseCombineMode(PhysicsMaterial0->RestitutionCombineMode,PhysicsMaterial1->RestitutionCombineMode);
-			CollisionMaterial.MaterialRestitution = FChaosPhysicsMaterial::CombineHelper(PhysicsMaterial0->Restitution, PhysicsMaterial1->Restitution, RestitutionCombineMode);
+			MaterialRestitution = FChaosPhysicsMaterial::CombineHelper(PhysicsMaterial0->Restitution, PhysicsMaterial1->Restitution, RestitutionCombineMode);
 
 			const FChaosPhysicsMaterial::ECombineMode FrictionCombineMode = FChaosPhysicsMaterial::ChooseCombineMode(PhysicsMaterial0->FrictionCombineMode,PhysicsMaterial1->FrictionCombineMode);
-			CollisionMaterial.MaterialDynamicFriction = FChaosPhysicsMaterial::CombineHelper(PhysicsMaterial0->Friction,PhysicsMaterial1->Friction, FrictionCombineMode);
+			MaterialDynamicFriction = FChaosPhysicsMaterial::CombineHelper(PhysicsMaterial0->Friction,PhysicsMaterial1->Friction, FrictionCombineMode);
 			const FReal StaticFriction0 = FMath::Max(PhysicsMaterial0->Friction, PhysicsMaterial0->StaticFriction);
 			const FReal StaticFriction1 = FMath::Max(PhysicsMaterial1->Friction, PhysicsMaterial1->StaticFriction);
-			CollisionMaterial.MaterialStaticFriction = FChaosPhysicsMaterial::CombineHelper(StaticFriction0, StaticFriction1, FrictionCombineMode);
+			MaterialStaticFriction = FChaosPhysicsMaterial::CombineHelper(StaticFriction0, StaticFriction1, FrictionCombineMode);
 		}
 		else if (PhysicsMaterial0)
 		{
 			const FReal StaticFriction0 = FMath::Max(PhysicsMaterial0->Friction, PhysicsMaterial0->StaticFriction);
-			CollisionMaterial.MaterialRestitution = PhysicsMaterial0->Restitution;
-			CollisionMaterial.MaterialDynamicFriction = PhysicsMaterial0->Friction;
-			CollisionMaterial.MaterialStaticFriction = StaticFriction0;
+			MaterialRestitution = PhysicsMaterial0->Restitution;
+			MaterialDynamicFriction = PhysicsMaterial0->Friction;
+			MaterialStaticFriction = StaticFriction0;
 		}
 		else if (PhysicsMaterial1)
 		{
 			const FReal StaticFriction1 = FMath::Max(PhysicsMaterial1->Friction, PhysicsMaterial1->StaticFriction);
-			CollisionMaterial.MaterialRestitution = PhysicsMaterial1->Restitution;
-			CollisionMaterial.MaterialDynamicFriction = PhysicsMaterial1->Friction;
-			CollisionMaterial.MaterialStaticFriction = StaticFriction1;
+			MaterialRestitution = PhysicsMaterial1->Restitution;
+			MaterialDynamicFriction = PhysicsMaterial1->Friction;
+			MaterialStaticFriction = StaticFriction1;
 		}
 		else
 		{
-			CollisionMaterial.MaterialDynamicFriction = DefaultCollisionFriction;
-			CollisionMaterial.MaterialStaticFriction = DefaultCollisionFriction;
-			CollisionMaterial.MaterialRestitution = DefaultCollisionRestitution;
+			MaterialDynamicFriction = DefaultCollisionFriction;
+			MaterialStaticFriction = DefaultCollisionFriction;
+			MaterialRestitution = DefaultCollisionRestitution;
 		}
 
-		CollisionMaterial.RestitutionThreshold = (CollisionRestitutionThresholdOverride >= 0.0f) ? CollisionRestitutionThresholdOverride : RestitutionThreshold;
+		MaterialRestitutionThreshold = RestitutionThreshold;
 
 		// Overrides for testing
 		if (CollisionFrictionOverride >= 0)
 		{
-			CollisionMaterial.MaterialDynamicFriction = CollisionFrictionOverride;
-			CollisionMaterial.MaterialStaticFriction = CollisionFrictionOverride;
+			MaterialDynamicFriction = CollisionFrictionOverride;
+			MaterialStaticFriction = CollisionFrictionOverride;
 		}
 		if (CollisionRestitutionOverride >= 0)
 		{
-			CollisionMaterial.MaterialRestitution = CollisionRestitutionOverride;
+			MaterialRestitution = CollisionRestitutionOverride;
+		}
+		if (CollisionRestitutionThresholdOverride >= 0.0f)
+		{
+			MaterialRestitutionThreshold = CollisionRestitutionThresholdOverride;
 		}
 		if (CollisionAngularFrictionOverride >= 0)
 		{
-			CollisionMaterial.MaterialStaticFriction = CollisionAngularFrictionOverride;
+			MaterialStaticFriction = CollisionAngularFrictionOverride;
 		}
 		if (!bEnableRestitution)
 		{
-			CollisionMaterial.MaterialRestitution = 0.0f;
+			MaterialRestitution = 0.0f;
 		}
 		
-		CollisionMaterial.ResetMaterialModifications();
+		Constraint.Material.MaterialRestitution = FRealSingle(MaterialRestitution);
+		Constraint.Material.RestitutionThreshold = FRealSingle(MaterialRestitutionThreshold);
+		Constraint.Material.MaterialStaticFriction = FRealSingle(MaterialStaticFriction);
+		Constraint.Material.MaterialDynamicFriction = FRealSingle(MaterialDynamicFriction);
+
+		Constraint.Material.ResetMaterialModifications();
 	}
 
 	void FPBDCollisionConstraints::BeginFrame()
diff --git a/Engine/Source/Runtime/Experimental/Chaos/Private/Chaos/PBDRigidsEvolutionGBF.cpp b/Engine/Source/Runtime/Experimental/Chaos/Private/Chaos/PBDRigidsEvolutionGBF.cpp
--- a/Engine/Source/Runtime/Experimental/Chaos/Private/Chaos/PBDRigidsEvolutionGBF.cpp
+++ b/Engine/Source/Runtime/Experimental/Chaos/Private/Chaos/PBDRigidsEvolutionGBF.cpp
@@ -706,7 +706,7 @@ void FPBDRigidsEvolutionGBF::TransferJointConstraintCollisions()
 				ChildParticle->Handle()->ParticleCollisions().VisitCollisions(
 					[&](const FPBDCollisionConstraint* VisitedConstraint)
 					{
-						if (VisitedConstraint->GetCCDType() != ECollisionCCDType::Disabled)
+						if (VisitedConstraint->GetCCDEnabled())
 						{
 							return;
 						}
diff --git a/Engine/Source/Runtime/Experimental/Chaos/Public/Chaos/Collision/CapsuleConvexContactPoint.h b/Engine/Source/Runtime/Experimental/Chaos/Public/Chaos/Collision/CapsuleConvexContactPoint.h
--- a/Engine/Source/Runtime/Experimental/Chaos/Public/Chaos/Collision/CapsuleConvexContactPoint.h
+++ b/Engine/Source/Runtime/Experimental/Chaos/Public/Chaos/Collision/CapsuleConvexContactPoint.h
@@ -1,22 +1,12 @@
 // Copyright Epic Games, Inc. All Rights Reserved.
 #pragma once
 
-#include "CoreMinimal.h"
 #include "Chaos/Core.h"
+#include "Chaos/CollisionResolutionTypes.h"
 #include "Chaos/ImplicitFwd.h"
 #include "Chaos/ImplicitObjectScaled.h"	// Cannot fwd declare scaled implicit
 
 namespace Chaos
 {
-	class FContactPoint;
-
 	extern CHAOS_API FContactPoint CapsuleConvexContactPoint(const FImplicitCapsule3& Capsule, const FRigidTransform3& CapsuleTransform, const FImplicitObject& Convex, const FRigidTransform3& ConvexTransform);
-
-	extern CHAOS_API FContactPoint CapsuleConvexContactPoint(const FImplicitCapsule3& Capsule, const FRigidTransform3& CapsuleTransform, const FImplicitConvex3& Convex, const FRigidTransform3& ConvexTransform);
-	extern CHAOS_API FContactPoint CapsuleConvexContactPoint(const FImplicitCapsule3& Capsule, const FRigidTransform3& CapsuleTransform, const TImplicitObjectInstanced<FImplicitConvex3>& Convex, const FRigidTransform3& ConvexTransform);
-	extern CHAOS_API FContactPoint CapsuleConvexContactPoint(const FImplicitCapsule3& Capsule, const FRigidTransform3& CapsuleTransform, const TImplicitObjectScaled<FImplicitConvex3>& Convex, const FRigidTransform3& ConvexTransform);
-
-	extern CHAOS_API FContactPoint ConvexCapsuleContactPoint(const FImplicitConvex3& Convex, const FRigidTransform3& ConvexTransform, const FImplicitCapsule3& Capsule, const FRigidTransform3& CapsuleTransform);
-	extern CHAOS_API FContactPoint ConvexCapsuleContactPoint(const TImplicitObjectInstanced<FImplicitConvex3>& Convex, const FRigidTransform3& ConvexTransform, const FImplicitCapsule3& Capsule, const FRigidTransform3& CapsuleTransform);
-	extern CHAOS_API FContactPoint ConvexCapsuleContactPoint(const TImplicitObjectScaled<FImplicitConvex3>& Convex, const FRigidTransform3& ConvexTransform, const FImplicitCapsule3& Capsule, const FRigidTransform3& CapsuleTransform);
 }
diff --git a/Engine/Source/Runtime/Experimental/Chaos/Public/Chaos/Collision/CollisionConstraintAllocator.h b/Engine/Source/Runtime/Experimental/Chaos/Public/Chaos/Collision/CollisionConstraintAllocator.h
--- a/Engine/Source/Runtime/Experimental/Chaos/Public/Chaos/Collision/CollisionConstraintAllocator.h
+++ b/Engine/Source/Runtime/Experimental/Chaos/Public/Chaos/Collision/CollisionConstraintAllocator.h
@@ -485,7 +485,7 @@ namespace Chaos
 			checkSlow(ActiveConstraints.Find(CollisionConstraint) == INDEX_NONE);
 			Cookie.ConstraintIndex = ActiveConstraints.Add(CollisionConstraint);
 
-			if (CollisionConstraint->GetCCDType() == ECollisionCCDType::Enabled)
+			if (CollisionConstraint->GetCCDEnabled())
 			{
 				checkSlow(ActiveSweptConstraints.Find(CollisionConstraint) == INDEX_NONE);
 				Cookie.SweptConstraintIndex = ActiveSweptConstraints.Add(CollisionConstraint);
diff --git a/Engine/Source/Runtime/Experimental/Chaos/Public/Chaos/Collision/ContactPoint.h b/Engine/Source/Runtime/Experimental/Chaos/Public/Chaos/Collision/ContactPoint.h
--- a/Engine/Source/Runtime/Experimental/Chaos/Public/Chaos/Collision/ContactPoint.h
+++ b/Engine/Source/Runtime/Experimental/Chaos/Public/Chaos/Collision/ContactPoint.h
@@ -26,20 +26,36 @@ namespace Chaos
 		VertexVertex,
 	};
 
+
+	/**
+	 * The value used for Phi (contact separation) to indicate that it is unset/invalid.
+	 */
+	template<typename T>
+	constexpr T InvalidPhi() 
+	{ 
+		return TNumericLimits<T>::Max();
+	}
+
 	/**
-	 * @brief Data returned by the low-level collision functions
+	 * Data returned by the low-level collision functions.
+	 * 
+	 * @note All data is invalid/uninitialized when IsSet() is false.
+	 * @see FContactPoint, FContactPointf
 	*/
-	class CHAOS_API FContactPoint
+	template<typename T>
+	class CHAOS_API TContactPoint
 	{
 	public:
+		using FRealType = T;
+
 		// Shape-space contact points on the two bodies
-		FVec3 ShapeContactPoints[2];
+		TVec3<FRealType> ShapeContactPoints[2];
 
 		// Shape-space contact normal on the second shape with direction that points away from shape 1
-		FVec3 ShapeContactNormal;
+		TVec3<FRealType> ShapeContactNormal;
 
 		// Contact separation (negative for overlap)
-		FReal Phi;
+		FRealType Phi;
 
 		// Face index of the shape we hit. Only valid for Heightfield and Trimesh contact points, otherwise INDEX_NONE
 		int32 FaceIndex;
@@ -47,29 +63,42 @@ namespace Chaos
 		// Whether this is a vertex-plane contact, edge-edge contact etc.
 		EContactPointType ContactType;
 
-		FContactPoint()
-			: Phi(TNumericLimits<FReal>::Max())
+		TContactPoint()
+			: Phi(InvalidPhi<FRealType>())
 			, FaceIndex(INDEX_NONE)
 			, ContactType(EContactPointType::Unknown)
 		{
 		}
 
+		template<typename U>
+		TContactPoint(const TContactPoint<U>& Other)
+			: ShapeContactPoints{ TVec3<FRealType>(Other.ShapeContactPoints[0]), TVec3<FRealType>(Other.ShapeContactPoints[1]) }
+			, ShapeContactNormal{ TVec3<FRealType>(Other.ShapeContactNormal) }
+			, Phi(FRealType(Other.Phi))
+			, FaceIndex(Other.FaceIndex)
+			, ContactType(Other.ContactType)
+		{
+		}
+
 		// Whether the contact point has been set up with contact data
-		bool IsSet() const { return (Phi != TNumericLimits<FReal>::Max()); }
+		bool IsSet() const 
+		{ 
+			return (Phi != InvalidPhi<FRealType>());
+		}
 
 		// Switch the shape indices. For use when calling a collision detection method which takes shape types in the opposite order to what you want.
-		// @todo(chaos): remove or fix this function
-		// WARNING: this function can no longer be used in isolation as it could when we were calculating world-space contact data. For this to
-		// work correctly, the normal must either already be in the space of the first shape, or will need to be transformed after.
-		// Alternatively we could start using EContactPointType to indicate normal ownership
-		FContactPoint& SwapShapes()
-		{
-			if (IsSet())
-			{
-				Swap(ShapeContactPoints[0], ShapeContactPoints[1]);
-				ShapeContactNormal = -ShapeContactNormal;
-			}
-			return *this;
-		}
+		// WARNING: this function is not "general purporse" and only works when ShapeContactPoints are both in the smae space, because
+		// the normal is in the space of the second body and negating it does not otherwise transform it to be relative to the other body. 
+		// However, it is useful in the cases where we perform collision detection in the space of one of the bodies before transforming into 
+		// shape space, whcih is actually most of the time.
+		//TContactPoint<FRealType>& SwapShapes()
+		//{
+		//	if (IsSet())
+		//	{
+		//		Swap(ShapeContactPoints[0], ShapeContactPoints[1]);
+		//		ShapeContactNormal = -ShapeContactNormal;
+		//	}
+		//	return *this;
+		//}
 	};
 }
\ No newline at end of file
diff --git a/Engine/Source/Runtime/Experimental/Chaos/Public/Chaos/Collision/ContactPointsMiscShapes.h b/Engine/Source/Runtime/Experimental/Chaos/Public/Chaos/Collision/ContactPointsMiscShapes.h
--- a/Engine/Source/Runtime/Experimental/Chaos/Public/Chaos/Collision/ContactPointsMiscShapes.h
+++ b/Engine/Source/Runtime/Experimental/Chaos/Public/Chaos/Collision/ContactPointsMiscShapes.h
@@ -1,15 +1,13 @@
 // Copyright Epic Games, Inc. All Rights Reserved.
 #pragma once
 
-#include "CoreMinimal.h"
 #include "Chaos/Core.h"
+#include "Chaos/CollisionResolutionTypes.h"
 #include "Chaos/ImplicitFwd.h"
-#include "Chaos/ImplicitObjectScaled.h"	// Cannot fwd declare scaled implicit
 #include "Chaos/TriangleMeshImplicitObject.h"
 
 namespace Chaos
 {
-	class FContactPoint;
 	class FHeightField;
 	
 	/**
@@ -38,38 +36,38 @@ namespace Chaos
 
 	
 	template <typename GeometryA, typename GeometryB>
-	FContactPoint GJKImplicitContactPoint(const FImplicitObject& A, const FRigidTransform3& ATransform, const GeometryB& B, const FRigidTransform3& BTransform, const FReal CullDistance, const FReal ShapePadding);
+	FContactPoint GJKImplicitContactPoint(const FImplicitObject& A, const FRigidTransform3& ATransform, const GeometryB& B, const FRigidTransform3& BTransform, const FReal CullDistance);
 
-	FContactPoint SphereSphereContactPoint(const TSphere<FReal, 3>& Sphere1, const FRigidTransform3& Sphere1Transform, const TSphere<FReal, 3>& Sphere2, const FRigidTransform3& Sphere2Transform, const FReal CullDistance, const FReal ShapePadding);
-	FContactPoint SpherePlaneContactPoint(const TSphere<FReal, 3>& Sphere, const FRigidTransform3& SphereTransform, const TPlane<FReal, 3>& Plane, const FRigidTransform3& PlaneTransform, const FReal ShapePadding);
-	FContactPoint SphereBoxContactPoint(const TSphere<FReal, 3>& Sphere, const FRigidTransform3& SphereTransform, const FImplicitBox3& Box, const FRigidTransform3& BoxTransform, const FReal ShapePadding);
-	FContactPoint SphereCapsuleContactPoint(const TSphere<FReal, 3>& A, const FRigidTransform3& ATransform, const FCapsule& B, const FRigidTransform3& BTransform, const FReal ShapePadding);
+	FContactPoint SphereSphereContactPoint(const TSphere<FReal, 3>& Sphere1, const FRigidTransform3& Sphere1Transform, const TSphere<FReal, 3>& Sphere2, const FRigidTransform3& Sphere2Transform, const FReal CullDistance);
+	FContactPoint SpherePlaneContactPoint(const TSphere<FReal, 3>& Sphere, const FRigidTransform3& SphereTransform, const TPlane<FReal, 3>& Plane, const FRigidTransform3& PlaneTransform);
+	FContactPoint SphereBoxContactPoint(const TSphere<FReal, 3>& Sphere, const FRigidTransform3& SphereTransform, const FImplicitBox3& Box, const FRigidTransform3& BoxTransform);
+	FContactPoint SphereCapsuleContactPoint(const TSphere<FReal, 3>& A, const FRigidTransform3& ATransform, const FCapsule& B, const FRigidTransform3& BTransform);
 
 	template <typename TriMeshType>
-	FContactPoint SphereTriangleMeshContactPoint(const TSphere<FReal, 3>& A, const FRigidTransform3& ATransform, const TriMeshType& B, const FRigidTransform3& BTransform, const FReal CullDistance, const FReal ShapePadding);
+	FContactPoint SphereTriangleMeshContactPoint(const TSphere<FReal, 3>& A, const FRigidTransform3& ATransform, const TriMeshType& B, const FRigidTransform3& BTransform, const FReal CullDistance);
 
 	template<typename TriMeshType>
 	FContactPoint SphereTriangleMeshSweptContactPoint(const TSphere<FReal, 3>& A, const FRigidTransform3& ATransform, const TriMeshType& B, const FRigidTransform3& BStartTransform, const FVec3& Dir, const FReal Length, const FReal IgnorePenetration, const FReal TargetPenetration, FReal& TOI);
 
-	FContactPoint BoxHeightFieldContactPoint(const FImplicitBox3& A, const FRigidTransform3& ATransform, const FHeightField& B, const FRigidTransform3& BTransform, const FReal CullDistance, const FReal ShapePadding);
+	FContactPoint BoxHeightFieldContactPoint(const FImplicitBox3& A, const FRigidTransform3& ATransform, const FHeightField& B, const FRigidTransform3& BTransform, const FReal CullDistance);
 
 	template <typename TriMeshType>
-	FContactPoint BoxTriangleMeshContactPoint(const FImplicitBox3& A, const FRigidTransform3& ATransform, const TriMeshType& B, const FRigidTransform3& BTransform, const FReal CullDistance, const FReal ShapePadding);
+	FContactPoint BoxTriangleMeshContactPoint(const FImplicitBox3& A, const FRigidTransform3& ATransform, const TriMeshType& B, const FRigidTransform3& BTransform, const FReal CullDistance);
 
-	FContactPoint SphereHeightFieldContactPoint(const TSphere<FReal, 3>& A, const FRigidTransform3& ATransform, const FHeightField& B, const FRigidTransform3& BTransform, const FReal CullDistance, const FReal ShapePadding);
-	FContactPoint CapsuleHeightFieldContactPoint(const FCapsule& A, const FRigidTransform3& ATransform, const FHeightField& B, const FRigidTransform3& BTransform, const FReal CullDistance, const FReal ShapePadding);
+	FContactPoint SphereHeightFieldContactPoint(const TSphere<FReal, 3>& A, const FRigidTransform3& ATransform, const FHeightField& B, const FRigidTransform3& BTransform, const FReal CullDistance);
+	FContactPoint CapsuleHeightFieldContactPoint(const FCapsule& A, const FRigidTransform3& ATransform, const FHeightField& B, const FRigidTransform3& BTransform, const FReal CullDistance);
 	template <typename TriMeshType>
-	FContactPoint CapsuleTriangleMeshContactPoint(const FCapsule& A, const FRigidTransform3& ATransform, const TriMeshType& B, const FRigidTransform3& BTransform, const FReal CullDistance, const FReal ShapePadding);
+	FContactPoint CapsuleTriangleMeshContactPoint(const FCapsule& A, const FRigidTransform3& ATransform, const TriMeshType& B, const FRigidTransform3& BTransform, const FReal CullDistance);
 	template <typename TriMeshType>
 	FContactPoint CapsuleTriangleMeshSweptContactPoint(const FCapsule& A, const FRigidTransform3& ATransform, const TriMeshType& B, const FRigidTransform3& BStartTransform, const FVec3& Dir, const FReal Length, const FReal IgnorePenetration, const FReal TargetPenetration, FReal& TOI);
 
-	FContactPoint ConvexHeightFieldContactPoint(const FImplicitObject& A, const FRigidTransform3& ATransform, const FHeightField& B, const FRigidTransform3& BTransform, const FReal CullDistance, const FReal ShapePadding);
-	FContactPoint ConvexTriangleMeshContactPoint(const FImplicitObject& A, const FRigidTransform3& ATransform, const FImplicitObject& B, const FRigidTransform3& BTransform, const FReal CullDistance, const FReal ShapePadding);
+	FContactPoint ConvexHeightFieldContactPoint(const FImplicitObject& A, const FRigidTransform3& ATransform, const FHeightField& B, const FRigidTransform3& BTransform, const FReal CullDistance);
+	FContactPoint ConvexTriangleMeshContactPoint(const FImplicitObject& A, const FRigidTransform3& ATransform, const FImplicitObject& B, const FRigidTransform3& BTransform, const FReal CullDistance);
 	template <typename TriMeshType>
 	FContactPoint ConvexTriangleMeshSweptContactPoint(const FImplicitObject& A, const FRigidTransform3& ATransform, const TriMeshType& B, const FRigidTransform3& BStartTransform, const FVec3& Dir, const FReal Length, const FReal IgnorePenetration, const FReal TargetPenetration, FReal& TOI);
 
-	FContactPoint CapsuleCapsuleContactPoint(const FCapsule& A, const FRigidTransform3& ATransform, const FCapsule& B, const FRigidTransform3& BTransform, const FReal ShapePadding);
-	FContactPoint CapsuleBoxContactPoint(const FCapsule& A, const FRigidTransform3& ATransform, const FImplicitBox3& B, const FRigidTransform3& BTransform, const FVec3& InitialDir, const FReal ShapePadding);
+	FContactPoint CapsuleCapsuleContactPoint(const FCapsule& A, const FRigidTransform3& ATransform, const FCapsule& B, const FRigidTransform3& BTransform);
+	FContactPoint CapsuleBoxContactPoint(const FCapsule& A, const FRigidTransform3& ATransform, const FImplicitBox3& B, const FRigidTransform3& BTransform, const FVec3& InitialDir);
 	
 	
 }
diff --git a/Engine/Source/Runtime/Experimental/Chaos/Public/Chaos/Collision/PBDCollisionConstraint.h b/Engine/Source/Runtime/Experimental/Chaos/Public/Chaos/Collision/PBDCollisionConstraint.h
--- a/Engine/Source/Runtime/Experimental/Chaos/Public/Chaos/Collision/PBDCollisionConstraint.h
+++ b/Engine/Source/Runtime/Experimental/Chaos/Public/Chaos/Collision/PBDCollisionConstraint.h
@@ -44,8 +44,8 @@ namespace Chaos
 	public:
 		FManifoldPoint() 
 			: ContactPoint()
-			, ShapeAnchorPoints{ FVec3(0), FVec3(0) }
-			, InitialShapeContactPoints{ FVec3(0), FVec3(0) }
+			, ShapeAnchorPoints{ FVec3f(0), FVec3f(0) }
+			, InitialShapeContactPoints{ FVec3f(0), FVec3f(0) }
 			, NetPushOut(0)
 			, NetImpulse(0)
 			, TargetPhi(0)
@@ -54,20 +54,20 @@ namespace Chaos
 
 		FManifoldPoint(const FContactPoint& InContactPoint) 
 			: ContactPoint(InContactPoint)
-			, ShapeAnchorPoints{ FVec3(0), FVec3(0) }
-			, InitialShapeContactPoints{ FVec3(0), FVec3(0) }
+			, ShapeAnchorPoints{ FVec3f(0), FVec3f(0) }
+			, InitialShapeContactPoints{ FVec3f(0), FVec3f(0) }
 			, NetPushOut(0)
 			, NetImpulse(0)
 			, TargetPhi(0)
 			, Flags()
 		{}
 
-		FContactPoint ContactPoint;			// Contact point results of low-level collision detection
-		FVec3 ShapeAnchorPoints[2];			// The contact points from a prior frame, if availble. Used by friction.
-		FVec3 InitialShapeContactPoints[2];	// ShapeContactPoints when the constraint was first initialized. Used to track reusablility
-		FVec3 NetPushOut;					// Total pushout applied at this contact point
-		FVec3 NetImpulse;					// Total impulse applied by this contact point
-		FReal TargetPhi;					// Usually 0, but can be used to add padding or penetration (e.g., via a collision modifer)
+		FContactPointf ContactPoint;			// Contact point results of low-level collision detection
+		FVec3f ShapeAnchorPoints[2];			// The contact points from a prior frame, if availble. Used by friction.
+		FVec3f InitialShapeContactPoints[2];	// ShapeContactPoints when the constraint was first initialized. Used to track reusablility
+		FVec3f NetPushOut;						// Total pushout applied at this contact point
+		FVec3f NetImpulse;						// Total impulse applied by this contact point
+		FRealSingle TargetPhi;					// Usually 0, but can be used to add padding or penetration (e.g., via a collision modifer)
 
 		union FFlags
 		{
@@ -94,7 +94,7 @@ namespace Chaos
 	class CHAOS_API FSavedManifoldPoint
 	{
 	public:
-		FVec3 ShapeContactPoints[2];
+		FVec3f ShapeContactPoints[2];
 	};
 
 	/*
@@ -110,31 +110,29 @@ namespace Chaos
 			, DynamicFriction(0)
 			, StaticFriction(0)
 			, Restitution(0)
-			, RestitutionPadding(0)
 			, RestitutionThreshold(0)
-			, InvMassScale0(1.f)
-			, InvMassScale1(1.f)
-			, InvInertiaScale0(1.f)
-			, InvInertiaScale1(1.f)
+			, InvMassScale0(1)
+			, InvMassScale1(1)
+			, InvInertiaScale0(1)
+			, InvInertiaScale1(1)
 		{
 		}
 
 		// Material properties pulled from the materials of the two shapes involved in the contact
-		FReal MaterialDynamicFriction;
-		FReal MaterialStaticFriction;
-		FReal MaterialRestitution;
+		FRealSingle MaterialDynamicFriction;
+		FRealSingle MaterialStaticFriction;
+		FRealSingle MaterialRestitution;
 
 		// Final material properties (post modifier) used by the solver. These get reset every frame to the material values above
-		FReal DynamicFriction;
-		FReal StaticFriction;
-		FReal Restitution;
+		FRealSingle DynamicFriction;
+		FRealSingle StaticFriction;
+		FRealSingle Restitution;
 
-		FReal RestitutionPadding; // For StandardPBD implementation of resitution, we pad constraints on initial contact to enforce outward velocity
-		FReal RestitutionThreshold;
-		FReal InvMassScale0;
-		FReal InvMassScale1;
-		FReal InvInertiaScale0;
-		FReal InvInertiaScale1;
+		FRealSingle RestitutionThreshold;
+		FRealSingle InvMassScale0;
+		FRealSingle InvMassScale1;
+		FRealSingle InvInertiaScale0;
+		FRealSingle InvInertiaScale1;
 
 		// Reset the material properties to those pulled from the shape's materials (i.e., back to the state before any contact modification)
 		void ResetMaterialModifications()
@@ -142,16 +140,10 @@ namespace Chaos
 			DynamicFriction = MaterialDynamicFriction;
 			StaticFriction = MaterialStaticFriction;
 			Restitution = MaterialRestitution;
-			InvMassScale0 = FReal(1);
-			InvMassScale1 = FReal(1);
-			InvInertiaScale0 = FReal(1);
-			InvInertiaScale1 = FReal(1);
-		}
-
-		// @todo(chaos): remove this
-		void Reset()
-		{
-			RestitutionPadding = 0;
+			InvMassScale0 = FRealSingle(1);
+			InvMassScale1 = FRealSingle(1);
+			InvInertiaScale0 = FRealSingle(1);
+			InvInertiaScale1 = FRealSingle(1);
 		}
 	};
 
@@ -284,17 +276,17 @@ namespace Chaos
 		virtual ~FPBDCollisionConstraint();
 
 		/**
-		 * @brief The current CCD state of this constraint
+		 * Whetehr CCD is enabled for this collision.
 		 * This may change from tick to tick as an object's velocity changes.
 		*/
-		ECollisionCCDType GetCCDType() const { return CCDType; }
+		bool GetCCDEnabled() const { return Flags.bCCDEnabled; }
 
 		/**
 		 * @brief Enable or disable CCD for this constraint
 		*/
 		void SetCCDEnabled(const bool bCCDEnabled)
 		{ 
-			CCDType = bCCDEnabled ? ECollisionCCDType::Enabled : ECollisionCCDType::Disabled;
+			Flags.bCCDEnabled = bCCDEnabled;
 
 			// Initialize the CCD thresholds (one time only)
 			if (bCCDEnabled && (CCDEnablePenetration == FReal(0)))
@@ -378,34 +370,31 @@ namespace Chaos
 		// @todo(chaos): remove (used by legacy RBAN collision solver)
 		FVec3 CalculateWorldContactLocation() const;
 
-		void SetInvMassScale0(const FReal InInvMassScale) { Material.InvMassScale0 = InInvMassScale; }
+		void SetInvMassScale0(const FReal InInvMassScale) { Material.InvMassScale0 = FRealSingle(InInvMassScale); }
 		FReal GetInvMassScale0() const { return Material.InvMassScale0; }
 
-		void SetInvMassScale1(const FReal InInvMassScale) { Material.InvMassScale1 = InInvMassScale; }
+		void SetInvMassScale1(const FReal InInvMassScale) { Material.InvMassScale1 = FRealSingle(InInvMassScale); }
 		FReal GetInvMassScale1() const { return Material.InvMassScale1; }
 
-		void SetInvInertiaScale0(const FReal InInvInertiaScale) { Material.InvInertiaScale0 = InInvInertiaScale; }
+		void SetInvInertiaScale0(const FReal InInvInertiaScale) { Material.InvInertiaScale0 = FRealSingle(InInvInertiaScale); }
 		FReal GetInvInertiaScale0() const { return Material.InvInertiaScale0; }
 
-		void SetInvInertiaScale1(const FReal InInvInertiaScale) { Material.InvInertiaScale1 = InInvInertiaScale; }
+		void SetInvInertiaScale1(const FReal InInvInertiaScale) { Material.InvInertiaScale1 = FRealSingle(InInvInertiaScale); }
 		FReal GetInvInertiaScale1() const { return Material.InvInertiaScale1; }
 
-		void SetStiffness(FReal InStiffness) { Stiffness = InStiffness; }
+		void SetStiffness(FReal InStiffness) { Stiffness = FRealSingle(InStiffness); }
 		FReal GetStiffness() const { return Stiffness; }
 
-		void SetRestitution(const FReal InRestitution) { Material.Restitution = InRestitution; }
+		void SetRestitution(const FReal InRestitution) { Material.Restitution = FRealSingle(InRestitution); }
 		FReal GetRestitution() const { return Material.Restitution; }
 
-		void SetRestitutionThreshold(const FReal InRestitutionThreshold) { Material.RestitutionThreshold = InRestitutionThreshold; }
+		void SetRestitutionThreshold(const FReal InRestitutionThreshold) { Material.RestitutionThreshold = FRealSingle(InRestitutionThreshold); }
 		FReal GetRestitutionThreshold() const { return Material.RestitutionThreshold; }
 
-		void SetRestitutionPadding(const FReal InRestitutionPadding) { Material.RestitutionPadding = InRestitutionPadding; }
-		FReal GetRestitutionPadding() const { return Material.RestitutionPadding; }
-
-		void SetStaticFriction(const FReal InStaticFriction) { Material.StaticFriction = InStaticFriction; }
+		void SetStaticFriction(const FReal InStaticFriction) { Material.StaticFriction = FRealSingle(InStaticFriction); }
 		FReal GetStaticFriction() const { return FMath::Max(Material.StaticFriction, Material.DynamicFriction); }
 
-		void SetDynamicFriction(const FReal InDynamicFriction) { Material.DynamicFriction = InDynamicFriction; }
+		void SetDynamicFriction(const FReal InDynamicFriction) { Material.DynamicFriction = FRealSingle(InDynamicFriction); }
 		FReal GetDynamicFriction() const { return Material.DynamicFriction; }
 
 		EContactShapesType GetShapesType() const { return ShapesType; }
@@ -413,7 +402,7 @@ namespace Chaos
 		FString ToString() const;
 
 		FReal GetCullDistance() const { return CullDistance; }
-		void SetCullDistance(FReal InCullDistance) { CullDistance = InCullDistance; }
+		void SetCullDistance(FReal InCullDistance) { CullDistance = FRealSingle(InCullDistance); }
 
 		// Whether we are using manifolds (either one-shot or incremental)
 		bool GetUseManifold() const { return Flags.bUseManifold; }
@@ -526,13 +515,13 @@ namespace Chaos
 		const FRigidTransform3& GetShapeRelativeTransform1() const { return ImplicitTransform[1]; }
 		const FRigidTransform3& GetShapeRelativeTransform(const int32 ParticleIndex) const { check((ParticleIndex >= 0) && (ParticleIndex < 2)); return ImplicitTransform[ParticleIndex]; }
 
-		const FRigidTransform3& GetShapeWorldTransform0() const { return ShapeWorldTransform0; }
-		const FRigidTransform3& GetShapeWorldTransform1() const { return ShapeWorldTransform1; }
+		const FRigidTransform3& GetShapeWorldTransform0() const { return ShapeWorldTransforms[0]; }
+		const FRigidTransform3& GetShapeWorldTransform1() const { return ShapeWorldTransforms[1]; }
 
 		void SetShapeWorldTransforms(const FRigidTransform3& InShapeWorldTransform0, const FRigidTransform3& InShapeWorldTransform1)
 		{
-			ShapeWorldTransform0 = InShapeWorldTransform0;
-			ShapeWorldTransform1 = InShapeWorldTransform1;
+			ShapeWorldTransforms[0] = InShapeWorldTransform0;
+			ShapeWorldTransforms[1] = InShapeWorldTransform1;
 		}
 
 		void SetLastShapeWorldTransforms(const FRigidTransform3& InShapeWorldTransform0, const FRigidTransform3& InShapeWorldTransform1)
@@ -597,14 +586,14 @@ namespace Chaos
 		/**
 		 * @brief Set the CCD TOI from the collision detection sweep
 		*/
-		void SetCCDTimeOfImpact(const FReal TOI) { CCDTimeOfImpact = TOI; }
+		void SetCCDTimeOfImpact(const FReal TOI) { CCDTimeOfImpact = FRealSingle(TOI); }
 
 		/**
 		 * \brief Store the results of CCD contact resolution, if active
 		 */
 		void SetCCDResults(const FVec3& InNetImpulse)
 		{
-			AccumulatedImpulse += InNetImpulse;
+			AccumulatedImpulse += FVec3f(InNetImpulse);
 		}
 
 		/**
@@ -612,10 +601,10 @@ namespace Chaos
 		*/
 		void SetSolverResults(
 			const int32 ManifoldPointIndex, 
-			const FVec3& NetPushOut, 
-			const FVec3& NetImpulse, 
-			const FReal StaticFrictionRatio,
-			const FReal Dt)
+			const FVec3f& NetPushOut, 
+			const FVec3f& NetImpulse, 
+			const FRealSingle StaticFrictionRatio,
+			const FRealSingle Dt)
 		{
 			FManifoldPoint& ManifoldPoint = ManifoldPoints[ManifoldPointIndex];
 
@@ -719,7 +708,7 @@ namespace Chaos
 		// @todo(chaos): Only intended for use by the legacy solvers - remove it
 		void UpdateManifoldPointPhi(const int32 ManifoldPointIndex);
 
-		void InitMarginsAndTolerances(const EImplicitObjectType ImplicitType0, const EImplicitObjectType ImplicitType1, const FReal Margin0, const FReal Margin1);
+		void InitMarginsAndTolerances(const EImplicitObjectType ImplicitType0, const EImplicitObjectType ImplicitType1, const FRealSingle Margin0, const FRealSingle Margin1);
 
 		void InitCCDThreshold();
 
@@ -735,42 +724,28 @@ namespace Chaos
 		const FBVHParticles* Simplicial[2];
 
 		FPBDCollisionConstraintMaterial Material;
-		FReal Stiffness;
+		FRealSingle Stiffness;
 
 	public:
-		FVec3 AccumulatedImpulse;					// @todo(chaos): we need to accumulate angular impulse separately
+		FVec3f AccumulatedImpulse;					// @todo(chaos): we need to accumulate angular impulse separately
 
 	private:
-		// Value in range [0,1] used to interpolate P between [X,P] that we will rollback to when solving at time of impact.
-		FReal CCDTimeOfImpact;
-
-		// The penetration at which CCD contacts get processed following a CCD sweep test.
-		// NOTE: also see GeometryParticle CCDAxisThreshold, which is used to determine when to enable sweeping during collision detection.
-		// Calculated from particle pair properties when constraint is created.
-		FReal CCDEnablePenetration;
-
-		// The penetration we leave behind when rolling back to a CCD time of impact. Should be less than or equal to CCDEnablePenetration.
-		// Calculated from particle pair properties when constraint is created.
-		FReal CCDTargetPenetration;
-
 		FPBDCollisionConstraintContainerCookie ContainerCookie;
 		EContactShapesType ShapesType;
-		ECollisionCCDType CCDType;
 
 		// The shape transforms at the current particle transforms
-		FRigidTransform3 ShapeWorldTransform0;
-		FRigidTransform3 ShapeWorldTransform1;
+		FRigidTransform3 ShapeWorldTransforms[2];
 
 		// The separation distance at which we don't track contacts
-		FReal CullDistance;
+		FRealSingle CullDistance;
 
 		// The margins to use during collision detection. We don't always use the margins on the shapes directly.
 		// E.g., we use the smallest non-zero margin for 2 convex shapes. See InitMarginsAndTolerances
-		FReal CollisionMargins[2];
+		FRealSingle CollisionMargins[2];
 
 		// The collision tolerance is used to determine whether a new contact matches an old on. It is derived from the
 		// margins of the two shapes, as well as their types
-		FReal CollisionTolerance;
+		FRealSingle CollisionTolerance;
 
 		// The index into ManifoldPoints of the point with the lowest Phi
 		int32 ClosestManifoldPointIndex;
@@ -791,14 +766,15 @@ namespace Chaos
 				uint32 bIsQuadratic1 : 1;
 				uint32 bIsProbeUnmodified : 1;  // Is this constraint a probe pre-contact-modification
 				uint32 bIsProbe : 1;            // Is this constraint currently a probe
+				uint32 bCCDEnabled : 1;
 			};
 			uint32 Bits;
 		} Flags;
 
 		// Relative transform the last time we ran the narrow phase
 		// Used to detect when the bodies have moved too far to reues the manifold
-		FVec3 LastShapeWorldPositionDelta;
-		FRotation3 LastShapeWorldRotationDelta;
+		FVec3f LastShapeWorldPositionDelta;
+		FRotation3f LastShapeWorldRotationDelta;
 
 		// These are only needed here while we still have the legacy solvers (not QuasiPBD)
 		FSolverBody* SolverBodies[2];
@@ -808,6 +784,18 @@ namespace Chaos
 
 		TCArray<FSavedManifoldPoint, MaxManifoldPoints> SavedManifoldPoints;
 		TCArray<FManifoldPoint, MaxManifoldPoints> ManifoldPoints;
+
+		// Value in range [0,1] used to interpolate P between [X,P] that we will rollback to when solving at time of impact.
+		FRealSingle CCDTimeOfImpact;
+
+		// The penetration at which CCD contacts get processed following a CCD sweep test.
+		// NOTE: also see GeometryParticle CCDAxisThreshold, which is used to determine when to enable sweeping during collision detection.
+		// Calculated from particle pair properties when constraint is created.
+		FRealSingle CCDEnablePenetration;
+
+		// The penetration we leave behind when rolling back to a CCD time of impact. Should be less than or equal to CCDEnablePenetration.
+		// Calculated from particle pair properties when constraint is created.
+		FRealSingle CCDTargetPenetration;
 	};
 
 
diff --git a/Engine/Source/Runtime/Experimental/Chaos/Public/Chaos/Collision/PBDCollisionConstraintHandle.h b/Engine/Source/Runtime/Experimental/Chaos/Public/Chaos/Collision/PBDCollisionConstraintHandle.h
--- a/Engine/Source/Runtime/Experimental/Chaos/Public/Chaos/Collision/PBDCollisionConstraintHandle.h
+++ b/Engine/Source/Runtime/Experimental/Chaos/Public/Chaos/Collision/PBDCollisionConstraintHandle.h
@@ -73,7 +73,7 @@ namespace Chaos
 		UE_DEPRECATED(4.27, "Use GetContact()")
 		FPBDCollisionConstraint& GetSweptPointContact() { return GetContact(); }
 
-		ECollisionCCDType GetCCDType() const;
+		bool GetCCDEnabled() const;
 
 		virtual void SetEnabled(bool InEnabled) override final;
 
diff --git a/Engine/Source/Runtime/Experimental/Chaos/Public/Chaos/Collision/ParticlePairMidPhase.h b/Engine/Source/Runtime/Experimental/Chaos/Public/Chaos/Collision/ParticlePairMidPhase.h
--- a/Engine/Source/Runtime/Experimental/Chaos/Public/Chaos/Collision/ParticlePairMidPhase.h
+++ b/Engine/Source/Runtime/Experimental/Chaos/Public/Chaos/Collision/ParticlePairMidPhase.h
@@ -131,21 +131,21 @@ namespace Chaos
 		FGeometryParticleHandle* Particle1;
 		const FPerShapeData* Shape0;
 		const FPerShapeData* Shape1;
-		EContactShapesType ShapePairType;
-		FReal SphereBoundsCheckSize;
+		FRealSingle SphereBoundsCheckSize;
 		int32 LastUsedEpoch;
+		EContactShapesType ShapePairType;
 		union FFlags
 		{
 			FFlags() : Bits(0) {}
 			struct 
 			{
-				uint32 bEnableAABBCheck : 1;
-				uint32 bEnableOBBCheck0 : 1;
-				uint32 bEnableOBBCheck1 : 1;
-				uint32 bEnableManifoldUpdate : 1;
-				uint32 bIsProbe : 1;
+				uint8 bEnableAABBCheck : 1;
+				uint8 bEnableOBBCheck0 : 1;
+				uint8 bEnableOBBCheck1 : 1;
+				uint8 bEnableManifoldUpdate : 1;
+				uint8 bIsProbe : 1;
 			};
-			uint32 Bits;
+			uint8 Bits;
 		} Flags;
 	};
 
diff --git a/Engine/Source/Runtime/Experimental/Chaos/Public/Chaos/Collision/SphereConvexContactPoint.h b/Engine/Source/Runtime/Experimental/Chaos/Public/Chaos/Collision/SphereConvexContactPoint.h
--- a/Engine/Source/Runtime/Experimental/Chaos/Public/Chaos/Collision/SphereConvexContactPoint.h
+++ b/Engine/Source/Runtime/Experimental/Chaos/Public/Chaos/Collision/SphereConvexContactPoint.h
@@ -1,15 +1,13 @@
 // Copyright Epic Games, Inc. All Rights Reserved.
 #pragma once
 
-#include "CoreMinimal.h"
 #include "Chaos/Core.h"
+#include "Chaos/CollisionResolutionTypes.h"
 #include "Chaos/ImplicitFwd.h"
 #include "Chaos/ImplicitObjectScaled.h"	// Cannot fwd declare scaled implicit
 
 namespace Chaos
 {
-	class FContactPoint;
-
 	extern CHAOS_API FContactPoint SphereConvexContactPoint(const FImplicitSphere3& Sphere, const FRigidTransform3& SphereTransform, const FImplicitObject& Convex, const FRigidTransform3& ConvexTransform);
 	extern CHAOS_API FContactPoint SphereConvexContactPoint(const FImplicitSphere3& Sphere, const FImplicitObject& Convex, const FRigidTransform3& SphereToConvexTransform);
 }
diff --git a/Engine/Source/Runtime/Experimental/Chaos/Public/Chaos/CollisionResolutionTypes.h b/Engine/Source/Runtime/Experimental/Chaos/Public/Chaos/CollisionResolutionTypes.h
--- a/Engine/Source/Runtime/Experimental/Chaos/Public/Chaos/CollisionResolutionTypes.h
+++ b/Engine/Source/Runtime/Experimental/Chaos/Public/Chaos/CollisionResolutionTypes.h
@@ -7,6 +7,9 @@
 
 namespace Chaos
 {
+	template<typename T> class TContactPoint;
+	using FContactPoint = TContactPoint<FReal>;
+	using FContactPointf = TContactPoint<FRealSingle>;
 
 	/** Specifies the type of work we should do*/
 	enum class ECollisionUpdateType
@@ -24,7 +27,7 @@ namespace Chaos
 	};
 
 	/** The shape types involved in a contact constraint. Used to look up the collision detection function */
-	enum class EContactShapesType
+	enum class EContactShapesType : int8
 	{
 		Unknown,
 		SphereSphere,
@@ -53,29 +56,10 @@ namespace Chaos
 		NumShapesTypes
 	};
 
-	enum class EContactManifoldType
-	{
-		None,			// No manifold - run collision detection whenever we need latest contact
-		OneShot,		// A manifold is created once and reused. The manifold consists of a plane attached to one shape and a set of points on the other
-		Incremental,	// Run collision detection whenever we need the latest contact point, but keep track of and match contact points
-	};
-
 	//
 	//
 	//
 
-	struct CHAOS_API FRigidBodyContactConstraintPGS
-	{
-		FRigidBodyContactConstraintPGS() : AccumulatedImpulse(0.f) {}
-		TGeometryParticleHandle<FReal, 3>* Particle;
-		TGeometryParticleHandle<FReal, 3>* Levelset;
-		TArray<FVec3> Normal;
-		TArray<FVec3> Location;
-		TArray<FReal> Phi;
-		FVec3 AccumulatedImpulse;
-	};
-
-
 	//TODO: move into a better forward declare location
 	class FPBDCollisionConstraintHandle;
 
diff --git a/Engine/Source/Runtime/Experimental/Chaos/Public/Chaos/CollisionResolutionUtil.h b/Engine/Source/Runtime/Experimental/Chaos/Public/Chaos/CollisionResolutionUtil.h
--- a/Engine/Source/Runtime/Experimental/Chaos/Public/Chaos/CollisionResolutionUtil.h
+++ b/Engine/Source/Runtime/Experimental/Chaos/Public/Chaos/CollisionResolutionUtil.h
@@ -8,7 +8,6 @@
 namespace Chaos
 {
 	class FBVHParticles;
-	class FContactPoint;
 	class FImplicitObject;
 	class FPBDCollisionConstraint;
 
diff --git a/Engine/Source/Runtime/Experimental/Chaos/Public/Chaos/HeightField.h b/Engine/Source/Runtime/Experimental/Chaos/Public/Chaos/HeightField.h
--- a/Engine/Source/Runtime/Experimental/Chaos/Public/Chaos/HeightField.h
+++ b/Engine/Source/Runtime/Experimental/Chaos/Public/Chaos/HeightField.h
@@ -2,6 +2,7 @@
 #pragma once
 
 #include "Chaos/Array.h"
+#include "Chaos/CollisionResolutionTypes.h"
 #include "ImplicitObject.h"
 #include "Box.h"
 #include "TriangleMeshImplicitObject.h"
@@ -16,7 +17,6 @@
 namespace Chaos
 {
 	class FHeightfieldRaycastVisitor;
-	class FContactPoint;
 	struct FMTDInfo;
 }
 
diff --git a/Engine/Source/Runtime/Experimental/Chaos/Public/Chaos/PBDCollisionConstraints.h b/Engine/Source/Runtime/Experimental/Chaos/Public/Chaos/PBDCollisionConstraints.h
--- a/Engine/Source/Runtime/Experimental/Chaos/Public/Chaos/PBDCollisionConstraints.h
+++ b/Engine/Source/Runtime/Experimental/Chaos/Public/Chaos/PBDCollisionConstraints.h
@@ -325,9 +325,9 @@ inline FPBDCollisionConstraint& FPBDCollisionConstraintHandle::GetContact()
 	return *GetConstraint();
 }
 
-inline ECollisionCCDType FPBDCollisionConstraintHandle::GetCCDType() const
+inline bool FPBDCollisionConstraintHandle::GetCCDEnabled() const
 {
-	return GetContact().GetCCDType();
+	return GetContact().GetCCDEnabled();
 }
 
 inline void FPBDCollisionConstraintHandle::SetEnabled(bool InEnabled)
diff --git a/Engine/Source/Runtime/Experimental/Chaos/Public/Chaos/TriangleMeshImplicitObject.h b/Engine/Source/Runtime/Experimental/Chaos/Public/Chaos/TriangleMeshImplicitObject.h
--- a/Engine/Source/Runtime/Experimental/Chaos/Public/Chaos/TriangleMeshImplicitObject.h
+++ b/Engine/Source/Runtime/Experimental/Chaos/Public/Chaos/TriangleMeshImplicitObject.h
@@ -2,6 +2,7 @@
 #pragma once
 
 #include "Chaos/Array.h"
+#include "Chaos/CollisionResolutionTypes.h"
 #include "Chaos/GeometryParticles.h"
 #include "Chaos/ImplicitFwd.h"
 #include "Chaos/ImplicitObjectScaled.h"
@@ -27,7 +28,6 @@ namespace Chaos
 	template <typename QueryGeomType>
 	struct FTriangleMeshOverlapVisitorNoMTD;
 	class FTriangleMeshImplicitObject;
-	class FContactPoint;
 	FArchive& operator<<(FArchive& Ar, FAABBVectorized& Bounds);
 
 	struct CHAOS_API FTrimeshBVH
