Chaos : GeometryCollection : fix potential crash when fixing collision for parent clusters

diff --git a/Engine/Source/Runtime/Experimental/Chaos/Private/PhysicsProxy/GeometryCollectionPhysicsProxy.cpp b/Engine/Source/Runtime/Experimental/Chaos/Private/PhysicsProxy/GeometryCollectionPhysicsProxy.cpp
--- a/Engine/Source/Runtime/Experimental/Chaos/Private/PhysicsProxy/GeometryCollectionPhysicsProxy.cpp
+++ b/Engine/Source/Runtime/Experimental/Chaos/Private/PhysicsProxy/GeometryCollectionPhysicsProxy.cpp
@@ -520,7 +520,6 @@ void FGeometryCollectionPhysicsProxy::Initialize(Chaos::FPBDRigidsEvolutionBase
 							// defer if any of the children is a cluster with no collision yet generated 
 							if (GameThreadCollection.Implicits[ChildIndex] == nullptr && GameThreadCollection.Children[ChildIndex].Num() > 0)
 							{
-								ParentToPotentiallyFix.Add(ParentToFixIndex);
 								bAllChildrenHaveCollision = false;
 								break;
 							}
