Chaos - only use margins on dynamic shapes, not kinematics or statics

diff --git a/Engine/Source/Runtime/Experimental/Chaos/Private/Chaos/Collision/PBDCollisionConstraint.cpp b/Engine/Source/Runtime/Experimental/Chaos/Private/Chaos/Collision/PBDCollisionConstraint.cpp
--- a/Engine/Source/Runtime/Experimental/Chaos/Private/Chaos/Collision/PBDCollisionConstraint.cpp
+++ b/Engine/Source/Runtime/Experimental/Chaos/Private/Chaos/Collision/PBDCollisionConstraint.cpp
@@ -308,8 +308,10 @@ namespace Chaos
 		// The collision tolerance is used for knowing whether a new contact matches an existing one.
 		//
 		// Margins: (Assuming convex margins are enabled...)
-		// If we have two polygonal shapes, we use the smallest of the two margins (unless one shape has zero margin, e.g. triangle).
-		// If we have a quadratic shape versus a polygonal shape, we use a zero margin on the polygoinal shape.
+		// If a polygonal shape is not dynamic, its margin is zero.
+		// Triangles always have zero margin (they are never dynamic anyway)
+		// If we have two polygonal shapes, we use the smallest of the two margins unless one shape has zero margin.
+		// If we have a quadratic shape versus a polygonal shape, we use a zero margin on the polygonal shape.
 		// Note: If we have a triangle, it is always the second shape (currently we do not support triangle-triangle collision)
 		//
 		// CollisionTolerance:
@@ -318,7 +320,7 @@ namespace Chaos
 		//
 		const bool bIsQuadratic0 = ((ImplicitType0 == ImplicitObjectType::Sphere) || (ImplicitType0 == ImplicitObjectType::Capsule));
 		const bool bIsQuadratic1 = ((ImplicitType1 == ImplicitObjectType::Sphere) || (ImplicitType1 == ImplicitObjectType::Capsule));
-		
+
 		// @todo(chaos): should probably be tunable. Used to use the same settings as the margin scale (for convex), but we want to support zero
 		// margins, but still have a non-zero CollisionTolerance (it is used for matching contact points for friction and manifold reuse)
 		const FRealSingle ToleranceScale = 0.1f;
@@ -326,6 +328,12 @@ namespace Chaos
 		
 		if (!bIsQuadratic0 && !bIsQuadratic1)
 		{
+			// Only dynamic polytopes objects use margins
+			const bool bIsDynamic0 = (GetParticle0() != nullptr) && (FConstGenericParticleHandle(GetParticle0())->IsDynamic());
+			const bool bIsDynamic1 = (GetParticle1() != nullptr) && (FConstGenericParticleHandle(GetParticle1())->IsDynamic());
+			const FRealSingle DynamicMargin0 = bIsDynamic0 ? FRealSingle(Margin0) : FRealSingle(0);
+			const FRealSingle DynamicMargin1 = bIsDynamic1 ? FRealSingle(Margin1) : FRealSingle(0);
+
 			const FRealSingle MaxSize0 = ((Implicit[0] != nullptr) && Implicit[0]->HasBoundingBox()) ? FRealSingle(Implicit[0]->BoundingBox().Extents().GetAbsMax()) : FRealSingle(0);
 			const FRealSingle MaxSize1 = ((Implicit[1] != nullptr) && Implicit[1]->HasBoundingBox()) ? FRealSingle(Implicit[1]->BoundingBox().Extents().GetAbsMax()) : FRealSingle(0);
 			const FRealSingle MaxSize = FMath::Min(MaxSize0, MaxSize1);
@@ -335,19 +343,19 @@ namespace Chaos
 			// If both shapes have a margin, use the smaller margin on both shapes. E.g., Box-Box
 			// We should never see both shapes with zero margin, but if we did we'd end up with a zero margin
 			const FRealSingle MinMargin = Chaos_Collision_ConvexZeroMargin;
-			if (Margin0 == FRealSingle(0))
+			if (DynamicMargin0 == FRealSingle(0))
 			{
 				CollisionMargins[0] = 0;
-				CollisionMargins[1] = FMath::Max(MinMargin, Margin1);
+				CollisionMargins[1] = FMath::Max(MinMargin, DynamicMargin1);
 			}
-			else if (Margin1 == FRealSingle(0))
+			else if (DynamicMargin1 == FRealSingle(0))
 			{
-				CollisionMargins[0] = FMath::Max(MinMargin, Margin0);
+				CollisionMargins[0] = FMath::Max(MinMargin, DynamicMargin0);
 				CollisionMargins[1] = 0;
 			}
 			else
 			{
-				CollisionMargins[0] = FMath::Min(Margin0, Margin1);
+				CollisionMargins[0] = FMath::Min(DynamicMargin0, DynamicMargin1);
 				CollisionMargins[1] = CollisionMargins[0];
 			}
 		}
