Fixed compilation in GJK.h

diff --git a/Engine/Source/Runtime/Experimental/Chaos/Public/Chaos/GJK.h b/Engine/Source/Runtime/Experimental/Chaos/Public/Chaos/GJK.h
--- a/Engine/Source/Runtime/Experimental/Chaos/Public/Chaos/GJK.h
+++ b/Engine/Source/Runtime/Experimental/Chaos/Public/Chaos/GJK.h
@@ -126,8 +126,15 @@ namespace Chaos
 		@param InitialDir The first direction we use to search the CSO
 		@return True if the geometries overlap, False otherwise 
 	 */
+	template <typename T, typename TGeometryA, typename TGeometryB>
+	bool GJKIntersection(const TGeometryA& RESTRICT A, const TGeometryB& RESTRICT B, const TRigidTransform<T, 3>& BToATM, const T InThicknessA = 0, const TVector<T, 3>& InitialDir = TVector<T, 3>(1, 0, 0))
+	{
+		const FReal EpsilonScale = FMath::Max<FReal>(A.TGeometryA::BoundingBox().Extents().Max(), B.TGeometryB::BoundingBox().Extents().Max());
+		return GJKIntersection(A, B, EpsilonScale, BToATM, InThicknessA, InitialDir);
+	}
+
 	template <typename T>
-	bool GJKIntersection(const FGeomGJKHelperSIMD& RESTRICT A, const FGeomGJKHelperSIMD& RESTRICT B, const TRigidTransform<T, 3>& BToATM, const T InThicknessA = 0, const TVector<T, 3>& InitialDir = TVector<T, 3>(1, 0, 0))
+	bool GJKIntersection(const FGeomGJKHelperSIMD& RESTRICT A, const FGeomGJKHelperSIMD& RESTRICT B, FReal EpsilonScale, const TRigidTransform<T, 3>& BToATM, const T InThicknessA = 0, const TVector<T, 3>& InitialDir = TVector<T, 3>(1, 0, 0))
 	{
 		const UE::Math::TQuat<T>& RotationDouble = BToATM.GetRotation();
 		VectorRegister4Float RotationSimd = MakeVectorRegisterFloatFromDouble(MakeVectorRegister(RotationDouble.X, RotationDouble.Y, RotationDouble.Z, RotationDouble.W));
@@ -139,11 +146,17 @@ namespace Chaos
 
 		const VectorRegister4Float InitialDirSimd = MakeVectorRegisterFloatFromDouble(MakeVectorRegister(InitialDir[0], InitialDir[1], InitialDir[2], 0.0));
 		
-		return GJKIntersectionSimd(A, B, TranslationSimd, RotationSimd, InThicknessA, InitialDirSimd);
+		return GJKIntersectionSimd(A, B, EpsilonScale, TranslationSimd, RotationSimd, InThicknessA, InitialDirSimd);
 	}
 
+	template <typename TGeometryA, typename TGeometryB>
+	bool GJKIntersectionSimd(const TGeometryA& RESTRICT A, const TGeometryB& RESTRICT B, const VectorRegister4Float& Translation, const VectorRegister4Float& Rotation, FReal InThicknessA, const VectorRegister4Float& InitialDir)
+	{
+		const FReal EpsilonScale = FMath::Max<FReal>(A.TGeometryA::BoundingBox().Extents().Max(), B.TGeometryB::BoundingBox().Extents().Max());
+		return GJKIntersectionSimd(A, B, EpsilonScale, Translation, Rotation, InThicknessA, InitialDir);
+	}
 
-	inline bool GJKIntersectionSimd(const FGeomGJKHelperSIMD& RESTRICT A, const FGeomGJKHelperSIMD& RESTRICT B, const VectorRegister4Float& Translation, const VectorRegister4Float& Rotation, FReal InThicknessA, const VectorRegister4Float& InitialDir)
+	inline bool GJKIntersectionSimd(const FGeomGJKHelperSIMD& RESTRICT A, const FGeomGJKHelperSIMD& RESTRICT B, FReal EpsilonScale, const VectorRegister4Float& Translation, const VectorRegister4Float& Rotation, FReal InThicknessA, const VectorRegister4Float& InitialDir)
 	{
 		VectorRegister4Float V = VectorNegate(InitialDir);
 		V = VectorNormalizeSafe(V, MakeVectorRegisterFloatConstant(-1.f, 0.f, 0.f, 0.f));
@@ -163,7 +176,6 @@ namespace Chaos
 		CalculateQueryMargins(A, B, ThicknessA, ThicknessB);
 		ThicknessA += InThicknessA;
 
-		FReal EpsilonScale = FMath::Max<FReal>(A.TGeometryA::BoundingBox().Extents().Max(), B.TGeometryB::BoundingBox().Extents().Max());
 		EpsilonScale = FMath::Min<FReal>(EpsilonScale, 1e5);
 		
 		const FReal InflationReal = ThicknessA + ThicknessB + 1e-6 * EpsilonScale;
