Chaos - fix crash in joint projection (missing check for disabled joint)

diff --git a/Engine/Source/Runtime/Experimental/Chaos/Private/Chaos/PBDJointConstraints.cpp b/Engine/Source/Runtime/Experimental/Chaos/Private/Chaos/PBDJointConstraints.cpp
--- a/Engine/Source/Runtime/Experimental/Chaos/Private/Chaos/PBDJointConstraints.cpp
+++ b/Engine/Source/Runtime/Experimental/Chaos/Private/Chaos/PBDJointConstraints.cpp
@@ -1307,6 +1307,11 @@ namespace Chaos
 
 		if (Settings.bUseLinearSolver)
 		{
+			if (!CanEvaluate(ConstraintIndex))
+			{
+				return;
+			}
+
 			FPBDJointCachedSolver& Solver = CachedConstraintSolvers[ConstraintIndex];
 			const FPBDJointSettings& JointSettings = ConstraintSettings[ConstraintIndex];
 			if (JointSettings.bProjectionEnabled && (Solver.IsDynamic(0) || Solver.IsDynamic(1)))
