Chaos: Add a memory barrier to TObjectPool to ensure no writes to the head block will reorder beyond grabbing a free item pointer
#b Chris.Caulfield

diff --git a/Engine/Source/Runtime/Experimental/ChaosCore/Public/Chaos/ObjectPool.h b/Engine/Source/Runtime/Experimental/ChaosCore/Public/Chaos/ObjectPool.h
--- a/Engine/Source/Runtime/Experimental/ChaosCore/Public/Chaos/ObjectPool.h
+++ b/Engine/Source/Runtime/Experimental/ChaosCore/Public/Chaos/ObjectPool.h
@@ -94,6 +94,10 @@ namespace Chaos
 				Swap(Blocks[0], Blocks[BackIndex]);
 			}
 
+			// Ensure any writes to Blocks[0] cannot be reordered so that the GetNextFree
+			// call is definitely handled by the right block
+			FPlatformMisc::MemoryBarrier();
+
 			// Blocks[0] is now a block with at least one free element
 			ObjectType* NewPtr = Blocks[0].GetNextFree();
 			Construct<ObjectType>(NewPtr, Forward<TArgs>(Args)...);
@@ -355,6 +359,8 @@ namespace Chaos
 					return &NewItem->Object;
 				}
 
+				checkf(false, TEXT("Attempt to request a free item from a full block (Freelist is empty, NumValid is %d, NumInBlock is %d, NumFree is %d)"), NumValid, NumInBlock, NumFree);
+
 				return nullptr;
 			}
 
