Chaos: replication imporvement
- Add a timer based hard snap for updating cluster position on the client
- Change remove on break to be replication aware to make sure server authoritative cluster cannot be broken by the client

diff --git a/Engine/Source/Runtime/Experimental/Chaos/Public/PhysicsProxy/GeometryCollectionPhysicsProxy.h b/Engine/Source/Runtime/Experimental/Chaos/Public/PhysicsProxy/GeometryCollectionPhysicsProxy.h
--- a/Engine/Source/Runtime/Experimental/Chaos/Public/PhysicsProxy/GeometryCollectionPhysicsProxy.h
+++ b/Engine/Source/Runtime/Experimental/Chaos/Public/PhysicsProxy/GeometryCollectionPhysicsProxy.h
@@ -384,8 +384,15 @@ public:
 	FProxyInterpolationData& GetInterpolationData() { return InterpolationData; }
 	const FProxyInterpolationData& GetInterpolationData() const { return InterpolationData; }
 
-	void SetInitializedForReplication() { bInitializedForReplication = true; }
-	bool GetInitializedForReplication() const { return bInitializedForReplication; }
+	enum class EReplicationMode: uint8
+	{
+		Unknown,
+		Server,
+		Client,
+	};
+
+	void SetReplicationMode(EReplicationMode Mode) { ReplicationMode = Mode; }
+	EReplicationMode GetReplicationMode() const { return ReplicationMode; }
 
 	void UpdateFilterData_External(const FCollisionFilterData& NewSimFilter, const FCollisionFilterData& NewQueryFilter);
 	
@@ -473,7 +480,7 @@ private:
 	bool IsObjectLoading; // Indicate when loaded
 	bool IsObjectDeleting; // Indicatge when pending deletion
 
-	bool bInitializedForReplication = false;	//Indicates that initialization for replication has finished
+	EReplicationMode ReplicationMode = EReplicationMode::Unknown;	
 
 	TManagedArray<TUniquePtr<Chaos::FGeometryParticle>> GTParticles;
 	TMap<Chaos::FGeometryParticle*, int32> GTParticlesToTransformGroupIndex;
diff --git a/Engine/Source/Runtime/Experimental/GeometryCollectionEngine/Private/GeometryCollection/GeometryCollectionComponent.cpp b/Engine/Source/Runtime/Experimental/GeometryCollectionEngine/Private/GeometryCollection/GeometryCollectionComponent.cpp
--- a/Engine/Source/Runtime/Experimental/GeometryCollectionEngine/Private/GeometryCollection/GeometryCollectionComponent.cpp
+++ b/Engine/Source/Runtime/Experimental/GeometryCollectionEngine/Private/GeometryCollection/GeometryCollectionComponent.cpp
@@ -253,7 +253,8 @@ UGeometryCollectionComponent::UGeometryCollectionComponent(const FObjectInitiali
 #endif
 	, bEnableReplication(false)
 	, bEnableAbandonAfterLevel(true)
-	, ReplicationAbandonClusterLevel(0)
+	, ReplicationAbandonClusterLevel_DEPRECATED(0)
+	, ReplicationAbandonAfterLevel(0)
 	, bRenderStateDirty(true)
 	, bEnableBoneSelection(false)
 	, ViewLevel(-1)
@@ -1170,6 +1171,21 @@ void UGeometryCollectionComponent::SetRestState(TArray<FTransform>&& InRestTrans
 	RefreshEmbeddedGeometry();
 }
 
+FORCEINLINE_DEBUGGABLE int32 ComputeParticleLevel(Chaos::FPBDRigidClusteredParticleHandle* Particle)
+{
+	int32 Level = 0;
+	if (Particle)
+	{
+		Chaos::FPBDRigidClusteredParticleHandle* Current = Particle;
+		while (Current->Parent())
+		{
+			Current = Current->Parent();
+			++Level;
+		}
+	}
+	return Level;
+};
+
 void UGeometryCollectionComponent::InitializeComponent()
 {
 	Super::InitializeComponent();
@@ -1198,57 +1214,44 @@ void UGeometryCollectionComponent::InitializeComponent()
 		return;
 	}
 
-	const ENetRole LocalRole = Owner->GetLocalRole();
-	const ENetMode NetMode = Owner->GetNetMode();
-
 	// If we're replicating we need some extra setup - check netmode as we don't need this for
 	// standalone runtimes where we aren't going to network the component
-	if(GetIsReplicated())
+	if(PhysicsProxy && GetIsReplicated())
 	{
-		if(LocalRole != ENetRole::ROLE_Authority)
-		{
+		const ENetRole LocalRole = Owner->GetLocalRole();
+		const ENetMode NetMode = Owner->GetNetMode();
 
-			// We're a replicated component and we're not in control.
-			Chaos::FPhysicsSolver* CurrSolver = GetSolver(*this);
-
-			if(CurrSolver)
+		const bool bIsReplicationAuthority = (LocalRole == ENetRole::ROLE_Authority);
+		// Client side : geometry collection children of parents below the rep level need to be infintely strong so that client cannot break it 
+		if (Chaos::FPhysicsSolver* CurrSolver = GetSolver(*this))
+		{
+			CurrSolver->RegisterSimOneShotCallback([Proxy = PhysicsProxy, bIsReplicationAuthority, AbandonAfterLevel = ReplicationAbandonAfterLevel, EnableAbandonAfterLevel = bEnableAbandonAfterLevel]()
 			{
-				CurrSolver->RegisterSimOneShotCallback([Prox = PhysicsProxy, ReplicationLevel = ReplicationAbandonClusterLevel, AbandonAfterLevel = bEnableAbandonAfterLevel, this]()
+				if (bIsReplicationAuthority)
 				{
+					Proxy->SetReplicationMode(FGeometryCollectionPhysicsProxy::EReplicationMode::Server);
+				}
+				else
+				{
+					Proxy->SetReplicationMode(FGeometryCollectionPhysicsProxy::EReplicationMode::Client);
+
 					// As we're not in control we make it so our simulated proxy cannot break clusters
 					// We have to set the strain to a high value but be below the max for the data type
 					// so releasing on authority demand works
-					const Chaos::FReal MaxStrain = TNumericLimits<Chaos::FReal>::Max() - TNumericLimits<Chaos::FReal>::Min();
-
-					TArray<Chaos::TPBDRigidClusteredParticleHandle<Chaos::FReal, 3>*> Particles = Prox->GetParticles();
-
-					for(Chaos::TPBDRigidClusteredParticleHandle<Chaos::FReal, 3> * P : Particles)
+					constexpr Chaos::FReal MaxStrain = TNumericLimits<Chaos::FReal>::Max() - TNumericLimits<Chaos::FReal>::Min();
+					for (Chaos::FPBDRigidClusteredParticleHandle* P : Proxy->GetParticles())
 					{
-						if(!P)
-						{
-							continue;
-						}
-
-						int32 Level = AbandonAfterLevel ? 0 : -1;
-						if(AbandonAfterLevel)
+						if (P)
 						{
-							Chaos::TPBDRigidClusteredParticleHandle<Chaos::FReal, 3>* Current = P;
-							while ( Current->Parent())
+							const int32 Level = EnableAbandonAfterLevel ? ComputeParticleLevel(P) : -1;
+							if (Level <= AbandonAfterLevel + 1)	//we only replicate up until level X, but it means we should replicate the breaking event of level X+1 (but not X+1's positions)
 							{
-								Current = Current->Parent();
-								++Level;
+								P->SetStrain(MaxStrain);
 							}
 						}
-
-						if (Level <= ReplicationLevel + 1)	//we only replicate up until level X, but it means we should replicate the breaking event of level X+1 (but not X+1's positions)
-						{
-							P->SetStrain(MaxStrain);
-						}
 					}
-
-					Prox->SetInitializedForReplication();
-				});
-			}
+				}
+			});
 		}
 	}
 }
@@ -1473,7 +1476,8 @@ void UGeometryCollectionComponent::UpdateRepData()
 		const FPBDRigidsSolver* Solver = PhysicsProxy->GetSolver<Chaos::FPBDRigidsSolver>();
 		const FRigidClustering& RigidClustering = Solver->GetEvolution()->GetRigidClustering();
 
-		const TManagedArray<int32>& Levels = PhysicsProxy->GetPhysicsCollection().GetAttribute<int32>("InitialLevel", FGeometryCollection::TransformGroup);
+		const TManagedArray<int32>& InitialLevels = PhysicsProxy->GetPhysicsCollection().GetAttribute<int32>("InitialLevel", FGeometryCollection::TransformGroup);
+		const TManagedArray<TSet<int32>>& InitialChildren = PhysicsProxy->GetPhysicsCollection().Children;
 
 		//see if we have any new clusters that are enabled
 		TSet<FPBDRigidClusteredParticleHandle*> Processed;
@@ -1511,23 +1515,23 @@ void UGeometryCollectionComponent::UpdateRepData()
 						ensureMsgf(TransformGroupIdx >= 0, TEXT("Non-internal cluster should always have a group index"));
 						ensureMsgf(TransformGroupIdx < TNumericLimits<uint16>::Max(), TEXT("Trying to replicate GC with more than 65k pieces. We assumed uint16 would suffice"));
 
-						Level = Levels[TransformGroupIdx];
+						Level = InitialLevels[TransformGroupIdx];
 					}
 					else
 					{
 						// Use internal cluster child's index to compute level.
 						const TArray<FPBDRigidParticleHandle*>& Children = RigidClustering.GetChildrenMap()[Root];
 						const int32 ChildTransformGroupIdx = PhysicsProxy->GetTransformGroupIndexFromHandle(Children[0]);
-						Level = Levels[ChildTransformGroupIdx] - 1;
+						Level = InitialLevels[ChildTransformGroupIdx] - 1;
 					}
 	
-					if (!bEnableAbandonAfterLevel|| Level < ReplicationAbandonClusterLevel)
+					if (!bEnableAbandonAfterLevel || Level <= ReplicationAbandonAfterLevel)
 					{
 						// not already replicated and not abandoned level, start replicating cluster
 						ClustersToRep->Add(Root);
 						bClustersChanged = true;
 					}
-	
+
 					if(Root->InternalCluster() == false && bFirstUpdate == false)	//if bFirstUpdate it must be that these are the initial roots of the GC. These did not break off so no need to replicate
 					{
 						//a one off so record it
@@ -1536,6 +1540,7 @@ void UGeometryCollectionComponent::UpdateRepData()
 	
 						// Because we cull ClustersToRep with abandoned level, we must make sure we don't add duplicates to one off activated.
 						// TODO: avoid search for entry for perf
+						// TODO: once we support deep fracture we should be able to remove one offs clusters that are now disabled, reducing the amount to be replicated
 						FGeometryCollectionActivatedCluster OneOffActivated(TransformGroupIdx, Root->V(), Root->W());
 						if(!RepData.OneOffActivated.Contains(OneOffActivated))
 						{
@@ -1543,12 +1548,21 @@ void UGeometryCollectionComponent::UpdateRepData()
 							RepData.OneOffActivated.Add(OneOffActivated);
 						}
 					}
+
+					// if we just hit the abandon level , let's disable all children 
+					if (bEnableAbandonAfterLevel && Level >= (ReplicationAbandonAfterLevel + 1))
+					{
+						if (!Root->Disabled())
+						{
+							Solver->GetEvolution()->DisableParticle(Root);
+						}
+					}
 				}
 			}
 		}
 		
 		INC_DWORD_STAT_BY(STAT_GCReplicatedFractures, RepData.OneOffActivated.Num());
-
+		
 		//build up clusters to replicate and compare with previous frame
 		TArray<FGeometryCollectionClusterRep> Clusters;
 
@@ -1590,19 +1604,18 @@ void UGeometryCollectionComponent::UpdateRepData()
 				if(!bClustersChanged)
 				{
 					//compare to previous frame data
-					if(RepData.Clusters.Num() >= Clusters.Num())
+					// this could be more efficient by having a way to find back the data from the idx
+					auto Predicate = [TransformGroupIdx](const FGeometryCollectionClusterRep& Entry)
+					{
+						return Entry.ClusterIdx == TransformGroupIdx;
+					};
+					if (const FGeometryCollectionClusterRep* PrevClusterData = RepData.Clusters.FindByPredicate(Predicate))
 					{
-						const FGeometryCollectionClusterRep& PrevCluster = RepData.Clusters[Clusters.Num() - 1];
-						if(ClusterRep.ClusterChanged(PrevCluster))
+						if (ClusterRep.ClusterChanged(*PrevClusterData))
 						{
 							bClustersChanged = true;
 						}
 					}
-					else
-					{
-						//must be some new clusters so definitely changed
-						bClustersChanged = true;
-					}
 				}
 			}
 		}
@@ -1630,26 +1643,36 @@ int32 GeometryCollectionHardMissingUpdatesSnapThreshold = 20;
 FAutoConsoleVariableRef CVarGeometryCollectionHardMissingUpdatesSnapThreshold(TEXT("p.GeometryCollectionHardMissingUpdatesSnapThreshold"), GeometryCollectionHardMissingUpdatesSnapThreshold,
 	TEXT("Determines how many missing updates before we trigger a hard snap"));
 
+int32 GeometryCollectionHardsnapThresholdMs = 100; // 10 Hz
+FAutoConsoleVariableRef CVarGeometryCollectionHardsnapThresholdMs(TEXT("p.GeometryCollectionHardsnapThresholdMs"), GeometryCollectionHardMissingUpdatesSnapThreshold,
+	TEXT("Determines how many ms since the last hardsnap to trigger a new one"));
+
 void UGeometryCollectionComponent::ProcessRepData()
 {
 	using namespace Chaos;
-	if(VersionProcessed == RepData.Version || PhysicsProxy->GetInitializedForReplication() == false)
+	if(VersionProcessed == RepData.Version || PhysicsProxy->GetReplicationMode() != FGeometryCollectionPhysicsProxy::EReplicationMode::Client)
 	{
 		return;
 	}
 
 	bool bHardSnap = false;
+	const int64 CurrentTimeInMs = FPlatformTime::ToMilliseconds64(FPlatformTime::Cycles64());
 	if(VersionProcessed < RepData.Version)
 	{
 		//TODO: this will not really work if a fracture happens and then immediately goes to sleep without updating client enough times
 		//A time method would work better here, but is limited to async mode. Maybe we can support both
-		bHardSnap = (RepData.Version - VersionProcessed) > GeometryCollectionHardMissingUpdatesSnapThreshold;
+		bHardSnap |= (RepData.Version - VersionProcessed) > GeometryCollectionHardMissingUpdatesSnapThreshold;
+		bHardSnap |= (CurrentTimeInMs - LastHardsnapTimeInMs) > GeometryCollectionHardsnapThresholdMs;
 	}
 	else
 	{
 		//rollover so just treat as hard snap - this case is extremely rare and a one off
 		bHardSnap = true;
 	}
+	if (bHardSnap)
+	{
+		LastHardsnapTimeInMs = CurrentTimeInMs;
+	}
 
 	FPBDRigidsSolver* Solver = PhysicsProxy->GetSolver<Chaos::FPBDRigidsSolver>();
 	FRigidClustering& RigidClustering = Solver->GetEvolution()->GetRigidClustering();
@@ -3910,6 +3933,9 @@ void UGeometryCollectionComponent::IncrementBreakTimer(float DeltaTime)
 		FGeometryCollectionDecayDynamicFacade DecayFacade(*DynamicCollection);
 		FGeometryCollectionDynamicStateFacade DynamicStateFacade(*DynamicCollection);
 
+		// if replication is on, client may not need to process this at all or only partially ( depending on the abandon cluster level )
+		const bool bIsReplicatedClient = GetIsReplicated() && PhysicsProxy->GetReplicationMode() == FGeometryCollectionPhysicsProxy::EReplicationMode::Client;
+
 		if (RemoveOnBreakFacade.IsValid()
 			&& DecayFacade.IsValid()
 			&& DynamicStateFacade.IsValid())
@@ -3917,6 +3943,8 @@ void UGeometryCollectionComponent::IncrementBreakTimer(float DeltaTime)
 			FGeometryCollectionDecayContext DecayContext(*PhysicsProxy, DecayFacade);
 			const TManagedArray<int32>& OriginalParents = RestCollection->GetGeometryCollection()->Parent;
 
+			const TManagedArray<int32>* InitialLevels = PhysicsProxy->GetPhysicsCollection().FindAttribute<int32>("InitialLevel", FGeometryCollection::TransformGroup);
+
 			TManagedArray<float>& Decay = DecayFacade.DecayAttribute.Modify();
 			for (int32 TransformIdx = 0; TransformIdx < Decay.Num(); ++TransformIdx)
 			{
@@ -3929,16 +3957,40 @@ void UGeometryCollectionComponent::IncrementBreakTimer(float DeltaTime)
 
 					if (OriginalParentIdx > INDEX_NONE && HasDynamicInternalClusterParent && RemoveOnBreakFacade.IsRemovalActive(OriginalParentIdx))
 					{
+						bool bIsAllowedClusterCrumbling = true;
+						if (bIsReplicatedClient && InitialLevels)
+						{
+							if (!bEnableAbandonAfterLevel || (*InitialLevels)[OriginalParentIdx] <= ReplicationAbandonAfterLevel)
+							{
+								bIsAllowedClusterCrumbling = false;
+							}
+						}
+
 						const bool UseClusterCrumbling = RemoveOnBreakFacade.UseClusterCrumbling(OriginalParentIdx);
-						const float UpdatedBreakDecay = RemoveOnBreakFacade.UpdateBreakTimerAndComputeDecay(TransformIdx, DeltaTime);
-						UpdateDecay(TransformIdx, UpdatedBreakDecay, UseClusterCrumbling, HasDynamicInternalClusterParent, DecayContext);
+						if (!UseClusterCrumbling || bIsAllowedClusterCrumbling)
+						{
+							const float UpdatedBreakDecay = RemoveOnBreakFacade.UpdateBreakTimerAndComputeDecay(TransformIdx, DeltaTime);
+							UpdateDecay(TransformIdx, UpdatedBreakDecay, UseClusterCrumbling, HasDynamicInternalClusterParent, DecayContext);
+						}
 					}
 				} 
 				else if (RemoveOnBreakFacade.IsRemovalActive(TransformIdx) && DynamicStateFacade.HasBrokenOff(TransformIdx))
 				{
+					bool bIsAllowedClusterCrumbling = true;
+					if (bIsReplicatedClient && InitialLevels)
+					{
+						if (!bEnableAbandonAfterLevel || (*InitialLevels)[TransformIdx] <= ReplicationAbandonAfterLevel)
+						{
+							bIsAllowedClusterCrumbling = false;
+						}
+					}
+
 					const bool UseClusterCrumbling = RemoveOnBreakFacade.UseClusterCrumbling(TransformIdx);
-					const float UpdatedBreakDecay = RemoveOnBreakFacade.UpdateBreakTimerAndComputeDecay(TransformIdx, DeltaTime);
-					UpdateDecay(TransformIdx, UpdatedBreakDecay, UseClusterCrumbling, false, DecayContext);
+					if (!UseClusterCrumbling || bIsAllowedClusterCrumbling)
+					{
+						const float UpdatedBreakDecay = RemoveOnBreakFacade.UpdateBreakTimerAndComputeDecay(TransformIdx, DeltaTime);
+						UpdateDecay(TransformIdx, UpdatedBreakDecay, UseClusterCrumbling, false, DecayContext);
+					}
 				}
 			}
 
diff --git a/Engine/Source/Runtime/Experimental/GeometryCollectionEngine/Public/GeometryCollection/GeometryCollectionComponent.h b/Engine/Source/Runtime/Experimental/GeometryCollectionEngine/Public/GeometryCollection/GeometryCollectionComponent.h
--- a/Engine/Source/Runtime/Experimental/GeometryCollectionEngine/Public/GeometryCollection/GeometryCollectionComponent.h
+++ b/Engine/Source/Runtime/Experimental/GeometryCollectionEngine/Public/GeometryCollection/GeometryCollectionComponent.h
@@ -1000,7 +1000,7 @@ protected:
 	bool bEnableReplication;
 
 	/** 
-	 * Enables use of ReplicationAbandonClusterLevel to stop providing network updates to
+	 * Enables use of ReplicationAbandonAfterLevel to stop providing network updates to
 	 * clients when the updated particle is of a level higher then specified.
 	 */
 	UPROPERTY(EditAnywhere, BlueprintReadOnly, Category = Network)
@@ -1012,8 +1012,15 @@ protected:
 	 * gameplay relevant to cut down on required bandwidth to update a collection.
 	 * @see bEnableAbandonAfterLevel
 	 */ 
+	UPROPERTY(meta = (DeprecatedProperty, DeprecationMessage = "GeometryCollection now uses ReplicationAbandonAfterLevel instead of ReplicationAbandonClusterLevel."))
+	int32 ReplicationAbandonClusterLevel_DEPRECATED;
+
+	/**
+	* If replicating - the cluster level after which replication will not happen 
+	* @see bEnableAbandonAfterLevel
+	*/
 	UPROPERTY(EditAnywhere, BlueprintReadOnly, Category = Network)
-	int32 ReplicationAbandonClusterLevel;
+	int32 ReplicationAbandonAfterLevel;
 
 	UPROPERTY(Replicated)
 	FGeometryCollectionRepData RepData;
@@ -1105,6 +1112,7 @@ private:
 	/** One off activation is processed in the same order as server so remember the last one we processed */
 	int32 OneOffActivatedProcessed = 0;
 	int32 VersionProcessed = INDEX_NONE;
+	double LastHardsnapTimeInMs = 0;
 
 	/** True if GeometryCollection transforms have changed from previous tick. */
 	bool bIsMoving;
