Broad Phase optimization: AbbTree use SIMD for Raycast

diff --git a/Engine/Source/Runtime/Experimental/Chaos/Public/Chaos/AABBTree.h b/Engine/Source/Runtime/Experimental/Chaos/Public/Chaos/AABBTree.h
--- a/Engine/Source/Runtime/Experimental/Chaos/Public/Chaos/AABBTree.h
+++ b/Engine/Source/Runtime/Experimental/Chaos/Public/Chaos/AABBTree.h
@@ -3,6 +3,7 @@
 
 #include "Chaos/AABB.h"
 #include "Chaos/AABBVectorized.h"
+#include "Chaos/AABBVectorizedDouble.h"
 #include "Chaos/AABBTreeDirtyGridUtils.h"
 #include "Chaos/Defines.h"
 #include "Chaos/GeometryParticles.h"
@@ -276,11 +277,17 @@ struct TAABBTreeLeafArray : public TBoundsWrapperHelper<TPayloadType, T, bComput
 	}
 
 	template <typename TSQVisitor, typename TQueryFastData>
-	FORCEINLINE_DEBUGGABLE bool RaycastFast(const TVec3<T>& Start, TQueryFastData& QueryFastData, TSQVisitor& Visitor, const TVec3<T>& Dir, const TVec3<T> InvDir, const bool bParallel[3]) const
+	FORCEINLINE_DEBUGGABLE bool RaycastFast(const TVec3<T>& Start, TQueryFastData& QueryFastData, TSQVisitor& Visitor, const TVec3<T>& Dir, const TVec3<T>& InvDir, const bool bParallel[3]) const
 	{
 		return RaycastSweepImp</*bSweep=*/false>(Start, QueryFastData, TVec3<T>((T)0), Visitor, Dir, InvDir, bParallel);
 	}
 
+	template <typename TSQVisitor, typename TQueryFastData>
+	FORCEINLINE_DEBUGGABLE bool RaycastFastSimd(const VectorRegister4Double& Start, TQueryFastData& QueryFastData, TSQVisitor& Visitor, const VectorRegister4Double& Dir, const VectorRegister4Double& InvDir, const VectorRegister4Double& Parallel, const VectorRegister4Double& Length) const
+	{
+		return RaycastImpSimd(Start, QueryFastData, Visitor, InvDir, Parallel, Length);
+	}
+
 	template <typename TSQVisitor, typename TQueryFastData>
 	FORCEINLINE_DEBUGGABLE bool SweepFast(const TVec3<T>& Start, TQueryFastData& QueryFastData, const TVec3<T>& QueryHalfExtents, TSQVisitor& Visitor,
 			const TVec3<T>& Dir, const TVec3<T> InvDir, const bool bParallel[3]) const
@@ -340,6 +347,33 @@ struct TAABBTreeLeafArray : public TBoundsWrapperHelper<TPayloadType, T, bComput
 		return true;
 	}
 
+	template <typename TQueryFastData, typename TSQVisitor>
+	FORCEINLINE_DEBUGGABLE bool RaycastImpSimd(const VectorRegister4Double& Start, TQueryFastData& QueryFastData, TSQVisitor& Visitor, const VectorRegister4Double& InvDir, const VectorRegister4Double& Parallel, const VectorRegister4Double& Length) const
+	{
+		PHYSICS_CSV_CUSTOM_VERY_EXPENSIVE(PhysicsCounters, MaxLeafSize, Elems.Num(), ECsvCustomStatOp::Max);
+		VectorRegister4Double TOI;
+		for (const auto& Elem : Elems)
+		{
+			if (PrePreFilterHelper(Elem.Payload, Visitor))
+			{
+				continue;
+			}
+
+			const FAABBVectorizedDouble InstanceBounds(Elem.Bounds);
+			if (InstanceBounds.RaycastFast(Start, InvDir, Parallel, Length, TOI))
+			{
+				TSpatialVisitorData<TPayloadType> VisitData(Elem.Payload, true, FAABB3(Elem.Bounds.Min(), Elem.Bounds.Max()));
+				const bool bContinue = Visitor.VisitRaycast(VisitData, QueryFastData);
+				if (!bContinue)
+				{
+					return false;
+				}
+			}
+		}
+		return true;
+	}
+
+
 	void RemoveElement(TPayloadType Payload)
 	{
 		for (int32 Idx = 0; Idx < Elems.Num(); ++Idx)
@@ -2641,6 +2675,22 @@ private:
 //		}
 //#endif
 
+
+		VectorRegister4Double StartSimd;
+		VectorRegister4Double DirSimd;
+		VectorRegister4Double Parallel;
+		VectorRegister4Double InvDirSimd;
+		VectorRegister4Double LengthSimd;
+
+		if constexpr (Query == EAABBQueryType::Raycast)
+		{
+			StartSimd = VectorLoadDouble3(&Start.X);
+			DirSimd = VectorLoadDouble3(&Dir.X);
+			Parallel = VectorCompareGT(GlobalVectorConstants::DoubleSmallNumber, VectorAbs(DirSimd));
+			InvDirSimd = VectorBitwiseNotAnd(Parallel, VectorDivide(VectorOne(), DirSimd));
+			LengthSimd = VectorSetDouble1(CurData.CurrentLength);
+		}
+
 		while (NodeStackNum)
 		{
 			PHYSICS_CSV_SCOPED_VERY_EXPENSIVE(PhysicsVerbose, QueryImp_NodeTraverse);
@@ -2680,7 +2730,7 @@ private:
 						return false;
 					}
 				}
-				else if (Leaf.RaycastFast(Start, CurData, Visitor, Dir, InvDir, bParallel) == false)
+				else if (Leaf.RaycastFastSimd(StartSimd, CurData, Visitor, DirSimd, InvDirSimd, Parallel, LengthSimd) == false)
 				{
 					return false;
 				}
@@ -2690,15 +2740,29 @@ private:
 				PHYSICS_CSV_SCOPED_VERY_EXPENSIVE(PhysicsVerbose, NodeTraverse_Branch);
 				int32 Idx = 0;
 				
-				if constexpr(Query != EAABBQueryType::Overlap)
+				if constexpr (Query != EAABBQueryType::Overlap)
 				{
-					FReal TOI0;
-					FAABB3 AABB0(Node.ChildrenBounds[0].Min(), Node.ChildrenBounds[0].Max());
-					bool bIntersect0 = TAABBTreeIntersectionHelper<TQueryFastData, Query>::Intersects(Start, CurData, TOI0, AABB0, QueryBounds, QueryHalfExtents, Dir, InvDir, bParallel);
+					bool bIntersect0, bIntersect1;
+					FReal TOI0, TOI1;
+					if constexpr (Query == EAABBQueryType::Raycast)
+					{
+						FAABBVectorizedDouble AABBVectorized(Node.ChildrenBounds[0]);
+						VectorRegister4Double TOISimd;
+						bIntersect0 = AABBVectorized.RaycastFast(StartSimd, InvDirSimd, Parallel, LengthSimd, TOISimd);
+						VectorStoreDouble1(TOISimd, &TOI0);
+
+						AABBVectorized = FAABBVectorizedDouble(Node.ChildrenBounds[1]);
+						bIntersect1 = AABBVectorized.RaycastFast(StartSimd, InvDirSimd, Parallel, LengthSimd, TOISimd);
+						VectorStoreDouble1(TOISimd, &TOI1);
+					}
+					else
+					{				
+						FAABB3 AABB0(Node.ChildrenBounds[0].Min(), Node.ChildrenBounds[0].Max());
+						bIntersect0 = TAABBTreeIntersectionHelper<TQueryFastData, Query>::Intersects(Start, CurData, TOI0, AABB0, QueryBounds, QueryHalfExtents, Dir, InvDir, bParallel);
 
-					FReal TOI1;
-					FAABB3 AABB1(Node.ChildrenBounds[1].Min(), Node.ChildrenBounds[1].Max());
-					bool bIntersect1 = TAABBTreeIntersectionHelper<TQueryFastData, Query>::Intersects(Start, CurData, TOI1, AABB1, QueryBounds, QueryHalfExtents, Dir, InvDir, bParallel);
+						FAABB3 AABB1(Node.ChildrenBounds[1].Min(), Node.ChildrenBounds[1].Max());
+						bIntersect1 = TAABBTreeIntersectionHelper<TQueryFastData, Query>::Intersects(Start, CurData, TOI1, AABB1, QueryBounds, QueryHalfExtents, Dir, InvDir, bParallel);
+					}
 					if (bIntersect0 && bIntersect1)
 					{
 						if (TOI1 > TOI0)
diff --git a/Engine/Source/Runtime/Experimental/Chaos/Public/Chaos/AABBVectorized.h b/Engine/Source/Runtime/Experimental/Chaos/Public/Chaos/AABBVectorized.h
--- a/Engine/Source/Runtime/Experimental/Chaos/Public/Chaos/AABBVectorized.h
+++ b/Engine/Source/Runtime/Experimental/Chaos/Public/Chaos/AABBVectorized.h
@@ -75,12 +75,7 @@ namespace Chaos
 			//Outside of slab before entering another
 			IsFalse = VectorCompareGT(LatestStartTime, EarliestEndTime);
 
-			if (VectorMaskBits(IsFalse))
-			{
-				return false;
-			}
-
-			return true;
+			return VectorMaskBits(IsFalse) == 0;
 		}
 
 		FORCEINLINE_DEBUGGABLE bool Intersects(const FAABBVectorized& Other) const
@@ -98,6 +93,7 @@ namespace Chaos
 		}
 
 	private:
-		VectorRegister4Float Min, Max;
+		VectorRegister4Float Min;
+		VectorRegister4Float Max;
 	};
 }
\ No newline at end of file
diff --git a/Engine/Source/Runtime/Experimental/Chaos/Public/Chaos/AABBVectorizedDouble.h b/Engine/Source/Runtime/Experimental/Chaos/Public/Chaos/AABBVectorizedDouble.h
new file mode 100644
--- /dev/null
+++ b/Engine/Source/Runtime/Experimental/Chaos/Public/Chaos/AABBVectorizedDouble.h
@@ -0,0 +1,97 @@
+// Copyright Epic Games, Inc. All Rights Reserved.
+#pragma once
+
+#include "Chaos/AABB.h"
+#include "Chaos/Core.h"
+#include "ChaosArchive.h"
+#include "Math/VectorRegister.h"
+#include "Math/UnrealMathVectorConstants.h"
+#include "Chaos/VectorUtility.h"
+
+
+namespace Chaos
+{
+
+	class alignas(32) FAABBVectorizedDouble
+	{
+	public:
+
+		FORCEINLINE_DEBUGGABLE FAABBVectorizedDouble()
+		{
+			Min = GlobalVectorConstants::DoubleBigNumber;
+			Max = VectorNegate(GlobalVectorConstants::DoubleBigNumber);
+		}
+
+		FORCEINLINE_DEBUGGABLE FAABBVectorizedDouble(const VectorRegister4Double& InMin, const VectorRegister4Double& InMax)
+			: Min(InMin)
+			, Max(InMax)
+		{
+		}
+		template<typename T>
+		FORCEINLINE_DEBUGGABLE explicit FAABBVectorizedDouble(const TAABB<T, 3>& AABB)
+		{
+			Min = MakeVectorRegisterDouble(AABB.Min().X, AABB.Min().Y, AABB.Min().Z, 0.0);
+			Max = MakeVectorRegisterDouble(AABB.Max().X, AABB.Max().Y, AABB.Max().Z, 0.0);
+		}
+
+		FORCEINLINE_DEBUGGABLE const VectorRegister4Double& GetMin() const { return Min; }
+		FORCEINLINE_DEBUGGABLE const VectorRegister4Double& GetMax() const { return Max; }
+
+		FORCEINLINE_DEBUGGABLE void SetMin(const VectorRegister4Double& InMin) { Min = InMin; }
+		FORCEINLINE_DEBUGGABLE void SetMax(const VectorRegister4Double& InMax) { Max = InMax; }
+
+		FORCEINLINE_DEBUGGABLE bool RaycastFast(const VectorRegister4Double& StartPoint, const VectorRegister4Double& InvDir, const VectorRegister4Double& Parallel, const VectorRegister4Double& Length, VectorRegister4Double& LatestStartTime) const
+		{
+			const VectorRegister4Double StarToMinGTZero = VectorCompareGT(Min, StartPoint);
+			const VectorRegister4Double ZeroGTStarToMax = VectorCompareGT(StartPoint, Max);
+			VectorRegister4Double IsFalse = VectorBitwiseAnd(VectorBitwiseOr(StarToMinGTZero, ZeroGTStarToMax), Parallel);
+
+			if (VectorMaskBits(IsFalse))
+			{
+				return false;	//parallel and outside
+			}
+
+			const VectorRegister4Double StartToMin = VectorSubtract(Min, StartPoint);
+			const VectorRegister4Double StartToMax = VectorSubtract(Max, StartPoint);
+			const VectorRegister4Double StartToMinInvDir = VectorMultiply(StartToMin, InvDir);
+			const VectorRegister4Double StartToMaxInvDir = VectorMultiply(StartToMax, InvDir);
+
+			const VectorRegister4Double Time1 = VectorBitwiseNotAnd(Parallel, StartToMinInvDir);
+			const VectorRegister4Double Time2 = VectorSelect(Parallel, Length, StartToMaxInvDir);
+
+			const VectorRegister4Double SortedTime1 = VectorMin(Time1, Time2);
+			const VectorRegister4Double SortedTime2 = VectorMax(Time1, Time2);
+
+			LatestStartTime = VectorMax(SortedTime1, VectorSwizzle(SortedTime1, 1, 2, 0, 3));
+			LatestStartTime = VectorMax(LatestStartTime, VectorSwizzle(SortedTime1, 2, 0, 1, 3));
+			LatestStartTime = VectorMax(LatestStartTime, VectorZero());
+
+			VectorRegister4Double EarliestEndTime = VectorMin(SortedTime2, VectorSwizzle(SortedTime2, 1, 2, 0, 3));
+			EarliestEndTime = VectorMin(EarliestEndTime, VectorSwizzle(SortedTime2, 2, 0, 1, 3));
+			EarliestEndTime = VectorMin(EarliestEndTime, Length);
+
+			//Outside of slab before entering another
+			IsFalse = VectorCompareGT(LatestStartTime, EarliestEndTime);
+
+			return VectorMaskBits(IsFalse) == 0;
+		}
+
+		FORCEINLINE_DEBUGGABLE bool Intersects(const FAABBVectorizedDouble& Other) const
+		{
+			VectorRegister4Double IsFalse = VectorBitwiseOr(VectorCompareGT(Min, Other.GetMax()), VectorCompareGT(Other.GetMin(), Max));
+
+			return !static_cast<bool>(VectorMaskBits(IsFalse));
+		}
+
+		FORCEINLINE_DEBUGGABLE FAABBVectorizedDouble& Thicken(const VectorRegister4Double& Thickness)
+		{
+			Min = VectorSubtract(Min, Thickness);
+			Max = VectorAdd(Max, Thickness);
+			return *this;
+		}
+
+	private:
+		VectorRegister4Double Min;
+		VectorRegister4Double Max;
+	};
+}
\ No newline at end of file
diff --git a/Engine/Source/Runtime/Experimental/Chaos/Public/Chaos/BoundingVolume.h b/Engine/Source/Runtime/Experimental/Chaos/Public/Chaos/BoundingVolume.h
--- a/Engine/Source/Runtime/Experimental/Chaos/Public/Chaos/BoundingVolume.h
+++ b/Engine/Source/Runtime/Experimental/Chaos/Public/Chaos/BoundingVolume.h
@@ -322,10 +322,18 @@ public:
 		return RaycastImp(Start, CurData, Visitor, Dir, InvDir, bParallel);
 	}
 
+	template <typename SQVisitor>
+	bool RaycastFastSimd(const VectorRegister4Double& Start, FQueryFastData& CurData, SQVisitor& Visitor, const VectorRegister4Double& Dir, const VectorRegister4Double& InvDir, const VectorRegister4Double& Parallel, const VectorRegister4Double& Length) const
+	{
+		FVec3 StartReal;
+		VectorStoreDouble3(Start, &StartReal[0]);
+		return RaycastImp(StartReal, CurData, Visitor, CurData.Dir, CurData.InvDir, CurData.bParallel);
+	}
+
 	template <typename SQVisitor, bool bPruneDuplicates = true>
 	void Raycast(const TVector<T, d>& Start, const TVector<T, d>& Dir, const T Length, SQVisitor& Visitor) const
 	{
-		FQueryFastData CurData(Dir, Length); 
+		FQueryFastData CurData(Dir, Length);
 		RaycastImp<SQVisitor, bPruneDuplicates>(Start, CurData, Visitor, CurData.Dir, CurData.InvDir, CurData.bParallel);
 	}
 
@@ -374,8 +382,8 @@ public:
 		return false;
 	}
 	
-	/** Set thye dirty flag onto the leaf 
-	* @param  bDirtyState Disrty flag to set 
+	/** Set the dirty flag onto the leaf 
+	* @param  bDirtyState Dirty flag to set 
 	*/
 	void SetDirtyState(const bool bDirtyState)
 	{}
diff --git a/Engine/Source/Runtime/Experimental/Chaos/Public/Chaos/VectorUtility.h b/Engine/Source/Runtime/Experimental/Chaos/Public/Chaos/VectorUtility.h
--- a/Engine/Source/Runtime/Experimental/Chaos/Public/Chaos/VectorUtility.h
+++ b/Engine/Source/Runtime/Experimental/Chaos/Public/Chaos/VectorUtility.h
@@ -96,3 +96,26 @@ FORCEINLINE VectorRegister4Float VectorBitwiseNotAnd(const VectorRegister4Float&
 #endif
 }
 
+FORCEINLINE VectorRegister4Double VectorBitwiseNotAnd(const VectorRegister4Double& A, const VectorRegister4Double& B)
+{
+	VectorRegister4Double Result;
+#if PLATFORM_ENABLE_VECTORINTRINSICS_NEON	
+	Result.XY = (VectorRegister2Double)vandq_u32(vmvnq_u32((VectorRegister2Double)A.XY), (VectorRegister2Double)B.XY);
+	Result.ZW = (VectorRegister2Double)vandq_u32(vmvnq_u32((VectorRegister2Double)A.ZW), (VectorRegister2Double)B.ZW);
+#elif PLATFORM_ENABLE_VECTORINTRINSICS
+#if !UE_PLATFORM_MATH_USE_AVX
+	Result.XY = _mm_cvtps_pd(_mm_andnot_ps(_mm_cvtpd_ps(A.XY), _mm_cvtpd_ps(B.XY)));
+	Result.ZW = _mm_cvtps_pd(_mm_andnot_ps(_mm_cvtpd_ps(A.ZW), _mm_cvtpd_ps(B.ZW)));
+#else
+	Result = _mm256_andnot_pd(A, B);
+#endif
+#else
+	Result = MakeVectorRegisterDouble(
+		uint64(~((uint64*)(A.V))[0] & ((uint64*)(B.V))[0]),
+		uint64(~((uint64*)(A.V))[1] & ((uint64*)(B.V))[1]),
+		uint64(~((uint64*)(A.V))[2] & ((uint64*)(B.V))[2]),
+		uint64(~((uint64*)(A.V))[3] & ((uint64*)(B.V))[3]));
+#endif
+	return Result;
+}
+
