Chaos : add option to display client and server shapes in different colors

diff --git a/Engine/Source/Runtime/Experimental/Chaos/Private/PBDRigidsSolver.cpp b/Engine/Source/Runtime/Experimental/Chaos/Private/PBDRigidsSolver.cpp
--- a/Engine/Source/Runtime/Experimental/Chaos/Private/PBDRigidsSolver.cpp
+++ b/Engine/Source/Runtime/Experimental/Chaos/Private/PBDRigidsSolver.cpp
@@ -86,6 +86,7 @@ namespace Chaos
 		int32 ChaosSolverDebugDrawSuspensionConstraints = 0;
 		int32 ChaosSolverDrawClusterConstraints = 0;
 		int32 ChaosSolverDebugDrawMeshContacts = 0;
+		int32 ChaosSolverDebugDrawColorShapeByClientServer = 0;
 		DebugDraw::FChaosDebugDrawJointFeatures ChaosSolverDrawJointFeatures = DebugDraw::FChaosDebugDrawJointFeatures::MakeDefault();
 		FAutoConsoleVariableRef CVarChaosSolverDrawShapes(TEXT("p.Chaos.Solver.DebugDrawShapes"), ChaosSolverDebugDrawShapes, TEXT("Draw Shapes (0 = never; 1 = end of frame)."));
 		FAutoConsoleVariableRef CVarChaosSolverDrawCollisions(TEXT("p.Chaos.Solver.DebugDrawCollisions"), ChaosSolverDebugDrawCollisions, TEXT("Draw Collisions (0 = never; 1 = end of frame)."));
@@ -113,6 +114,7 @@ namespace Chaos
 		FAutoConsoleVariableRef CVarChaosSolverDrawSuspensionConstraints(TEXT("p.Chaos.Solver.DebugDrawSuspension"), ChaosSolverDebugDrawSuspensionConstraints, TEXT("Draw Suspension (0 = never; 1 = end of frame)."));
 		FAutoConsoleVariableRef CVarChaosSolverDrawClusterConstraints(TEXT("p.Chaos.Solver.DebugDraw.Cluster.Constraints"), ChaosSolverDrawClusterConstraints, TEXT("Draw Active Cluster Constraints (0 = never; 1 = end of frame)."));
 		FAutoConsoleVariableRef CVarChaosSolverDrawMeshContacts(TEXT("p.Chaos.Solver.DebugDrawMeshContacts"), ChaosSolverDebugDrawMeshContacts, TEXT("Draw Mesh contacts"));
+		FAutoConsoleVariableRef CVarChaosSolverDebugDrawColorShapeByClientServer(TEXT("p.Chaos.Solver.DebugDraw.ColorShapeByClientServer"), ChaosSolverDebugDrawColorShapeByClientServer, TEXT("Color shape according to client and server: red = server / blue = client "));
 
 
 		DebugDraw::FChaosDebugDrawSettings ChaosSolverDebugDebugDrawSettings(
@@ -146,6 +148,30 @@ namespace Chaos
 			/* InBoundsColorsPerShapeType=  */ DebugDraw::GetDefaultBoundsColorsByShapeType()
 		);
 
+		static DebugDraw::FChaosDebugDrawColorsByState GetSolverShapesColorsByState_Server()
+		{
+			static DebugDraw::FChaosDebugDrawColorsByState SolverShapesColorsByState_Server
+			{
+				/* InDynamicColor =	  */ FColor(255, 0, 0),
+				/* InSleepingColor =  */ FColor(128, 0, 0),
+				/* InKinematicColor = */ FColor(255, 0, 0),
+				/* InStaticColor =	  */ FColor(255, 0, 0),
+			};
+			return SolverShapesColorsByState_Server;
+		}
+
+		static DebugDraw::FChaosDebugDrawColorsByState GetSolverShapesColorsByState_Client()
+		{
+			static DebugDraw::FChaosDebugDrawColorsByState SolverShapesColorsByState_Client
+			{
+				/* InDynamicColor =	  */ FColor(0, 0, 255),
+				/* InSleepingColor =  */ FColor(0, 0, 128),
+				/* InKinematicColor = */ FColor(0, 0, 255),
+				/* InStaticColor =	  */ FColor(0, 0, 255),
+			};
+			return SolverShapesColorsByState_Client;
+		}
+
 		FAutoConsoleVariableRef CVarChaosSolverArrowSize(TEXT("p.Chaos.Solver.DebugDraw.ArrowSize"), ChaosSolverDebugDebugDrawSettings.ArrowSize, TEXT("ArrowSize."));
 		FAutoConsoleVariableRef CVarChaosSolverBodyAxisLen(TEXT("p.Chaos.Solver.DebugDraw.BodyAxisLen"), ChaosSolverDebugDebugDrawSettings.BodyAxisLen, TEXT("BodyAxisLen."));
 		FAutoConsoleVariableRef CVarChaosSolverContactLen(TEXT("p.Chaos.Solver.DebugDraw.ContactLen"), ChaosSolverDebugDebugDrawSettings.ContactLen, TEXT("ContactLen."));
@@ -1526,6 +1552,25 @@ CSV_CUSTOM_STAT(PhysicsCounters, Name, Value, ECsvCustomStatOp::Set);
 	{
 #if CHAOS_DEBUG_DRAW
 		QUICK_SCOPE_CYCLE_COUNTER(SolverDebugDraw);
+
+#if !UE_BUILD_SHIPPING && !UE_BUILD_TEST
+		if (ChaosSolverDebugDrawColorShapeByClientServer)
+		{
+			if (DebugName.ToString().StartsWith(TEXT("Server")))
+			{
+				ChaosSolverDebugDebugDrawSettings.ShapesColorsPerState = GetSolverShapesColorsByState_Server();
+			}
+			else
+			{
+				ChaosSolverDebugDebugDrawSettings.ShapesColorsPerState = GetSolverShapesColorsByState_Client();
+			}
+		}
+		else
+		{
+			ChaosSolverDebugDebugDrawSettings.ShapesColorsPerState = DebugDraw::GetDefaultShapesColorsByState();
+		}
+#endif
+
 		if (ChaosSolverDebugDrawShapes == 1)
 		{
 			if (ChaosSolverDrawShapesShowStatic)
