Removed unecessary TEXT macro

diff --git a/Engine/Source/Runtime/Experimental/Chaos/Public/RewindData.h b/Engine/Source/Runtime/Experimental/Chaos/Public/RewindData.h
--- a/Engine/Source/Runtime/Experimental/Chaos/Public/RewindData.h
+++ b/Engine/Source/Runtime/Experimental/Chaos/Public/RewindData.h
@@ -1061,7 +1061,7 @@ private:
 		int32 DirtyDynamics = INDEX_NONE;	//Only used by particles, indicates the dirty properties was written to.
 		int32 LastDirtyFrame;	//Track how recently this was made dirty
 		int32 InitializedOnStep = INDEX_NONE;	//if not INDEX_NONE, it indicates we saw initialization during rewind history window
-		UE_DEPRECATED(5.1, TEXT("bResimAsSlave is deprecated - please use bResimAsFollower"))
+		UE_DEPRECATED(5.1, "bResimAsSlave is deprecated - please use bResimAsFollower")
 		bool bResimAsSlave = true;
 		bool bResimAsFollower = true;	//Indicates the particle will always resim in the exact same way from game thread data
 
