Removed unecessary TEXT macro

diff --git a/Engine/Source/Runtime/Experimental/Chaos/Public/Chaos/Evolution/SimulationSpace.h b/Engine/Source/Runtime/Experimental/Chaos/Public/Chaos/Evolution/SimulationSpace.h
--- a/Engine/Source/Runtime/Experimental/Chaos/Public/Chaos/Evolution/SimulationSpace.h
+++ b/Engine/Source/Runtime/Experimental/Chaos/Public/Chaos/Evolution/SimulationSpace.h
@@ -59,7 +59,7 @@ namespace Chaos
 		{
 		}
 
-		UE_DEPRECATED(5.1, TEXT("MasterAlpha is deprecated - please use Alpha"))
+		UE_DEPRECATED(5.1, "MasterAlpha is deprecated - please use Alpha")
 		FReal MasterAlpha;
 
 		// Global multipler on the effects of simulation space movement
