Explicitly declare copy and copy assign for simulation space settings to avoid implicit operators throwing deprecation warnings

diff --git a/Engine/Source/Runtime/Experimental/Chaos/Public/Chaos/Evolution/SimulationSpace.h b/Engine/Source/Runtime/Experimental/Chaos/Public/Chaos/Evolution/SimulationSpace.h
--- a/Engine/Source/Runtime/Experimental/Chaos/Public/Chaos/Evolution/SimulationSpace.h
+++ b/Engine/Source/Runtime/Experimental/Chaos/Public/Chaos/Evolution/SimulationSpace.h
@@ -59,6 +59,11 @@ namespace Chaos
 		{
 		}
 
+PRAGMA_DISABLE_DEPRECATION_WARNINGS
+		FSimulationSpaceSettings(const FSimulationSpaceSettings&) = default;
+		FSimulationSpaceSettings& operator=(const FSimulationSpaceSettings&) = default;
+PRAGMA_ENABLE_DEPRECATION_WARNINGS
+
 		UE_DEPRECATED(5.1, "MasterAlpha is deprecated - please use Alpha")
 		FReal MasterAlpha;
 
