Disable clang analyzer warnings for Voro++ library

diff --git a/Engine/Source/Runtime/Experimental/Voronoi/Voronoi.Build.cs b/Engine/Source/Runtime/Experimental/Voronoi/Voronoi.Build.cs
--- a/Engine/Source/Runtime/Experimental/Voronoi/Voronoi.Build.cs
+++ b/Engine/Source/Runtime/Experimental/Voronoi/Voronoi.Build.cs
@@ -13,6 +13,9 @@ namespace UnrealBuildTool.Rules
 					"CoreUObject"
                 }
                 );
+
+			// Disable all static analysis checkers for this module ( third party code ) 
+			StaticAnalyzerDisabledCheckers.Add("all");
 		}
 	}
 }
