ensure that actors and constraints are safe to use

diff --git a/Engine/Plugins/Animation/ControlRig/Source/ControlRigEditor/Private/Editor/Constraints/SConstraintsWidget.cpp b/Engine/Plugins/Animation/ControlRig/Source/ControlRigEditor/Private/Editor/Constraints/SConstraintsWidget.cpp
--- a/Engine/Plugins/Animation/ControlRig/Source/ControlRigEditor/Private/Editor/Constraints/SConstraintsWidget.cpp
+++ b/Engine/Plugins/Animation/ControlRig/Source/ControlRigEditor/Private/Editor/Constraints/SConstraintsWidget.cpp
@@ -94,7 +94,7 @@ TArray<AActor*> GetCurrentSelection()
 
 	return CurrentLevel->Actors.FilterByPredicate( [](const AActor* Actor)
 	{
-		return Actor && Actor->IsSelected();
+		return IsValid(Actor) && Actor->IsSelected();
 	});	
 }
 	
diff --git a/Engine/Source/Runtime/Experimental/Animation/Constraints/Private/ConstraintsManager.cpp b/Engine/Source/Runtime/Experimental/Animation/Constraints/Private/ConstraintsManager.cpp
--- a/Engine/Source/Runtime/Experimental/Animation/Constraints/Private/ConstraintsManager.cpp
+++ b/Engine/Source/Runtime/Experimental/Animation/Constraints/Private/ConstraintsManager.cpp
@@ -503,7 +503,7 @@ TArray< TObjectPtr<UTickableConstraint> > FConstraintsManagerController::GetPare
 	TArray< TObjectPtr<UTickableConstraint> > FilteredConstraints =
 	Manager->Constraints.FilterByPredicate( [InTargetHash](const ConstraintPtr& Constraint)
 	{
-		return Constraint->GetTargetHash() == InTargetHash;
+		return IsValid(Constraint) && Constraint->GetTargetHash() == InTargetHash;
 	} );
 	
 	if (bSorted)
