Broad Phase optimization : Vectorized overlap test

diff --git a/Engine/Source/Runtime/Experimental/Chaos/Public/Chaos/AABBTree.h b/Engine/Source/Runtime/Experimental/Chaos/Public/Chaos/AABBTree.h
--- a/Engine/Source/Runtime/Experimental/Chaos/Public/Chaos/AABBTree.h
+++ b/Engine/Source/Runtime/Experimental/Chaos/Public/Chaos/AABBTree.h
@@ -300,8 +300,57 @@ struct TAABBTreeLeafArray : public TBoundsWrapperHelper<TPayloadType, T, bComput
 	{
 		PHYSICS_CSV_CUSTOM_VERY_EXPENSIVE(PhysicsCounters, MaxLeafSize, Elems.Num(), ECsvCustomStatOp::Max);
 
-		for (const auto& Elem : Elems)
+		const int32 NumElems = Elems.Num();
+		const int32 SimdIters = NumElems / 4;
+
+		for (int32 SimdIter = 0; SimdIter < SimdIters; ++SimdIter)
+		{
+			const int32 StartIndex = SimdIter * 4;
+			VectorRegister4Double MinX = MakeVectorRegisterDouble(Elems[StartIndex].Bounds.Min()[0], Elems[StartIndex + 1].Bounds.Min()[0], Elems[StartIndex + 2].Bounds.Min()[0], Elems[StartIndex + 3].Bounds.Min()[0]);
+			VectorRegister4Double MaxX = MakeVectorRegisterDouble(Elems[StartIndex].Bounds.Max()[0], Elems[StartIndex + 1].Bounds.Max()[0], Elems[StartIndex + 2].Bounds.Max()[0], Elems[StartIndex + 3].Bounds.Max()[0]);
+			VectorRegister4Double MinY = MakeVectorRegisterDouble(Elems[StartIndex].Bounds.Min()[1], Elems[StartIndex + 1].Bounds.Min()[1], Elems[StartIndex + 2].Bounds.Min()[1], Elems[StartIndex + 3].Bounds.Min()[1]);
+			VectorRegister4Double MaxY = MakeVectorRegisterDouble(Elems[StartIndex].Bounds.Max()[1], Elems[StartIndex + 1].Bounds.Max()[1], Elems[StartIndex + 2].Bounds.Max()[1], Elems[StartIndex + 3].Bounds.Max()[1]);
+			VectorRegister4Double MinZ = MakeVectorRegisterDouble(Elems[StartIndex].Bounds.Min()[2], Elems[StartIndex + 1].Bounds.Min()[2], Elems[StartIndex + 2].Bounds.Min()[2], Elems[StartIndex + 3].Bounds.Min()[2]);
+			VectorRegister4Double MaxZ = MakeVectorRegisterDouble(Elems[StartIndex].Bounds.Max()[2], Elems[StartIndex + 1].Bounds.Max()[2], Elems[StartIndex + 2].Bounds.Max()[2], Elems[StartIndex + 3].Bounds.Max()[2]);
+
+			const VectorRegister4Double OtherMinX = VectorSetDouble1(QueryBounds.Min().X);
+			const VectorRegister4Double OtherMinY = VectorSetDouble1(QueryBounds.Min().Y);
+			const VectorRegister4Double OtherMinZ = VectorSetDouble1(QueryBounds.Min().Z);
+
+			const VectorRegister4Double OtherMaxX = VectorSetDouble1(QueryBounds.Max().X);
+			const VectorRegister4Double OtherMaxY = VectorSetDouble1(QueryBounds.Max().Y);
+			const VectorRegister4Double OtherMaxZ = VectorSetDouble1(QueryBounds.Max().Z);
+
+			VectorRegister4Double IsFalseX = VectorBitwiseOr(VectorCompareGT(MinX, OtherMaxX), VectorCompareGT(OtherMinX, MaxX));
+			VectorRegister4Double IsFalseY = VectorBitwiseOr(VectorCompareGT(MinY, OtherMaxY), VectorCompareGT(OtherMinY, MaxY));
+			VectorRegister4Double IsFalseZ = VectorBitwiseOr(VectorCompareGT(MinZ, OtherMaxZ), VectorCompareGT(OtherMinZ, MaxZ));
+
+			VectorRegister4Double IsFalse = VectorBitwiseOr(VectorBitwiseOr(IsFalseX, IsFalseY), IsFalseZ);
+
+			int32 MaskBitFalse = VectorMaskBits(IsFalse);
+
+			for (int32 MaskIndex = 0; MaskIndex < 4; ++MaskIndex)
+			{
+				if ((MaskBitFalse & (1 << MaskIndex)) == 0)
+				{
+					const TPayloadBoundsElement<TPayloadType, T >& Elem = Elems[StartIndex + MaskIndex];
+					if (PrePreFilterHelper(Elem.Payload, Visitor))
+					{
+						continue;
+					}
+					const FAABB3 InstanceBounds(Elem.Bounds.Min(), Elem.Bounds.Max());
+					TSpatialVisitorData<TPayloadType> VisitData(Elem.Payload, true, InstanceBounds);
+					if (Visitor.VisitOverlap(VisitData) == false)
+					{
+						return false;
+					}
+				}
+			}
+		}
+
+		for (int32 SlowIter = SimdIters * 4; SlowIter < NumElems; ++SlowIter)
 		{
+			const TPayloadBoundsElement<TPayloadType, T >& Elem = Elems[SlowIter];
 			if (PrePreFilterHelper(Elem.Payload, Visitor))
 			{
 				continue;
