Chaos - collision optimization
- aggregate some RigidParticle data into a single structure to optimize  broadphase filtering
- optimize some data access patterns in broadphase filtering

diff --git a/Engine/Source/Runtime/Experimental/Chaos/Private/Chaos/Collision/ParticlePairMidPhase.cpp b/Engine/Source/Runtime/Experimental/Chaos/Private/Chaos/Collision/ParticlePairMidPhase.cpp
--- a/Engine/Source/Runtime/Experimental/Chaos/Private/Chaos/Collision/ParticlePairMidPhase.cpp
+++ b/Engine/Source/Runtime/Experimental/Chaos/Private/Chaos/Collision/ParticlePairMidPhase.cpp
@@ -154,6 +154,7 @@ namespace Chaos
 		Flags.bEnableManifoldUpdate = !bIsSphere0 && !bIsSphere1 && !bIsCapsule0 && !bIsCapsule1 && !bIsTriangle0 && !bIsTriangle1 && !bIsLevelSet;
 
 		// Mark probe flag now so we know which GenerateCollisions to use
+		// @todo(chaos): it looks like this can be changed by a collision modifier so we should not be caching it
 		Flags.bIsProbe = Shape0->GetIsProbe() || Shape1->GetIsProbe();
 	}
 
@@ -293,36 +294,29 @@ namespace Chaos
 		{
 			PHYSICS_CSV_SCOPED_EXPENSIVE(PhysicsVerbose, NarrowPhase_UpdateConstraint);
 
-			const FImplicitObject* Implicit0 = Shape0->GetLeafGeometry();
-			const FImplicitObject* Implicit1 = Shape1->GetLeafGeometry();
-			const FRigidTransform3& ShapeWorldTransform0 = Shape0->GetLeafWorldTransform(GetParticle0());
-			const FRigidTransform3& ShapeWorldTransform1 = Shape1->GetLeafWorldTransform(GetParticle1());
+			const FRigidTransform3 ShapeWorldTransform0 = Shape0->GetLeafWorldTransform(GetParticle0());
+			const FRigidTransform3 ShapeWorldTransform1 = Shape1->GetLeafWorldTransform(GetParticle1());
 			const int32 CurrentEpoch = MidPhase.GetCollisionAllocator().GetCurrentEpoch();
 			const int32 LastEpoch = CurrentEpoch - 1;
 			const bool bWasUpdatedLastTick = IsUsedSince(LastEpoch);
 
 			// Update the world shape transforms on the constraint (we cannot just give it the PerShapeData 
 			// pointer because of Unions - see FMultiShapePairCollisionDetector)
-			// NOTE: these are not used by CCD which continuously moves the particles
 			Constraint->SetShapeWorldTransforms(ShapeWorldTransform0, ShapeWorldTransform1);
 
-			Constraint->ResetModifications();
-			Constraint->SetCCDEnabled(false);
-
-			// If the constraint was not used last frame, it needs to be reset. 
-			// Otherwise we will try to reuse it below
+			// If the constraint was not used last frame, it needs to be reset, otherwise we will try to reuse
 			if (!bWasUpdatedLastTick || (Constraint->GetManifoldPoints().Num() == 0))
 			{
 				// Clear all manifold data including saved contact data
 				Constraint->ResetManifold();
 			}
-
+			
 			bool bWasManifoldRestored = false;
 			if (Context.GetSettings().bAllowManifoldReuse && Flags.bEnableManifoldUpdate && bWasUpdatedLastTick)
 			{
 				// Update the existing manifold. We can re-use as-is if none of the points have moved much and the bodies have not moved much
 				// NOTE: this can succeed in "restoring" even if we have no manifold points
-				// NOTE: this uses the previous world-space shape transforms, so we can only do this if we were updated last tick
+				// NOTE: this uses the transforms from SetLastShapeWorldTransforms, so we can only do this if we were updated last tick
 				bWasManifoldRestored = Constraint->UpdateAndTryRestoreManifold();
 			}
 			else
@@ -333,7 +327,7 @@ namespace Chaos
 
 			if (!bWasManifoldRestored)
 			{
-				// We will be updating the manifold, if only partially, so update the restore comparison transforms
+				// We will be updating the manifold so update transforms used to check for movement in UpdateAndTryRestoreManifold on future ticks
 				Constraint->SetLastShapeWorldTransforms(ShapeWorldTransform0, ShapeWorldTransform1);
 
 				if (!Context.GetSettings().bDeferNarrowPhase)
@@ -344,11 +338,15 @@ namespace Chaos
 			}
 
 			// If we have a valid contact, add it to the active list
-			// We also add it to the active list if collision detection is deferred (which is if per-iteration collision detection is enabled like with RBAN)
+			// We also add it to the active list if collision detection is deferred because the data will be filled in later and we
+			// don't know in advance whether we will pass the Phi check (deferred narrow phase is used with RBAN)
 			if (Constraint->GetPhi() <= CullDistance || Context.GetSettings().bDeferNarrowPhase)
 			{
 				if (MidPhase.GetCollisionAllocator().ActivateConstraint(Constraint.Get()))
 				{
+					// If we had any mods last tick, undo them
+					Constraint->ResetModifications();
+
 					LastUsedEpoch = CurrentEpoch;
 					return 1;
 				}
diff --git a/Engine/Source/Runtime/Experimental/Chaos/Private/Chaos/Collision/SolverCollisionContainer.cpp b/Engine/Source/Runtime/Experimental/Chaos/Private/Chaos/Collision/SolverCollisionContainer.cpp
--- a/Engine/Source/Runtime/Experimental/Chaos/Private/Chaos/Collision/SolverCollisionContainer.cpp
+++ b/Engine/Source/Runtime/Experimental/Chaos/Private/Chaos/Collision/SolverCollisionContainer.cpp
@@ -371,6 +371,8 @@ namespace Chaos
 					FRealSingle(Dt));
 			}
 
+			Constraint->EndTick();
+
 			Unbind();
 		}
 
diff --git a/Engine/Source/Runtime/Experimental/Chaos/Private/Chaos/ContactModification.cpp b/Engine/Source/Runtime/Experimental/Chaos/Private/Chaos/ContactModification.cpp
--- a/Engine/Source/Runtime/Experimental/Chaos/Private/Chaos/ContactModification.cpp
+++ b/Engine/Source/Runtime/Experimental/Chaos/Private/Chaos/ContactModification.cpp
@@ -86,6 +86,8 @@ namespace Chaos
 
 	void FContactPairModifier::ModifyTargetSeparation(FReal TargetSeparation, int32 ContactPointIdx)
 	{	
+		Constraint->SetModifierApplied();
+
 		FManifoldPoint& ManifoldPoint = Constraint->GetManifoldPoint(ContactPointIdx);
 		ManifoldPoint.TargetPhi = FRealSingle(TargetSeparation);
 	}
@@ -102,6 +104,8 @@ namespace Chaos
 
 	void FContactPairModifier::ModifyWorldNormal(const FVec3& Normal, int32 ContactPointIdx)
 	{
+		Constraint->SetModifierApplied();
+
 		FVec3 WorldContactPoint0, WorldContactPoint1;
 		GetWorldContactLocations(ContactPointIdx, WorldContactPoint0, WorldContactPoint1);
 
@@ -119,6 +123,8 @@ namespace Chaos
 
 	void FContactPairModifier::GetWorldContactLocations(int32 ContactPointIdx, FVec3& OutLocation0, FVec3& OutLocation1) const
 	{
+		Constraint->SetModifierApplied();
+
 		const FRigidTransform3& ShapeTransform0 = Constraint->GetShapeWorldTransform0();
 		const FRigidTransform3& ShapeTransform1 = Constraint->GetShapeWorldTransform1();
 
@@ -137,6 +143,8 @@ namespace Chaos
 
 	void FContactPairModifier::ModifyWorldContactLocations(const FVec3& Location0, const FVec3& Location1, int32 ContactPointIdx)
 	{
+		Constraint->SetModifierApplied();
+
 		TArrayView<FManifoldPoint> ManifoldPoints = Constraint->GetManifoldPoints();
 		FManifoldPoint& ManifoldPoint = ManifoldPoints[ContactPointIdx];
 
@@ -157,6 +165,8 @@ namespace Chaos
 
 	void FContactPairModifier::ModifyRestitution(FReal Restitution)
 	{
+		Constraint->SetModifierApplied();
+
 		Constraint->SetRestitution(Restitution);
 	}
 
@@ -167,6 +177,8 @@ namespace Chaos
 
 	void FContactPairModifier::ModifyRestitutionThreshold(FReal Threshold)
 	{
+		Constraint->SetModifierApplied();
+
 		Constraint->SetRestitutionThreshold(Threshold);
 	}
 
@@ -177,6 +189,8 @@ namespace Chaos
 
 	void FContactPairModifier::ModifyDynamicFriction(FReal DynamicFriction)
 	{
+		Constraint->SetModifierApplied();
+
 		Constraint->SetDynamicFriction(DynamicFriction);
 	}
 
@@ -187,6 +201,8 @@ namespace Chaos
 
 	void FContactPairModifier::ModifyStaticFriction(FReal StaticFriction)
 	{
+		Constraint->SetModifierApplied();
+
 		Constraint->SetStaticFriction(StaticFriction);
 	}
 
@@ -205,6 +221,8 @@ namespace Chaos
 
 	void FContactPairModifier::ModifyParticleVelocity(FVec3 Velocity, int32 ParticleIdx)
 	{
+		Constraint->SetModifierApplied();
+
 		FGeometryParticleHandle* Particle = Constraint->GetParticle(ParticleIdx);
 		FKinematicGeometryParticleHandle* KinematicHandle = Particle->CastToKinematicParticle();
 		if (!ensure(KinematicHandle))
@@ -242,6 +260,8 @@ namespace Chaos
 
 	void FContactPairModifier::ModifyParticleAngularVelocity(FVec3 AngularVelocity, int32 ParticleIdx)
 	{
+		Constraint->SetModifierApplied();
+
 		FGeometryParticleHandle* Particle = Constraint->GetParticle(ParticleIdx);
 		FKinematicGeometryParticleHandle* KinematicHandle = Particle->CastToKinematicParticle();
 		if (!ensure(KinematicHandle))
@@ -284,6 +304,8 @@ namespace Chaos
 
 	void FContactPairModifier::ModifyParticlePosition(FVec3 Position, bool bMaintainVelocity, int32 ParticleIdx)
 	{
+		Constraint->SetModifierApplied();
+
 		FGeometryParticleHandle* Particle = Constraint->GetParticle(ParticleIdx);
 
 		Modifier->MarkConstraintForManifoldUpdate(*Constraint);
@@ -347,6 +369,8 @@ namespace Chaos
 
 	void FContactPairModifier::ModifyParticleRotation(FRotation3 Rotation, bool bMaintainVelocity, int32 ParticleIdx)
 	{
+		Constraint->SetModifierApplied();
+
 		FGeometryParticleHandle* Particle = Constraint->GetParticle(ParticleIdx);
 
 		Modifier->MarkConstraintForManifoldUpdate(*Constraint);
@@ -399,6 +423,8 @@ namespace Chaos
 
 	void FContactPairModifier::ModifyInvInertiaScale(FReal InInvInertiaScale, int32 ParticleIdx)
 	{
+		Constraint->SetModifierApplied();
+
 		if (ParticleIdx == 0)
 		{
 			Constraint->SetInvInertiaScale0(InInvInertiaScale);
@@ -416,6 +442,8 @@ namespace Chaos
 
 	void FContactPairModifier::ModifyInvMassScale(FReal InInvMassScale, int32 ParticleIdx)
 	{
+		Constraint->SetModifierApplied();
+
 		if (ParticleIdx == 0)
 		{
 			Constraint->SetInvMassScale0(InInvMassScale);
@@ -473,16 +501,22 @@ namespace Chaos
 
 	void FCollisionContactModifier::DisableConstraint(FPBDCollisionConstraint& Constraint)
 	{
+		Constraint.SetModifierApplied();
+
 		Constraint.SetDisabled(true);
 	}
 
 	void FCollisionContactModifier::EnableConstraint(FPBDCollisionConstraint& Constraint)
 	{
+		Constraint.SetModifierApplied();
+
 		Constraint.SetDisabled(false);
 	}
 
 	void FCollisionContactModifier::ConvertToProbeConstraint(FPBDCollisionConstraint& Constraint)
 	{
+		Constraint.SetModifierApplied();
+
 		Constraint.SetIsProbe(true);
 	}
 
diff --git a/Engine/Source/Runtime/Experimental/Chaos/Public/Chaos/Collision/CollisionKeys.h b/Engine/Source/Runtime/Experimental/Chaos/Public/Chaos/Collision/CollisionKeys.h
--- a/Engine/Source/Runtime/Experimental/Chaos/Public/Chaos/Collision/CollisionKeys.h
+++ b/Engine/Source/Runtime/Experimental/Chaos/Public/Chaos/Collision/CollisionKeys.h
@@ -34,13 +34,15 @@ namespace Chaos
 	 * Used to order particles in a consistent way in Broadphase and Resim. Returns whether the partcile order should be reversed.
 	 * We want the particle with the lower ID first, unless only one is dynamic in which case that one is first.
 	 */
-	inline bool ShouldSwapParticleOrder(const FGeometryParticleHandle* Particle0, const FGeometryParticleHandle* Particle1, const bool bIsParticle0Preferred)
+	inline bool ShouldSwapParticleOrder(const bool bIsDynamicOrSleeping0, const bool bIsParticle0Preferred)
 	{
-		return !FConstGenericParticleHandle(Particle0)->IsDynamic() || !bIsParticle0Preferred;
+		return !bIsDynamicOrSleeping0 || !bIsParticle0Preferred;
 	}
 	inline bool ShouldSwapParticleOrder(const FGeometryParticleHandle* Particle0, const FGeometryParticleHandle* Particle1)
 	{
-		return ShouldSwapParticleOrder(Particle0, Particle1, AreParticlesInPreferredOrder(Particle0, Particle1));
+		const EObjectStateType ObjectState0 = Particle0->ObjectState();
+		const bool bIsDynamic0 = (ObjectState0 == EObjectStateType::Dynamic) || (ObjectState0 == EObjectStateType::Sleeping);
+		return ShouldSwapParticleOrder(bIsDynamic0, AreParticlesInPreferredOrder(Particle0, Particle1));
 	}
 
 	/**
diff --git a/Engine/Source/Runtime/Experimental/Chaos/Public/Chaos/Collision/PBDCollisionConstraint.h b/Engine/Source/Runtime/Experimental/Chaos/Public/Chaos/Collision/PBDCollisionConstraint.h
--- a/Engine/Source/Runtime/Experimental/Chaos/Public/Chaos/Collision/PBDCollisionConstraint.h
+++ b/Engine/Source/Runtime/Experimental/Chaos/Public/Chaos/Collision/PBDCollisionConstraint.h
@@ -370,6 +370,8 @@ namespace Chaos
 		// @todo(chaos): remove (used by legacy RBAN collision solver)
 		FVec3 CalculateWorldContactLocation() const;
 
+		void SetModifierApplied() { Flags.bModifierApplied = true; }
+
 		void SetInvMassScale0(const FReal InInvMassScale) { Material.InvMassScale0 = FRealSingle(InInvMassScale); }
 		FReal GetInvMassScale0() const { return Material.InvMassScale0; }
 
@@ -419,10 +421,15 @@ namespace Chaos
 		*/
 		inline void ResetModifications()
 		{
-			Material.ResetMaterialModifications();
+			if (Flags.bModifierApplied)
+			{
+				Material.ResetMaterialModifications();
+
+				// Reset other properties which may have changed in contact modification last frame
+				Flags.bIsProbe = Flags.bIsProbeUnmodified;
 
-			// Reset other properties which may have changed in contact modification last frame
-			Flags.bIsProbe = Flags.bIsProbeUnmodified;
+				Flags.bModifierApplied = false;
+			}
 		}
 
 		/**
@@ -578,6 +585,15 @@ namespace Chaos
 			ResetSavedManifoldPoints();
 		}
 
+		/**
+		 * Called after the simulation to reset any state that need to be reset before the next tick
+		 */
+		void EndTick()
+		{
+			// CCD state gets set each tick based on speed, size, etc
+			SetCCDEnabled(false);
+		}
+
 		/**
 		 * @brief Time of impact from CCD sweep test if CCD is activate.Otherwise undefined.
 		*/
@@ -766,7 +782,8 @@ namespace Chaos
 				uint32 bIsQuadratic1 : 1;
 				uint32 bIsProbeUnmodified : 1;  // Is this constraint a probe pre-contact-modification
 				uint32 bIsProbe : 1;            // Is this constraint currently a probe
-				uint32 bCCDEnabled : 1;
+				uint32 bCCDEnabled : 1;			// Is CCD enabled for the current tick
+				uint32 bModifierApplied;		// Was a constraint modifier applied this tick
 			};
 			uint32 Bits;
 		} Flags;
diff --git a/Engine/Source/Runtime/Experimental/Chaos/Public/Chaos/Collision/SpatialAccelerationBroadPhase.h b/Engine/Source/Runtime/Experimental/Chaos/Public/Chaos/Collision/SpatialAccelerationBroadPhase.h
--- a/Engine/Source/Runtime/Experimental/Chaos/Public/Chaos/Collision/SpatialAccelerationBroadPhase.h
+++ b/Engine/Source/Runtime/Experimental/Chaos/Public/Chaos/Collision/SpatialAccelerationBroadPhase.h
@@ -20,13 +20,30 @@ namespace Chaos
 	class ISpatialAcceleration;
 	class IResimCacheBase;
 
+	// Data collected from the spatial partitioning system
+	class FSimOverlapVisitorEntry
+	{
+	public:
+		FSimOverlapVisitorEntry()
+			: Particle(nullptr)
+		{}
+
+		FSimOverlapVisitorEntry(const FAccelerationStructureHandle& Handle)
+			: Particle(Handle.GetGeometryParticleHandle_PhysicsThread())
+		{
+		}
+
+		FGeometryParticleHandle* Particle;
+	};
+
 	/**
-	 *
+	 * A visitor for the spatial partitioning system used to build the set of objects overlapping a bounding box.
 	 */
 	struct FSimOverlapVisitor
 	{
-		TArray<FAccelerationStructureHandle>& Intersections;
-		FSimOverlapVisitor(FGeometryParticleHandle* ParticleHandle, const FCollisionFilterData& InSimFilterData, TArray<FAccelerationStructureHandle>& InIntersections)
+		TArray<FSimOverlapVisitorEntry>& Intersections;
+
+		FSimOverlapVisitor(FGeometryParticleHandle* ParticleHandle, const FCollisionFilterData& InSimFilterData, TArray<FSimOverlapVisitorEntry>& InIntersections)
 			: Intersections(InIntersections)
 			, SimFilterData(InSimFilterData)
 			, ParticleUniqueIdx(ParticleHandle ? ParticleHandle->UniqueIdx() : FUniqueIdx(0))
@@ -36,7 +53,7 @@ namespace Chaos
 
 		bool VisitOverlap(const TSpatialVisitorData<FAccelerationStructureHandle>& Instance)
 		{
-			Intersections.Add(Instance.Payload);
+			Intersections.Emplace(FSimOverlapVisitorEntry(Instance.Payload));
 			return true;
 		}
 
@@ -152,8 +169,7 @@ namespace Chaos
 		{
 			 OverlapView.ParallelFor([&](auto& Particle1,int32 ActiveIdxIdx)
 			 {
-			 	FGenericParticleHandleImp GenericHandle(Particle1.Handle());
-			 	ProduceParticleOverlaps<bNeedsResim,bOnlyRigid>(Dt,GenericHandle, InSpatialAcceleration,NarrowPhase,ActiveIdxIdx);
+			 	ProduceParticleOverlaps<bNeedsResim,bOnlyRigid>(Dt, Particle1.Handle(), InSpatialAcceleration,NarrowPhase,ActiveIdxIdx);
 			},bDisableCollisionParallelFor);
 		}
 
@@ -193,223 +209,255 @@ namespace Chaos
 		FIgnoreCollisionManager& GetIgnoreCollisionManager() { return IgnoreCollisionManager; }
 
 	private:
-		template<bool bIsResimming, bool bOnlyRigid, typename THandle, typename T_SPATIALACCELERATION>
+
+		template<bool bIsResimming, bool bOnlyRigid, typename T_SPATIALACCELERATION>
 		void ProduceParticleOverlaps(
 		    FReal Dt,
-		    THandle& Particle1,
+		    FGeometryParticleHandle* Particle1,
 		    const T_SPATIALACCELERATION& InSpatialAcceleration,
 		    FNarrowPhase& NarrowPhase,
 			int32 EntryIndex)
 		{
-			TArray<FAccelerationStructureHandle> PotentialIntersections;  
+			TArray<FSimOverlapVisitorEntry> PotentialIntersections;
+
+			bool bIsKinematic1 = true;
+			bool bIsDynamicAwake1 = false;
+			bool bIsDynamicAsleep1 = false;
+			bool bUseIgnoreCollisionManager1 = false;
+			bool bDisabled1 = false;
+			int32 CollisionGroup1 = 0;
+			const FPBDRigidParticleHandle* Rigid1 = Particle1->CastToRigidParticle();
+			if (Rigid1 != nullptr)
+			{
+				bIsKinematic1 = Rigid1->IsKinematic();
+				bIsDynamicAsleep1 = !bIsKinematic1 && Rigid1->IsSleeping();
+				bIsDynamicAwake1 = !bIsKinematic1 && !bIsDynamicAsleep1;
+				bUseIgnoreCollisionManager1 = Rigid1->HasCollisionConstraintFlag(ECollisionConstraintFlags::CCF_BroadPhaseIgnoreCollisions);
+				bDisabled1 = Rigid1->Disabled();
+				CollisionGroup1 = Rigid1->CollisionGroup();
+			}
 
-			const bool bHasValidState = bOnlyRigid ?
-					(Particle1.ObjectState() == EObjectStateType::Dynamic || Particle1.ObjectState() == EObjectStateType::Sleeping) :
-					(Particle1.ObjectState() == EObjectStateType::Dynamic || Particle1.ObjectState() == EObjectStateType::Kinematic);
+			// @todo(chaos): We should not be getting disabled particles, but we do from GeometryCollections
+			//check(!bDisabled1);
+			if (bDisabled1)
+			{
+				return;
+			}
+
+			// Skip particles we are not interested in
+			// @todo(chaos): This should already be handled by selecting the appropriate particle view in the calling function. GeometryCollections?
+			const bool bHasValidState = bOnlyRigid ? (bIsDynamicAwake1 || bIsDynamicAsleep1) : (bIsDynamicAwake1 || bIsKinematic1);
+			if (!bHasValidState && !bIsResimming)
+			{
+				return;
+			}
 			
-			if (bIsResimming || bHasValidState)
+			const bool bBody1Bounded = Particle1->HasBounds();
+			if (bBody1Bounded)
 			{
-				const bool bBody1Bounded = HasBoundingBox(Particle1);
-				{
-					//SCOPE_CYCLE_COUNTER(STAT_Collisions_AABBTree);
-					if (bBody1Bounded)
-					{
-						// @todo(chaos): cache this on the particle?
-						FCollisionFilterData ParticleSimData;
-						FAccelerationStructureHandle::ComputeParticleSimFilterDataFromShapes(Particle1, ParticleSimData);
+				//SCOPE_CYCLE_COUNTER(STAT_Collisions_AABBTree);
+				// @todo(chaos): cache this on the particle?
+				FCollisionFilterData ParticleSimData;
+				FAccelerationStructureHandle::ComputeParticleSimFilterDataFromShapes(*Particle1, ParticleSimData);
 
-						const FAABB3 Box1 = Particle1.WorldSpaceInflatedBounds();
+				const FAABB3 Box1 = Particle1->WorldSpaceInflatedBounds();
 
-						{
-							PHYSICS_CSV_SCOPED_EXPENSIVE(PhysicsVerbose, DetectCollisions_BroadPhase);
-							FSimOverlapVisitor OverlapVisitor(Particle1.Handle() , ParticleSimData, PotentialIntersections);
-							InSpatialAcceleration.Overlap(Box1, OverlapVisitor);
-						}
+				{
+					PHYSICS_CSV_SCOPED_EXPENSIVE(PhysicsVerbose, DetectCollisions_BroadPhase);
+					FSimOverlapVisitor OverlapVisitor(Particle1, ParticleSimData, PotentialIntersections);
+					InSpatialAcceleration.Overlap(Box1, OverlapVisitor);
+				}
 						
-					}
-					else
-					{
-						const auto& GlobalElems = InSpatialAcceleration.GlobalObjects();
-						PotentialIntersections.Reserve(GlobalElems.Num());
+			}
+			else
+			{
+				//SCOPE_CYCLE_COUNTER(STAT_Collisions_AABBTree);
+				const auto& GlobalElems = InSpatialAcceleration.GlobalObjects();
+				PotentialIntersections.Reserve(GlobalElems.Num());
 
-						for (auto& Elem : GlobalElems)
-						{
-							PotentialIntersections.Add(Elem.Payload);
-						}
-					}
+				for (auto& Elem : GlobalElems)
+				{
+					PotentialIntersections.Add(Elem.Payload);
 				}
+			}
 
-				//SCOPE_CYCLE_COUNTER(STAT_Collisions_Filtering);
-				const int32 NumPotentials = PotentialIntersections.Num();
-				int32 NumIntoNarrowPhase = 0;
-				TArray<FParticlePairMidPhase*> MidPhasePairs;
-				MidPhasePairs.Reserve(NumPotentials);
-				for (int32 i = 0; i < NumPotentials; ++i)
+			//SCOPE_CYCLE_COUNTER(STAT_Collisions_Filtering);
+			const int32 NumPotentials = PotentialIntersections.Num();
+			TArray<FParticlePairMidPhase*> MidPhases;
+			MidPhases.Reserve(NumPotentials);
+			for (int32 PotentialIndex = 0; PotentialIndex < NumPotentials; ++PotentialIndex)
+			{
+				FGeometryParticleHandle* Particle2 = PotentialIntersections[PotentialIndex].Particle;
+				if (Particle1 == Particle2)
 				{
-					auto& Particle2 = *PotentialIntersections[i].GetGeometryParticleHandle_PhysicsThread();
-					const FGenericParticleHandle Particle2Generic(&Particle2);
+					continue;
+				}
 
-					// Broad Phase Culling
-					// CollisionGroup == 0 : Collide_With_Everything
-					// CollisionGroup == INDEX_NONE : Disabled collisions
-					// CollisionGroup_A != CollisionGroup_B : Skip Check
-					
-					// collision pair is valid if at least one of the 2 particles is a rigid one (will discard static/kinematic-static/kinematic pairs)
-					if(Particle1.CastToRigidParticle() == nullptr && Particle2.CastToRigidParticle() == nullptr)
-					{
-						continue;
-					}
-					
-					if (Particle1.HasCollisionConstraintFlag(ECollisionConstraintFlags::CCF_BroadPhaseIgnoreCollisions) || 
-						Particle2Generic->HasCollisionConstraintFlag(ECollisionConstraintFlags::CCF_BroadPhaseIgnoreCollisions))
-					{
-						if (IgnoreCollisionManager.IgnoresCollision(Particle1.UniqueIdx(), Particle2.UniqueIdx()))
-						{
-							continue;
-						}
-					}					
-				
-					if (Particle1.CollisionGroup() == INDEX_NONE || Particle2Generic->CollisionGroup() == INDEX_NONE)
-					{
-						continue;
-					}
-					if (Particle1.CollisionGroup() && Particle2Generic->CollisionGroup() && Particle1.CollisionGroup() != Particle2Generic->CollisionGroup())
-					{
-						continue;
-					}
+				bool bIsKinematic2 = true;
+				bool bIsDynamicAwake2 = false;
+				bool bIsDynamicAsleep2 = false;
+				bool bUseIgnoreCollisionManager2 = false;
+				bool bDisabled2 = false;
+				int32 CollisionGroup2 = 0;
+				const FPBDRigidParticleHandle* Rigid2 = Particle2->CastToRigidParticle();
+				if (Rigid2 != nullptr)
+				{
+					bIsKinematic2 = Rigid2->IsKinematic();
+					bIsDynamicAsleep2 = !bIsKinematic2 && Rigid2->IsSleeping();
+					bIsDynamicAwake2 = !bIsKinematic2 && !bIsDynamicAsleep2;
+					bUseIgnoreCollisionManager2 = Rigid2->HasCollisionConstraintFlag(ECollisionConstraintFlags::CCF_BroadPhaseIgnoreCollisions);
+					bDisabled2 = Rigid2->Disabled();
+					CollisionGroup2 = Rigid2->CollisionGroup();
+				}
 
-					if (!Particle1.Geometry() && !Particle2.Geometry())
-					{
-						continue;
-					}
+				// @todo(chaos): This should not be happening if the disabled particles are removed from the active particles list, but GeometryCollection particles leave there there
+				//check(!bDisabled2);
+				if (bDisabled2)
+				{
+					continue;
+				}
 
-					if (Particle1.Handle() == Particle2.Handle())
-					{
-						continue;
-					}
+				// At least one particle needs to be dynamic to generate a collision response
+				if (bIsKinematic1 && bIsKinematic2)
+				{
+					continue;
+				}
+				check((Rigid1 != nullptr) || (Rigid2 != nullptr));
 
-					// HACK : This should not be happening if the disabled particles are properly removed from the active particles list. 
-					if (Particle1.Disabled() || Particle2Generic->Disabled())
+				// Is this particle interaction governed by the IgnoreCollisionManager? If so, check to see if interaction is allowed
+				if (bUseIgnoreCollisionManager1 || bUseIgnoreCollisionManager2)
+				{
+					if (IgnoreCollisionManager.IgnoresCollision(Particle1->UniqueIdx(), Particle2->UniqueIdx()))
 					{
 						continue;
 					}
-					bool bAcceptParticlePair = false;
-
-					// In both cases (resim or not) we will generate (1) dynamic-(sleeping,kinematic(steady+moving),static) pairs + (2) sleeping-moving kinematic ones
-					// Sleeping particles could collide with dynamic ones but these collisions are already handled in case 1
-					// Sleeping particles won't collide with static or steady kinematic particles since neither are moving
-					// Sleeping particles will collide against moving kinematic particles
-					if ((Particle1.ObjectState() == EObjectStateType::Dynamic && Particle2.ObjectState() != EObjectStateType::Dynamic) ||
-						(Particle1.ObjectState() == EObjectStateType::Sleeping && Particle2.ObjectState() == EObjectStateType::Kinematic))
-					{
-						bAcceptParticlePair = true;
-					}
+				}
+
+				// CollisionGroups are used by geometry collections for high-level collision filtering
+				// CollisionGroup == 0 : Collide_With_Everything
+				// CollisionGroup == INDEX_NONE : Disabled collisions
+				// CollisionGroup1 != CollisionGroup2 : Disabled collisions (if other conditions not met)
+				if (CollisionGroup1 == INDEX_NONE || CollisionGroup2 == INDEX_NONE)
+				{
+					continue;
+				}
+				if ((CollisionGroup1 != 0) && (CollisionGroup1 != 0) && (CollisionGroup1 != CollisionGroup2))
+				{
+					continue;
+				}
+
+				// In both cases (resim or not) we will generate (1) dynamic-(sleeping,kinematic(steady+moving),static) pairs + (2) sleeping-moving kinematic ones
+				// Sleeping particles could collide with dynamic ones but these collisions are already handled in case 1
+				// Sleeping particles won't collide with static or steady kinematic particles since neither are moving
+				// Sleeping particles will collide against moving kinematic particles
+				bool bAcceptParticlePair = false;
+				if ((bIsDynamicAwake1 && !bIsDynamicAwake2) || (bIsDynamicAsleep1 && bIsKinematic2))
+				{
+					bAcceptParticlePair = true;
+				}
 					
-					// Used to determine a winner in cases where we will visit particle pairs in both orders
-					const bool bIsParticle1Preferred = AreParticlesInPreferredOrder(Particle1.Handle(), Particle2.Handle());
+				// Used to determine a winner in cases where we will visit particle pairs in both orders
+				const bool bIsParticle1Preferred = AreParticlesInPreferredOrder(Particle1, Particle2);
 					
-					if(!bIsResimming)
+				if(!bIsResimming)
+				{
+					// Normally (not resimming) we iterate over dynamic and asleep|kinematic particles, so:
+					// - Particle1 is dynamic, asleep OR kinematic
+					// - Particle2 may be static, kinematic, dynamic, asleep
+
+					// If Particle1 is non dynamic but particle 2 is dynamic, the case should already be handled by (1)
+					if(!bIsDynamicAwake1 && bIsDynamicAwake2)
 					{
-						// Normally (not resimming) we iterate over dynamic and asleep|kinematic particles, so:
-						// - Particle1 is dynamic, asleep OR kinematic
-						// - Particle2 may be static, kinematic, dynamic, asleep
-
-						// If Particle1 is non dynamic but particle 2 is dynamic, the case should already be handled by (1)
-						if(Particle1.ObjectState() != EObjectStateType::Dynamic && Particle2.ObjectState() == EObjectStateType::Dynamic)
-						{
-							bAcceptParticlePair = false;
-						}
-						// If Particle1 and Particle2 are dynamic we validate the pair if particle1 has higher ID to discard duplicates since we will visit twice the pairs
-						else if(Particle1.ObjectState() == EObjectStateType::Dynamic && Particle2.ObjectState() == EObjectStateType::Dynamic)
-						{
-							bAcceptParticlePair = bIsParticle1Preferred;
-						}
-						// If Particle1 is kinematic we should in theory discard the pairs against sleeping particles
-						// since the sleeping-kinematic case has been validated in (2). But Particle1 is asleep OR kinematic so
-						// when entering this condition we are sure that we enver entered (2). It is why we validate the kinematic-sleeping pairs as well
-						else if(Particle1.ObjectState() == EObjectStateType::Kinematic && Particle2.ObjectState() == EObjectStateType::Sleeping)
-						{
-							bAcceptParticlePair = true;
-						}
+						bAcceptParticlePair = false;
 					}
-					else
+					// If Particle1 and Particle2 are dynamic we validate the pair if particle1 has higher ID to discard duplicates since we will visit twice the pairs
+					else if(bIsDynamicAwake1 && bIsDynamicAwake2)
 					{
-						// When resimming we iterate over "desynced" particles which may be kinematic so:
-						// - Particle1 is always desynced
-						// - Particle2 may also be desynced, in which case we will also visit the opposite ordering regardless of dynamic/kinematic status
-						// - Particle1 may be static, kinematic, dynamic, asleep
-						// - Particle2 may be static, kinematic, dynamic, asleep
-						// 
-						// Even though Particle1 may be kinematic when resimming, we want to create the contacts in the original order (i.e., dynamic first)
-						// 
-						const bool bIsParticle2Desynced = bIsResimming && (Particle2.SyncState() == ESyncState::HardDesync);
-
-						// If Particle1 is non dynamic but particle 2 is dynamic, the case should already be handled by (1) for
-						// the desynced dynamic - synced/desynced (static,kinematic,asleep) pairs. But we still need to process
-						// the desynced (static,kinematic,asleep) against the synced dynamic since this case have not been handled by (1)
-						if(Particle1.ObjectState() != EObjectStateType::Dynamic && Particle2.ObjectState() == EObjectStateType::Dynamic)
-						{
-							bAcceptParticlePair = !bIsParticle2Desynced;
-						}
-						// If Particle1 and Particle2 are dynamic we validate the pair if particle1 has higher ID to discard duplicates since we will visit twice the pairs
-						// We validate the pairs as well if the particle 2 is synced since  we will never visit the opposite order (we only iterate over desynced particles)
-						else if(Particle1.ObjectState() == EObjectStateType::Dynamic && Particle2.ObjectState() == EObjectStateType::Dynamic)
-						{
-							bAcceptParticlePair = bIsParticle1Preferred || !bIsParticle2Desynced;
-						}
-						// If Particle1 is kinematic we are discarding the pairs against sleeping desynced particles since the case has been handled by (2)
-						else if(Particle1.ObjectState() == EObjectStateType::Kinematic && Particle2.ObjectState() == EObjectStateType::Sleeping)
-						{
-							bAcceptParticlePair = !bIsParticle2Desynced;
-						}
+						bAcceptParticlePair = bIsParticle1Preferred;
 					}
-					
-					if (!bAcceptParticlePair)
+					// If Particle1 is kinematic we should in theory discard the pairs against sleeping particles
+					// since the sleeping-kinematic case has been validated in (2). But Particle1 is asleep OR kinematic so
+					// when entering this condition we are sure that we enver entered (2). It is why we validate the kinematic-sleeping pairs as well
+					else if(bIsKinematic1 && bIsDynamicAsleep2)
 					{
-						continue;
+						bAcceptParticlePair = true;
 					}
-
-					// Constraints have a key generated from the Particle IDs and we don't want to have to deal with constraint being created from particles in the two orders.
-					// Since particles can change type (between Kinematic and Dynamic) we may visit them in different orders at different times, but if we allow
-					// that it would break Resim and constraint re-use. Also, if only one particle is Dynamic, we want it in first position. This isn't a strtct 
-					// requirement but some downstream systems assume this is true (e.g., CCD, TriMesh collision).
-					TGeometryParticleHandle<FReal, 3>* ParticleA = Particle1.Handle();
-					TGeometryParticleHandle<FReal, 3>* ParticleB = Particle2.Handle();
-					const bool bSwapOrder = ShouldSwapParticleOrder(ParticleA, ParticleB, bIsParticle1Preferred);
-					if (bSwapOrder)
+				}
+				else
+				{
+					// When resimming we iterate over "desynced" particles which may be kinematic so:
+					// - Particle1 is always desynced
+					// - Particle2 may also be desynced, in which case we will also visit the opposite ordering regardless of dynamic/kinematic status
+					// - Particle1 may be static, kinematic, dynamic, asleep
+					// - Particle2 may be static, kinematic, dynamic, asleep
+					// 
+					// Even though Particle1 may be kinematic when resimming, we want to create the contacts in the original order (i.e., dynamic first)
+					// 
+					const bool bIsParticle2Desynced = bIsResimming && (Particle2->SyncState() == ESyncState::HardDesync);
+
+					// If Particle1 is non dynamic but particle 2 is dynamic, the case should already be handled by (1) for
+					// the desynced dynamic - synced/desynced (static,kinematic,asleep) pairs. But we still need to process
+					// the desynced (static,kinematic,asleep) against the synced dynamic since this case have not been handled by (1)
+					if(!bIsDynamicAwake1 && bIsDynamicAwake2)
 					{
-						Swap(ParticleA, ParticleB);
+						bAcceptParticlePair = !bIsParticle2Desynced;
 					}
-
-					FParticlePairMidPhase* MidPhase = NarrowPhase.GetParticlePairMidPhase(ParticleA, ParticleB, Particle1.Handle());
-					if (MidPhase)
+					// If Particle1 and Particle2 are dynamic we validate the pair if particle1 has higher ID to discard duplicates since we will visit twice the pairs
+					// We validate the pairs as well if the particle 2 is synced since  we will never visit the opposite order (we only iterate over desynced particles)
+					else if(bIsDynamicAwake1 && bIsDynamicAwake2)
 					{
-						MidPhasePairs.Add(MidPhase);
+						bAcceptParticlePair = bIsParticle1Preferred || !bIsParticle2Desynced;
+					}
+					// If Particle1 is kinematic we are discarding the pairs against sleeping desynced particles since the case has been handled by (2)
+					else if(bIsKinematic1 && bIsDynamicAsleep2)
+					{
+						bAcceptParticlePair = !bIsParticle2Desynced;
 					}
 				}
+					
+				if (!bAcceptParticlePair)
+				{
+					continue;
+				}
 
-				const int32  PrefetchLookahead = 4;
-
-				for (int32 Index = 0; Index < MidPhasePairs.Num() && Index < PrefetchLookahead; Index++)
+				// Since particles can change type (between Kinematic and Dynamic) we may visit them in different orders at different times, but if we allow
+				// that it would break Resim and constraint re-use. Also, if only one particle is Dynamic, we want it in first position. This isn't a strtct 
+				// requirement but some downstream systems assume this is true (e.g., CCD, TriMesh collision).
+				const bool bSwapOrder = ShouldSwapParticleOrder((bIsDynamicAwake1 || bIsDynamicAsleep1), bIsParticle1Preferred);
+				FGeometryParticleHandle* ParticleA = bSwapOrder ? Particle2 : Particle1;
+				FGeometryParticleHandle* ParticleB = bSwapOrder ? Particle1 : Particle2;
+
+				// NOTE: This searches the array of existing midphases on the particle specified in the last parameter. The assumption
+				// here is that Particle1 is more likely to have less collision on it because it is dynamic. Kinematics can often be
+				// large and may collide with many objects. Maybe we should consider using size instead.
+				FParticlePairMidPhase* MidPhase = NarrowPhase.GetParticlePairMidPhase(ParticleA, ParticleB, Particle1);
+				if (MidPhase != nullptr)
 				{
-					MidPhasePairs[Index]->CachePrefetch();
+					MidPhases.Add(MidPhase);
 				}
+			}
 
+			// Prefetch initial set of MidPhases
+			const int32 PrefetchLookahead = 4;
+			const int32 NumMidPhases = MidPhases.Num();
+			for (int32 Index = 0; Index < NumMidPhases && Index < PrefetchLookahead; Index++)
+			{
+				MidPhases[Index]->CachePrefetch();
+			}
 
-				for (int32 Index = 0; Index < MidPhasePairs.Num(); Index++)
+			for (int32 Index = 0; Index < NumMidPhases; Index++)
+			{
+				// Prefetch next MidPhase
+				if (Index + PrefetchLookahead < NumMidPhases)
 				{
-					// Prefetch next pair
-					if (Index + PrefetchLookahead  < MidPhasePairs.Num())
-					{
-						MidPhasePairs[Index + PrefetchLookahead]->CachePrefetch();
-					}
-
-					FParticlePairMidPhase* MidPhasePair = MidPhasePairs[Index];
-					NarrowPhase.GenerateCollisions(Dt, MidPhasePair);
+					MidPhases[Index + PrefetchLookahead]->CachePrefetch();
 				}
 
-				PHYSICS_CSV_CUSTOM_EXPENSIVE(PhysicsCounters, NumFromBroadphase, NumPotentials, ECsvCustomStatOp::Accumulate);
+				// Run MidPhase + NarrowPhase
+				NarrowPhase.GenerateCollisions(Dt, MidPhases[Index]);
 			}
+
+			PHYSICS_CSV_CUSTOM_EXPENSIVE(PhysicsCounters, NumFromBroadphase, NumPotentials, ECsvCustomStatOp::Accumulate);
 		}
 
 		const FPBDRigidsSOAs& Particles;
diff --git a/Engine/Source/Runtime/Experimental/Chaos/Public/Chaos/ParticleHandle.h b/Engine/Source/Runtime/Experimental/Chaos/Public/Chaos/ParticleHandle.h
--- a/Engine/Source/Runtime/Experimental/Chaos/Public/Chaos/ParticleHandle.h
+++ b/Engine/Source/Runtime/Experimental/Chaos/Public/Chaos/ParticleHandle.h
@@ -886,6 +886,7 @@ public:
 
 	operator TPBDRigidParticleHandleImp<T, d, false>& () { return reinterpret_cast<TPBDRigidParticleHandleImp<T, d, false>&>(*this); }
 
+	bool IsKinematic() const { return ObjectState() == EObjectStateType::Kinematic; }
 	bool IsDynamic() const { return (ObjectState() == EObjectStateType::Dynamic) || (ObjectState() == EObjectStateType::Sleeping); }
 
 	const TUniquePtr<TBVHParticles<T, d>>& CollisionParticles() const { return PBDRigidParticles->CollisionParticles(ParticleIdx); }
@@ -1050,6 +1051,8 @@ public:
 	void SetObjectStateLowLevel(EObjectStateType InState) { PBDRigidParticles->SetObjectState(ParticleIdx, InState); }
 	void SetPreObjectStateLowLevel(EObjectStateType InState) { PBDRigidParticles->PreObjectState(ParticleIdx) = InState; }
 	
+	bool IsSleeping() const { return ObjectState() == EObjectStateType::Sleeping; }
+
 	bool Sleeping() const { return PBDRigidParticles->Sleeping(ParticleIdx); }
 	void SetSleeping(bool bSleeping) { PBDRigidParticles->SetSleeping(ParticleIdx, bSleeping); }
 
diff --git a/Engine/Source/Runtime/Experimental/Chaos/Public/Chaos/RigidParticles.h b/Engine/Source/Runtime/Experimental/Chaos/Public/Chaos/RigidParticles.h
--- a/Engine/Source/Runtime/Experimental/Chaos/Public/Chaos/RigidParticles.h
+++ b/Engine/Source/Runtime/Experimental/Chaos/Public/Chaos/RigidParticles.h
@@ -69,37 +69,13 @@ public:
 	CHAOS_API TRigidParticles()
 	    : TKinematicGeometryParticles<T, d>()
 	{
-		TArrayCollection::AddArray(&MVSmooth);
-		TArrayCollection::AddArray(&MWSmooth);
-		TArrayCollection::AddArray(&MAcceleration);
-		TArrayCollection::AddArray(&MAngularAcceleration);
-		TArrayCollection::AddArray(&MLinearImpulseVelocity);
-		TArrayCollection::AddArray(&MAngularImpulseVelocity);
-		TArrayCollection::AddArray(&MI);
-		TArrayCollection::AddArray(&MInvI);
-		TArrayCollection::AddArray(&MInvIConditioning);
-		TArrayCollection::AddArray(&MM);
-		TArrayCollection::AddArray(&MInvM);
-		TArrayCollection::AddArray(&MCenterOfMass);
-		TArrayCollection::AddArray(&MRotationOfMass);
-		TArrayCollection::AddArray(&MLinearEtherDrag);
-		TArrayCollection::AddArray(&MAngularEtherDrag);
-		TArrayCollection::AddArray(&MaxLinearSpeedsSq);
-		TArrayCollection::AddArray(&MaxAngularSpeedsSq);
-		TArrayCollection::AddArray(&MCollisionParticles);
-		TArrayCollection::AddArray(&MCollisionGroup);
-		TArrayCollection::AddArray(&MCollisionConstraintFlags);
-		TArrayCollection::AddArray(&MControlFlags);
-		TArrayCollection::AddArray(&MTransientFlags);
-		TArrayCollection::AddArray(&MDisabled);
-		TArrayCollection::AddArray(&MObjectState);
-		TArrayCollection::AddArray(&MPreObjectState);
-		TArrayCollection::AddArray(&MIslandIndex);
-		TArrayCollection::AddArray(&MSleepType);
+		RegisterArrays();
 	}
+
 	TRigidParticles(const TRigidParticles<T, d>& Other) = delete;
 	CHAOS_API TRigidParticles(TRigidParticles<T, d>&& Other)
 	    : TKinematicGeometryParticles<T, d>(MoveTemp(Other))
+		, CoreData(MoveTemp(Other.CoreData))
 		, MVSmooth(MoveTemp(Other.MVSmooth))
 		, MWSmooth(MoveTemp(Other.MWSmooth))
 		, MAcceleration(MoveTemp(Other.MAcceleration))
@@ -118,16 +94,16 @@ public:
 		, MaxLinearSpeedsSq(MoveTemp(Other.MaxLinearSpeedsSq))
 		, MaxAngularSpeedsSq(MoveTemp(Other.MaxAngularSpeedsSq))
 		, MCollisionParticles(MoveTemp(Other.MCollisionParticles))
-		, MCollisionGroup(MoveTemp(Other.MCollisionGroup))
-		, MCollisionConstraintFlags(MoveTemp(Other.MCollisionConstraintFlags))
-		, MControlFlags(MoveTemp(Other.MControlFlags))
-		, MTransientFlags(MoveTemp(Other.MTransientFlags))
 		, MIslandIndex(MoveTemp(Other.MIslandIndex))
-		, MObjectState(MoveTemp(Other.MObjectState))
-		, MPreObjectState(MoveTemp(Other.MPreObjectState))
 		, MSleepType(MoveTemp(Other.MSleepType))
-		, MDisabled(MoveTemp(Other.MDisabled))
 	{
+		RegisterArrays();
+	}
+
+	void RegisterArrays()
+	{
+		TArrayCollection::AddArray(&CoreData);
+
 		TArrayCollection::AddArray(&MVSmooth);
 		TArrayCollection::AddArray(&MWSmooth);
 		TArrayCollection::AddArray(&MAcceleration);
@@ -146,15 +122,9 @@ public:
 		TArrayCollection::AddArray(&MaxLinearSpeedsSq);
 		TArrayCollection::AddArray(&MaxAngularSpeedsSq);
 		TArrayCollection::AddArray(&MCollisionParticles);
-		TArrayCollection::AddArray(&MCollisionGroup);
-		TArrayCollection::AddArray(&MCollisionConstraintFlags);
-		TArrayCollection::AddArray(&MControlFlags);
-		TArrayCollection::AddArray(&MTransientFlags);
-		TArrayCollection::AddArray(&MDisabled);
-		TArrayCollection::AddArray(&MObjectState);
-		TArrayCollection::AddArray(&MPreObjectState);
 		TArrayCollection::AddArray(&MIslandIndex);
 		TArrayCollection::AddArray(&MSleepType);
+
 	}
 
 	CHAOS_API virtual ~TRigidParticles()
@@ -219,32 +189,33 @@ public:
 	FORCEINLINE const TUniquePtr<TBVHParticles<T, d>>& CollisionParticles(const int32 Index) const { return MCollisionParticles[Index]; }
 	FORCEINLINE TUniquePtr<TBVHParticles<T, d>>& CollisionParticles(const int32 Index) { return MCollisionParticles[Index]; }
 
-	FORCEINLINE const int32 CollisionGroup(const int32 Index) const { return MCollisionGroup[Index]; }
-	FORCEINLINE int32& CollisionGroup(const int32 Index) { return MCollisionGroup[Index]; }
+	FORCEINLINE const int32 CollisionGroup(const int32 Index) const { return CoreData[Index].CollisionGroup; }
+	FORCEINLINE int32& CollisionGroup(const int32 Index) { return CoreData[Index].CollisionGroup; }
 
-	FORCEINLINE bool HasCollisionConstraintFlag(const ECollisionConstraintFlags Flag, const int32 Index) const { return (MCollisionConstraintFlags[Index] & (uint32)Flag) != 0; }
-	FORCEINLINE void AddCollisionConstraintFlag(const ECollisionConstraintFlags Flag, const int32 Index) { MCollisionConstraintFlags[Index] |= (uint32)Flag; }
-	FORCEINLINE void RemoveCollisionConstraintFlag(const ECollisionConstraintFlags Flag, const int32 Index) { MCollisionConstraintFlags[Index] &= ~(uint32)Flag; }
-	FORCEINLINE void SetCollisionConstraintFlags(const int32 Index, const uint32 Flags) { MCollisionConstraintFlags[Index] = Flags; }
-	FORCEINLINE uint32 CollisionConstraintFlags(const int32 Index) const { return MCollisionConstraintFlags[Index]; }
+	FORCEINLINE bool HasCollisionConstraintFlag(const ECollisionConstraintFlags Flag, const int32 Index) const { return (CoreData[Index].CollisionConstraintFlags & (uint32)Flag) != 0; }
+	FORCEINLINE void AddCollisionConstraintFlag(const ECollisionConstraintFlags Flag, const int32 Index) { CoreData[Index].CollisionConstraintFlags |= (uint32)Flag; }
+	FORCEINLINE void RemoveCollisionConstraintFlag(const ECollisionConstraintFlags Flag, const int32 Index) { CoreData[Index].CollisionConstraintFlags &= ~(uint32)Flag; }
+	FORCEINLINE void SetCollisionConstraintFlags(const int32 Index, const uint32 Flags) { CoreData[Index].CollisionConstraintFlags = Flags; }
+	FORCEINLINE uint32 CollisionConstraintFlags(const int32 Index) const { return CoreData[Index].CollisionConstraintFlags; }
 
-	FORCEINLINE const bool Disabled(const int32 Index) const { return MDisabled[Index]; }
+	FORCEINLINE const bool Disabled(const int32 Index) const { return CoreData[Index].bDisabled; }
 
-	FORCEINLINE bool& DisabledRef(const int32 Index) { return MDisabled[Index]; }
+	FORCEINLINE bool& DisabledRef(const int32 Index) { return CoreData[Index].bDisabled; }
 
 	// DisableParticle/EnableParticle on Evolution should be used. Don't disable particles with this.
     // Using this will break stuff. This is for solver's use only, and possibly some particle construction/copy code.
-	FORCEINLINE void SetDisabledLowLevel(const int32 Index, bool disabled) { MDisabled[Index] = disabled; }
+	FORCEINLINE void SetDisabledLowLevel(const int32 Index, bool InDisabled) { CoreData[Index].bDisabled = InDisabled; }
 
-	FORCEINLINE const FRigidParticleControlFlags& ControlFlags(const int32 Index) const { return MControlFlags[Index]; }
-	FORCEINLINE FRigidParticleControlFlags& ControlFlags(const int32 Index) { return MControlFlags[Index]; }
+	FORCEINLINE const FRigidParticleControlFlags& ControlFlags(const int32 Index) const { return CoreData[Index].ControlFlags; }
+	FORCEINLINE FRigidParticleControlFlags& ControlFlags(const int32 Index) { return CoreData[Index].ControlFlags; }
 
-	FORCEINLINE const FRigidParticleTransientFlags& TransientFlags(const int32 Index) const { return MTransientFlags[Index]; }
-	FORCEINLINE FRigidParticleTransientFlags& TransientFlags(const int32 Index) { return MTransientFlags[Index]; }
+	FORCEINLINE const FRigidParticleTransientFlags& TransientFlags(const int32 Index) const { return CoreData[Index].TransientFlags; }
+	FORCEINLINE FRigidParticleTransientFlags& TransientFlags(const int32 Index) { return CoreData[Index].TransientFlags; }
 
 	FORCEINLINE ESleepType SleepType(const int32 Index) const { return MSleepType[Index]; }
 	FORCEINLINE ESleepType& SleepType(const int32 Index) { return MSleepType[Index]; }
 
+	// @todo(chaos): This data should be marshalled via the proxies like everything else. There is probably a particle recycling bug here.
 	FORCEINLINE TArray<TSleepData<T, d>>& GetSleepData() { return MSleepData; }
 	FORCEINLINE	void AddSleepData(TGeometryParticleHandle<T, d>* Particle, bool Sleeping)
 	{ 
@@ -264,11 +235,11 @@ public:
 	}
 	FORCEINLINE FRWLock& GetSleepDataLock() { return SleepDataLock; }
 
-	FORCEINLINE const EObjectStateType ObjectState(const int32 Index) const { return MObjectState[Index]; }
-	FORCEINLINE EObjectStateType& ObjectState(const int32 Index) { return MObjectState[Index]; }
+	FORCEINLINE const EObjectStateType ObjectState(const int32 Index) const { return CoreData[Index].ObjectState; }
+	FORCEINLINE EObjectStateType& ObjectState(const int32 Index) { return CoreData[Index].ObjectState; }
 
-	FORCEINLINE const EObjectStateType PreObjectState(const int32 Index) const { return MPreObjectState[Index]; }
-	FORCEINLINE EObjectStateType& PreObjectState(const int32 Index) { return MPreObjectState[Index]; }
+	FORCEINLINE const EObjectStateType PreObjectState(const int32 Index) const { return CoreData[Index].PreObjectState; }
+	FORCEINLINE EObjectStateType& PreObjectState(const int32 Index) { return CoreData[Index].PreObjectState; }
 
 	FORCEINLINE const bool Dynamic(const int32 Index) const { return ObjectState(Index) == EObjectStateType::Dynamic; }
 
@@ -279,18 +250,26 @@ public:
 	FORCEINLINE const int32 IslandIndex(const int32 Index) const { return MIslandIndex[Index]; }
 	FORCEINLINE int32& IslandIndex(const int32 Index) { return MIslandIndex[Index]; }
 
-	FORCEINLINE FString ToString(int32 index) const
+	FORCEINLINE FString ToString(int32 Index) const
 	{
-		FString BaseString = TKinematicGeometryParticles<T, d>::ToString(index);
+		FString BaseString = TKinematicGeometryParticles<T, d>::ToString(Index);
 		return FString::Printf(TEXT("%s, MAcceleration:%s, MAngularAcceleration:%s, MLinearImpulseVelocity:%s, MAngularImpulseVelocity:%s, MI:%s, MInvI:%s, MM:%f, MInvM:%f, MCenterOfMass:%s, MRotationOfMass:%s, MCollisionParticles(num):%d, MCollisionGroup:%d, MDisabled:%d, MSleeping:%d, MIslandIndex:%d"),
-			*BaseString, *Acceleration(index).ToString(), *AngularAcceleration(index).ToString(), *LinearImpulseVelocity(index).ToString(), *AngularImpulseVelocity(index).ToString(),
-			*I(index).ToString(), *InvI(index).ToString(), M(index), InvM(index), *CenterOfMass(index).ToString(), *RotationOfMass(index).ToString(), CollisionParticlesSize(index),
-			CollisionGroup(index), Disabled(index), Sleeping(index), IslandIndex(index));
+			*BaseString, *Acceleration(Index).ToString(), *AngularAcceleration(Index).ToString(), *LinearImpulseVelocity(Index).ToString(), *AngularImpulseVelocity(Index).ToString(),
+			*I(Index).ToString(), *InvI(Index).ToString(), M(Index), InvM(Index), *CenterOfMass(Index).ToString(), *RotationOfMass(Index).ToString(), CollisionParticlesSize(Index),
+			CollisionGroup(Index), Disabled(Index), Sleeping(Index), IslandIndex(Index));
 	}
 
 	CHAOS_API virtual void Serialize(FChaosArchive& Ar) override
 	{
 		TKinematicGeometryParticles<T,d>::Serialize(Ar);
+
+		// To avoid bumping file version, serialize to/from previous structures
+		// If we aren't loading (i.e., we are saving or copying) CoreData will be valid, so copy that to the legacy structure
+		FLegacyData LegacyData;
+		if (!Ar.IsLoading())
+		{
+			LegacyData.CopyFromCoreData(CoreData);
+		}
 		
 		Ar.UsingCustomVersion(FExternalPhysicsCustomObjectVersion::GUID);
 		if(Ar.CustomVer(FExternalPhysicsCustomObjectVersion::GUID) >= FExternalPhysicsCustomObjectVersion::KinematicCentersOfMass)
@@ -335,12 +314,67 @@ public:
 			Ar << MLinearEtherDrag << MAngularEtherDrag;
 		}
 
-		Ar << MCollisionParticles << MCollisionGroup << MIslandIndex << MDisabled << MObjectState << MSleepType;
+		Ar << MCollisionParticles << LegacyData.MCollisionGroup << MIslandIndex << LegacyData.MDisabled << LegacyData.MObjectState << MSleepType;
+		// @todo(chaos): what about ControlFlags, TransientFlags, PreObjectState, ..?
 
-		// @todo(chaos): what about MControlFlags and MTransientFlags?
+		// If we loaded into the legacy structure, copy to CoreData
+		if (Ar.IsLoading())
+		{
+			LegacyData.CopyToCoreData(CoreData);
+		}
 	}
 
 private:
+	// Data that is commonly accessed together
+	// This contains all properties accessed in the broadphase filtering (FSpatialAccelerationBroadPhase)
+	struct FCoreData
+	{
+		int32 CollisionGroup;							// 4 bytes
+		uint32 CollisionConstraintFlags;				// 4 bytes
+		FRigidParticleControlFlags ControlFlags;		// 1 byte
+		FRigidParticleTransientFlags TransientFlags;	// 1 byte
+		EObjectStateType ObjectState;					// 1 byte
+		EObjectStateType PreObjectState;				// 1 byte
+		bool bDisabled;									// 1 byte
+	};
+
+	// Used during serialization to avoid bumping the file version as we switch to aggregated strunctures like FCoreData.
+	// Note: Only serialized data is needed here and not all data is serialized so some elements in the aggregates are not represented.
+	struct FLegacyData
+	{
+		TArrayCollectionArray<int32> MCollisionGroup;
+		TArrayCollectionArray<EObjectStateType> MObjectState;
+		TArrayCollectionArray<bool> MDisabled;
+
+		void CopyFromCoreData(const TArrayCollectionArray<FCoreData>& Source)
+		{
+			MCollisionGroup.Resize(Source.Num());
+			MObjectState.Resize(Source.Num());
+			MDisabled.Resize(Source.Num());
+
+			for (int32 Index = 0; Index < Source.Num(); ++Index)
+			{
+				MCollisionGroup[Index] = Source[Index].CollisionGroup;
+				MObjectState[Index] = Source[Index].ObjectState;
+				MDisabled[Index] = Source[Index].bDisabled;
+			}
+		}
+
+		void CopyToCoreData(TArrayCollectionArray<FCoreData>& Dest)
+		{
+			Dest.Resize(MCollisionGroup.Num());
+
+			for (int32 Index = 0; Index < Dest.Num(); ++Index)
+			{
+				Dest[Index].CollisionGroup = MCollisionGroup[Index];
+				Dest[Index].ObjectState = MObjectState[Index];
+				Dest[Index].bDisabled = MDisabled[Index];
+			}
+		}
+	};
+
+	TArrayCollectionArray<FCoreData> CoreData;
+
 	TArrayCollectionArray<TVector<T, d>> MVSmooth;
 	TArrayCollectionArray<TVector<T, d>> MWSmooth;
 	TArrayCollectionArray<TVector<T, d>> MAcceleration;
@@ -359,15 +393,8 @@ private:
 	TArrayCollectionArray<T> MaxLinearSpeedsSq;
 	TArrayCollectionArray<T> MaxAngularSpeedsSq;
 	TArrayCollectionArray<TUniquePtr<TBVHParticles<T, d>>> MCollisionParticles;
-	TArrayCollectionArray<int32> MCollisionGroup;
-	TArrayCollectionArray<uint32> MCollisionConstraintFlags;
-	TArrayCollectionArray<FRigidParticleControlFlags> MControlFlags;
-	TArrayCollectionArray<FRigidParticleTransientFlags> MTransientFlags;
 	TArrayCollectionArray<int32> MIslandIndex;
-	TArrayCollectionArray<EObjectStateType> MObjectState;
-	TArrayCollectionArray<EObjectStateType> MPreObjectState;
 	TArrayCollectionArray<ESleepType> MSleepType;
-	TArrayCollectionArray<bool> MDisabled;
 
 	TArray<TSleepData<T, d>> MSleepData;
 	FRWLock SleepDataLock;
