Chaos - fix collision filter typo and static analysis warning

diff --git a/Engine/Source/Runtime/Experimental/Chaos/Public/Chaos/Collision/SpatialAccelerationBroadPhase.h b/Engine/Source/Runtime/Experimental/Chaos/Public/Chaos/Collision/SpatialAccelerationBroadPhase.h
--- a/Engine/Source/Runtime/Experimental/Chaos/Public/Chaos/Collision/SpatialAccelerationBroadPhase.h
+++ b/Engine/Source/Runtime/Experimental/Chaos/Public/Chaos/Collision/SpatialAccelerationBroadPhase.h
@@ -341,7 +341,7 @@ namespace Chaos
 				{
 					continue;
 				}
-				if ((CollisionGroup1 != 0) && (CollisionGroup1 != 0) && (CollisionGroup1 != CollisionGroup2))
+				if ((CollisionGroup1 != 0) && (CollisionGroup2 != 0) && (CollisionGroup1 != CollisionGroup2))
 				{
 					continue;
 				}
