Chaos : fix damage threshold not being set from asset when adding it to a level  ( got wrongly commented out with a previous CL )

diff --git a/Engine/Source/Runtime/Experimental/GeometryCollectionEngine/Private/GeometryCollection/GeometryCollectionComponent.cpp b/Engine/Source/Runtime/Experimental/GeometryCollectionEngine/Private/GeometryCollection/GeometryCollectionComponent.cpp
--- a/Engine/Source/Runtime/Experimental/GeometryCollectionEngine/Private/GeometryCollection/GeometryCollectionComponent.cpp
+++ b/Engine/Source/Runtime/Experimental/GeometryCollectionEngine/Private/GeometryCollection/GeometryCollectionComponent.cpp
@@ -2701,7 +2701,7 @@ void UGeometryCollectionComponent::SetRestCollection(const UGeometryCollection*
 		}
 
 		// initialize the component per level damage threshold from the asset defaults 
-		//DamageThreshold = RestCollection->DamageThreshold;
+		DamageThreshold = RestCollection->DamageThreshold;
 
 		// initialize the component damage progataion data from the asset defaults 
 		DamagePropagationData = RestCollection->DamagePropagationData;
