Disable UndefinedBinaryOperatorResult analysis test for Chaos module due to false positives in mass calculation

diff --git a/Engine/Source/Runtime/Experimental/Chaos/Chaos.Build.cs b/Engine/Source/Runtime/Experimental/Chaos/Chaos.Build.cs
--- a/Engine/Source/Runtime/Experimental/Chaos/Chaos.Build.cs
+++ b/Engine/Source/Runtime/Experimental/Chaos/Chaos.Build.cs
@@ -38,6 +38,7 @@ namespace UnrealBuildTool.Rules
 			UnsafeTypeCastWarningLevel = WarningLevel.Error;
 
 			StaticAnalyzerDisabledCheckers.Add("cplusplus.NewDeleteLeaks"); // To be reevalulated, believed to be invalid warnings.
+			StaticAnalyzerDisabledCheckers.Add("core.UndefinedBinaryOperatorResult"); // Invalid warning in mass property calculation.
 		}
     }
 }
