Fracture tool : fix crashes when selecting using the hide unselected feature
- crash is caused by constantData section having no triangles and causing bogus section data in the render proxy

diff --git a/Engine/Source/Runtime/Experimental/GeometryCollectionEngine/Private/GeometryCollection/GeometryCollectionSceneProxy.cpp b/Engine/Source/Runtime/Experimental/GeometryCollectionEngine/Private/GeometryCollection/GeometryCollectionSceneProxy.cpp
--- a/Engine/Source/Runtime/Experimental/GeometryCollectionEngine/Private/GeometryCollection/GeometryCollectionSceneProxy.cpp
+++ b/Engine/Source/Runtime/Experimental/GeometryCollectionEngine/Private/GeometryCollection/GeometryCollectionSceneProxy.cpp
@@ -607,17 +607,16 @@ void FGeometryCollectionSceneProxy::SetConstantData_RenderThread(FGeometryCollec
 		}
 
 		// Update mesh sections
-		check(Sections.Num() == ConstantData->Sections.Num());
+		Sections.Reset(ConstantData->Sections.Num());
 		// #todo(dmp): We should restructure the component/SceneProxy usage to avoid this messy stuff.  We need to know the sections
 		// when we create the sceneproxy for the hit proxy to work, but then we are updating the sections here with potentially differing
 		// vertex counts due to hiding geometry.  Ideally, the SceneProxy is treated as const and recreated whenever the geometry
 		// changes rather than this.  SetConstantData_RenderThread should be done in the constructor for the sceneproxy, most likely
-		int i = 0;
 		for (FGeometryCollectionSection Section : ConstantData->Sections)
 		{
 			if (Section.NumTriangles > 0)
 			{
-				FGeometryCollectionSection &NewSection = Sections[i++];
+				FGeometryCollectionSection& NewSection = Sections.AddDefaulted_GetRef();
 				NewSection.MaterialID = Section.MaterialID;
 				NewSection.FirstIndex = Section.FirstIndex;
 				NewSection.NumTriangles = Section.NumTriangles;
