GeometryCollection : AppendSkeletalMesh
Made GeometryCollection::AppendSkeletalMesh more generic so it can take just a FManagedArrayCollection.
- Added a static version of FGeometryCollection::UpdateBoundingBox
- Added a static version of FGeometryCollection::ReindexMaterials
- Added Interface for standard attribute creation

diff --git a/Engine/Source/Runtime/Experimental/Chaos/Private/GeometryCollection/GeometryCollection.cpp b/Engine/Source/Runtime/Experimental/Chaos/Private/GeometryCollection/GeometryCollection.cpp
--- a/Engine/Source/Runtime/Experimental/Chaos/Private/GeometryCollection/GeometryCollection.cpp
+++ b/Engine/Source/Runtime/Experimental/Chaos/Private/GeometryCollection/GeometryCollection.cpp
@@ -39,6 +39,49 @@ FGeometryCollection::FGeometryCollection()
 	Construct();
 }
 
+void FGeometryCollection::DefineGeometryScheam(FManagedArrayCollection& InCollection)
+{
+	FTransformCollection::DefineTransformSchema(InCollection);
+
+	FManagedArrayCollection::FConstructionParameters TransformDependency(FTransformCollection::TransformGroup);
+	FManagedArrayCollection::FConstructionParameters VerticesDependency(FGeometryCollection::VerticesGroup);
+	FManagedArrayCollection::FConstructionParameters FacesDependency(FGeometryCollection::FacesGroup);
+
+	// Transform Group
+	InCollection.AddAttribute<int32>("TransformToGeometryIndex", FTransformCollection::TransformGroup);
+	InCollection.AddAttribute<int32>("SimulationType", FTransformCollection::TransformGroup);
+	InCollection.AddAttribute<int32>("StatusFlags", FTransformCollection::TransformGroup);
+	InCollection.AddAttribute<int32>("InitialDynamicState", FTransformCollection::TransformGroup);
+	InCollection.AddAttribute<int32>("ExemplarIndex", FTransformCollection::TransformGroup);
+
+	// Vertices Group
+	InCollection.AddAttribute<FVector3f>("Vertex", FGeometryCollection::VerticesGroup);
+	InCollection.AddAttribute<FVector3f>("Normal", FGeometryCollection::VerticesGroup);
+	InCollection.AddAttribute<TArray<FVector2f>>("UVs", FGeometryCollection::VerticesGroup);
+	InCollection.AddAttribute<FLinearColor>("Color", FGeometryCollection::VerticesGroup);
+	InCollection.AddAttribute<FVector3f>("TangentU", FGeometryCollection::VerticesGroup);
+	InCollection.AddAttribute<FVector3f>("TangentV", FGeometryCollection::VerticesGroup);
+	InCollection.AddAttribute<int32>("BoneMap", FGeometryCollection::VerticesGroup, TransformDependency);
+
+	// Faces Group
+	InCollection.AddAttribute<FIntVector>("Indices", FGeometryCollection::FacesGroup, VerticesDependency);
+	InCollection.AddAttribute<bool>("Visible", FGeometryCollection::FacesGroup);
+	InCollection.AddAttribute<int32>("MaterialIndex", FGeometryCollection::FacesGroup);
+	InCollection.AddAttribute<int32>("MaterialID", FGeometryCollection::FacesGroup);
+
+	// Geometry Group
+	InCollection.AddAttribute<int32>("TransformIndex", FGeometryCollection::GeometryGroup, TransformDependency);
+	InCollection.AddAttribute<FBox>("BoundingBox", FGeometryCollection::GeometryGroup);
+	InCollection.AddAttribute<float>("InnerRadius", FGeometryCollection::GeometryGroup);
+	InCollection.AddAttribute<float>("OuterRadius", FGeometryCollection::GeometryGroup);
+	InCollection.AddAttribute<int32>("VertexStart", FGeometryCollection::GeometryGroup, VerticesDependency);
+	InCollection.AddAttribute<int32>("VertexCount", FGeometryCollection::GeometryGroup);
+	InCollection.AddAttribute<int32>("FaceStart", FGeometryCollection::GeometryGroup, FacesDependency);
+	InCollection.AddAttribute<int32>("FaceCount", FGeometryCollection::GeometryGroup);
+
+	// Material Group
+	InCollection.AddAttribute<FGeometryCollectionSection>("Sections", FGeometryCollection::MaterialGroup, FacesDependency);
+}
 
 void FGeometryCollection::Construct()
 {
@@ -365,24 +408,42 @@ void FGeometryCollection::ReindexExemplarIndices(TArray<int32>& SortedRemovedInd
 // This will rebuild all mesh sections
 void FGeometryCollection::ReindexMaterials()
 {
+	FGeometryCollection::ReindexMaterials(*this);
+}
+
+void FGeometryCollection::ReindexMaterials(FManagedArrayCollection& InCollection)
+{
+
+	if (!InCollection.HasAttribute("MaterialID", FGeometryCollection::FacesGroup) ||
+		!InCollection.HasAttribute("Sections", FGeometryCollection::MaterialGroup) ||
+		!InCollection.HasAttribute("MaterialIndex", FGeometryCollection::FacesGroup))
+	{
+		return;
+	}
+
+	TManagedArray<int32>& MaterialID = InCollection.ModifyAttribute<int32>("MaterialID", FGeometryCollection::FacesGroup);
+	TManagedArray<FGeometryCollectionSection>& Sections = InCollection.ModifyAttribute<FGeometryCollectionSection>("Sections", FGeometryCollection::MaterialGroup);
+	TManagedArray<int32>& MaterialIndex = InCollection.ModifyAttribute<int32>("MaterialIndex", FGeometryCollection::FacesGroup);
+
 	// clear all sections	
 	TArray<int32> DelSections;
-	GeometryCollectionAlgo::ContiguousArray(DelSections, NumElements(FGeometryCollection::MaterialGroup));
-	Super::RemoveElements(FGeometryCollection::MaterialGroup, DelSections);
+	GeometryCollectionAlgo::ContiguousArray(DelSections, InCollection.NumElements(FGeometryCollection::MaterialGroup));
+	InCollection.RemoveElements(FGeometryCollection::MaterialGroup, DelSections);
 	DelSections.Reset(0);
 
 
 	// rebuild sections		
 
 	// count the number of triangles for each material section, adding a new section if the material ID is higher than the current number of sections
-	for (int FaceElement = 0, nf = NumElements(FGeometryCollection::FacesGroup) ; FaceElement < nf ; ++FaceElement)
+	
+	for (int FaceElement = 0, nf = InCollection.NumElements(FGeometryCollection::FacesGroup); FaceElement < nf ; ++FaceElement)
 	{
 		int32 Section = MaterialID[FaceElement];
 
-		while (Section + 1 > NumElements(FGeometryCollection::MaterialGroup))
+		while (Section + 1 > InCollection.NumElements(FGeometryCollection::MaterialGroup))
 		{
 			// add a new material section
-			int32 Element = AddElements(1, FGeometryCollection::MaterialGroup);
+			int32 Element = InCollection.AddElements(1, FGeometryCollection::MaterialGroup);
 
 			Sections[Element].MaterialID = Element;
 			Sections[Element].FirstIndex = -1;
@@ -395,7 +456,7 @@ void FGeometryCollection::ReindexMaterials()
 	}
 
 	// fixup the section FirstIndex and MaxVertexIndex
-	for (int SectionElement = 0; SectionElement < NumElements(FGeometryCollection::MaterialGroup); SectionElement++)
+	for (int SectionElement = 0; SectionElement < InCollection.NumElements(FGeometryCollection::MaterialGroup); SectionElement++)
 	{
 		if (SectionElement == 0)
 		{
@@ -408,7 +469,7 @@ void FGeometryCollection::ReindexMaterials()
 			Sections[SectionElement].FirstIndex = Sections[SectionElement - 1].FirstIndex + Sections[SectionElement - 1].NumTriangles * 3;
 		}
 
-		Sections[SectionElement].MaxVertexIndex = NumElements(FGeometryCollection::VerticesGroup) - 1;
+		Sections[SectionElement].MaxVertexIndex = InCollection.NumElements(FGeometryCollection::VerticesGroup) - 1;
 
 		// if a material group no longer has any triangles in it then add material section for removal
 		if (Sections[SectionElement].NumTriangles == 0)
@@ -419,9 +480,9 @@ void FGeometryCollection::ReindexMaterials()
 
 	// remap indices so the materials appear to be grouped
 	int Idx = 0;
-	for (int Section=0; Section < NumElements(FGeometryCollection::MaterialGroup); Section++)
+	for (int Section=0; Section < InCollection.NumElements(FGeometryCollection::MaterialGroup); Section++)
 	{
-		for (int FaceElement = 0; FaceElement < NumElements(FGeometryCollection::FacesGroup); FaceElement++)
+		for (int FaceElement = 0; FaceElement < InCollection.NumElements(FGeometryCollection::FacesGroup); FaceElement++)
 		{
 			int32 ID = (MaterialID)[FaceElement];
 	
@@ -435,7 +496,7 @@ void FGeometryCollection::ReindexMaterials()
 	// delete unused material sections
 	if (DelSections.Num())
 	{
-		Super::RemoveElements(FGeometryCollection::MaterialGroup, DelSections);
+		InCollection.RemoveElements(FGeometryCollection::MaterialGroup, DelSections);
 	}
 }
 
@@ -920,6 +981,25 @@ bool FGeometryCollection::HasVisibleGeometry() const
 
 void FGeometryCollection::UpdateBoundingBox()
 {
+	FGeometryCollection::UpdateBoundingBox(*this, /*bSkipCheck*/ true);
+}
+
+void FGeometryCollection::UpdateBoundingBox(FManagedArrayCollection& InCollection, bool bSkipCheck)
+{
+	if (!bSkipCheck && (
+		!InCollection.HasAttribute("BoundingBox", FGeometryCollection::GeometryGroup) ||
+		!InCollection.HasAttribute("Vertex", FGeometryCollection::VerticesGroup) ||
+		!InCollection.HasAttribute("BoneMap", FGeometryCollection::VerticesGroup) ||
+		!InCollection.HasAttribute("TransformToGeometryIndex", FTransformCollection::TransformGroup))) 
+	{
+		return;
+	}
+
+	TManagedArray<FBox>& BoundingBox = InCollection.ModifyAttribute<FBox>("BoundingBox", FGeometryCollection::GeometryGroup);
+	const TManagedArray<FVector3f>& Vertex = InCollection.GetAttribute<FVector3f>("Vertex", FGeometryCollection::VerticesGroup);
+	const TManagedArray<int32>& BoneMap = InCollection.GetAttribute<int32>("BoneMap", FGeometryCollection::VerticesGroup);
+	const TManagedArray<int32>& TransformToGeometryIndex = InCollection.GetAttribute<int32>("TransformToGeometryIndex", FTransformCollection::TransformGroup);
+
 	if (BoundingBox.Num())
 	{
 		// Initialize BoundingBox
@@ -937,6 +1017,7 @@ void FGeometryCollection::UpdateBoundingBox()
 	}
 }
 
+
 FBoxSphereBounds FGeometryCollection::GetBoundingBox() const 
 {
 	TArray<FTransform> GlobalTransformArray;
diff --git a/Engine/Source/Runtime/Experimental/Chaos/Private/GeometryCollection/GeometryCollectionUtility.cpp b/Engine/Source/Runtime/Experimental/Chaos/Private/GeometryCollection/GeometryCollectionUtility.cpp
--- a/Engine/Source/Runtime/Experimental/Chaos/Private/GeometryCollection/GeometryCollectionUtility.cpp
+++ b/Engine/Source/Runtime/Experimental/Chaos/Private/GeometryCollection/GeometryCollectionUtility.cpp
@@ -275,33 +275,33 @@ namespace GeometryCollection
 		OutOuterRadius = FMath::Sqrt(OutOuterRadius);
 	}
 	
-	void AddGeometryProperties(FGeometryCollection * Collection)
+	void AddGeometryProperties(FManagedArrayCollection* InCollection)
 	{
-		if (Collection)
+		if (InCollection)
 		{
-			if (!Collection->NumElements(FGeometryCollection::GeometryGroup))
+			if (!InCollection->NumElements(FGeometryCollection::GeometryGroup))
 			{
-				int32 NumVertices = Collection->Vertex.Num();
+				int32 NumVertices = InCollection->NumElements(FGeometryCollection::VerticesGroup);
 				if (NumVertices)
 				{
 					// transforms group
-					TManagedArray<FTransform>& Transform = Collection->Transform;
-					TManagedArray<int32>& TransformToGeometryIndex = Collection->TransformToGeometryIndex;
+					TManagedArray<FTransform>& Transform = InCollection->ModifyAttribute<FTransform>(FTransformCollection::TransformAttribute, FTransformCollection::TransformGroup);
+					TManagedArray<int32>& TransformToGeometryIndex = InCollection->ModifyAttribute<int32>("TransformToGeometryIndex", FTransformCollection::TransformGroup);
 					// vertices group
-					TManagedArray<int32> & BoneMap = Collection->BoneMap;
-					TManagedArray<FVector3f> & Vertex = Collection->Vertex;
+					TManagedArray<int32>& BoneMap = InCollection->ModifyAttribute<int32>("BoneMap", FGeometryCollection::VerticesGroup);
+					TManagedArray<FVector3f>& Vertex = InCollection->ModifyAttribute<FVector3f>("Vertex", FGeometryCollection::VerticesGroup);
 					// faces
-					TManagedArray<FIntVector> & FaceIndices = Collection->Indices;
+					TManagedArray<FIntVector>& FaceIndices = InCollection->ModifyAttribute<FIntVector>("Indices", FGeometryCollection::FacesGroup);
 
 					// geometry group
-					TManagedArray<int32>& TransformIndex = Collection->TransformIndex;
-					TManagedArray<FBox>& BoundingBox = Collection->BoundingBox;
-					TManagedArray<float>& InnerRadius = Collection->InnerRadius;
-					TManagedArray<float>& OuterRadius = Collection->OuterRadius;
-					TManagedArray<int32>& VertexCount = Collection->VertexCount;
-					TManagedArray<int32>& VertexStart = Collection->VertexStart;
-					TManagedArray<int32>& FaceCount = Collection->FaceCount;
-					TManagedArray<int32>& FaceStart = Collection->FaceStart;
+					TManagedArray<int32>& TransformIndex = InCollection->ModifyAttribute<int32>("TransformIndex", FGeometryCollection::GeometryGroup);
+					TManagedArray<FBox>& BoundingBox = InCollection->ModifyAttribute<FBox>("BoundingBox", FGeometryCollection::GeometryGroup);
+					TManagedArray<float>& InnerRadius = InCollection->ModifyAttribute<float>("InnerRadius", FGeometryCollection::GeometryGroup);
+					TManagedArray<float>& OuterRadius = InCollection->ModifyAttribute<float>("OuterRadius", FGeometryCollection::GeometryGroup);
+					TManagedArray<int32>& VertexCount = InCollection->ModifyAttribute<int32>("VertexCount", FGeometryCollection::GeometryGroup);
+					TManagedArray<int32>& VertexStart = InCollection->ModifyAttribute<int32>("VertexStart", FGeometryCollection::GeometryGroup);
+					TManagedArray<int32>& FaceCount = InCollection->ModifyAttribute<int32>("FaceCount", FGeometryCollection::GeometryGroup);
+					TManagedArray<int32>& FaceStart = InCollection->ModifyAttribute<int32>("FaceStart", FGeometryCollection::GeometryGroup);
 
 					// gather unique geometries
 					TSet<int32> TransformIndexOfGeometry;
@@ -314,7 +314,7 @@ namespace GeometryCollection
 					TArray<int32> ReverseMap;
 					ReverseMap.Init(FGeometryCollection::Invalid, Transform.Num());
 
-					Collection->AddElements(TransformIndexOfGeometry.Num(), FGeometryCollection::GeometryGroup);
+					InCollection->AddElements(TransformIndexOfGeometry.Num(), FGeometryCollection::GeometryGroup);
 					TArray<int32> GeometryIndices = TransformIndexOfGeometry.Array();
 					for (int32 Index = 0; Index < GeometryIndices.Num(); Index++)
 					{
@@ -440,7 +440,7 @@ namespace GeometryCollection
 					}
 				}
 
-				Collection->UpdateBoundingBox();
+				FGeometryCollection::UpdateBoundingBox(*InCollection);
 			}
 		}
 	}
diff --git a/Engine/Source/Runtime/Experimental/Chaos/Private/GeometryCollection/TransformCollection.cpp b/Engine/Source/Runtime/Experimental/Chaos/Private/GeometryCollection/TransformCollection.cpp
--- a/Engine/Source/Runtime/Experimental/Chaos/Private/GeometryCollection/TransformCollection.cpp
+++ b/Engine/Source/Runtime/Experimental/Chaos/Private/GeometryCollection/TransformCollection.cpp
@@ -13,13 +13,22 @@ const FName FTransformCollection::ParentAttribute = "Parent";
 const FName FTransformCollection::ChildrenAttribute = "Children";
 const FName FTransformCollection::ParticlesAttribute = "Particles";
 
-
 FTransformCollection::FTransformCollection()
 	: FManagedArrayCollection()
 {
 	Construct();
 }
 
+void FTransformCollection::DefineTransformSchema(FManagedArrayCollection& InCollection)
+{
+	// Hierarchy Group
+	InCollection.AddAttribute<FTransform>(FTransformCollection::TransformAttribute, FTransformCollection::TransformGroup);
+	InCollection.AddAttribute<FString>("BoneName", FTransformCollection::TransformGroup);
+	InCollection.AddAttribute<FLinearColor>("BoneColor", FTransformCollection::TransformGroup);
+	InCollection.AddAttribute<int32>(FTransformCollection::ParentAttribute, FTransformCollection::TransformGroup);
+	InCollection.AddAttribute<TSet<int32>>(FTransformCollection::ChildrenAttribute, FTransformCollection::TransformGroup);
+}
+
 void FTransformCollection::Construct()
 {
 	FManagedArrayCollection::FConstructionParameters TransformDependency(FTransformCollection::TransformGroup);
diff --git a/Engine/Source/Runtime/Experimental/Chaos/Public/GeometryCollection/GeometryCollection.h b/Engine/Source/Runtime/Experimental/Chaos/Public/GeometryCollection/GeometryCollection.h
--- a/Engine/Source/Runtime/Experimental/Chaos/Public/GeometryCollection/GeometryCollection.h
+++ b/Engine/Source/Runtime/Experimental/Chaos/Public/GeometryCollection/GeometryCollection.h
@@ -131,6 +131,7 @@ public:
 	 */
 	static FGeometryCollection* NewGeometryCollection(const TArray<float>& RawVertexArray, const TArray<int32>& RawIndicesArray, bool ReverseVertexOrder = true);
 	static void Init(FGeometryCollection* Collection, const TArray<float>& RawVertexArray, const TArray<int32>& RawIndicesArray, bool ReverseVertexOrder = true);
+	static void DefineGeometryScheam(FManagedArrayCollection&);
 
 	/**
 	* Create a GeometryCollection from Vertex, Indices, BoneMap, Transform, BoneHierarchy arrays
@@ -192,6 +193,7 @@ public:
 	*  Update bounding box entries for the geometry
 	*/
 	void UpdateBoundingBox();
+	static void UpdateBoundingBox(FManagedArrayCollection&, bool bSkipCheck=false);
 
 	/**  
 	* GetBoundingBox 
@@ -207,7 +209,8 @@ public:
 	* Reindex sections to keep polys with same materials together to reduce the number of draw calls
 	*/
 	void ReindexMaterials();
-		
+	static void ReindexMaterials(FManagedArrayCollection&);
+
 	/**
 	* Builds mesh sections for a given index buffer that could be a subset.
 	* Currently, this call assumes that the indices are ordered by MaterialID
diff --git a/Engine/Source/Runtime/Experimental/Chaos/Public/GeometryCollection/GeometryCollectionUtility.h b/Engine/Source/Runtime/Experimental/Chaos/Public/GeometryCollection/GeometryCollectionUtility.h
--- a/Engine/Source/Runtime/Experimental/Chaos/Public/GeometryCollection/GeometryCollectionUtility.h
+++ b/Engine/Source/Runtime/Experimental/Chaos/Public/GeometryCollection/GeometryCollectionUtility.h
@@ -14,7 +14,7 @@ namespace GeometryCollection
 	*/
 	void
 	CHAOS_API
-	AddGeometryProperties(FGeometryCollection* Collection);
+	AddGeometryProperties(FManagedArrayCollection* Collection);
 
 	/****
 	* MakeMeshElement
diff --git a/Engine/Source/Runtime/Experimental/Chaos/Public/GeometryCollection/TransformCollection.h b/Engine/Source/Runtime/Experimental/Chaos/Public/GeometryCollection/TransformCollection.h
--- a/Engine/Source/Runtime/Experimental/Chaos/Public/GeometryCollection/TransformCollection.h
+++ b/Engine/Source/Runtime/Experimental/Chaos/Public/GeometryCollection/TransformCollection.h
@@ -24,7 +24,6 @@ public:
 	FTransformCollection(FTransformCollection&&) = default;
 	FTransformCollection& operator=(FTransformCollection&&) = default;
 
-
 	/***
 	*  Attribute Groups
 	*
@@ -51,6 +50,10 @@ public:
 	static const FName ChildrenAttribute;
 	static const FName ParticlesAttribute;
 
+	/** Schema */
+	static void DefineTransformSchema(FManagedArrayCollection&);
+
+
 	/** Serialize */
 	virtual void Serialize(Chaos::FChaosArchive& Ar) override;
 
diff --git a/Engine/Source/Runtime/Experimental/GeometryCollectionEngine/Private/GeometryCollection/GeometryCollectionEngineConversion.cpp b/Engine/Source/Runtime/Experimental/GeometryCollectionEngine/Private/GeometryCollection/GeometryCollectionEngineConversion.cpp
--- a/Engine/Source/Runtime/Experimental/GeometryCollectionEngine/Private/GeometryCollection/GeometryCollectionEngineConversion.cpp
+++ b/Engine/Source/Runtime/Experimental/GeometryCollectionEngine/Private/GeometryCollection/GeometryCollectionEngineConversion.cpp
@@ -912,13 +912,14 @@ void FGeometryCollectionEngineConversion::AppendGeometryCollection(const UGeomet
 }
 
 
-bool FGeometryCollectionEngineConversion::AppendSkeletalMesh(const USkeletalMesh* SkeletalMesh, int32 MaterialStartIndex, const FTransform& SkeletalMeshTransform, FGeometryCollection* GeometryCollection, bool bReindexMaterials)
+bool FGeometryCollectionEngineConversion::AppendSkeletalMesh(const USkeletalMesh* SkeletalMesh, int32 MaterialStartIndex, const FTransform& SkeletalMeshTransform, FManagedArrayCollection* InCollection, bool bReindexMaterials)
 {
 	//UE_LOG(UGeometryCollectionConversionLogging, Log, TEXT("FGeometryCollectionEngineConversion::AppendSkeletalMesh()"));
-	if (!GeometryCollection)
+	if (!InCollection)
 	{
 		return false;
 	}
+	FGeometryCollection::DefineGeometryScheam(*InCollection);
 
 	const FSkeletalMeshLODRenderData* MeshLODData = GetSkeletalMeshLOD(SkeletalMesh, 0);
 	if (!MeshLODData)
@@ -954,8 +955,8 @@ bool FGeometryCollectionEngineConversion::AppendSkeletalMesh(const USkeletalMesh
 	// sections.
 	//
 	const USkeleton* Skeleton = SkeletalMesh->GetSkeleton();
-	TManagedArray<FTransform>& Transform = GeometryCollection->Transform;
-	int32 TransformBaseIndex = GeometryCollection->AddElements(SkeletalBoneMap.Num(), FGeometryCollection::TransformGroup);
+	TManagedArray<FTransform>& Transform = InCollection->ModifyAttribute<FTransform>(FTransformCollection::TransformAttribute, FTransformCollection::TransformGroup);
+	int32 TransformBaseIndex = InCollection->AddElements(SkeletalBoneMap.Num(), FGeometryCollection::TransformGroup);
 	const FReferenceSkeleton & ReferenceSkeletion = Skeleton->GetReferenceSkeleton();
 	const TArray<FTransform> & RestArray = Skeleton->GetRefLocalPoses();
 	for (int32 BoneIndex = 0; BoneIndex < SkeletalBoneMap.Num(); BoneIndex++)
@@ -968,18 +969,18 @@ bool FGeometryCollectionEngineConversion::AppendSkeletalMesh(const USkeletalMesh
 	//
 	// The Triangle Indices
 	//
-	TManagedArray<FIntVector>& Indices = GeometryCollection->Indices;
-	TManagedArray<bool>& Visible = GeometryCollection->Visible;
-	TManagedArray<int32>& MaterialID = GeometryCollection->MaterialID;
-	TManagedArray<int32>& MaterialIndex = GeometryCollection->MaterialIndex;
+	TManagedArray<FIntVector>& Indices = InCollection->ModifyAttribute<FIntVector>("Indices", FGeometryCollection::FacesGroup);
+	TManagedArray<bool>& Visible = InCollection->ModifyAttribute<bool>("Visible", FGeometryCollection::FacesGroup);
+	TManagedArray<int32>& MaterialIndex = InCollection->ModifyAttribute<int32>("MaterialIndex", FGeometryCollection::FacesGroup); 
+	TManagedArray<int32>& MaterialID = InCollection->ModifyAttribute<int32>("MaterialID", FGeometryCollection::FacesGroup);
 
 	TArray<uint32> IndexBuffer;
 	SkeletalMeshLODRenderData.MultiSizeIndexContainer.GetIndexBuffer(IndexBuffer);
 
 	const int32 IndicesCount = IndexBuffer.Num() / 3;
-	int NumVertices = GeometryCollection->NumElements(FGeometryCollection::VerticesGroup);
-	int InitialNumIndices = GeometryCollection->NumElements(FGeometryCollection::FacesGroup);
-	int IndicesBaseIndex = GeometryCollection->AddElements(IndicesCount, FGeometryCollection::FacesGroup);
+	int NumVertices = InCollection->NumElements(FGeometryCollection::VerticesGroup);
+	int InitialNumIndices = InCollection->NumElements(FGeometryCollection::FacesGroup);
+	int IndicesBaseIndex = InCollection->AddElements(IndicesCount, FGeometryCollection::FacesGroup);
 	for (int32 IndicesIndex = 0, StaticIndex = 0; IndicesIndex < IndicesCount; IndicesIndex++, StaticIndex += 3)
 	{
 		int32 IndicesOffset = IndicesBaseIndex + IndicesIndex;
@@ -995,23 +996,23 @@ bool FGeometryCollectionEngineConversion::AppendSkeletalMesh(const USkeletalMesh
 	//
 	// Vertex Attributes
 	//
-	TManagedArray<FVector3f>& Vertex = GeometryCollection->Vertex;
-	TManagedArray<FVector3f>& TangentU = GeometryCollection->TangentU;
-	TManagedArray<FVector3f>& TangentV = GeometryCollection->TangentV;
-	TManagedArray<FVector3f>& Normal = GeometryCollection->Normal;
-	TManagedArray<TArray<FVector2f>>& UVs = GeometryCollection->UVs;
-	TManagedArray<FLinearColor>& Color = GeometryCollection->Color;
-	TManagedArray<int32>& BoneMap = GeometryCollection->BoneMap;
-	TManagedArray<FLinearColor>& BoneColor = GeometryCollection->BoneColor;
-	TManagedArray<FString>& BoneName = GeometryCollection->BoneName;
+	TManagedArray<FVector3f>& Vertex = InCollection->ModifyAttribute<FVector3f>("Vertex", FGeometryCollection::VerticesGroup);
+	TManagedArray<FVector3f>& Normal = InCollection->ModifyAttribute<FVector3f>("Normal", FGeometryCollection::VerticesGroup);
+	TManagedArray<TArray<FVector2f>>& UVs = InCollection->ModifyAttribute<TArray<FVector2f>>("UVs", FGeometryCollection::VerticesGroup);
+	TManagedArray<FLinearColor>& Color = InCollection->ModifyAttribute<FLinearColor>("Color", FGeometryCollection::VerticesGroup);
+	TManagedArray<FVector3f>& TangentU = InCollection->ModifyAttribute<FVector3f>("TangentU", FGeometryCollection::VerticesGroup);
+	TManagedArray<FVector3f>& TangentV = InCollection->ModifyAttribute<FVector3f>("TangentV", FGeometryCollection::VerticesGroup);
+	TManagedArray<int32>& BoneMap = InCollection->ModifyAttribute<int32>("BoneMap", FGeometryCollection::VerticesGroup);
+	TManagedArray<FLinearColor>& BoneColor = InCollection->ModifyAttribute<FLinearColor>("BoneColor", FTransformCollection::TransformGroup);
+	TManagedArray<FString>& BoneName = InCollection->ModifyAttribute<FString>("BoneName", FTransformCollection::TransformGroup);
 
 	// 
 	// Transform Attributes 
 	// 
-	TManagedArray<int32>& Parent = GeometryCollection->Parent;
-	TManagedArray<int32>& SimulationType = GeometryCollection->SimulationType;
-	int InitialNumVertices = GeometryCollection->NumElements(FGeometryCollection::VerticesGroup);
-	int VertexBaseIndex = GeometryCollection->AddElements(VertexCount, FGeometryCollection::VerticesGroup);
+	TManagedArray<int32>& Parent = InCollection->ModifyAttribute<int32>(FTransformCollection::ParentAttribute, FTransformCollection::TransformGroup);
+	TManagedArray<int32>& SimulationType = InCollection->ModifyAttribute<int32>("SimulationType", FTransformCollection::TransformGroup);
+	int InitialNumVertices = InCollection->NumElements(FGeometryCollection::VerticesGroup);
+	int VertexBaseIndex = InCollection->AddElements(VertexCount, FGeometryCollection::VerticesGroup);
 	const int32 NumUVLayers = VertexBuffers.StaticMeshVertexBuffer.GetNumTexCoords();
 	for (int32 VertexIndex = 0; VertexIndex < VertexCount; VertexIndex++)
 	{
@@ -1028,13 +1029,13 @@ bool FGeometryCollectionEngineConversion::AppendSkeletalMesh(const USkeletalMesh
 		TangentU[VertexOffset] = VertexBuffers.StaticMeshVertexBuffer.VertexTangentX(VertexIndex);
 		TangentV[VertexOffset] = VertexBuffers.StaticMeshVertexBuffer.VertexTangentY(VertexIndex);
 		Normal[VertexOffset] = VertexBuffers.StaticMeshVertexBuffer.VertexTangentZ(VertexIndex);
-							
+
 		UVs[VertexOffset].SetNum(NumUVLayers);
 		for (int32 UVLayerIdx = 0; UVLayerIdx < NumUVLayers; ++UVLayerIdx)
 		{
 			UVs[VertexOffset][UVLayerIdx] = VertexBuffers.StaticMeshVertexBuffer.GetVertexUV(VertexIndex, UVLayerIdx);
 		}
-							
+
 		if (VertexBuffers.ColorVertexBuffer.GetNumVertices() == VertexCount)
 			Color[VertexOffset] = VertexBuffers.ColorVertexBuffer.VertexColor(VertexIndex);
 		else
@@ -1067,19 +1068,19 @@ bool FGeometryCollectionEngineConversion::AppendSkeletalMesh(const USkeletalMesh
 
 	// Geometry Group
 	TArray<int32> GeometryIndices;
-	GeometryCollectionAlgo::ContiguousArray(GeometryIndices, GeometryCollection->NumElements(FGeometryCollection::GeometryGroup));
-	GeometryCollection->RemoveDependencyFor(FGeometryCollection::GeometryGroup);
-	GeometryCollection->RemoveElements(FGeometryCollection::GeometryGroup, GeometryIndices);
-	::GeometryCollection::AddGeometryProperties(GeometryCollection);
+	GeometryCollectionAlgo::ContiguousArray(GeometryIndices, InCollection->NumElements(FGeometryCollection::GeometryGroup));
+	InCollection->RemoveDependencyFor(FGeometryCollection::GeometryGroup);
+	InCollection->RemoveElements(FGeometryCollection::GeometryGroup, GeometryIndices);
+	::GeometryCollection::AddGeometryProperties(InCollection);
 
 	const TArray<FSkelMeshRenderSection> &StaticMeshSections = SkeletalMesh->GetResourceForRendering()->LODRenderData[0].RenderSections;
 
-	TManagedArray<FGeometryCollectionSection> & Sections = GeometryCollection->Sections;
+	TManagedArray<FGeometryCollectionSection> & Sections = InCollection->ModifyAttribute<FGeometryCollectionSection>("Sections", FGeometryCollection::MaterialGroup);
 
 	for (const FSkelMeshRenderSection &CurrSection : StaticMeshSections)
 	{
 		// create new section
-		int32 SectionIndex = GeometryCollection->AddElements(1, FGeometryCollection::MaterialGroup);
+		int32 SectionIndex = InCollection->AddElements(1, FGeometryCollection::MaterialGroup);
 						
 		Sections[SectionIndex].MaterialID = MaterialStartIndex + CurrSection.MaterialIndex;
 
@@ -1101,7 +1102,7 @@ bool FGeometryCollectionEngineConversion::AppendSkeletalMesh(const USkeletalMesh
 
 	if (bReindexMaterials)
 	{
-		GeometryCollection->ReindexMaterials();
+		FGeometryCollection::ReindexMaterials(*InCollection);
 	}
 
 	return true;
diff --git a/Engine/Source/Runtime/Experimental/GeometryCollectionEngine/Public/GeometryCollection/GeometryCollectionEngineConversion.h b/Engine/Source/Runtime/Experimental/GeometryCollectionEngine/Public/GeometryCollection/GeometryCollectionEngineConversion.h
--- a/Engine/Source/Runtime/Experimental/GeometryCollectionEngine/Public/GeometryCollection/GeometryCollectionEngineConversion.h
+++ b/Engine/Source/Runtime/Experimental/GeometryCollectionEngine/Public/GeometryCollection/GeometryCollectionEngineConversion.h
@@ -10,6 +10,7 @@ class UGeometryCollection;
 class UMaterialInterface;
 class UGeometryCollectionComponent;
 class FGeometryCollection;
+struct FManagedArrayCollection;
 class UBodySetup;
 class FSkeletalMeshLODRenderData;
 struct FMeshDescription;
@@ -103,7 +104,7 @@ public:
 	*  @param SkeletalMeshTransform : Mesh transform.
 	*  @param GeometryCollection    : Collection to append the mesh into.
 	*/
-	static bool AppendSkeletalMesh(const USkeletalMesh* SkeletalMesh, int32 MaterialStartIndex, const FTransform& SkeletalMeshTransform, FGeometryCollection* GeometryCollection, bool bReindexMaterials = true);
+	static bool AppendSkeletalMesh(const USkeletalMesh* SkeletalMesh, int32 MaterialStartIndex, const FTransform& SkeletalMeshTransform, FManagedArrayCollection* InCollection, bool bReindexMaterials = true);
 
 	/**
 	*  Appends a skeletal mesh to a GeometryCollectionComponent.
