Chaos - fix NaN in CullDistance threshold calculation for unbounded objects

diff --git a/Engine/Source/Runtime/Experimental/Chaos/Private/Chaos/Collision/ParticlePairMidPhase.cpp b/Engine/Source/Runtime/Experimental/Chaos/Private/Chaos/Collision/ParticlePairMidPhase.cpp
--- a/Engine/Source/Runtime/Experimental/Chaos/Private/Chaos/Collision/ParticlePairMidPhase.cpp
+++ b/Engine/Source/Runtime/Experimental/Chaos/Private/Chaos/Collision/ParticlePairMidPhase.cpp
@@ -905,8 +905,10 @@ namespace Chaos
 	{
 		// @todo(chaos): improve this threshold calculation for thin objects? Dynamic thin objects have bigger problems so maybe we don't care
 		// @todo(chaos): Spheres and capsules need smaller position tolerance - the restore test doesn't work well with rolling
-		bool bIsDynamic0 = FConstGenericParticleHandle(Particle0)->IsDynamic();
-		bool bIsDynamic1 = FConstGenericParticleHandle(Particle1)->IsDynamic();
+		const bool bIsDynamic0 = FConstGenericParticleHandle(Particle0)->IsDynamic();
+		const bool bIsDynamic1 = FConstGenericParticleHandle(Particle1)->IsDynamic();
+		const bool bIsBounded0 = Particle0->HasBounds();
+		const bool bIsBounded1 = Particle1->HasBounds();
 
 		// NOTE: If CullDistance ends up smaller than the thresholds used to restore collisions, we can end up missing
 		// collisions as the objects move if we restore a "zero contact" manifold after movement greater than the cull distance. 
@@ -914,8 +916,8 @@ namespace Chaos
 		// @todo(chaos): Add a way to enforce a CullDistance big enough to support the reuse thresholds
 		const FReal CullDistanceReferenceSizeInv = FReal(Chaos_Collision_CullDistanceScaleInverseSize);
 		const FReal MinCullDistanceScale = FReal(Chaos_Collision_MinCullDistanceScale);
-		const FReal MaxBoundsSize0 = bIsDynamic0 ? Particle0->LocalBounds().Extents().GetMax() : FReal(0);
-		const FReal MaxBoundsSize1 = bIsDynamic1 ? Particle1->LocalBounds().Extents().GetMax() : FReal(0);
+		const FReal MaxBoundsSize0 = (bIsDynamic0 && bIsBounded0) ? Particle0->LocalBounds().Extents().GetMax() : FReal(0);
+		const FReal MaxBoundsSize1 = (bIsDynamic1 && bIsBounded1) ? Particle1->LocalBounds().Extents().GetMax() : FReal(0);
 		const FReal CullDistanceScale0 = MaxBoundsSize0 * CullDistanceReferenceSizeInv;
 		const FReal CullDistanceScale1 = MaxBoundsSize1 * CullDistanceReferenceSizeInv;
 		CullDistanceScale = (FRealSingle)FMath::Max3(CullDistanceScale0, CullDistanceScale1, MinCullDistanceScale);
