Chaos geometry collection : add asset side ability to use size specific data and apply propgate to instance when placing it in a level

diff --git a/Engine/Source/Runtime/Experimental/GeometryCollectionEngine/Private/GeometryCollection/GeometryCollectionComponent.cpp b/Engine/Source/Runtime/Experimental/GeometryCollectionEngine/Private/GeometryCollection/GeometryCollectionComponent.cpp
--- a/Engine/Source/Runtime/Experimental/GeometryCollectionEngine/Private/GeometryCollection/GeometryCollectionComponent.cpp
+++ b/Engine/Source/Runtime/Experimental/GeometryCollectionEngine/Private/GeometryCollection/GeometryCollectionComponent.cpp
@@ -2702,6 +2702,7 @@ void UGeometryCollectionComponent::SetRestCollection(const UGeometryCollection*
 
 		// initialize the component per level damage threshold from the asset defaults 
 		DamageThreshold = RestCollection->DamageThreshold;
+		bUseSizeSpecificDamageThreshold = RestCollection->bUseSizeSpecificDamageThreshold;
 
 		// initialize the component damage progataion data from the asset defaults 
 		DamagePropagationData = RestCollection->DamagePropagationData;
diff --git a/Engine/Source/Runtime/Experimental/GeometryCollectionEngine/Private/GeometryCollection/GeometryCollectionObject.cpp b/Engine/Source/Runtime/Experimental/GeometryCollectionEngine/Private/GeometryCollection/GeometryCollectionObject.cpp
--- a/Engine/Source/Runtime/Experimental/GeometryCollectionEngine/Private/GeometryCollection/GeometryCollectionObject.cpp
+++ b/Engine/Source/Runtime/Experimental/GeometryCollectionEngine/Private/GeometryCollection/GeometryCollectionObject.cpp
@@ -73,6 +73,7 @@ UGeometryCollection::UGeometryCollection(const FObjectInitializer& ObjectInitial
 	, ClusterGroupIndex(0)
 	, MaxClusterLevel(100)
 	, DamageThreshold({ 500000.f, 50000.f, 5000.f })
+	, bUseSizeSpecificDamageThreshold(false)
 	, PerClusterOnlyDamageThreshold(false)
 	, ClusterConnectionType(EClusterConnectionTypeEnum::Chaos_MinimalSpanningSubsetDelaunayTriangulation)
 	, ConnectionGraphBoundsFilteringMargin(0)
diff --git a/Engine/Source/Runtime/Experimental/GeometryCollectionEngine/Public/GeometryCollection/GeometryCollectionObject.h b/Engine/Source/Runtime/Experimental/GeometryCollectionEngine/Public/GeometryCollection/GeometryCollectionObject.h
--- a/Engine/Source/Runtime/Experimental/GeometryCollectionEngine/Public/GeometryCollection/GeometryCollectionObject.h
+++ b/Engine/Source/Runtime/Experimental/GeometryCollectionEngine/Public/GeometryCollection/GeometryCollectionObject.h
@@ -455,9 +455,13 @@ public:
 	int32 MaxClusterLevel;
 
 	/** Damage threshold for clusters at different levels. */
-	UPROPERTY(EditAnywhere, Category = "Damage")
+	UPROPERTY(EditAnywhere, Category = "Damage", meta = (EditCondition = "!bUseSizeSpecificDamageThreshold"))
 	TArray<float> DamageThreshold;
 
+	/** whether to use size specific damage threshold instead of level based ones ( see Size Specific Data array ). */
+	UPROPERTY(EditAnywhere, Category = "Damage")
+	bool bUseSizeSpecificDamageThreshold;
+
 	/** compatibility check, when true, only cluster compute damage from parameters and propagate to direct children
 	 *  when false, each child will compute it's damage threshold allowing for more precise and intuitive destruction behavior
 	 */
