Chaos - fix rigid particle natvis

diff --git a/Engine/Source/Runtime/Experimental/Chaos/Public/Chaos/RigidParticles.h b/Engine/Source/Runtime/Experimental/Chaos/Public/Chaos/RigidParticles.h
--- a/Engine/Source/Runtime/Experimental/Chaos/Public/Chaos/RigidParticles.h
+++ b/Engine/Source/Runtime/Experimental/Chaos/Public/Chaos/RigidParticles.h
@@ -58,6 +58,20 @@ constexpr int8 LowBitsMask(const int8 NumBits)
 // Count N, the number of bits needed to store an object state
 static constexpr int8 ObjectStateBitCount = NumBitsNeeded((int8)EObjectStateType::Count - (int8)1);
 
+// RigidParticle data that is commonly accessed together.
+// This contains all properties accessed in the broadphase filtering (FSpatialAccelerationBroadPhase)
+// NOTE: not a member class for easier natvis
+struct FRigidParticleCoreData
+{
+	int32 CollisionGroup;							// 4 bytes
+	uint32 CollisionConstraintFlags;				// 4 bytes
+	FRigidParticleControlFlags ControlFlags;		// 1 byte
+	FRigidParticleTransientFlags TransientFlags;	// 1 byte
+	EObjectStateType ObjectState;					// 1 byte
+	EObjectStateType PreObjectState;				// 1 byte
+	bool bDisabled;									// 1 byte
+};
+
 template<class T, int d>
 class TRigidParticles : public TKinematicGeometryParticles<T, d>
 {
@@ -325,20 +339,7 @@ public:
 	}
 
 private:
-	// Data that is commonly accessed together
-	// This contains all properties accessed in the broadphase filtering (FSpatialAccelerationBroadPhase)
-	struct FCoreData
-	{
-		int32 CollisionGroup;							// 4 bytes
-		uint32 CollisionConstraintFlags;				// 4 bytes
-		FRigidParticleControlFlags ControlFlags;		// 1 byte
-		FRigidParticleTransientFlags TransientFlags;	// 1 byte
-		EObjectStateType ObjectState;					// 1 byte
-		EObjectStateType PreObjectState;				// 1 byte
-		bool bDisabled;									// 1 byte
-	};
-
-	// Used during serialization to avoid bumping the file version as we switch to aggregated strunctures like FCoreData.
+	// Used during serialization to avoid bumping the file version as we switch to aggregated strunctures like FRigidParticleCoreData.
 	// Note: Only serialized data is needed here and not all data is serialized so some elements in the aggregates are not represented.
 	struct FLegacyData
 	{
@@ -346,7 +347,7 @@ private:
 		TArrayCollectionArray<EObjectStateType> MObjectState;
 		TArrayCollectionArray<bool> MDisabled;
 
-		void CopyFromCoreData(const TArrayCollectionArray<FCoreData>& Source)
+		void CopyFromCoreData(const TArrayCollectionArray<FRigidParticleCoreData>& Source)
 		{
 			MCollisionGroup.Resize(Source.Num());
 			MObjectState.Resize(Source.Num());
@@ -360,7 +361,7 @@ private:
 			}
 		}
 
-		void CopyToCoreData(TArrayCollectionArray<FCoreData>& Dest)
+		void CopyToCoreData(TArrayCollectionArray<FRigidParticleCoreData>& Dest)
 		{
 			Dest.Resize(MCollisionGroup.Num());
 
@@ -373,7 +374,7 @@ private:
 		}
 	};
 
-	TArrayCollectionArray<FCoreData> CoreData;
+	TArrayCollectionArray<FRigidParticleCoreData> CoreData;
 
 	TArrayCollectionArray<TVector<T, d>> MVSmooth;
 	TArrayCollectionArray<TVector<T, d>> MWSmooth;
