Iris
- Fixed issue where custom property overrides was not properly initialzied for iris if set before beginplay/beginreplication
- Cleaned up delegates propagating custom property overrides to iris-replicationsystem
- Marked interface methods in IRepChangedPropertyTracker to be deprecated in the future and added routing of SetCustomActive override via FNetPropertyConditionManager::SetPropertyActiveOverride to allow for future removal of the interface

diff --git a/Engine/Source/Runtime/CoreUObject/Public/UObject/CoreNet.h b/Engine/Source/Runtime/CoreUObject/Public/UObject/CoreNet.h
--- a/Engine/Source/Runtime/CoreUObject/Public/UObject/CoreNet.h
+++ b/Engine/Source/Runtime/CoreUObject/Public/UObject/CoreNet.h
@@ -54,6 +54,11 @@ namespace UE::Net
 	struct FReplicationStateDescriptor;
 	typedef FReplicationFragment* (*CreateAndRegisterReplicationFragmentFunc)(UObject* Owner, const FReplicationStateDescriptor* Descriptor, FFragmentRegistrationContext& Context);
 #endif
+
+	namespace Private
+	{
+		class FNetPropertyConditionManager;
+	};
 }
 
 
@@ -536,6 +541,7 @@ public:
 	IRepChangedPropertyTracker() { }
 	virtual ~IRepChangedPropertyTracker() { }
 
+	UE_DEPRECATED(5.3, "Please use FPropertyConditions::SetActiveOverride instead.")
 	virtual void SetCustomIsActiveOverride(UObject* OwningObject, const uint16 RepIndex, const bool bIsActive) = 0;
 
 	UE_DEPRECATED(5.0, "Please use UReplaySubsystem::SetExternalDataForObject instead.")
@@ -546,6 +552,13 @@ public:
 	* This should include the dynamically allocated data, as well as the classes size.
 	*/
 	virtual void CountBytes(FArchive& Ar) const {};
+
+private:
+	PRAGMA_DISABLE_DEPRECATION_WARNINGS
+	void CallSetCustomIsActiveOverride(UObject* OwningObject, const uint16 RepIndex, const bool bIsActive) { SetCustomIsActiveOverride(OwningObject, RepIndex, bIsActive); }
+	PRAGMA_ENABLE_DEPRECATION_WARNINGS
+	
+	friend UE::Net::Private::FNetPropertyConditionManager;
 };
 
 class FCustomPropertyConditionState
diff --git a/Engine/Source/Runtime/Engine/Private/Net/Iris/ReplicationSystem/ReplicationSystemUtil.cpp b/Engine/Source/Runtime/Engine/Private/Net/Iris/ReplicationSystem/ReplicationSystemUtil.cpp
--- a/Engine/Source/Runtime/Engine/Private/Net/Iris/ReplicationSystem/ReplicationSystemUtil.cpp
+++ b/Engine/Source/Runtime/Engine/Private/Net/Iris/ReplicationSystem/ReplicationSystemUtil.cpp
@@ -240,47 +240,6 @@ void FReplicationSystemUtil::SetActorComponentNetCondition(const UActorComponent
 	}
 }
 
-void FReplicationSystemUtil::SetPropertyCustomCondition(const UObject* Object, uint16 RepIndex, bool bIsActive)
-{
-	// If this is an actor component we can get the proper bridge through its owner, regardless of whether it's a default replicated component or not.
-	if (const UActorComponent* ActorComponent = Cast<UActorComponent>(Object))
-	{
-		const AActor* Owner = ActorComponent->GetOwner();
-		const UActorReplicationBridge* Bridge = GetActorReplicationBridge(Owner);
-		if (Bridge == nullptr)
-		{
-			return;
-		}
-
-		// Try to get a handle for the component. This can fail if we incorporate default components in the protocol for the actor.
-		FNetHandle NetHandle = Bridge->GetReplicatedHandle(Object);
-		if (!NetHandle.IsValid())
-		{
-			NetHandle = Bridge->GetReplicatedHandle(Owner);
-			if (!NetHandle.IsValid())
-			{
-				return;
-			}
-		}
-
-		Net::GetReplicationSystem(NetHandle.GetReplicationSystemId())->SetPropertyCustomCondition(NetHandle, Object, RepIndex, bIsActive);
-	}
-	else if (const AActor* Actor = Cast<AActor>(Object))
-	{
-		FNetHandle NetHandle = GetNetHandle(Actor);
-		if (!NetHandle.IsValid())
-		{
-			return;
-		}
-
-		Net::GetReplicationSystem(NetHandle.GetReplicationSystemId())->SetPropertyCustomCondition(NetHandle, Object, RepIndex, bIsActive);
-	}
-	else
-	{
-		checkf(false, TEXT("Unexpected class %s trying to set custom condition."), ToCStr(Object->GetFName().GetPlainNameString()));
-	}
-}
-
 void FReplicationSystemUtil::BeginReplicationForActorsInWorld(UWorld* World)
 {
 	// We only do this if the world already is initialized
diff --git a/Engine/Source/Runtime/Engine/Public/Net/Iris/ReplicationSystem/ReplicationSystemUtil.h b/Engine/Source/Runtime/Engine/Public/Net/Iris/ReplicationSystem/ReplicationSystemUtil.h
--- a/Engine/Source/Runtime/Engine/Public/Net/Iris/ReplicationSystem/ReplicationSystemUtil.h
+++ b/Engine/Source/Runtime/Engine/Public/Net/Iris/ReplicationSystem/ReplicationSystemUtil.h
@@ -106,11 +106,6 @@ struct FReplicationSystemUtil
 
 	/** Trigger replication of dirty state for actor wanting to be dormant. */
 	ENGINE_API static void FlushNetDormancy(AActor* Actor, bool bWasDormInitial);
-
-private:
-	friend FRepChangedPropertyTracker;
-
-	ENGINE_API static void SetPropertyCustomCondition(const UObject* Object, uint16 RepIndex, bool bIsActive);
 };
 
 }
diff --git a/Engine/Source/Runtime/Engine/Public/Net/UnrealNetwork.h b/Engine/Source/Runtime/Engine/Public/Net/UnrealNetwork.h
--- a/Engine/Source/Runtime/Engine/Public/Net/UnrealNetwork.h
+++ b/Engine/Source/Runtime/Engine/Public/Net/UnrealNetwork.h
@@ -294,7 +294,7 @@ static FProperty* GetReplicatedProperty(const UClass* CallingClass, const UClass
 #define DOREPLIFETIME_ACTIVE_OVERRIDE_FAST(c,v,active) \
 { \
 	static const bool bIsValid_##c_##v = ValidateReplicatedClassInheritance(StaticClass(), c::StaticClass(), TEXT(#v)); \
-	ChangedPropertyTracker.SetCustomIsActiveOverride(this, (int32)c::ENetFields_Private::v, active); \
+	UE::Net::Private::FNetPropertyConditionManager::SetPropertyActiveOverride(ChangedPropertyTracker, this, (int32)c::ENetFields_Private::v, active); \
 }
 
 #define DOREPLIFETIME_ACTIVE_OVERRIDE_FAST_STATIC_ARRAY(c,v,active) \
@@ -302,7 +302,7 @@ static FProperty* GetReplicatedProperty(const UClass* CallingClass, const UClass
 	static const bool bIsValid_##c_##v = ValidateReplicatedClassInheritance(StaticClass(), c::StaticClass(), TEXT(#v)); \
 	for (int32 i = 0; i < (int32)c::EArrayDims_Private::v; ++i) \
 	{ \
-		ChangedPropertyTracker.SetCustomIsActiveOverride(this, (int32)c::ENetFields_Private::v##_STATIC_ARRAY + i, active); \
+		UE::Net::Private::FNetPropertyConditionManager::SetPropertyActiveOverride(ChangedPropertyTracker, this, (int32)c::ENetFields_Private::v##_STATIC_ARRAY + i, active); \
 	} \
 }
 
@@ -311,7 +311,7 @@ static FProperty* GetReplicatedProperty(const UClass* CallingClass, const UClass
 	static FProperty* sp##v = GetReplicatedProperty(StaticClass(), c::StaticClass(),GET_MEMBER_NAME_CHECKED(c,v)); \
 	for (int32 i = 0; i < sp##v->ArrayDim; i++) \
 	{ \
-		ChangedPropertyTracker.SetCustomIsActiveOverride(this, sp##v->RepIndex + i, active); \
+		UE::Net::Private::FNetPropertyConditionManager::SetPropertyActiveOverride(ChangedPropertyTracker, this, sp##v->RepIndex + i, active); \
 	} \
 }
 
@@ -328,9 +328,10 @@ static FProperty* GetReplicatedProperty(const UClass* CallingClass, const UClass
 
 #define DOREPCUSTOMCONDITION_SETACTIVE_FAST_STATIC_ARRAY(c,v,active) \
 { \
+	UE::Net::Private::FNetPropertyConditionManager& PropertyConditionManager = UE::Net::Private::FNetPropertyConditionManager::Get(); \
 	for (int32 i = 0; i < (int32)c::EArrayDims_Private::v; ++i) \
 	{ \
-		UE::Net::Private::FNetPropertyConditionManager::Get().SetPropertyActive(this, (uint16)c::ENetFields_Private::v##_STATIC_ARRAY + i, active); \
+		PropertyConditionManager.SetPropertyActive(this, (uint16)c::ENetFields_Private::v##_STATIC_ARRAY + i, active); \
 	} \
 }
 
diff --git a/Engine/Source/Runtime/Experimental/Iris/Core/Private/Iris/IrisCoreModule.cpp b/Engine/Source/Runtime/Experimental/Iris/Core/Private/Iris/IrisCoreModule.cpp
--- a/Engine/Source/Runtime/Experimental/Iris/Core/Private/Iris/IrisCoreModule.cpp
+++ b/Engine/Source/Runtime/Experimental/Iris/Core/Private/Iris/IrisCoreModule.cpp
@@ -11,38 +11,10 @@
 
 #include "Misc/CommandLine.h"
 
-// $IRIS: $TODO: Remove when UE-158358 is completed
-#include "Iris/ReplicationSystem/ReplicationSystem.h"
-#include "Iris/ReplicationSystem/ObjectReplicationBridge.h"
-#include "Net/Core/PropertyConditions/RepChangedPropertyTracker.h"
-// End Remove when UE-158358 is completed
-
 class FIrisCoreModule : public IModuleInterface
 {
 private:
 
-	// $IRIS: $TODO: Make proper impl of support for SetPropertyCustomCondition - UE-158358
-	static void SetPropertyCustomCondition(const UObject* Owner, uint16 RepIndex, bool bEnable)
-	{
-		using namespace UE::Net;
-		// See if we can find the instance in any replication system
-		for (uint32 RepSystemIt = 0; RepSystemIt < FReplicationSystemFactory::MaxReplicationSystemCount; ++RepSystemIt)
-		{
-			UReplicationSystem* ReplicationSystem = GetReplicationSystem(RepSystemIt);
-			if (ReplicationSystem)
-			{
-				if (const UObjectReplicationBridge* Bridge = ReplicationSystem->GetReplicationBridgeAs<UObjectReplicationBridge>())
-				{
-					const FNetHandle Handle = Bridge->GetReplicatedHandle(Owner);
-					if (Handle.IsValid())
-					{
-						ReplicationSystem->SetPropertyCustomCondition(Handle, Owner, RepIndex, bEnable);
-					}
-				}
-			}
-		}
-	}
-
 	void RegisterPropertyNetSerializerSelectorTypes()
 	{
 		using namespace UE::Net;
@@ -73,8 +45,6 @@ private:
 		RegisterPropertyNetSerializerSelectorTypes();
 
 		UE_NET_IRIS_INIT_LEGACY_PUSH_MODEL();
-
-		UE::Net::Private::SetIrisSetPropertyCustomConditionDelegate(UE::Net::Private::FIrisSetPropertyCustomCondition::CreateStatic(&FIrisCoreModule::SetPropertyCustomCondition));
 		
 		ModulesChangedHandle = FModuleManager::Get().OnModulesChanged().AddRaw(this, &FIrisCoreModule::OnModulesChanged);
 	}
@@ -87,8 +57,6 @@ private:
 			ModulesChangedHandle.Reset();
 		}
 	
-		UE::Net::Private::SetIrisSetPropertyCustomConditionDelegate({});
-
 		UE_NET_IRIS_SHUTDOWN_LEGACY_PUSH_MODEL();
 	}
 
diff --git a/Engine/Source/Runtime/Experimental/Iris/Core/Private/Iris/ReplicationSystem/Conditionals/ReplicationConditionals.cpp b/Engine/Source/Runtime/Experimental/Iris/Core/Private/Iris/ReplicationSystem/Conditionals/ReplicationConditionals.cpp
--- a/Engine/Source/Runtime/Experimental/Iris/Core/Private/Iris/ReplicationSystem/Conditionals/ReplicationConditionals.cpp
+++ b/Engine/Source/Runtime/Experimental/Iris/Core/Private/Iris/ReplicationSystem/Conditionals/ReplicationConditionals.cpp
@@ -14,6 +14,8 @@
 #include "Iris/ReplicationSystem/Filtering/ReplicationFiltering.h"
 #include "Containers/ArrayView.h"
 #include "UObject/CoreNetTypes.h"
+#include "Net/Core/PropertyConditions/RepChangedPropertyTracker.h"
+#include "Net/Core/PropertyConditions/PropertyConditions.h"
 
 namespace UE::Net::Private
 {
@@ -108,6 +110,83 @@ bool FReplicationConditionals::SetCondition(FInternalNetHandle ObjectIndex, ERep
 	return false;
 }
 
+void FReplicationConditionals::InitPropertyCustomConditions(FInternalNetHandle ObjectIndex)
+{
+	IRIS_PROFILER_SCOPE(FReplicationConditionals_InitPropertyCustomConditions);
+
+	const FNetHandleManager::FReplicatedObjectData& ReplicatedObjectData = NetHandleManager->GetReplicatedObjectDataNoCheck(ObjectIndex);
+	const FReplicationProtocol* Protocol = ReplicatedObjectData.Protocol;
+
+	if (NetHandleManager->GetReplicatedObjectStateBufferNoCheck(ObjectIndex) == nullptr || !EnumHasAnyFlags(Protocol->ProtocolTraits, EReplicationProtocolTraits::HasLifetimeConditionals))
+	{
+		return;
+	}
+
+	// Set up fragment owner collector once, currently we only support a single owner
+	constexpr uint32 MaxFragmentOwnerCount = 1U;
+	UObject* FragmentOwners[MaxFragmentOwnerCount] = {};
+	FReplicationStateOwnerCollector FragmentOwnerCollector(FragmentOwners, MaxFragmentOwnerCount);
+
+	FReplicationFragment* const * Fragments = ReplicatedObjectData.InstanceProtocol->Fragments;
+	const FReplicationInstanceProtocol* InstanceProtocol = ReplicatedObjectData.InstanceProtocol;
+	const uint32 FragmentCount = InstanceProtocol->FragmentCount;
+	const SIZE_T FirstRelevantStateIndex = Protocol->FirstLifetimeConditionalsStateIndex;
+
+	const UObject* LastOwner = nullptr;
+	for (const FReplicationStateDescriptor*& StateDescriptor : MakeArrayView(Protocol->ReplicationStateDescriptors + FirstRelevantStateIndex, Protocol->ReplicationStateCount - FirstRelevantStateIndex))
+	{
+		if (EnumHasAnyFlags(StateDescriptor->Traits, EReplicationStateTraits::HasLifetimeConditionals))
+		{
+			const SIZE_T StateIndex = &StateDescriptor - Protocol->ReplicationStateDescriptors;
+
+			// Get Owner
+			const FReplicationFragment* ReplicationFragment = InstanceProtocol->Fragments[StateIndex];
+			FragmentOwnerCollector.Reset();
+			ReplicationFragment->CollectOwner(&FragmentOwnerCollector);
+			if (FragmentOwnerCollector.GetOwnerCount() == 0U)
+			{
+				// We have no owner
+				continue;
+			}
+
+			const UObject* CurrentOwner = FragmentOwnerCollector.GetOwners()[0];
+			TSharedPtr<FRepChangedPropertyTracker> ChangedPropertyTracker;
+			if (CurrentOwner != LastOwner)
+			{				
+				ChangedPropertyTracker = FNetPropertyConditionManager::Get().FindOrCreatePropertyTracker(CurrentOwner);
+			}
+
+			const FReplicationInstanceProtocol::FFragmentData& Fragment = InstanceProtocol->FragmentData[StateIndex];
+			FNetBitArrayView ConditionalChangeMask = GetMemberConditionalChangeMask(Fragment.ExternalSrcBuffer, StateDescriptor);
+
+			// Initialize conditionals based on the state of the ChangedPropertyTracker
+			const FRepChangedPropertyTracker* Tracker = ChangedPropertyTracker.Get();
+			for (const FReplicationStateMemberLifetimeConditionDescriptor& MemberLifeTimeConditionDescriptor : MakeArrayView(StateDescriptor->MemberLifetimeConditionDescriptors, StateDescriptor->MemberCount))
+			{
+				if (MemberLifeTimeConditionDescriptor.Condition == ELifetimeCondition::COND_Custom)
+				{
+					const SIZE_T MemberIndex = &MemberLifeTimeConditionDescriptor - StateDescriptor->MemberLifetimeConditionDescriptors;
+					const uint16 RepIndex = StateDescriptor->MemberProperties[MemberIndex]->RepIndex;
+					if (!Tracker->IsParentActive(RepIndex))
+					{
+						const FReplicationStateMemberChangeMaskDescriptor& MemberChangeMaskDescriptor = StateDescriptor->MemberChangeMaskDescriptors[MemberIndex];
+						ConditionalChangeMask.ClearBits(MemberChangeMaskDescriptor.BitOffset, MemberChangeMaskDescriptor.BitCount);
+					}
+				}
+				else
+				{
+					const SIZE_T MemberIndex = &MemberLifeTimeConditionDescriptor - StateDescriptor->MemberLifetimeConditionDescriptors;
+					const uint16 RepIndex = StateDescriptor->MemberProperties[MemberIndex]->RepIndex;
+					if (!Tracker->IsParentActive(RepIndex))
+					{
+						UE_LOG(LogIris, Warning, TEXT("Trying to change non-existing custom conditional for RepIndex %u in protocol %s"), RepIndex, ToCStr(Protocol->DebugName));
+					}
+				}
+			}
+		}
+	}
+}
+
 // N.B. Calls can come for properties that have been disabled. We must handle such cases gracefully.
 bool FReplicationConditionals::SetPropertyCustomCondition(FInternalNetHandle ObjectIndex, const void* Owner, uint16 RepIndex, bool bIsActive)
 {
@@ -210,7 +289,6 @@ bool FReplicationConditionals::SetPropertyCustomCondition(FInternalNetHandle Obj
 				const FReplicationInstanceProtocol::FFragmentData& Fragment = InstanceProtocol->FragmentData[StateIndex];
 
 				const FReplicationStateMemberChangeMaskDescriptor& ChangeMaskDescriptor = StateDescriptor->MemberChangeMaskDescriptors[RepIndexToMemberIndexDescriptor.MemberIndex];
-				checkSlow(ChangeMaskDescriptor.BitCount == 1);
 				FNetBitArrayView ConditionalChangeMask = GetMemberConditionalChangeMask(Fragment.ExternalSrcBuffer, StateDescriptor);
 
 				if (bIsActive)
diff --git a/Engine/Source/Runtime/Experimental/Iris/Core/Private/Iris/ReplicationSystem/Conditionals/ReplicationConditionals.h b/Engine/Source/Runtime/Experimental/Iris/Core/Private/Iris/ReplicationSystem/Conditionals/ReplicationConditionals.h
--- a/Engine/Source/Runtime/Experimental/Iris/Core/Private/Iris/ReplicationSystem/Conditionals/ReplicationConditionals.h
+++ b/Engine/Source/Runtime/Experimental/Iris/Core/Private/Iris/ReplicationSystem/Conditionals/ReplicationConditionals.h
@@ -57,6 +57,7 @@ public:
 	bool SetCondition(FInternalNetHandle ObjectIndex, EReplicationCondition Condition, bool bEnable);
 
 	// For property custom conditions only
+	void InitPropertyCustomConditions(FInternalNetHandle ObjectIndex);
 	bool SetPropertyCustomCondition(FInternalNetHandle ObjectIndex, const void* Owner, uint16 RepIndex, bool bIsActive);
 
 	void Update();
diff --git a/Engine/Source/Runtime/Experimental/Iris/Core/Private/Iris/ReplicationSystem/ObjectReplicationBridge.cpp b/Engine/Source/Runtime/Experimental/Iris/Core/Private/Iris/ReplicationSystem/ObjectReplicationBridge.cpp
--- a/Engine/Source/Runtime/Experimental/Iris/Core/Private/Iris/ReplicationSystem/ObjectReplicationBridge.cpp
+++ b/Engine/Source/Runtime/Experimental/Iris/Core/Private/Iris/ReplicationSystem/ObjectReplicationBridge.cpp
@@ -10,6 +10,7 @@
 
 #include "Net/Core/Trace/NetTrace.h"
 #include "Net/Core/Trace/NetDebugName.h"
+#include "Net/Core/PropertyConditions/PropertyConditionsDelegates.h"
 
 #include "Iris/ReplicationSystem/LegacyPushModel.h"
 #include "Iris/ReplicationSystem/ObjectPollFrequencyLimiter.h"
@@ -121,10 +122,29 @@ void UObjectReplicationBridge::Initialize(UReplicationSystem* InReplicationSyste
 	GarbageCollectionAffectedObjects.Init(MaxActiveObjectCount);
 
 	LoadConfig();
+
+	// Hookup delegate when a property custom conditions is changed
+	OnCustomConditionChangedHandle = UE::Net::Private::FPropertyConditionDelegates::GetOnPropertyCustomConditionChangedDelegate().AddLambda([this, InReplicationSystem](const UObject* Owner, uint16 RepIndex, bool bEnable)
+	{
+		using namespace UE::Net;
+		using namespace UE::Net::Private;
+
+		const FNetHandle Handle = this->GetReplicatedHandle(Owner);
+		if (Handle.IsValid())
+		{
+			FReplicationSystemInternal* ReplicationSystemInternal = GetReplicationSystem()->GetReplicationSystemInternal();
+			const FNetHandleManager& LocalNetHandleManager = ReplicationSystemInternal->GetNetHandleManager();
+			FReplicationConditionals& Conditionals = ReplicationSystemInternal->GetConditionals();
+
+			Conditionals.SetPropertyCustomCondition(LocalNetHandleManager.GetInternalIndex(Handle), Owner, RepIndex, bEnable);
+		}
+	});
 }
 
 void UObjectReplicationBridge::Deinitialize()
 {
+	UE::Net::Private::FPropertyConditionDelegates::GetOnPropertyCustomConditionChangedDelegate().Remove(OnCustomConditionChangedHandle);
+	OnCustomConditionChangedHandle.Reset();
 	PollFrequencyLimiter->Deinit();
 	Super::Deinitialize();
 }
@@ -198,8 +218,7 @@ UE::Net::FNetHandle UObjectReplicationBridge::BeginReplication(UObject* Instance
 
 	const FReplicationFragments& RegisteredFragments = FFragmentRegistrationContextPrivateAccessor::GetReplicationFragments(FragmentRegistrationContext);
 
-	// We currently identify protocols by local archetype or CDO pointer and verified the protocol id received from server
-	// We also should verify the default state that we use for delta compression https://jira.it.epicgames.com/browse/UE-131344
+	// We currently identify protocols by local archetype or CDO pointer and verified the protocol id received from server and the hash of the default state
 	const UObject* ArchetypeOrCDOUsedAsKey = Instance->GetArchetype();
 
 	// Create Protocols
@@ -235,9 +254,13 @@ UE::Net::FNetHandle UObjectReplicationBridge::BeginReplication(UObject* Instance
 			// Attach the instance and Bind the Instance protocol to dirty tracking
 			constexpr bool bBindInstanceProtocol = true;
 			InternalAttachInstanceToNetHandle(Handle, bBindInstanceProtocol, InstanceProtocol.Get(), Instance);
-
+			
 			const FInternalNetHandle InternalReplicationIndex = NetHandleManager->GetInternalIndex(Handle);
 
+			// Initialize conditionals
+			FReplicationSystemInternal* ReplicationSystemInternal = GetReplicationSystem()->GetReplicationSystemInternal();
+			ReplicationSystemInternal->GetConditionals().InitPropertyCustomConditions(InternalReplicationIndex);
+
 			// Keep track of handles with object references for garbage collection's sake.
 			ObjectsWithObjectReferences.SetBitValue(InternalReplicationIndex, EnumHasAnyFlags(InstanceProtocol->InstanceTraits, EReplicationInstanceProtocolTraits::HasObjectReference));
 
diff --git a/Engine/Source/Runtime/Experimental/Iris/Core/Private/Iris/ReplicationSystem/ReplicationSystem.cpp b/Engine/Source/Runtime/Experimental/Iris/Core/Private/Iris/ReplicationSystem/ReplicationSystem.cpp
--- a/Engine/Source/Runtime/Experimental/Iris/Core/Private/Iris/ReplicationSystem/ReplicationSystem.cpp
+++ b/Engine/Source/Runtime/Experimental/Iris/Core/Private/Iris/ReplicationSystem/ReplicationSystem.cpp
@@ -1107,21 +1107,6 @@ bool UReplicationSystem::SetReplicationCondition(FNetHandle Handle, UE::Net::ERe
 	return Conditionals.SetCondition(ObjectInternalIndex, Condition, bEnable);
 }
 
-bool UReplicationSystem::SetPropertyCustomCondition(FNetHandle Handle, const void* Owner, uint16 RepIndex, bool bEnable)
-{
-	using namespace UE::Net::Private;
-
-	FNetHandleManager& NetHandleManager = Impl->ReplicationSystemInternal.GetNetHandleManager();
-	const uint32 ObjectInternalIndex = NetHandleManager.GetInternalIndex(Handle);
-	if (ObjectInternalIndex == FNetHandleManager::InvalidInternalIndex)
-	{
-		return false;
-	}
-
-	FReplicationConditionals& Conditionals = Impl->ReplicationSystemInternal.GetConditionals();
-	return Conditionals.SetPropertyCustomCondition(ObjectInternalIndex, Owner, RepIndex, bEnable);
-}
-
 void UReplicationSystem::SetDeltaCompressionStatus(FNetHandle Handle, UE::Net::ENetObjectDeltaCompressionStatus Status)
 {
 	using namespace UE::Net::Private;
diff --git a/Engine/Source/Runtime/Experimental/Iris/Core/Public/Iris/ReplicationSystem/ObjectReplicationBridge.h b/Engine/Source/Runtime/Experimental/Iris/Core/Public/Iris/ReplicationSystem/ObjectReplicationBridge.h
--- a/Engine/Source/Runtime/Experimental/Iris/Core/Public/Iris/ReplicationSystem/ObjectReplicationBridge.h
+++ b/Engine/Source/Runtime/Experimental/Iris/Core/Public/Iris/ReplicationSystem/ObjectReplicationBridge.h
@@ -8,6 +8,7 @@
 #include "Iris/ReplicationSystem/ReplicationBridge.h"
 #include "Iris/ReplicationState/ReplicationStateDescriptor.h"
 #include "Iris/ReplicationSystem/ReplicationFragment.h"
+#include "Delegates/IDelegateInstance.h"
 
 #include "ObjectReplicationBridge.generated.h"
 
@@ -288,6 +289,8 @@ private:
 
 	UE::Net::FNetObjectFilterHandle DefaultSpatialFilterHandle;
 
+	FDelegateHandle OnCustomConditionChangedHandle;
+
 	bool bHasPollOverrides = false;
 	bool bHasDirtyClassesInPollPeriodOverrides = false;
 };
diff --git a/Engine/Source/Runtime/Experimental/Iris/Core/Public/Iris/ReplicationSystem/ReplicationSystem.h b/Engine/Source/Runtime/Experimental/Iris/Core/Public/Iris/ReplicationSystem/ReplicationSystem.h
--- a/Engine/Source/Runtime/Experimental/Iris/Core/Public/Iris/ReplicationSystem/ReplicationSystem.h
+++ b/Engine/Source/Runtime/Experimental/Iris/Core/Public/Iris/ReplicationSystem/ReplicationSystem.h
@@ -511,9 +511,6 @@ private:
 	IRISCORE_API void NotifyStreamingLevelUnload(const UObject* Level);
 
 private:
-	friend UE::Net::FReplicationSystemUtil;
-	friend class FIrisCoreModule;
-	IRISCORE_API bool SetPropertyCustomCondition(FNetHandle Handle, const void* Owner, uint16 RepIndex, bool bEnable);
 
 	void PostGarbageCollection();
 	void CollectGarbage();
diff --git a/Engine/Source/Runtime/Net/Core/Private/Net/Core/PropertyConditions/PropertyConditions.cpp b/Engine/Source/Runtime/Net/Core/Private/Net/Core/PropertyConditions/PropertyConditions.cpp
--- a/Engine/Source/Runtime/Net/Core/Private/Net/Core/PropertyConditions/PropertyConditions.cpp
+++ b/Engine/Source/Runtime/Net/Core/Private/Net/Core/PropertyConditions/PropertyConditions.cpp
@@ -2,6 +2,7 @@
 
 #include "Net/Core/PropertyConditions/PropertyConditions.h"
 #include "Net/Core/PropertyConditions/RepChangedPropertyTracker.h"
+#include "Net/Core/PropertyConditions/PropertyConditionsDelegates.h"
 #include "UObject/UObjectGlobals.h"
 #include "UObject/Object.h"
 #include "Serialization/ArchiveCountMem.h"
@@ -11,6 +12,9 @@ DECLARE_CYCLE_STAT(TEXT("PropertyConditions PostGarbageCollect"), STAT_PropertyC
 
 namespace UE::Net::Private
 {
+#if UE_WITH_IRIS
+FPropertyConditionDelegates::FOnPropertyCustomConditionChanged FPropertyConditionDelegates::OnPropertyCustomConditionChangedDelegate;
+#endif // UE_WITH_IRIS
 
 FNetPropertyConditionManager::FNetPropertyConditionManager()
 {
@@ -33,7 +37,7 @@ void FNetPropertyConditionManager::SetPropertyActive(const FObjectKey ObjectKey,
 	TSharedPtr<FRepChangedPropertyTracker> Tracker = FindPropertyTracker(ObjectKey);
 	if (Tracker.IsValid())
 	{
-		Tracker->SetCustomIsActiveOverride(ObjectKey.ResolveObjectPtr(), RepIndex, bActive);
+		Tracker->CallSetCustomIsActiveOverride(ObjectKey.ResolveObjectPtr(), RepIndex, bActive);
 	}
 }
 
@@ -108,4 +112,9 @@ void FNetPropertyConditionManager::LogMemory(FOutputDevice& Ar)
 	Ar.Logf(TEXT("  Property Condition Memory: %u"), CountBytes);
 }
 
+void FNetPropertyConditionManager::SetPropertyActiveOverride(IRepChangedPropertyTracker& Tracker, UObject* OwningObject, const uint16 RepIndex, const bool bIsActive)
+{
+	Tracker.CallSetCustomIsActiveOverride(OwningObject, RepIndex, bIsActive);
+}
+
 }; // UE::Net::Private
\ No newline at end of file
diff --git a/Engine/Source/Runtime/Net/Core/Private/Net/Core/PropertyConditions/RepChangedPropertyTracker.cpp b/Engine/Source/Runtime/Net/Core/Private/Net/Core/PropertyConditions/RepChangedPropertyTracker.cpp
--- a/Engine/Source/Runtime/Net/Core/Private/Net/Core/PropertyConditions/RepChangedPropertyTracker.cpp
+++ b/Engine/Source/Runtime/Net/Core/Private/Net/Core/PropertyConditions/RepChangedPropertyTracker.cpp
@@ -1,6 +1,7 @@
 // Copyright Epic Games, Inc. All Rights Reserved.
 
 #include "Net/Core/PropertyConditions/RepChangedPropertyTracker.h"
+#include "Net/Core/PropertyConditions/PropertyConditionsDelegates.h"
 #include "Net/Core/PushModel/PushModel.h"
 
 PRAGMA_DISABLE_DEPRECATION_WARNINGS
@@ -15,18 +16,6 @@ FRepChangedPropertyTracker::FRepChangedPropertyTracker(const bool InbIsReplay, c
 {}
 PRAGMA_ENABLE_DEPRECATION_WARNINGS
 
-namespace UE::Net::Private
-{
-#if UE_WITH_IRIS
-static FIrisSetPropertyCustomCondition IrisSetPropertyCustomConditionDelegate;
-
-void SetIrisSetPropertyCustomConditionDelegate(const FIrisSetPropertyCustomCondition& Delegate)
-{
-	IrisSetPropertyCustomConditionDelegate = Delegate;
-}
-#endif // UE_WITH_IRIS
-}
-
 void FRepChangedPropertyTracker::SetCustomIsActiveOverride(UObject* OwningObject, const uint16 RepIndex, const bool bIsActive)
 {
 	const bool bOldActive = ActiveState.GetActiveState(RepIndex);
@@ -40,9 +29,9 @@ void FRepChangedPropertyTracker::SetCustomIsActiveOverride(UObject* OwningObject
 #endif
 
 #if UE_WITH_IRIS
-	if (bOldActive != bIsActive && UE::Net::Private::IrisSetPropertyCustomConditionDelegate.IsBound())
+	if (bOldActive != bIsActive && UE::Net::Private::FPropertyConditionDelegates::GetOnPropertyCustomConditionChangedDelegate().IsBound())
 	{
-		UE::Net::Private::IrisSetPropertyCustomConditionDelegate.Execute(OwningObject, RepIndex, bIsActive);
+		UE::Net::Private::FPropertyConditionDelegates::GetOnPropertyCustomConditionChangedDelegate().Broadcast(OwningObject, RepIndex, bIsActive);
 	}
 #endif // UE_WITH_IRIS
 }
diff --git a/Engine/Source/Runtime/Net/Core/Public/Net/Core/PropertyConditions/PropertyConditions.h b/Engine/Source/Runtime/Net/Core/Public/Net/Core/PropertyConditions/PropertyConditions.h
--- a/Engine/Source/Runtime/Net/Core/Public/Net/Core/PropertyConditions/PropertyConditions.h
+++ b/Engine/Source/Runtime/Net/Core/Public/Net/Core/PropertyConditions/PropertyConditions.h
@@ -32,6 +32,8 @@ public:
 
 	void LogMemory(FOutputDevice& Ar);
 
+	static void SetPropertyActiveOverride(IRepChangedPropertyTracker& Tracker, UObject* OwningObject, const uint16 RepIndex, const bool bIsActive);
+
 private:
 	void PostGarbageCollect();
 
diff --git a/Engine/Source/Runtime/Net/Core/Public/Net/Core/PropertyConditions/PropertyConditionsDelegates.h b/Engine/Source/Runtime/Net/Core/Public/Net/Core/PropertyConditions/PropertyConditionsDelegates.h
new file mode 100644
--- /dev/null
+++ b/Engine/Source/Runtime/Net/Core/Public/Net/Core/PropertyConditions/PropertyConditionsDelegates.h
@@ -0,0 +1,31 @@
+// Copyright Epic Games, Inc. All Rights Reserved.
+
+#pragma once
+
+// Kept as iris only until we we know if we even need this delegate. 
+// We are considering to move Iris into NetCore when we leave the experimental stage at that point we can call directly into iris code
+// rather than going though a delegate.
+
+#if UE_WITH_IRIS
+
+#include "Delegates/Delegate.h"
+
+class FRepChangedPropertyTracker;
+
+namespace UE::Net::Private
+{
+
+class NETCORE_API FPropertyConditionDelegates
+{
+public:
+	DECLARE_MULTICAST_DELEGATE_ThreeParams(FOnPropertyCustomConditionChanged, const UObject*, const uint16, const bool);
+	
+	static FOnPropertyCustomConditionChanged& GetOnPropertyCustomConditionChangedDelegate() { return OnPropertyCustomConditionChangedDelegate; }
+
+private:
+	static FOnPropertyCustomConditionChanged OnPropertyCustomConditionChangedDelegate;
+};
+
+}; // UE::Net::Private
+
+#endif
\ No newline at end of file
diff --git a/Engine/Source/Runtime/Net/Core/Public/Net/Core/PropertyConditions/RepChangedPropertyTracker.h b/Engine/Source/Runtime/Net/Core/Public/Net/Core/PropertyConditions/RepChangedPropertyTracker.h
--- a/Engine/Source/Runtime/Net/Core/Public/Net/Core/PropertyConditions/RepChangedPropertyTracker.h
+++ b/Engine/Source/Runtime/Net/Core/Public/Net/Core/PropertyConditions/RepChangedPropertyTracker.h
@@ -43,6 +43,7 @@ public:
 		* @param RepIndex		Replication index for the Property.
 		* @param bIsActive		The new Active state.
 		*/
+	UE_DEPRECATED(5.3, "Please use FNetPropertyConditionManager::SetPropertyActiveOverride instead.")
 	virtual void SetCustomIsActiveOverride(UObject* OwningObject, const uint16 RepIndex, const bool bIsActive) override;
 
 	/**
@@ -83,12 +84,3 @@ public:
 	UE_DEPRECATED(5.0, "No longer used, see UReplaySubsystem::SetExternalDataForObject")
 	uint32 ExternalDataNumBits;
 };
-
-namespace UE::Net::Private
-{
-#if UE_WITH_IRIS
-DECLARE_DELEGATE_ThreeParams(FIrisSetPropertyCustomCondition, const UObject*, const uint16, const bool);
-NETCORE_API void SetIrisSetPropertyCustomConditionDelegate(const FIrisSetPropertyCustomCondition& Delegate);
-#endif
-}
-
