ensure that both constraint and target tick functions belong to the same world when setting up dependencies

diff --git a/Engine/Plugins/Animation/ControlRig/Source/ControlRig/Private/Constraints/ControlRigTransformableHandle.cpp b/Engine/Plugins/Animation/ControlRig/Source/ControlRig/Private/Constraints/ControlRigTransformableHandle.cpp
--- a/Engine/Plugins/Animation/ControlRig/Source/ControlRig/Private/Constraints/ControlRigTransformableHandle.cpp
+++ b/Engine/Plugins/Animation/ControlRig/Source/ControlRig/Private/Constraints/ControlRigTransformableHandle.cpp
@@ -404,7 +404,11 @@ void UTransformableControlHandle::OnObjectBoundToControlRig(UObject* InObject)
 	const UObject* CurrentObject = Binding ? Binding->GetBoundObject() : nullptr;
 	if (CurrentObject == InObject)
 	{
-		OnHandleModified.Broadcast(this, EHandleEvent::ComponentUpdated);
+		const UWorld* ThisWorld = GetWorld();
+		if (ThisWorld && InObject->GetWorld() == ThisWorld)
+		{
+			OnHandleModified.Broadcast(this, EHandleEvent::ComponentUpdated);
+		}
 	}
 }
 
diff --git a/Engine/Source/Runtime/Experimental/Animation/Constraints/Private/TransformConstraint.cpp b/Engine/Source/Runtime/Experimental/Animation/Constraints/Private/TransformConstraint.cpp
--- a/Engine/Source/Runtime/Experimental/Animation/Constraints/Private/TransformConstraint.cpp
+++ b/Engine/Source/Runtime/Experimental/Animation/Constraints/Private/TransformConstraint.cpp
@@ -143,8 +143,25 @@ void UTickableTransformConstraint::Setup()
 
 void UTickableTransformConstraint::SetupDependencies()
 {
-	FTickFunction* ParentTickFunction = ParentTRSHandle->IsValid() ? ParentTRSHandle->GetTickFunction() : nullptr;
-	FTickFunction* ChildTickFunction = ChildTRSHandle ? ChildTRSHandle->GetTickFunction() : nullptr;
+	auto GetTickableFunction = [this](const TObjectPtr<UTransformableHandle>& InHandle) -> FTickFunction*
+	{
+		if (!IsValid(InHandle) || !InHandle->IsValid())
+		{
+			return nullptr;
+		}
+
+		// avoid creating dependencies between functions that are registered in levels that don't leave in the same world 
+		const UObject* PrerequisiteObject = InHandle->GetPrerequisiteObject();
+		if (!PrerequisiteObject || PrerequisiteObject->GetWorld() != GetWorld())
+		{
+			return nullptr;
+		}
+		
+		return InHandle->GetTickFunction();
+	};
+	
+	FTickFunction* ParentTickFunction = GetTickableFunction(ParentTRSHandle);
+	FTickFunction* ChildTickFunction = GetTickableFunction(ChildTRSHandle);
 	
 	if (ParentTickFunction && (ChildTickFunction != ParentTickFunction))
 	{
@@ -166,7 +183,10 @@ void UTickableTransformConstraint::SetupDependencies()
 void UTickableTransformConstraint::PostLoad()
 {
 	Super::PostLoad();
-	ConstraintTick.RegisterFunction(GetFunction() );
+	if (ConstraintTick.ConstraintFunctions.IsEmpty())
+	{
+		ConstraintTick.RegisterFunction(GetFunction() );
+	}
 
 	SetupDependencies();	
 	RegisterDelegates();
@@ -175,7 +195,10 @@ void UTickableTransformConstraint::PostLoad()
 void UTickableTransformConstraint::PostDuplicate(bool bDuplicateForPIE)
 {
 	Super::PostDuplicate(bDuplicateForPIE);
-	ConstraintTick.RegisterFunction(GetFunction());
+	if (ConstraintTick.ConstraintFunctions.IsEmpty())
+	{
+		ConstraintTick.RegisterFunction(GetFunction());
+	}
 
 	SetupDependencies();
 	RegisterDelegates();
