Sequencer: Constraints  Fix redo constraint deletion due to the ordering of the modifies for the transform section and it's constraint object, and make sure constraint manager also undo's correctly and is reflected in the widget.

diff --git a/Engine/Plugins/Animation/ControlRig/Source/ControlRigEditor/Private/Editor/Constraints/SConstraintsWidget.cpp b/Engine/Plugins/Animation/ControlRig/Source/ControlRigEditor/Private/Editor/Constraints/SConstraintsWidget.cpp
--- a/Engine/Plugins/Animation/ControlRig/Source/ControlRigEditor/Private/Editor/Constraints/SConstraintsWidget.cpp
+++ b/Engine/Plugins/Animation/ControlRig/Source/ControlRigEditor/Private/Editor/Constraints/SConstraintsWidget.cpp
@@ -788,6 +788,16 @@ void SConstraintsEditionWidget::RemoveItem(const TSharedPtr<FEditableConstraintI
 	RefreshConstraintList();
 }
 
+void SConstraintsEditionWidget::PostUndo(bool bSuccess)
+{
+	InvalidateConstraintList();
+}
+
+void SConstraintsEditionWidget::PostRedo(bool bSuccess)
+{
+	InvalidateConstraintList();
+}
+
 void SConstraintsEditionWidget::InvalidateConstraintList()
 {
 	bNeedsRefresh = true;
diff --git a/Engine/Plugins/Animation/ControlRig/Source/ControlRigEditor/Private/Editor/Constraints/SConstraintsWidget.h b/Engine/Plugins/Animation/ControlRig/Source/ControlRigEditor/Private/Editor/Constraints/SConstraintsWidget.h
--- a/Engine/Plugins/Animation/ControlRig/Source/ControlRigEditor/Private/Editor/Constraints/SConstraintsWidget.h
+++ b/Engine/Plugins/Animation/ControlRig/Source/ControlRigEditor/Private/Editor/Constraints/SConstraintsWidget.h
@@ -5,6 +5,7 @@
 #include "Widgets/SCompoundWidget.h"
 #include "Widgets/Views/SListView.h"
 #include "Constraint.h"
+#include "EditorUndoClient.h"
 
 class AActor;
 class SConstraintsCreationWidget;
@@ -192,7 +193,7 @@ private:
  * SConstraintsEditionWidget
  */
 
-class CONTROLRIGEDITOR_API SConstraintsEditionWidget : public SCompoundWidget
+class CONTROLRIGEDITOR_API SConstraintsEditionWidget : public SCompoundWidget, public FEditorUndoClient
 {
 public:
 	SLATE_BEGIN_ARGS(SConstraintsEditionWidget)	{}
@@ -225,6 +226,11 @@ public:
 	/**  */
 	void RemoveItem(const TSharedPtr<FEditableConstraintItem>& Item);
 
+	/* FEditorUndoClient interface */
+	virtual void PostUndo(bool bSuccess);
+	virtual void PostRedo(bool bSuccess);
+	/* End FEditorUndoClient interface */
+
 private:
 	/** Types */
 	using ItemSharedPtr = TSharedPtr<FEditableConstraintItem>;
diff --git a/Engine/Source/Runtime/Experimental/Animation/Constraints/Private/ConstraintsManager.cpp b/Engine/Source/Runtime/Experimental/Animation/Constraints/Private/ConstraintsManager.cpp
--- a/Engine/Source/Runtime/Experimental/Animation/Constraints/Private/ConstraintsManager.cpp
+++ b/Engine/Source/Runtime/Experimental/Animation/Constraints/Private/ConstraintsManager.cpp
@@ -210,7 +210,7 @@ UConstraintsManager* UConstraintsManager::Get(UWorld* InWorld)
 #if WITH_EDITOR
 	ConstraintsActor->SetActorLabel(TEXT("Constraints Manager"));
 #endif // WITH_EDITOR
-
+	ConstraintsActor->ConstraintsManager->SetFlags(RF_Transactional);
 	return ConstraintsActor->ConstraintsManager;
 }
 
@@ -440,17 +440,18 @@ bool FConstraintsManagerController::RemoveConstraint(const int32 InConstraintInd
 	UTickableConstraint* Constraint = Manager->Constraints[InConstraintIndex];
 	
 	// notify deletion
+	Manager->Modify();
+	Manager->Constraints[InConstraintIndex]->Modify();
 	ConstraintRemoved.Broadcast(ConstraintName);
 	Manager->OnConstraintRemoved_BP.Broadcast(Manager, Constraint);
 
-	Manager->Modify();
-	Manager->Constraints[InConstraintIndex]->ConstraintTick.UnRegisterTickFunction();
+	Manager->Constraints[InConstraintIndex]->SetActive(false);
 	Manager->Constraints.RemoveAt(InConstraintIndex);
 
 	// destroy constraints actor if no constraints left
 	if (Manager->Constraints.IsEmpty())
 	{
-		DestroyManager();
+		// todo will re-evaluate this based upon spawnable work. DestroyManager();
 	}
 	return true;
 }
