Chaos : destruction : fix ensure wrongly checking for parent in the caching code path

diff --git a/Engine/Source/Runtime/Experimental/Chaos/Private/Chaos/PBDRigidClustering.cpp b/Engine/Source/Runtime/Experimental/Chaos/Private/Chaos/PBDRigidClustering.cpp
--- a/Engine/Source/Runtime/Experimental/Chaos/Private/Chaos/PBDRigidClustering.cpp
+++ b/Engine/Source/Runtime/Experimental/Chaos/Private/Chaos/PBDRigidClustering.cpp
@@ -581,6 +581,12 @@ namespace Chaos
 	{
 		SCOPE_CYCLE_COUNTER(STAT_ReleaseClusterParticles_STRAIN);
 
+		if (!ensureMsgf((ClusteredParticle->Parent() == nullptr), TEXT("Removing a cluster that still has a parent")))
+		{
+			TSet<FPBDRigidParticleHandle*> EmptySet;
+			return EmptySet;
+		};
+
 		return ReleaseClusterParticlesImpl(ClusteredParticle, bForceRelease, true /*bCreateNewClusters*/);
 	}
 	
@@ -606,11 +612,6 @@ namespace Chaos
 	{	
 		TSet<FPBDRigidParticleHandle*> ActivatedChildren;
 
-		if (!ensureMsgf((ClusteredParticle->Parent()==nullptr), TEXT("Removing a cluster that still has a parent")))
-		{
-			return ActivatedChildren;;
-		};
-
 		if (!ensureMsgf(MChildren.Contains(ClusteredParticle), TEXT("Removing Cluster that does not exist!")))
 		{
 			return ActivatedChildren;
