GeometryCollection � Type inspection
- Allow the derived types to be inspected on the base class.

diff --git a/Engine/Source/Runtime/Experimental/Chaos/Private/GeometryCollection/GeometryCollection.cpp b/Engine/Source/Runtime/Experimental/Chaos/Private/GeometryCollection/GeometryCollection.cpp
--- a/Engine/Source/Runtime/Experimental/Chaos/Private/GeometryCollection/GeometryCollection.cpp
+++ b/Engine/Source/Runtime/Experimental/Chaos/Private/GeometryCollection/GeometryCollection.cpp
@@ -32,8 +32,8 @@ const FName FGeometryCollection::SimulatableParticlesAttribute("SimulatableParti
 const FName FGeometryCollection::SimulationTypeAttribute("SimulationType");
 const FName FGeometryCollection::StatusFlagsAttribute("StatusFlags");
 
-FGeometryCollection::FGeometryCollection()
-	: FTransformCollection()
+FGeometryCollection::FGeometryCollection(FClassType InType)
+	: FTransformCollection(InType.Union({ FGeometryCollection::StaticType() }))
 	, FGeometryCollectionConvexPropertiesInterface(this)
 {
 	Construct();
diff --git a/Engine/Source/Runtime/Experimental/Chaos/Private/GeometryCollection/ManagedArrayCollection.cpp b/Engine/Source/Runtime/Experimental/Chaos/Private/GeometryCollection/ManagedArrayCollection.cpp
--- a/Engine/Source/Runtime/Experimental/Chaos/Private/GeometryCollection/ManagedArrayCollection.cpp
+++ b/Engine/Source/Runtime/Experimental/Chaos/Private/GeometryCollection/ManagedArrayCollection.cpp
@@ -11,12 +11,13 @@ DEFINE_LOG_CATEGORY_STATIC(FManagedArrayCollectionLogging, NoLogging, All);
 int8 FManagedArrayCollection::Invalid = INDEX_NONE;
 
 
-FManagedArrayCollection::FManagedArrayCollection()
+FManagedArrayCollection::FManagedArrayCollection(FClassType InType)
+	: ClassType(InType)
 {
+	ClassType.Add(FManagedArrayCollection::StaticType());
 	Version = 9;
 }
 
-
 void FManagedArrayCollection::AddGroup(FName Group)
 {
 	ensure(!GroupInfo.Contains(Group));
diff --git a/Engine/Source/Runtime/Experimental/Chaos/Private/GeometryCollection/TransformCollection.cpp b/Engine/Source/Runtime/Experimental/Chaos/Private/GeometryCollection/TransformCollection.cpp
--- a/Engine/Source/Runtime/Experimental/Chaos/Private/GeometryCollection/TransformCollection.cpp
+++ b/Engine/Source/Runtime/Experimental/Chaos/Private/GeometryCollection/TransformCollection.cpp
@@ -13,8 +13,8 @@ const FName FTransformCollection::ParentAttribute = "Parent";
 const FName FTransformCollection::ChildrenAttribute = "Children";
 const FName FTransformCollection::ParticlesAttribute = "Particles";
 
-FTransformCollection::FTransformCollection()
-	: FManagedArrayCollection()
+FTransformCollection::FTransformCollection(FClassType InType)
+	: FManagedArrayCollection(InType.Union({ FTransformCollection::StaticType() }))
 {
 	Construct();
 }
diff --git a/Engine/Source/Runtime/Experimental/Chaos/Public/GeometryCollection/GeometryCollection.h b/Engine/Source/Runtime/Experimental/Chaos/Public/GeometryCollection/GeometryCollection.h
--- a/Engine/Source/Runtime/Experimental/Chaos/Public/GeometryCollection/GeometryCollection.h
+++ b/Engine/Source/Runtime/Experimental/Chaos/Public/GeometryCollection/GeometryCollection.h
@@ -31,7 +31,7 @@ class CHAOS_API FGeometryCollection : public FTransformCollection,
 {
 
 public:
-	FGeometryCollection();
+	FGeometryCollection(FClassType InType = {});
 	FGeometryCollection(FGeometryCollection &) = delete;
 	FGeometryCollection& operator=(const FGeometryCollection &) = delete;
 	FGeometryCollection(FGeometryCollection &&) = default;
@@ -126,6 +126,9 @@ public:
 	//
 	//
 
+	/** Type name for this class. */
+	static FName StaticType() { return FName("FGeometryCollection"); }
+
 	/**
 	 * Create a GeometryCollection from Vertex and Indices arrays
 	 */
diff --git a/Engine/Source/Runtime/Experimental/Chaos/Public/GeometryCollection/ManagedArrayCollection.h b/Engine/Source/Runtime/Experimental/Chaos/Public/GeometryCollection/ManagedArrayCollection.h
--- a/Engine/Source/Runtime/Experimental/Chaos/Public/GeometryCollection/ManagedArrayCollection.h
+++ b/Engine/Source/Runtime/Experimental/Chaos/Public/GeometryCollection/ManagedArrayCollection.h
@@ -47,7 +47,10 @@ struct CHAOS_API FManagedArrayCollection
 	friend FSimulationProperties;
 
 public:
-	FManagedArrayCollection();
+	typedef TSet<FName> FClassType;
+
+
+	FManagedArrayCollection(FClassType InType = {});
 	virtual ~FManagedArrayCollection() {}
 
 	FManagedArrayCollection(const FManagedArrayCollection& In) { In.CopyTo(this); }
@@ -57,7 +60,6 @@ public:
 
 	static int8 Invalid;
 	typedef EManagedArrayType EArrayType;
-
 	/**
 	*
 	*/
@@ -97,6 +99,20 @@ public:
 			FManagedType(ManagedArrayType<T>(), InName, InGroup) {}			
 	};
 
+	/**
+	* Type name for this class.
+	*/
+	static FName StaticType() { return FName("FManagedArrayCollection"); }
+
+	template<class T>
+	bool IsA() { return ClassType.Contains(T::StaticType())?true:false; }
+
+	template<class T>
+	T* Cast() { return ClassType.Contains(T::StaticType())?static_cast<T*>(this):nullptr; }
+
+	template<class T>
+	const T* Cast() const { return ClassType.Contains(T::StaticType()) ? static_cast<const T*>(this) : nullptr; }
+
 	/**
 	* Add an attribute of Type(T) to the group
 	* @param Name - The name of the attribute
@@ -594,6 +610,7 @@ private:
 
 	virtual void SetDefaults(FName Group, uint32 StartSize, uint32 NumElements) {};
 
+	FClassType ClassType;
 	TMap< FKeyType, FValueType> Map;	//data is owned by the map explicitly
 	TMap< FName, FGroupInfo> GroupInfo;
 	bool bDirty;
diff --git a/Engine/Source/Runtime/Experimental/Chaos/Public/GeometryCollection/TransformCollection.h b/Engine/Source/Runtime/Experimental/Chaos/Public/GeometryCollection/TransformCollection.h
--- a/Engine/Source/Runtime/Experimental/Chaos/Public/GeometryCollection/TransformCollection.h
+++ b/Engine/Source/Runtime/Experimental/Chaos/Public/GeometryCollection/TransformCollection.h
@@ -18,7 +18,7 @@ class CHAOS_API FTransformCollection : public FManagedArrayCollection
 public:
 	typedef FManagedArrayCollection Super;
 
-	FTransformCollection();
+	FTransformCollection(FClassType InType = {});
 	FTransformCollection(FTransformCollection &) = delete;
 	FTransformCollection& operator=(const FTransformCollection&) = delete;
 	FTransformCollection(FTransformCollection&&) = default;
@@ -50,6 +50,9 @@ public:
 	static const FName ChildrenAttribute;
 	static const FName ParticlesAttribute;
 
+	/** Type name for this class. */
+	static FName StaticType() { return FName("FTransformCollection"); }
+
 	/** Schema */
 	static void DefineTransformSchema(FManagedArrayCollection&);
 
