Fracture editor : better handling of existing vertex color with the bone color display option
- When ShowBoneColors option is off display vertex color

diff --git a/Engine/Plugins/Experimental/ChaosEditor/Source/FractureEditor/Private/FractureEditorModeToolkit.cpp b/Engine/Plugins/Experimental/ChaosEditor/Source/FractureEditor/Private/FractureEditorModeToolkit.cpp
--- a/Engine/Plugins/Experimental/ChaosEditor/Source/FractureEditor/Private/FractureEditorModeToolkit.cpp
+++ b/Engine/Plugins/Experimental/ChaosEditor/Source/FractureEditor/Private/FractureEditorModeToolkit.cpp
@@ -1248,9 +1248,8 @@ void FFractureEditorModeToolkit::ToggleShowBoneColors()
 
 	for (UGeometryCollectionComponent* Comp : GeomCompSelection)
 	{
-		Comp->SetShowBoneColors(!Comp->GetShowBoneColors());
-		Comp->MarkRenderStateDirty();
-		Comp->MarkRenderDynamicDataDirty();
+		FScopedColorEdit EditBoneColor(Comp, true /*bForceUpdate*/); // the property has already changed; this will trigger the color update + render state updates
+		EditBoneColor.SetShowBoneColors(!EditBoneColor.GetShowBoneColors());
 	}
 }
 
diff --git a/Engine/Source/Runtime/Experimental/GeometryCollectionEngine/Private/GeometryCollection/GeometryCollectionComponent.cpp b/Engine/Source/Runtime/Experimental/GeometryCollectionEngine/Private/GeometryCollection/GeometryCollectionComponent.cpp
--- a/Engine/Source/Runtime/Experimental/GeometryCollectionEngine/Private/GeometryCollection/GeometryCollectionComponent.cpp
+++ b/Engine/Source/Runtime/Experimental/GeometryCollectionEngine/Private/GeometryCollection/GeometryCollectionComponent.cpp
@@ -3351,34 +3351,45 @@ void FScopedColorEdit::UpdateBoneColors()
 		{
 			FLinearColor BoneColor = FLinearColor(FColor::Black);
 
-			if (Component->ViewLevel == -1)
+			if (Component->bShowBoneColors)
 			{
-				BoneColor = RandomColors[BoneIndex % RandomColors.Num()];
-			}
-			else
-			{
-				if (HasLevelAttribute && (*Levels)[BoneIndex] >= Component->ViewLevel)
+				if (Component->ViewLevel == -1)
+				{
+					BoneColor = RandomColors[BoneIndex % RandomColors.Num()];
+				}
+				else
 				{
-					// go up until we find parent at the required ViewLevel
-					int32 Bone = BoneIndex;
-					while (Bone != -1 && (*Levels)[Bone] > Component->ViewLevel)
+					if (HasLevelAttribute && (*Levels)[BoneIndex] >= Component->ViewLevel)
 					{
-						Bone = Parents[Bone];
-					}
+						// go up until we find parent at the required ViewLevel
+						int32 Bone = BoneIndex;
+						while (Bone != -1 && (*Levels)[Bone] > Component->ViewLevel)
+						{
+							Bone = Parents[Bone];
+						}
 
-					int32 ColorIndex = Bone + 1; // parent can be -1 for root, range [-1..n]
-					BoneColor = RandomColors[ColorIndex % RandomColors.Num()];
+						int32 ColorIndex = Bone + 1; // parent can be -1 for root, range [-1..n]
+						BoneColor = RandomColors[ColorIndex % RandomColors.Num()];
 
-					BoneColor.LinearRGBToHSV();
-					BoneColor.B *= .5;
-					BoneColor.HSVToLinearRGB();
+						BoneColor.LinearRGBToHSV();
+						BoneColor.B *= .5;
+						BoneColor.HSVToLinearRGB();
+					}
+					else
+					{
+						BoneColor = BlankColor;
+					}
 				}
-				else
+			}
+			else
+			{
+				BoneColor = FLinearColor(FColor::White);
+				if (Component->ViewLevel != INDEX_NONE && HasLevelAttribute && (*Levels)[BoneIndex] < Component->ViewLevel)
 				{
-					BoneColor = BlankColor;
+					BoneColor = FLinearColor(FColor(128U, 128U, 128U, 255U));
 				}
+				
 			}
-
 			// store the bone selected toggle in alpha so we can use it in the shader
 			BoneColor.A = IsBoneHighlighted(BoneIndex) ? 1 : 0;
 
diff --git a/Engine/Source/Runtime/Experimental/GeometryCollectionEngine/Private/GeometryCollection/GeometryCollectionSceneProxy.cpp b/Engine/Source/Runtime/Experimental/GeometryCollectionEngine/Private/GeometryCollection/GeometryCollectionSceneProxy.cpp
--- a/Engine/Source/Runtime/Experimental/GeometryCollectionEngine/Private/GeometryCollection/GeometryCollectionSceneProxy.cpp
+++ b/Engine/Source/Runtime/Experimental/GeometryCollectionEngine/Private/GeometryCollection/GeometryCollectionSceneProxy.cpp
@@ -447,9 +447,11 @@ void FGeometryCollectionSceneProxy::BuildGeometry( const FGeometryCollectionCons
 			FDynamicMeshVertex(
 				ConstantDataIn->Vertices[PointIdx],
 				ConstantDataIn->UVs[PointIdx][0],
-				bShowBoneColors||bEnableBoneSelection ?
-				ConstantDataIn->BoneColors[PointIdx].ToFColor(true) :
-				ConstantDataIn->Colors[PointIdx].ToFColor(true)
+				(bShowBoneColors||bEnableBoneSelection) 
+					? bShowBoneColors 
+						? ConstantDataIn->BoneColors[PointIdx].ToFColor(true)
+						: (ConstantDataIn->BoneColors[PointIdx] * ConstantDataIn->Colors[PointIdx]).ToFColor(true)
+					: ConstantDataIn->Colors[PointIdx].ToFColor(true)
 			);
 		OutVertices[PointIdx].SetTangents(ConstantDataIn->TangentU[PointIdx], ConstantDataIn->TangentV[PointIdx], ConstantDataIn->Normals[PointIdx]);
 
