Chaos - collision detection optrimizations
- per-worker collision context with midphase and collision buffers, no locks or lock-free queues
- midphases now allocated from a pool

diff --git a/Engine/Source/Programs/HeadlessChaos/Private/HeadlessChaosTestJoints.cpp b/Engine/Source/Programs/HeadlessChaos/Private/HeadlessChaosTestJoints.cpp
--- a/Engine/Source/Programs/HeadlessChaos/Private/HeadlessChaosTestJoints.cpp
+++ b/Engine/Source/Programs/HeadlessChaos/Private/HeadlessChaosTestJoints.cpp
@@ -359,7 +359,7 @@ namespace ChaosTest {
 
 			Test.GetParticle(0)->X() = RootPosition;
 
-			Test.Evolution.GetCollisionDetector().GetNarrowPhase().SetBoundsVelocityInflation(1);
+			Test.Evolution.GetCollisionDetector().SetBoundsVelocityInflation(1);
 			Test.Evolution.AdvanceOneTimeStep(Dt);
 			Test.Evolution.EndFrame(Dt);
 
@@ -433,7 +433,7 @@ namespace ChaosTest {
 		}
 
 		Test.Create();
-		Test.Evolution.GetCollisionDetector().GetNarrowPhase().SetBoundsVelocityInflation(1);
+		Test.Evolution.GetCollisionDetector().SetBoundsVelocityInflation(1);
 
 		const FVec3 Box2LocalSpaceJointPosition = Test.JointPositions[0] - Test.ParticlePositions[1];
 
diff --git a/Engine/Source/Programs/HeadlessChaos/Private/HeadlessChaosTestMinEvolution.cpp b/Engine/Source/Programs/HeadlessChaos/Private/HeadlessChaosTestMinEvolution.cpp
--- a/Engine/Source/Programs/HeadlessChaos/Private/HeadlessChaosTestMinEvolution.cpp
+++ b/Engine/Source/Programs/HeadlessChaos/Private/HeadlessChaosTestMinEvolution.cpp
@@ -2,7 +2,6 @@
 
 #include "HeadlessChaos.h"
 
-#include "Chaos/Collision/NarrowPhase.h"
 #include "Chaos/Collision/ParticlePairBroadPhase.h"
 #include "Chaos/Collision/ParticlePairCollisionDetector.h"
 #include "Chaos/Evolution/PBDMinEvolution.h"
@@ -41,8 +40,7 @@ namespace ChaosTest
 		TArrayCollectionArray<FRotation3> ParticlePrevRs;
 		FCollisionConstraints Collisions(ParticlesContainer, CollidedParticles, ParticleMaterials, PerParticleMaterials, nullptr);
 		FBasicBroadPhase BroadPhase(&ActivePotentiallyCollidingPairs, nullptr, nullptr);
-		FNarrowPhase NarrowPhase(0, 0, Collisions.GetConstraintAllocator());
-		FCollisionDetector CollisionDetector(BroadPhase, NarrowPhase, Collisions);
+		FCollisionDetector CollisionDetector(BroadPhase, Collisions);
 		// End collisions stuff
 
 		// Springs
diff --git a/Engine/Source/Programs/HeadlessChaos/Private/HeadlessChaosTestObjectPool.cpp b/Engine/Source/Programs/HeadlessChaos/Private/HeadlessChaosTestObjectPool.cpp
--- a/Engine/Source/Programs/HeadlessChaos/Private/HeadlessChaosTestObjectPool.cpp
+++ b/Engine/Source/Programs/HeadlessChaos/Private/HeadlessChaosTestObjectPool.cpp
@@ -17,7 +17,7 @@ namespace ChaosTest
 		// Fill fist page
 		Pool.Alloc(1);
 		Pool.Alloc(2);
-		FPool::Ptr Ptr = Pool.Alloc(3);
+		FPool::FPtr Ptr = Pool.Alloc(3);
 
 		const int32 PoolSize = Pool.GetAllocatedSize();
 
@@ -53,7 +53,7 @@ namespace ChaosTest
 		// Fill fist page
 		Pool.Alloc(1, 1.0f);
 		Pool.Alloc(2, 2.0f);
-		FPool::Ptr Ptr = Pool.Alloc(3, 3.0f);
+		FPool::FPtr Ptr = Pool.Alloc(3, 3.0f);
 
 		const int32 PoolSize = Pool.GetAllocatedSize();
 
@@ -99,7 +99,7 @@ namespace ChaosTest
 		// Fill fist page
 		Pool.Alloc(ConstructorCount, DestructorCount);
 		Pool.Alloc(ConstructorCount, DestructorCount);
-		FPool::Ptr Ptr = Pool.Alloc(ConstructorCount, DestructorCount);
+		FPool::FPtr Ptr = Pool.Alloc(ConstructorCount, DestructorCount);
 
 		const int32 PoolSize = Pool.GetAllocatedSize();
 
@@ -125,7 +125,7 @@ namespace ChaosTest
 		using FPool = TObjectPool<int>;
 		FPool Pool(32);
 
-		TArray<FPool::Ptr> PooledInts;
+		TArray<FPool::FPtr> PooledInts;
 
 		for(int32 i = 0; i < 32; ++i)
 		{
@@ -174,10 +174,10 @@ namespace ChaosTest
 		using FPool = TObjectPool<FRefIncrementOnConstructDestroy>;
 		FPool Pool(3);
 
-		TArray<FPool::Ptr> PooledInts;
+		TArray<FPool::FPtr> PooledInts;
 
 		// Fill fist page
-		TArray<FPool::Ptr> Ptrs;
+		TArray<FPool::FPtr> Ptrs;
 
 		// Add 32 items
 		for(int32 i = 0; i < 32; ++i)
@@ -208,7 +208,7 @@ namespace ChaosTest
 		using FPool = TObjectPool<int>;
 		FPool Pool(3);
 
-		FPool::Ptr Objects[9];
+		FPool::FPtr Objects[9];
 
 		// Fill 3 pages
 		for(int32 i = 0; i < 9; ++i)
diff --git a/Engine/Source/Programs/HeadlessChaos/Private/HeadlessChaosTestShapePadding.cpp b/Engine/Source/Programs/HeadlessChaos/Private/HeadlessChaosTestShapePadding.cpp
--- a/Engine/Source/Programs/HeadlessChaos/Private/HeadlessChaosTestShapePadding.cpp
+++ b/Engine/Source/Programs/HeadlessChaos/Private/HeadlessChaosTestShapePadding.cpp
@@ -86,8 +86,9 @@ namespace ChaosTest {
 		EXPECT_NEAR(BoxBounds1.Extents().Z, Size.Z, Tolerance);
 
 		FCollisionConstraintAllocator CollisionAllocator;
+		CollisionAllocator.SetMaxContexts(1);
 
-		FPBDCollisionConstraintPtr Constraint = CollisionAllocator.CreateCollisionConstraint(
+		FPBDCollisionConstraintPtr Constraint = CollisionAllocator.GetContextAllocator(0)->CreateConstraint(
 			Box0,
 			Box0->Geometry().Get(),
 			Box0->ShapesArray()[0].Get(),
@@ -198,8 +199,9 @@ namespace ChaosTest {
 		EXPECT_NEAR(BoxBounds1.Extents().Z, Size.Z, Tolerance);
 
 		FCollisionConstraintAllocator CollisionAllocator;
+		CollisionAllocator.SetMaxContexts(1);
 
-		FPBDCollisionConstraintPtr Constraint = CollisionAllocator.CreateCollisionConstraint(
+		FPBDCollisionConstraintPtr Constraint = CollisionAllocator.GetContextAllocator(0)->CreateConstraint(
 			Box0,
 			Box0->Geometry().Get(),
 			Box0->ShapesArray()[0].Get(),
diff --git a/Engine/Source/Programs/HeadlessChaos/Public/HeadlessChaosCollisionConstraints.h b/Engine/Source/Programs/HeadlessChaos/Public/HeadlessChaosCollisionConstraints.h
--- a/Engine/Source/Programs/HeadlessChaos/Public/HeadlessChaosCollisionConstraints.h
+++ b/Engine/Source/Programs/HeadlessChaos/Public/HeadlessChaosCollisionConstraints.h
@@ -39,13 +39,14 @@ public:
 		: EmptyParticles(UniqueIndices)
 		, SpatialAcceleration(EmptyParticles.GetNonDisabledView())
 		, CollisionConstraints(EmptyParticles, EmptyCollided, EmptyPhysicsMaterials, EmptyUniquePhysicsMaterials, nullptr)
-		, NarrowPhase((FReal)1, (FReal)0, CollisionConstraints.GetConstraintAllocator())
 		, BroadPhase(EmptyParticles)
-		, CollisionDetector(BroadPhase, NarrowPhase, CollisionConstraints)
+		, CollisionDetector(BroadPhase, CollisionConstraints)
 		, ConstraintSolver()
 	{
 		CollisionConstraints.SetContainerId(0);
 
+		CollisionDetector.SetBoundsExpansion(FReal(1));
+
 		ConstraintSolver.SetConstraintSolver(CollisionConstraints.GetContainerId(), CollisionConstraints.CreateSceneSolver(0));
 	}
 
@@ -54,13 +55,14 @@ public:
 		: EmptyParticles(UniqueIndices)
 		, SpatialAcceleration(InParticles.GetNonDisabledView())
 		, CollisionConstraints(InParticles, Collided, PerParticleMaterials, PerParticleUniqueMaterials, nullptr)
-		, NarrowPhase((FReal)1, (FReal)0, CollisionConstraints.GetConstraintAllocator())
 		, BroadPhase(InParticles)
-		, CollisionDetector(BroadPhase, NarrowPhase, CollisionConstraints)
+		, CollisionDetector(BroadPhase, CollisionConstraints)
 		, ConstraintSolver()
 	{
 		CollisionConstraints.SetContainerId(0);
-		
+
+		CollisionDetector.SetBoundsExpansion(FReal(1));
+
 		ConstraintSolver.SetConstraintSolver(CollisionConstraints.GetContainerId(), CollisionConstraints.CreateSceneSolver(0));
 	}
 
@@ -69,10 +71,12 @@ public:
 	void ComputeConstraints(FReal Dt)
 	{
 		CollisionDetector.GetBroadPhase().SetSpatialAcceleration(&SpatialAcceleration);
-		CollisionDetector.GetNarrowPhase().GetContext().GetSettings().bFilteringEnabled = true;
-		CollisionDetector.GetNarrowPhase().GetContext().GetSettings().bDeferNarrowPhase = false;
-		CollisionDetector.GetNarrowPhase().GetContext().GetSettings().bAllowManifolds = true;
-		CollisionDetector.GetCollisionContainer().SetDetectorSettings(CollisionDetector.GetNarrowPhase().GetContext().GetSettings());
+
+		FCollisionDetectorSettings DetectorSettings = CollisionDetector.GetSettings();
+		DetectorSettings.bFilteringEnabled = true;
+		DetectorSettings.bDeferNarrowPhase = false;
+		DetectorSettings.bAllowManifolds = true;
+		CollisionDetector.SetSettings(DetectorSettings);
 		
 		CollisionDetector.DetectCollisions(Dt, nullptr);
 		
@@ -165,7 +169,6 @@ public:
 
 	FAccelerationStructure SpatialAcceleration;
 	FCollisionConstraints CollisionConstraints;
-	FNarrowPhase NarrowPhase;
 	FSpatialAccelerationBroadPhase BroadPhase;
 	FCollisionDetector CollisionDetector;
 	FPBDSceneConstraintGroupSolver ConstraintSolver;
diff --git a/Engine/Source/Programs/HeadlessChaos/Public/HeadlessChaosTestUtility.h b/Engine/Source/Programs/HeadlessChaos/Public/HeadlessChaosTestUtility.h
--- a/Engine/Source/Programs/HeadlessChaos/Public/HeadlessChaosTestUtility.h
+++ b/Engine/Source/Programs/HeadlessChaos/Public/HeadlessChaosTestUtility.h
@@ -18,6 +18,7 @@
 #include "Chaos/ParticleHandleFwd.h"
 #include "Chaos/PBDRigidsEvolutionFwd.h"
 #include "Chaos/Collision/CollisionApplyType.h"
+#include "Chaos/Collision/CollisionContext.h"
 
 namespace Chaos
 {
@@ -148,10 +149,11 @@ namespace ChaosTest {
 	{
 		// Settings used for unit tests
 		const float CullDistance = 3.0f;
-		Evolution.GetNarrowPhase().SetBoundsExpansion(CullDistance);
-		Evolution.GetCollisionDetector().GetNarrowPhase().GetContext().GetSettings().bDeferNarrowPhase = false;
-		Evolution.GetCollisionDetector().GetNarrowPhase().GetContext().GetSettings().bAllowManifolds = true;
-		Evolution.GetCollisionConstraints().SetDetectorSettings(Evolution.GetCollisionDetector().GetNarrowPhase().GetContext().GetSettings());
+		FCollisionDetectorSettings DetectorSettings = Evolution.GetCollisionDetector().GetSettings();
+		DetectorSettings.BoundsExpansion = CullDistance;
+		DetectorSettings.bDeferNarrowPhase = false;
+		DetectorSettings.bAllowManifolds = true;
+		Evolution.GetCollisionDetector().SetSettings(DetectorSettings);
 	}
 
 	template<typename T_SOLVER>
diff --git a/Engine/Source/Runtime/Engine/Private/PhysicsEngine/ImmediatePhysics/ImmediatePhysicsChaos/ImmediatePhysicsSimulation_Chaos.cpp b/Engine/Source/Runtime/Engine/Private/PhysicsEngine/ImmediatePhysics/ImmediatePhysicsChaos/ImmediatePhysicsSimulation_Chaos.cpp
--- a/Engine/Source/Runtime/Engine/Private/PhysicsEngine/ImmediatePhysics/ImmediatePhysicsChaos/ImmediatePhysicsSimulation_Chaos.cpp
+++ b/Engine/Source/Runtime/Engine/Private/PhysicsEngine/ImmediatePhysics/ImmediatePhysicsChaos/ImmediatePhysicsSimulation_Chaos.cpp
@@ -10,7 +10,6 @@
 #include "PhysicsEngine/ConstraintInstance.h"
 
 #include "Chaos/ChaosDebugDraw.h"
-#include "Chaos/Collision/NarrowPhase.h"
 #include "Chaos/Collision/ParticlePairBroadPhase.h"
 #include "Chaos/Collision/ParticlePairCollisionDetector.h"
 #include "Chaos/DebugDrawQueue.h"
@@ -301,8 +300,7 @@ namespace ImmediatePhysics_Chaos
 			, Joints()
 			, Collisions(Particles, CollidedParticles, ParticleMaterials, PerParticleMaterials, nullptr, ChaosImmediate_Evolution_NumCollisionsPerBlock, 2000)
 			, BroadPhase(&ActivePotentiallyCollidingPairs, nullptr, nullptr)
-			, NarrowPhase(FReal(0), FReal(0), Collisions.GetConstraintAllocator())
-			, CollisionDetector(BroadPhase, NarrowPhase, Collisions)
+			, CollisionDetector(BroadPhase, Collisions)
 			, Evolution(Particles, ParticlePrevXs, ParticlePrevRs, CollisionDetector, FReal(0))
 			, NumActiveDynamicActorHandles(0)
 			, SimulationSpace()
@@ -339,7 +337,6 @@ namespace ImmediatePhysics_Chaos
 		Chaos::FPBDJointConstraints Joints;
 		Chaos::FPBDCollisionConstraints Collisions;
 		Chaos::FBasicBroadPhase BroadPhase;
-		Chaos::FNarrowPhase NarrowPhase;
 		Chaos::FBasicCollisionDetector CollisionDetector;
 		Chaos::FPBDMinEvolution Evolution;
 
@@ -380,12 +377,13 @@ namespace ImmediatePhysics_Chaos
 		// RBAN collision customization
 		Implementation->Collisions.DisableHandles();
 
-		Implementation->NarrowPhase.GetContext().GetSettings().bFilteringEnabled = false;
-		Implementation->NarrowPhase.GetContext().GetSettings().bAllowManifoldReuse = false;
-		Implementation->NarrowPhase.GetContext().GetSettings().bDeferNarrowPhase = (ChaosImmediate_Collision_DeferNarrowPhase != 0);
-		Implementation->NarrowPhase.GetContext().GetSettings().bAllowManifolds = (ChaosImmediate_Collision_UseManifolds != 0);
-		Implementation->NarrowPhase.GetContext().GetSettings().bAllowCCD = false;
-		Implementation->Collisions.SetDetectorSettings(Implementation->NarrowPhase.GetContext().GetSettings());
+		FCollisionDetectorSettings DetectorSettings = Implementation->CollisionDetector.GetSettings();
+		DetectorSettings.bFilteringEnabled = false;
+		DetectorSettings.bAllowManifoldReuse = false;
+		DetectorSettings.bDeferNarrowPhase = (ChaosImmediate_Collision_DeferNarrowPhase != 0);;
+		DetectorSettings.bAllowManifolds = (ChaosImmediate_Collision_UseManifolds != 0);;
+		DetectorSettings.bAllowCCD = false; 
+		Implementation->CollisionDetector.SetSettings(DetectorSettings);
 	}
 
 	FSimulation::~FSimulation()
@@ -750,7 +748,7 @@ namespace ImmediatePhysics_Chaos
 
 		if (CullDistance >= FReal(0))
 		{
-			Implementation->NarrowPhase.SetBoundsExpansion(CullDistance);
+			Implementation->CollisionDetector.SetBoundsExpansion(CullDistance);
 		}
 
 		if (MaxDepenetrationVelocity >= FReal(0))
@@ -861,10 +859,11 @@ namespace ImmediatePhysics_Chaos
 
 			Implementation->Evolution.SetConstraintContainerPriority(Implementation->Collisions.GetContainerId(), ChaosImmediate_Collision_Priority);
 
-			Implementation->NarrowPhase.GetContext().GetSettings().bAllowManifoldReuse = false;
-			Implementation->NarrowPhase.GetContext().GetSettings().bDeferNarrowPhase = (ChaosImmediate_Collision_DeferNarrowPhase != 0);
-			Implementation->NarrowPhase.GetContext().GetSettings().bAllowManifolds = (ChaosImmediate_Collision_UseManifolds != 0);
-			Implementation->Collisions.SetDetectorSettings(Implementation->NarrowPhase.GetContext().GetSettings());
+			FCollisionDetectorSettings DetectorSettings = Implementation->CollisionDetector.GetSettings();
+			DetectorSettings.bAllowManifoldReuse = false;
+			DetectorSettings.bDeferNarrowPhase = (ChaosImmediate_Collision_DeferNarrowPhase != 0);;
+			DetectorSettings.bAllowManifolds = (ChaosImmediate_Collision_UseManifolds != 0);;
+			Implementation->CollisionDetector.SetSettings(DetectorSettings);
 
 			Implementation->Evolution.SetBoundsExtension(ChaosImmediate_Evolution_BoundsExtension);
 
diff --git a/Engine/Source/Runtime/Engine/Public/Physics/ImmediatePhysics/ImmediatePhysicsChaos/ImmediatePhysicsCore_Chaos.h b/Engine/Source/Runtime/Engine/Public/Physics/ImmediatePhysics/ImmediatePhysicsChaos/ImmediatePhysicsCore_Chaos.h
--- a/Engine/Source/Runtime/Engine/Public/Physics/ImmediatePhysics/ImmediatePhysicsChaos/ImmediatePhysicsCore_Chaos.h
+++ b/Engine/Source/Runtime/Engine/Public/Physics/ImmediatePhysics/ImmediatePhysicsChaos/ImmediatePhysicsCore_Chaos.h
@@ -12,7 +12,6 @@
 namespace Chaos
 {
 	class FImplicitObject;
-	class FNarrowPhase;
 	class FBasicBroadPhase;
 	class FBasicCollisionDetector;
 	class FPBDJointConstraintHandle;
diff --git a/Engine/Source/Runtime/Experimental/Chaos/Private/Chaos/Collision/CollisionConstraintAllocator.cpp b/Engine/Source/Runtime/Experimental/Chaos/Private/Chaos/Collision/CollisionConstraintAllocator.cpp
--- a/Engine/Source/Runtime/Experimental/Chaos/Private/Chaos/Collision/CollisionConstraintAllocator.cpp
+++ b/Engine/Source/Runtime/Experimental/Chaos/Private/Chaos/Collision/CollisionConstraintAllocator.cpp
@@ -3,4 +3,73 @@
 
 namespace Chaos
 {
+	void FCollisionContextAllocator::ProcessNewMidPhases(FCollisionConstraintAllocator* Allocator)
+	{
+		// NOTE: Called from the physics thread, and never from a physics task/parallel-for. No need for locks.
+
+		// Register the midphases with each of their particles and add the midphase to the central list
+		for (FParticlePairMidPhase* MidPhase : NewMidPhases)
+		{
+			FGeometryParticleHandle* Particle0 = MidPhase->GetParticle0();
+			FGeometryParticleHandle* Particle1 = MidPhase->GetParticle1();
+
+			// NOTE: It may look weird that we are passing the particle's self to the AddMidPhase method. 
+			// The midphase represents two particles and it needs to know which one we're currently adding it to.
+			// We could make this API nicer but providing a helper on the Particle, but it's a bit awkard to do so
+			// and since AddMidPahse is only called from here we'll leave it like this for now.
+			Particle0->ParticleCollisions().AddMidPhase(Particle0, MidPhase);
+			Particle1->ParticleCollisions().AddMidPhase(Particle1, MidPhase);
+
+#if CHAOS_COLLISION_OBJECTPOOL_ENABLED 
+			Allocator->AddMidPhase(FParticlePairMidPhasePtr(MidPhase, FParticlePairMidPhaseDeleter(MidPhasePool)));
+#else
+			Allocator->AddMidPhase(FParticlePairMidPhasePtr(MidPhase));
+#endif
+		}
+		NewMidPhases.Reset();
+	}
+
+	void FCollisionContextAllocator::ProcessNewConstraints(FCollisionConstraintAllocator* Allocator)
+	{
+		// NOTE: Called from the physics thread, and never from a physics task/parallel-for. No need for locks.
+
+		for (FPBDCollisionConstraint* NewConstraint : NewActiveConstraints)
+		{
+			Allocator->ActivateConstraintImp(NewConstraint);
+		}
+		NewActiveConstraints.Reset();
+	}
+
+	void FCollisionConstraintAllocator::PruneExpiredMidPhases()
+	{
+		// NOTE: Called from the physics thread, and never from a physics task/parallel-for. No need for locks.
+
+		// Determine which particle pairs are no longer overlapping
+		// Prune all pairs which were not updated this tick as part of the collision
+		// detection loop and are not asleep
+		for (int32 Index = ParticlePairMidPhases.Num() - 1; Index >= 0; --Index)
+		{
+			FParticlePairMidPhasePtr& MidPhase = ParticlePairMidPhases[Index];
+
+			// We could also add !MidPhase->IsInConstraintGraph() here, but we know that we will not be in the graph if we were
+			// not active this tick and were not asleep. The constraint graph ejects all non-sleeping constraints each tick.
+			// (There is a check in the collision destructor that verified this).
+			if (!MidPhase->IsUsedSince(CurrentEpoch) && !MidPhase->IsSleeping())
+			{
+				// Remove from the particles' lists of contacts
+				DetachParticlePairMidPhase(MidPhase.Get());
+
+				// Destroy the midphase and its collisions
+				MidPhase.Reset();
+
+				// Remove the midphase from the list. 
+				// ParticlePairMidPhases can get large so we allow it to shrink from time to time
+				const int32 MaxSlack = 1000;
+				const int32 Slack = ParticlePairMidPhases.Max() - ParticlePairMidPhases.Num();
+				const bool bAllowShrink = (Slack > MaxSlack);
+				ParticlePairMidPhases.RemoveAtSwap(Index, 1, bAllowShrink);
+			}
+		}
+	}
+
 }
\ No newline at end of file
diff --git a/Engine/Source/Runtime/Experimental/Chaos/Private/Chaos/Collision/PBDCollisionConstraint.cpp b/Engine/Source/Runtime/Experimental/Chaos/Private/Chaos/Collision/PBDCollisionConstraint.cpp
--- a/Engine/Source/Runtime/Experimental/Chaos/Private/Chaos/Collision/PBDCollisionConstraint.cpp
+++ b/Engine/Source/Runtime/Experimental/Chaos/Private/Chaos/Collision/PBDCollisionConstraint.cpp
@@ -407,7 +407,7 @@ namespace Chaos
 		// constraints on the same particle pair is a little wasteful. It early-outs on subsequent calls, but still not ideal.
 		// @todo(chaos): we only need to set sleeping on particle pairs or particles, not constraints (See UpdateSleepState in IslandManager.cpp)
 		check(ContainerCookie.MidPhase != nullptr);
-		ContainerCookie.MidPhase->SetIsSleeping(bInIsSleeping);
+		ContainerCookie.MidPhase->SetIsSleeping(bInIsSleeping, ConcreteContainer()->GetConstraintAllocator().GetCurrentEpoch());
 	}
 
 	FVec3 FPBDCollisionConstraint::CalculateWorldContactLocation() const
diff --git a/Engine/Source/Runtime/Experimental/Chaos/Private/Chaos/Collision/ParticlePairMidPhase.cpp b/Engine/Source/Runtime/Experimental/Chaos/Private/Chaos/Collision/ParticlePairMidPhase.cpp
--- a/Engine/Source/Runtime/Experimental/Chaos/Private/Chaos/Collision/ParticlePairMidPhase.cpp
+++ b/Engine/Source/Runtime/Experimental/Chaos/Private/Chaos/Collision/ParticlePairMidPhase.cpp
@@ -69,7 +69,7 @@ namespace Chaos
 		const FReal CullDistance,
 		const EContactShapesType ShapePairType,
 		const bool bUseManifold,
-		FCollisionConstraintAllocator& Allocator)
+		const FCollisionContext& Context)
 	{
 		const FImplicitObject* Implicit0 = InShape0->GetLeafGeometry();
 		const FBVHParticles* BVHParticles0 = FConstGenericParticleHandle(Particle0)->CollisionParticles().Get();
@@ -78,7 +78,7 @@ namespace Chaos
 		const FBVHParticles* BVHParticles1 = FConstGenericParticleHandle(Particle1)->CollisionParticles().Get();
 		const FRigidTransform3& ShapeRelativeTransform1 = (FRigidTransform3)InShape1->GetLeafRelativeTransform();
 
-		return Allocator.CreateCollisionConstraint(Particle0, Implicit0, InShape0, BVHParticles0, ShapeRelativeTransform0, Particle1, Implicit1, InShape1, BVHParticles1, ShapeRelativeTransform1, CullDistance, bUseManifold, ShapePairType);
+		return Context.GetAllocator()->CreateConstraint(Particle0, Implicit0, InShape0, BVHParticles0, ShapeRelativeTransform0, Particle1, Implicit1, InShape1, BVHParticles1, ShapeRelativeTransform1, CullDistance, bUseManifold, ShapePairType);
 	}
 
 	FPBDCollisionConstraintPtr CreateImplicitPairConstraint(
@@ -95,9 +95,9 @@ namespace Chaos
 		const FReal CullDistance,
 		const EContactShapesType ShapePairType,
 		const bool bUseManifold,
-		FCollisionConstraintAllocator& Allocator)
+		const FCollisionContext& Context)
 	{
-		return Allocator.CreateCollisionConstraint(Particle0, Implicit0, Shape0, BVHParticles0, ShapeRelativeTransform0, Particle1, Implicit1, Shape1, BVHParticles1, ShapeRelativeTransform1, CullDistance, bUseManifold, ShapePairType);
+		return Context.GetAllocator()->CreateConstraint(Particle0, Implicit0, Shape0, BVHParticles0, ShapeRelativeTransform0, Particle1, Implicit1, Shape1, BVHParticles1, ShapeRelativeTransform1, CullDistance, bUseManifold, ShapePairType);
 	}
 
 
@@ -175,7 +175,7 @@ namespace Chaos
 	{
 	}
 
-	bool FSingleShapePairCollisionDetector::DoBoundsOverlap(const FReal CullDistance)
+	bool FSingleShapePairCollisionDetector::DoBoundsOverlap(const FReal CullDistance, const int32 CurrentEpoch)
 	{
 		PHYSICS_CSV_SCOPED_EXPENSIVE(PhysicsVerbose, NarrowPhase_ShapeBounds);
 
@@ -209,7 +209,7 @@ namespace Chaos
 		// However, it is quite expensive to do this all the time so we only do this test when we did not 
 		// collide last frame. This is ok if we assume not much changes from frame to frame, but it means
 		// we might call the narrow phase one time too many when shapes become separated.
-		const int32 LastEpoch = MidPhase.GetCollisionAllocator().GetCurrentEpoch() - 1;
+		const int32 LastEpoch = CurrentEpoch - 1;
 		const bool bCollidedLastTick = IsUsedSince(LastEpoch);
 		if ((Flags.bEnableOBBCheck0 || Flags.bEnableOBBCheck1) && !bCollidedLastTick)
 		{
@@ -241,9 +241,10 @@ namespace Chaos
 	int32 FSingleShapePairCollisionDetector::GenerateCollision(
 		const FReal CullDistance,
 		const FReal Dt,
-		FCollisionContext& Context)
+		const FCollisionContext& Context)
 	{
-		if (DoBoundsOverlap(CullDistance))
+		const int32 CurrentEpoch = Context.GetAllocator()->GetCurrentEpoch();
+		if (DoBoundsOverlap(CullDistance, CurrentEpoch))
 		{
 			return GenerateCollisionImpl(CullDistance, Dt, Context);
 		}
@@ -253,28 +254,29 @@ namespace Chaos
 	int32 FSingleShapePairCollisionDetector::GenerateCollisionCCD(
 		const FReal CullDistance,
 		const FReal Dt,
-		FCollisionContext& Context)
+		const FCollisionContext& Context)
 	{
 		return GenerateCollisionCCDImpl(CullDistance, Dt, Context);
 	}
 
-	void FSingleShapePairCollisionDetector::CreateConstraint(const FReal CullDistance, FCollisionContext& Context)
+	void FSingleShapePairCollisionDetector::CreateConstraint(const FReal CullDistance, const FCollisionContext& Context)
 	{
 		PHYSICS_CSV_SCOPED_EXPENSIVE(PhysicsVerbose, NarrowPhase_CreateConstraint);
 		check(!Constraint.IsValid());
 
-		Constraint = CreateShapePairConstraint(GetParticle0(), Shape0, GetParticle1(), Shape1, CullDistance, ShapePairType, Context.GetSettings().bAllowManifolds, MidPhase.GetCollisionAllocator());
+		Constraint = CreateShapePairConstraint(GetParticle0(), Shape0, GetParticle1(), Shape1, CullDistance, ShapePairType, Context.GetSettings().bAllowManifolds, Context);
 
+		const int32 CurrentEpoch = Context.GetAllocator()->GetCurrentEpoch();
 		Constraint->GetContainerCookie().MidPhase = &MidPhase;
 		Constraint->GetContainerCookie().bIsMultiShapePair = false;
-		Constraint->GetContainerCookie().CreationEpoch = MidPhase.GetCollisionAllocator().GetCurrentEpoch();
+		Constraint->GetContainerCookie().CreationEpoch = CurrentEpoch;
 		LastUsedEpoch = -1;
 	}
 
 	int32 FSingleShapePairCollisionDetector::GenerateCollisionImpl(
 		const FReal CullDistance, 
 		const FReal Dt,
-		FCollisionContext& Context)
+		const FCollisionContext& Context)
 	{
 		if (Flags.bIsProbe)
 		{
@@ -296,7 +298,7 @@ namespace Chaos
 
 			const FRigidTransform3 ShapeWorldTransform0 = Shape0->GetLeafWorldTransform(GetParticle0());
 			const FRigidTransform3 ShapeWorldTransform1 = Shape1->GetLeafWorldTransform(GetParticle1());
-			const int32 CurrentEpoch = MidPhase.GetCollisionAllocator().GetCurrentEpoch();
+			const int32 CurrentEpoch = Context.GetAllocator()->GetCurrentEpoch();
 			const int32 LastEpoch = CurrentEpoch - 1;
 			const bool bWasUpdatedLastTick = IsUsedSince(LastEpoch);
 
@@ -342,7 +344,7 @@ namespace Chaos
 			// don't know in advance whether we will pass the Phi check (deferred narrow phase is used with RBAN)
 			if (Constraint->GetPhi() <= CullDistance || Context.GetSettings().bDeferNarrowPhase)
 			{
-				if (MidPhase.GetCollisionAllocator().ActivateConstraint(Constraint.Get()))
+				if (Context.GetAllocator()->ActivateConstraint(Constraint.Get()))
 				{
 					// If we had any mods last tick, undo them
 					Constraint->ResetModifications();
@@ -359,7 +361,7 @@ namespace Chaos
 	int32 FSingleShapePairCollisionDetector::GenerateCollisionCCDImpl(
 		const FReal CullDistance, 
 		const FReal Dt,
-		FCollisionContext& Context)
+		const FCollisionContext& Context)
 	{
 		if (Flags.bIsProbe)
 		{
@@ -413,8 +415,8 @@ namespace Chaos
 				Constraint->SetCCDEnabled(false);
 			}
 
-			MidPhase.GetCollisionAllocator().ActivateConstraint(Constraint.Get());
-			LastUsedEpoch = MidPhase.GetCollisionAllocator().GetCurrentEpoch();
+			Context.GetAllocator()->ActivateConstraint(Constraint.Get());
+			LastUsedEpoch = Context.GetAllocator()->GetCurrentEpoch();
 
 			return 1;
 		}
@@ -425,7 +427,7 @@ namespace Chaos
 	int32 FSingleShapePairCollisionDetector::GenerateCollisionProbeImpl(
 		const FReal CullDistance, 
 		const FReal Dt,
-		FCollisionContext& Context)
+		const FCollisionContext& Context)
 	{
 		// Same as regular constraint generation, but always defer narrow phase.
 		// Don't do any initial constraint computations.
@@ -439,15 +441,15 @@ namespace Chaos
 		{
 			PHYSICS_CSV_SCOPED_EXPENSIVE(PhysicsVerbose, NarrowPhase_UpdateConstraintProbe);
 
-			MidPhase.GetCollisionAllocator().ActivateConstraint(Constraint.Get());
-			LastUsedEpoch = MidPhase.GetCollisionAllocator().GetCurrentEpoch();
+			Context.GetAllocator()->ActivateConstraint(Constraint.Get());
+			LastUsedEpoch = Context.GetAllocator()->GetCurrentEpoch();
 			return 1;
 		}
 
 		return 0;
 	}
 
-	void FSingleShapePairCollisionDetector::WakeCollision(const int32 SleepEpoch)
+	void FSingleShapePairCollisionDetector::WakeCollision(const int32 SleepEpoch, const int32 CurrentEpoch)
 	{
 		if (Constraint.IsValid() && IsUsedSince(SleepEpoch))
 		{
@@ -458,8 +460,7 @@ namespace Chaos
 
 			// We need to refresh the epoch so that the constraint state will be used as the previous
 			// state if the pair is still colliding in the next tick
-			const int32 CurrentEpoch = MidPhase.GetCollisionAllocator().GetCurrentEpoch();
-			Constraint->GetContainerCookie().LastUsedEpoch = MidPhase.GetCollisionAllocator().GetCurrentEpoch();
+			Constraint->GetContainerCookie().LastUsedEpoch = CurrentEpoch;
 			LastUsedEpoch = CurrentEpoch;
 
 			// We have skipped collision detection for this particle because it was asleep, so we need to update the transforms...
@@ -469,14 +470,14 @@ namespace Chaos
 		}
 	}
 
-	// This fucntion is called by the resim system to add/replace a collision with one from the history buffer
-	void FSingleShapePairCollisionDetector::SetCollision(const FPBDCollisionConstraint& SourceConstraint)
+	// This function is called by the resim system to add/replace a collision with one from the history buffer
+	void FSingleShapePairCollisionDetector::SetCollision(const FPBDCollisionConstraint& SourceConstraint, const FCollisionContext& Context)
 	{
-		const int32 CurrentEpoch = MidPhase.GetCollisionAllocator().GetCurrentEpoch();
+		const int32 CurrentEpoch = Context.GetAllocator()->GetCurrentEpoch();
 
 		if (!Constraint.IsValid())
 		{
-			Constraint = MidPhase.GetCollisionAllocator().CreateCollisionConstraint();
+			Constraint = Context.GetAllocator()->CreateConstraint();
 			Constraint->GetContainerCookie().MidPhase = &MidPhase;
 			Constraint->GetContainerCookie().bIsMultiShapePair = false;
 			Constraint->GetContainerCookie().CreationEpoch = CurrentEpoch;
@@ -487,7 +488,7 @@ namespace Chaos
 
 		// Add the constraint to the active list
 		// If the constraint already existed and was already active, this will do nothing
-		MidPhase.GetCollisionAllocator().ActivateConstraint(Constraint.Get());
+		Context.GetAllocator()->ActivateConstraint(Constraint.Get());
 		LastUsedEpoch = CurrentEpoch;
 	}
 
@@ -531,7 +532,7 @@ namespace Chaos
 	int32 FMultiShapePairCollisionDetector::GenerateCollisions(
 		const FReal CullDistance,
 		const FReal Dt,
-		FCollisionContext& Context)
+		const FCollisionContext& Context)
 	{
 		FConstGenericParticleHandle P0 = Particle0;
 		FConstGenericParticleHandle P1 = Particle1;
@@ -565,7 +566,7 @@ namespace Chaos
 			Dt,
 			LocalContext);
 
-		int32 NumActiveConstraints = ProcessNewConstraints();
+		int32 NumActiveConstraints = ProcessNewConstraints(Context);
 
 		// @todo(chaos): we could clean up unused collisions between this pair, but probably not worth it
 		// and we would have to prevent cleanup for sleeping particles because the collisions will still
@@ -588,7 +589,8 @@ namespace Chaos
 		const FRigidTransform3& ShapeRelativeTransform1,
 		const FReal CullDistance,
 		const EContactShapesType ShapePairType,
-		const bool bUseManifold)
+		const bool bUseManifold,
+		const FCollisionContext& Context)
 	{
 		// This is a callback from the low-level collision function. It should always be the same two particles, though the
 		// shapes may be some children in the implicit hierarchy. The particles could be in the opposite order though, and
@@ -619,7 +621,7 @@ namespace Chaos
 		if (Constraint == nullptr)
 		{
 			// NOTE: Using InParticle0 and InParticle1 here because the order may be different to what we have stored
-			Constraint = CreateConstraint(InParticle0, Implicit0, InShape0, BVHParticles0, ShapeRelativeTransform0, InParticle1, Implicit1, InShape1, BVHParticles1, ShapeRelativeTransform1, CullDistance, ShapePairType, bUseManifold, Key);
+			Constraint = CreateConstraint(InParticle0, Implicit0, InShape0, BVHParticles0, ShapeRelativeTransform0, InParticle1, Implicit1, InShape1, BVHParticles1, ShapeRelativeTransform1, CullDistance, ShapePairType, bUseManifold, Key, Context);
 		}
 
 		// @todo(chaos): we already have the shape world transforms at the calling site - pass them in
@@ -648,10 +650,11 @@ namespace Chaos
 		const FBVHParticles* BVHParticles1,
 		const FRigidTransform3& ShapeRelativeTransform1,
 		const FReal CullDistance,
-		const EContactShapesType ShapePairType)
+		const EContactShapesType ShapePairType,
+		const FCollisionContext& Context)
 	{
 		const bool bUseManifold = true;
-		FPBDCollisionConstraint* Constraint = FindOrCreateConstraint(InParticle0, Implicit0, InShape0, BVHParticles0, ShapeRelativeTransform0, InParticle1, Implicit1, InShape1, BVHParticles1, ShapeRelativeTransform1, CullDistance, ShapePairType, bUseManifold);
+		FPBDCollisionConstraint* Constraint = FindOrCreateConstraint(InParticle0, Implicit0, InShape0, BVHParticles0, ShapeRelativeTransform0, InParticle1, Implicit1, InShape1, BVHParticles1, ShapeRelativeTransform1, CullDistance, ShapePairType, bUseManifold, Context);
 		if (Constraint != nullptr)
 		{
 			Constraint->SetCCDEnabled(true);
@@ -683,21 +686,22 @@ namespace Chaos
 		const FReal CullDistance,
 		const EContactShapesType ShapePairType,
 		const bool bUseManifold,
-		const FCollisionParticlePairConstraintKey& Key)
+		const FCollisionParticlePairConstraintKey& Key,
+		const FCollisionContext& Context)
 	{
 		PHYSICS_CSV_SCOPED_EXPENSIVE(PhysicsVerbose, NarrowPhase_CreateConstraint);
-		FPBDCollisionConstraintPtr Constraint = CreateImplicitPairConstraint(InParticle0, Implicit0, InShape0, BVHParticles0, ShapeRelativeTransform0, InParticle1, Implicit1, InShape1, BVHParticles1, ShapeRelativeTransform1, CullDistance, ShapePairType, bUseManifold, MidPhase.GetCollisionAllocator());
+		FPBDCollisionConstraintPtr Constraint = CreateImplicitPairConstraint(InParticle0, Implicit0, InShape0, BVHParticles0, ShapeRelativeTransform0, InParticle1, Implicit1, InShape1, BVHParticles1, ShapeRelativeTransform1, CullDistance, ShapePairType, bUseManifold, Context);
 
+		const int32 CurrentEpoch = Context.GetAllocator()->GetCurrentEpoch();
 		Constraint->GetContainerCookie().MidPhase = &MidPhase;
 		Constraint->GetContainerCookie().bIsMultiShapePair = true;
-		Constraint->GetContainerCookie().CreationEpoch = MidPhase.GetCollisionAllocator().GetCurrentEpoch();
+		Constraint->GetContainerCookie().CreationEpoch = CurrentEpoch;
 
 		return Constraints.Add(Key.GetKey(), MoveTemp(Constraint)).Get();
 	}
 
-	void FMultiShapePairCollisionDetector::WakeCollisions(const int32 SleepEpoch)
+	void FMultiShapePairCollisionDetector::WakeCollisions(const int32 SleepEpoch, const int32 CurrentEpoch)
 	{
-		const int32 CurrentEpoch = MidPhase.GetCollisionAllocator().GetCurrentEpoch();
 		for (auto& KVP : Constraints)
 		{
 			FPBDCollisionConstraintPtr& Constraint = KVP.Value;
@@ -708,14 +712,14 @@ namespace Chaos
 		}
 	}
 
-	int32 FMultiShapePairCollisionDetector::ProcessNewConstraints()
+	int32 FMultiShapePairCollisionDetector::ProcessNewConstraints(const FCollisionContext& Context)
 	{
 		int32 NumActiveConstraints = 0;
 		for (FPBDCollisionConstraint* Constraint : NewConstraints)
 		{
 			if (Constraint->GetPhi() < Constraint->GetCullDistance())
 			{
-				MidPhase.GetCollisionAllocator().ActivateConstraint(Constraint);
+				Context.GetAllocator()->ActivateConstraint(Constraint);
 				++NumActiveConstraints;
 			}
 		}
@@ -723,13 +727,11 @@ namespace Chaos
 		return NumActiveConstraints;
 	}
 
-	void FMultiShapePairCollisionDetector::PruneConstraints()
+	void FMultiShapePairCollisionDetector::PruneConstraints(const int32 CurrentEpoch)
 	{
 		// We don't prune from NewCollisions - must call ProcessNewCollisions before Prune
 		check(NewConstraints.Num() == 0);
 
-		const int32 CurrentEpoch = MidPhase.GetCollisionAllocator().GetCurrentEpoch();
-
 		// Find all the expired collisions
 		TArray<uint32> Pruned;
 		for (auto& KVP : Constraints)
@@ -762,7 +764,6 @@ namespace Chaos
 		, CullDistanceScale(1)
 		, Flags()
 		, Key()
-		, CollisionAllocator(nullptr)
 		, LastUsedEpoch(INDEX_NONE)
 		, NumActiveConstraints(0)
 		, ParticleCollisionsIndex0(INDEX_NONE)
@@ -777,11 +778,6 @@ namespace Chaos
 		Reset();
 	}
 
-	int32 FParticlePairMidPhase::GetCurrentEpoch() const
-	{
-		return CollisionAllocator->GetCurrentEpoch();
-	}
-
 	void FParticlePairMidPhase::DetachParticle(FGeometryParticleHandle* Particle)
 	{
 		Reset();
@@ -809,15 +805,13 @@ namespace Chaos
 	void FParticlePairMidPhase::Init(
 		FGeometryParticleHandle* InParticle0,
 		FGeometryParticleHandle* InParticle1,
-		const FCollisionParticlePairKey& InKey,
-		FCollisionConstraintAllocator& InCollisionAllocator)
+		const FCollisionParticlePairKey& InKey)
 	{
 		PHYSICS_CSV_SCOPED_EXPENSIVE(PhysicsVerbose, NarrowPhase_Filter);
 
 		Particle0 = InParticle0;
 		Particle1 = InParticle1;
 		Key = InKey;
-		CollisionAllocator = &InCollisionAllocator;
 
 		Flags.bIsCCD = FConstGenericParticleHandle(Particle0)->CCDEnabled() || FConstGenericParticleHandle(Particle1)->CCDEnabled();
 
@@ -926,7 +920,7 @@ namespace Chaos
 	void FParticlePairMidPhase::GenerateCollisions(
 		const FReal InCullDistance,
 		const FReal Dt,
-		FCollisionContext& Context)
+		const FCollisionContext& Context)
 	{
 		//SCOPE_CYCLE_COUNTER(STAT_Collisions_GenerateCollisions);
 		PHYSICS_CSV_SCOPED_EXPENSIVE(PhysicsVerbose, DetectCollisions_NarrowPhase);
@@ -968,10 +962,10 @@ namespace Chaos
 			NumActiveConstraints += MultiShapePair.GenerateCollisions(CullDistance, Dt, Context);
 		}
 
-		LastUsedEpoch = CollisionAllocator->GetCurrentEpoch();
+		LastUsedEpoch = Context.GetAllocator()->GetCurrentEpoch();
 	}
 
-	void FParticlePairMidPhase::InjectCollision(const FPBDCollisionConstraint& Constraint)
+	void FParticlePairMidPhase::InjectCollision(const FPBDCollisionConstraint& Constraint, const FCollisionContext& Context)
 	{
 		if (!Constraint.GetContainerCookie().bIsMultiShapePair)
 		{
@@ -984,7 +978,7 @@ namespace Chaos
 			{
 				if (((Shape0 == ShapePair.GetShape0()) && (Shape1 == ShapePair.GetShape1())) || ((Shape0 == ShapePair.GetShape1()) && (Shape1 == ShapePair.GetShape0())))
 				{
-					ShapePair.SetCollision(Constraint);
+					ShapePair.SetCollision(Constraint, Context);
 					break;
 				}
 			}
@@ -995,10 +989,10 @@ namespace Chaos
 			ensure(false);
 		}
 
-		LastUsedEpoch = CollisionAllocator->GetCurrentEpoch();
+		LastUsedEpoch = Context.GetAllocator()->GetCurrentEpoch();
 	}
 
-	void FParticlePairMidPhase::SetIsSleeping(const bool bInIsSleeping)
+	void FParticlePairMidPhase::SetIsSleeping(const bool bInIsSleeping, const int32 CurrentEpoch)
 	{
 		// This can be called from two locations:
 		// 1)	At the start of the tick as a results of some state change from the game thread such as an explicit wake event,
@@ -1019,18 +1013,18 @@ namespace Chaos
 			const bool bWakingUp = !bInIsSleeping;
 			if (bWakingUp)
 			{
-				if (LastUsedEpoch < CollisionAllocator->GetCurrentEpoch())
+				if (LastUsedEpoch < CurrentEpoch)
 				{
 					// Restore all constraints that were active when we were put to sleep
 					for (FSingleShapePairCollisionDetector& ShapePair : ShapePairDetectors)
 					{
-						ShapePair.WakeCollision(LastUsedEpoch);
+						ShapePair.WakeCollision(LastUsedEpoch, CurrentEpoch);
 					}
 					for (FMultiShapePairCollisionDetector& MultiShapePair : MultiShapePairDetectors)
 					{
-						MultiShapePair.WakeCollisions(LastUsedEpoch);
+						MultiShapePair.WakeCollisions(LastUsedEpoch, CurrentEpoch);
 					}
-					LastUsedEpoch = CollisionAllocator->GetCurrentEpoch();
+					LastUsedEpoch = CurrentEpoch;
 				}
 			}
 			// If we are going to sleep, there is nothing to do (see comments above)
diff --git a/Engine/Source/Runtime/Experimental/Chaos/Private/Chaos/CollisionResolution.cpp b/Engine/Source/Runtime/Experimental/Chaos/Private/Chaos/CollisionResolution.cpp
--- a/Engine/Source/Runtime/Experimental/Chaos/Private/Chaos/CollisionResolution.cpp
+++ b/Engine/Source/Runtime/Experimental/Chaos/Private/Chaos/CollisionResolution.cpp
@@ -421,7 +421,7 @@ namespace Chaos
 			const TBox<FReal, 3>* Object1 = Implicit1->template GetObject<const TBox<FReal, 3> >();
 			if (ensure(Object0 && Object1))
 			{
-				FPBDCollisionConstraint* Constraint = Context.MultiShapeCollisionDetector->FindOrCreateConstraint(Particle0, Implicit0, Shape0, nullptr, LocalTransform0, Particle1, Implicit1, Shape1, nullptr, LocalTransform1, CullDistance, EContactShapesType::BoxBox, Context.GetSettings().bAllowManifolds);
+				FPBDCollisionConstraint* Constraint = Context.MultiShapeCollisionDetector->FindOrCreateConstraint(Particle0, Implicit0, Shape0, nullptr, LocalTransform0, Particle1, Implicit1, Shape1, nullptr, LocalTransform1, CullDistance, EContactShapesType::BoxBox, Context.GetSettings().bAllowManifolds, Context);
 				if (T_TRAITS::bImmediateUpdate && (Constraint != nullptr))
 				{
 					const FRigidTransform3 WorldTransform0 = LocalTransform0 * ParticleWorldTransform0;
@@ -458,7 +458,7 @@ namespace Chaos
 			if (ensure(Object0 && Object1))
 			{
 				// @todo(chaos): one-shot manifold
-				FPBDCollisionConstraint* Constraint = Context.MultiShapeCollisionDetector->FindOrCreateConstraint(Particle0, Implicit0, Shape0, nullptr, LocalTransform0, Particle1, Implicit1, Shape1, nullptr, LocalTransform1, CullDistance, EContactShapesType::BoxHeightField, Context.GetSettings().bAllowManifolds);
+				FPBDCollisionConstraint* Constraint = Context.MultiShapeCollisionDetector->FindOrCreateConstraint(Particle0, Implicit0, Shape0, nullptr, LocalTransform0, Particle1, Implicit1, Shape1, nullptr, LocalTransform1, CullDistance, EContactShapesType::BoxHeightField, Context.GetSettings().bAllowManifolds, Context);
 				if (T_TRAITS::bImmediateUpdate && (Constraint != nullptr))
 				{
 					FRigidTransform3 WorldTransform0 = LocalTransform0 * ParticleWorldTransform0;
@@ -550,7 +550,7 @@ namespace Chaos
 			const TPlane<FReal, 3>* Object1 = Implicit1->template GetObject<const TPlane<FReal, 3> >();
 			if (ensure(Object0 && Object1))
 			{
-				FPBDCollisionConstraint* Constraint = Context.MultiShapeCollisionDetector->FindOrCreateConstraint(Particle0, Implicit0, Shape0, nullptr, LocalTransform0, Particle1, Implicit1, Shape1, nullptr, LocalTransform1, CullDistance, EContactShapesType::BoxPlane, Context.GetSettings().bAllowManifolds);
+				FPBDCollisionConstraint* Constraint = Context.MultiShapeCollisionDetector->FindOrCreateConstraint(Particle0, Implicit0, Shape0, nullptr, LocalTransform0, Particle1, Implicit1, Shape1, nullptr, LocalTransform1, CullDistance, EContactShapesType::BoxPlane, Context.GetSettings().bAllowManifolds, Context);
 				if (T_TRAITS::bImmediateUpdate && (Constraint != nullptr))
 				{
 					FRigidTransform3 WorldTransform0 = LocalTransform0 * ParticleWorldTransform0;
@@ -587,7 +587,7 @@ namespace Chaos
 				if (const TImplicitObjectScaled<FTriangleMeshImplicitObject>* ScaledTriangleMesh = Implicit1->template GetObject<const TImplicitObjectScaled<FTriangleMeshImplicitObject>>())
 				{
 					// @todo(chaos): one-shot manifold
-					FPBDCollisionConstraint* Constraint = Context.MultiShapeCollisionDetector->FindOrCreateConstraint(Particle0, Implicit0, Shape0, nullptr, LocalTransform0, Particle1, Implicit1, Shape1, nullptr, LocalTransform1, CullDistance, EContactShapesType::BoxTriMesh, Context.GetSettings().bAllowManifolds);
+					FPBDCollisionConstraint* Constraint = Context.MultiShapeCollisionDetector->FindOrCreateConstraint(Particle0, Implicit0, Shape0, nullptr, LocalTransform0, Particle1, Implicit1, Shape1, nullptr, LocalTransform1, CullDistance, EContactShapesType::BoxTriMesh, Context.GetSettings().bAllowManifolds, Context);
 					if (T_TRAITS::bImmediateUpdate && (Constraint != nullptr))
 					{
 						FRigidTransform3 WorldTransform0 = LocalTransform0 * ParticleWorldTransform0;
@@ -598,7 +598,7 @@ namespace Chaos
 				else if (const FTriangleMeshImplicitObject* TriangleMesh = Implicit1->template GetObject<const FTriangleMeshImplicitObject>())
 				{
 					// @todo(chaos): one-shot manifold
-					FPBDCollisionConstraint* Constraint = Context.MultiShapeCollisionDetector->FindOrCreateConstraint(Particle0, Implicit0, Shape0, nullptr, LocalTransform0, Particle1, Implicit1, Shape1, nullptr, LocalTransform1, CullDistance, EContactShapesType::BoxTriMesh, Context.GetSettings().bAllowManifolds);
+					FPBDCollisionConstraint* Constraint = Context.MultiShapeCollisionDetector->FindOrCreateConstraint(Particle0, Implicit0, Shape0, nullptr, LocalTransform0, Particle1, Implicit1, Shape1, nullptr, LocalTransform1, CullDistance, EContactShapesType::BoxTriMesh, Context.GetSettings().bAllowManifolds, Context);
 					if (T_TRAITS::bImmediateUpdate && (Constraint != nullptr))
 					{
 						FRigidTransform3 WorldTransform0 = LocalTransform0 * ParticleWorldTransform0;
@@ -636,7 +636,7 @@ namespace Chaos
 			const TSphere<FReal, 3>* Object1 = Implicit1->template GetObject<const TSphere<FReal, 3> >();
 			if (ensure(Object0 && Object1))
 			{
-				FPBDCollisionConstraint* Constraint = Context.MultiShapeCollisionDetector->FindOrCreateConstraint(Particle0, Implicit0, Shape0, nullptr, LocalTransform0, Particle1, Implicit1, Shape1, nullptr, LocalTransform1, CullDistance, EContactShapesType::SphereSphere, Context.GetSettings().bAllowManifolds);
+				FPBDCollisionConstraint* Constraint = Context.MultiShapeCollisionDetector->FindOrCreateConstraint(Particle0, Implicit0, Shape0, nullptr, LocalTransform0, Particle1, Implicit1, Shape1, nullptr, LocalTransform1, CullDistance, EContactShapesType::SphereSphere, Context.GetSettings().bAllowManifolds, Context);
 				if (T_TRAITS::bImmediateUpdate && (Constraint != nullptr))
 				{
 					FRigidTransform3 WorldTransform0 = LocalTransform0 * ParticleWorldTransform0;
@@ -677,7 +677,7 @@ namespace Chaos
 			const FHeightField* Object1 = Implicit1->template GetObject<const FHeightField >();
 			if (ensure(Object0 && Object1))
 			{
-				FPBDCollisionConstraint* Constraint = Context.MultiShapeCollisionDetector->FindOrCreateConstraint(Particle0, Implicit0, Shape0, nullptr, LocalTransform0, Particle1, Implicit1, Shape1, nullptr, LocalTransform1, CullDistance, EContactShapesType::SphereHeightField, Context.GetSettings().bAllowManifolds);
+				FPBDCollisionConstraint* Constraint = Context.MultiShapeCollisionDetector->FindOrCreateConstraint(Particle0, Implicit0, Shape0, nullptr, LocalTransform0, Particle1, Implicit1, Shape1, nullptr, LocalTransform1, CullDistance, EContactShapesType::SphereHeightField, Context.GetSettings().bAllowManifolds, Context);
 				if (T_TRAITS::bImmediateUpdate && (Constraint != nullptr))
 				{
 					FRigidTransform3 WorldTransform0 = LocalTransform0 * ParticleWorldTransform0;
@@ -693,7 +693,7 @@ namespace Chaos
 			const FHeightField* Object1 = Implicit1->template GetObject<const FHeightField>();
 			if (ensure(Object0 && Object1))
 			{
-				FPBDCollisionConstraint* Constraint = Context.MultiShapeCollisionDetector->FindOrCreateSweptConstraint(Particle0, Implicit0, Shape0, nullptr, LocalTransform0, Particle1, Implicit1, Shape1, nullptr, LocalTransform1, CullDistance, EContactShapesType::SphereHeightField);
+				FPBDCollisionConstraint* Constraint = Context.MultiShapeCollisionDetector->FindOrCreateSweptConstraint(Particle0, Implicit0, Shape0, nullptr, LocalTransform0, Particle1, Implicit1, Shape1, nullptr, LocalTransform1, CullDistance, EContactShapesType::SphereHeightField, Context);
 				FRigidTransform3 WorldTransformX0 = LocalTransform0 * FRigidTransform3(Particle0->X(), Particle0->R());
 				FRigidTransform3 WorldTransform1 = LocalTransform1 * ParticleWorldTransform1;
 				UpdateSphereHeightFieldConstraintSwept(Particle0, *Object0, WorldTransformX0, *Object1, WorldTransform1, Dir, Length, Dt, *Constraint);
@@ -724,7 +724,7 @@ namespace Chaos
 			const TPlane<FReal, 3>* Object1 = Implicit1->template GetObject<const TPlane<FReal, 3> >();
 			if (ensure(Object0 && Object1))
 			{
-				FPBDCollisionConstraint* Constraint = Context.MultiShapeCollisionDetector->FindOrCreateConstraint(Particle0, Implicit0, Shape0, nullptr, LocalTransform0, Particle1, Implicit1, Shape1, nullptr, LocalTransform1, CullDistance, EContactShapesType::SpherePlane, Context.GetSettings().bAllowManifolds);
+				FPBDCollisionConstraint* Constraint = Context.MultiShapeCollisionDetector->FindOrCreateConstraint(Particle0, Implicit0, Shape0, nullptr, LocalTransform0, Particle1, Implicit1, Shape1, nullptr, LocalTransform1, CullDistance, EContactShapesType::SpherePlane, Context.GetSettings().bAllowManifolds, Context);
 				if (T_TRAITS::bImmediateUpdate && (Constraint != nullptr))
 				{
 					FRigidTransform3 WorldTransform0 = LocalTransform0 * ParticleWorldTransform0;
@@ -759,7 +759,7 @@ namespace Chaos
 			const TBox<FReal, 3>* Object1 = Implicit1->template GetObject<const TBox<FReal, 3> >();
 			if (ensure(Object0 && Object1))
 			{
-				FPBDCollisionConstraint* Constraint = Context.MultiShapeCollisionDetector->FindOrCreateConstraint(Particle0, Implicit0, Shape0, nullptr, LocalTransform0, Particle1, Implicit1, Shape1, nullptr, LocalTransform1, CullDistance, EContactShapesType::SphereBox, Context.GetSettings().bAllowManifolds);
+				FPBDCollisionConstraint* Constraint = Context.MultiShapeCollisionDetector->FindOrCreateConstraint(Particle0, Implicit0, Shape0, nullptr, LocalTransform0, Particle1, Implicit1, Shape1, nullptr, LocalTransform1, CullDistance, EContactShapesType::SphereBox, Context.GetSettings().bAllowManifolds, Context);
 				if (T_TRAITS::bImmediateUpdate && (Constraint != nullptr))
 				{
 					FRigidTransform3 WorldTransform0 = LocalTransform0 * ParticleWorldTransform0;
@@ -794,7 +794,7 @@ namespace Chaos
 			const FCapsule* Object1 = Implicit1->template GetObject<const FCapsule >();
 			if (ensure(Object0 && Object1))
 			{
-				FPBDCollisionConstraint* Constraint = Context.MultiShapeCollisionDetector->FindOrCreateConstraint(Particle0, Implicit0, Shape0, nullptr, LocalTransform0, Particle1, Implicit1, Shape1, nullptr, LocalTransform1, CullDistance, EContactShapesType::SphereCapsule, Context.GetSettings().bAllowManifolds);
+				FPBDCollisionConstraint* Constraint = Context.MultiShapeCollisionDetector->FindOrCreateConstraint(Particle0, Implicit0, Shape0, nullptr, LocalTransform0, Particle1, Implicit1, Shape1, nullptr, LocalTransform1, CullDistance, EContactShapesType::SphereCapsule, Context.GetSettings().bAllowManifolds, Context);
 				if (T_TRAITS::bImmediateUpdate && (Constraint != nullptr))
 				{
 					FRigidTransform3 WorldTransform0 = LocalTransform0 * ParticleWorldTransform0;
@@ -827,7 +827,7 @@ namespace Chaos
 			const FImplicitObject* Object1 = Implicit1;
 			if (ensure(Object0 && Object1))
 			{
-				FPBDCollisionConstraint* Constraint = Context.MultiShapeCollisionDetector->FindOrCreateConstraint(Particle0, Implicit0, Shape0, nullptr, LocalTransform0, Particle1, Implicit1, Shape1, nullptr, LocalTransform1, CullDistance, EContactShapesType::SphereConvex, Context.GetSettings().bAllowManifolds);
+				FPBDCollisionConstraint* Constraint = Context.MultiShapeCollisionDetector->FindOrCreateConstraint(Particle0, Implicit0, Shape0, nullptr, LocalTransform0, Particle1, Implicit1, Shape1, nullptr, LocalTransform1, CullDistance, EContactShapesType::SphereConvex, Context.GetSettings().bAllowManifolds, Context);
 				if (T_TRAITS::bImmediateUpdate && (Constraint != nullptr))
 				{
 					FRigidTransform3 WorldTransform0 = LocalTransform0 * ParticleWorldTransform0;
@@ -871,7 +871,7 @@ namespace Chaos
 			{
 				if (const TImplicitObjectScaled<FTriangleMeshImplicitObject>* ScaledTriangleMesh = Implicit1->template GetObject<const TImplicitObjectScaled<FTriangleMeshImplicitObject>>())
 				{
-					FPBDCollisionConstraint* Constraint = Context.MultiShapeCollisionDetector->FindOrCreateConstraint(Particle0, Implicit0, Shape0, nullptr, LocalTransform0, Particle1, Implicit1, Shape1, nullptr, LocalTransform1, CullDistance, EContactShapesType::SphereTriMesh, Context.GetSettings().bAllowManifolds);
+					FPBDCollisionConstraint* Constraint = Context.MultiShapeCollisionDetector->FindOrCreateConstraint(Particle0, Implicit0, Shape0, nullptr, LocalTransform0, Particle1, Implicit1, Shape1, nullptr, LocalTransform1, CullDistance, EContactShapesType::SphereTriMesh, Context.GetSettings().bAllowManifolds, Context);
 					if (T_TRAITS::bImmediateUpdate && (Constraint != nullptr))
 					{
 						FRigidTransform3 WorldTransform0 = LocalTransform0 * ParticleWorldTransform0;
@@ -881,7 +881,7 @@ namespace Chaos
 				}
 				else if (const FTriangleMeshImplicitObject* TriangleMesh = Implicit1->template GetObject<const FTriangleMeshImplicitObject>())
 				{
-					FPBDCollisionConstraint* Constraint = Context.MultiShapeCollisionDetector->FindOrCreateConstraint(Particle0, Implicit0, Shape0, nullptr, LocalTransform0, Particle1, Implicit1, Shape1, nullptr, LocalTransform1, CullDistance, EContactShapesType::SphereTriMesh, Context.GetSettings().bAllowManifolds);
+					FPBDCollisionConstraint* Constraint = Context.MultiShapeCollisionDetector->FindOrCreateConstraint(Particle0, Implicit0, Shape0, nullptr, LocalTransform0, Particle1, Implicit1, Shape1, nullptr, LocalTransform1, CullDistance, EContactShapesType::SphereTriMesh, Context.GetSettings().bAllowManifolds, Context);
 					if (T_TRAITS::bImmediateUpdate && (Constraint != nullptr))
 					{
 						FRigidTransform3 WorldTransform0 = LocalTransform0 * ParticleWorldTransform0;
@@ -904,14 +904,14 @@ namespace Chaos
 			{
 				if (const TImplicitObjectScaled<FTriangleMeshImplicitObject>* ScaledTriangleMesh = Implicit1->template GetObject<const TImplicitObjectScaled<FTriangleMeshImplicitObject>>())
 				{
-					FPBDCollisionConstraint* Constraint = Context.MultiShapeCollisionDetector->FindOrCreateSweptConstraint(Particle0, Implicit0, Shape0, nullptr, LocalTransform0, Particle1, Implicit1, Shape1, nullptr, LocalTransform1, CullDistance, EContactShapesType::SphereTriMesh);
+					FPBDCollisionConstraint* Constraint = Context.MultiShapeCollisionDetector->FindOrCreateSweptConstraint(Particle0, Implicit0, Shape0, nullptr, LocalTransform0, Particle1, Implicit1, Shape1, nullptr, LocalTransform1, CullDistance, EContactShapesType::SphereTriMesh, Context);
 					FRigidTransform3 WorldTransformX0 = LocalTransform0 * FRigidTransform3(Particle0->X(), Particle0->R());
 					FRigidTransform3 WorldTransform1 = LocalTransform1 * ParticleWorldTransform1;
 					UpdateSphereTriangleMeshConstraintSwept(Particle0, *Object0, WorldTransformX0, *ScaledTriangleMesh, WorldTransform1, Dir, Length, Dt, *Constraint);
 				}
 				else if (const FTriangleMeshImplicitObject* TriangleMesh = Implicit1->template GetObject<const FTriangleMeshImplicitObject>())
 				{
-					FPBDCollisionConstraint* Constraint = Context.MultiShapeCollisionDetector->FindOrCreateSweptConstraint(Particle0, Implicit0, Shape0, nullptr, LocalTransform0, Particle1, Implicit1, Shape1, nullptr, LocalTransform1, CullDistance, EContactShapesType::SphereTriMesh);
+					FPBDCollisionConstraint* Constraint = Context.MultiShapeCollisionDetector->FindOrCreateSweptConstraint(Particle0, Implicit0, Shape0, nullptr, LocalTransform0, Particle1, Implicit1, Shape1, nullptr, LocalTransform1, CullDistance, EContactShapesType::SphereTriMesh, Context);
 					FRigidTransform3 WorldTransformX0 = LocalTransform0 * FRigidTransform3(Particle0->X(), Particle0->R());
 					FRigidTransform3 WorldTransform1 = LocalTransform1 * ParticleWorldTransform1;
 					UpdateSphereTriangleMeshConstraintSwept(Particle0, *Object0, WorldTransformX0, *TriangleMesh, WorldTransform1, Dir, Length, Dt, *Constraint);
@@ -949,7 +949,7 @@ namespace Chaos
 			const FCapsule* Object1 = Implicit1->template GetObject<const FCapsule >();
 			if (ensure(Object0 && Object1))
 			{
-				FPBDCollisionConstraint* Constraint = Context.MultiShapeCollisionDetector->FindOrCreateConstraint(Particle0, Implicit0, Shape0, nullptr, LocalTransform0, Particle1, Implicit1, Shape1, nullptr, LocalTransform1, CullDistance, EContactShapesType::CapsuleCapsule, Context.GetSettings().bAllowManifolds);
+				FPBDCollisionConstraint* Constraint = Context.MultiShapeCollisionDetector->FindOrCreateConstraint(Particle0, Implicit0, Shape0, nullptr, LocalTransform0, Particle1, Implicit1, Shape1, nullptr, LocalTransform1, CullDistance, EContactShapesType::CapsuleCapsule, Context.GetSettings().bAllowManifolds, Context);
 				if (T_TRAITS::bImmediateUpdate && (Constraint != nullptr))
 				{
 					FRigidTransform3 WorldTransform0 = LocalTransform0 * ParticleWorldTransform0;
@@ -990,7 +990,7 @@ namespace Chaos
 			const TBox<FReal, 3>* Object1 = Implicit1->template GetObject<const TBox<FReal, 3>>();
 			if (ensure(Object0 && Object1))
 			{
-				FPBDCollisionConstraint* Constraint = Context.MultiShapeCollisionDetector->FindOrCreateConstraint(Particle0, Implicit0, Shape0, nullptr, LocalTransform0, Particle1, Implicit1, Shape1, nullptr, LocalTransform1, CullDistance, EContactShapesType::CapsuleBox, Context.GetSettings().bAllowManifolds);
+				FPBDCollisionConstraint* Constraint = Context.MultiShapeCollisionDetector->FindOrCreateConstraint(Particle0, Implicit0, Shape0, nullptr, LocalTransform0, Particle1, Implicit1, Shape1, nullptr, LocalTransform1, CullDistance, EContactShapesType::CapsuleBox, Context.GetSettings().bAllowManifolds, Context);
 				if (T_TRAITS::bImmediateUpdate && (Constraint != nullptr))
 				{
 					FRigidTransform3 WorldTransform0 = LocalTransform0 * ParticleWorldTransform0;
@@ -1072,7 +1072,7 @@ namespace Chaos
 			const FImplicitObject* Object1 = Implicit1;
 			if (ensure(Object0 && Object1))
 			{
-				FPBDCollisionConstraint* Constraint = Context.MultiShapeCollisionDetector->FindOrCreateConstraint(Particle0, Implicit0, Shape0, nullptr, LocalTransform0, Particle1, Implicit1, Shape1, nullptr, LocalTransform1, CullDistance, EContactShapesType::CapsuleConvex, Context.GetSettings().bAllowManifolds);
+				FPBDCollisionConstraint* Constraint = Context.MultiShapeCollisionDetector->FindOrCreateConstraint(Particle0, Implicit0, Shape0, nullptr, LocalTransform0, Particle1, Implicit1, Shape1, nullptr, LocalTransform1, CullDistance, EContactShapesType::CapsuleConvex, Context.GetSettings().bAllowManifolds, Context);
 				if (T_TRAITS::bImmediateUpdate && (Constraint != nullptr))
 				{
 					FRigidTransform3 WorldTransform0 = LocalTransform0 * ParticleWorldTransform0;
@@ -1118,7 +1118,7 @@ namespace Chaos
 			const FHeightField* Object1 = Implicit1->template GetObject<const FHeightField >();
 			if (ensure(Object0 && Object1))
 			{
-				FPBDCollisionConstraint* Constraint = Context.MultiShapeCollisionDetector->FindOrCreateConstraint(Particle0, Implicit0, Shape0, nullptr, LocalTransform0, Particle1, Implicit1, Shape1, nullptr, LocalTransform1, CullDistance, EContactShapesType::CapsuleHeightField, Context.GetSettings().bAllowManifolds);
+				FPBDCollisionConstraint* Constraint = Context.MultiShapeCollisionDetector->FindOrCreateConstraint(Particle0, Implicit0, Shape0, nullptr, LocalTransform0, Particle1, Implicit1, Shape1, nullptr, LocalTransform1, CullDistance, EContactShapesType::CapsuleHeightField, Context.GetSettings().bAllowManifolds, Context);
 				if (T_TRAITS::bImmediateUpdate && (Constraint != nullptr))
 				{
 					FRigidTransform3 WorldTransform0 = LocalTransform0 * ParticleWorldTransform0;
@@ -1135,7 +1135,7 @@ namespace Chaos
 			const FHeightField* Object1 = Implicit1->template GetObject<const FHeightField >();
 			if (ensure(Object0 && Object1))
 			{
-				FPBDCollisionConstraint* Constraint = Context.MultiShapeCollisionDetector->FindOrCreateSweptConstraint(Particle0, Implicit0, Shape0, nullptr, LocalTransform0, Particle1, Implicit1, Shape1, nullptr, LocalTransform1, CullDistance, EContactShapesType::CapsuleHeightField);
+				FPBDCollisionConstraint* Constraint = Context.MultiShapeCollisionDetector->FindOrCreateSweptConstraint(Particle0, Implicit0, Shape0, nullptr, LocalTransform0, Particle1, Implicit1, Shape1, nullptr, LocalTransform1, CullDistance, EContactShapesType::CapsuleHeightField, Context);
 				FRigidTransform3 WorldTransformX0 = LocalTransform0 * FRigidTransform3(Particle0->X(), Particle0->R());
 				FRigidTransform3 WorldTransform1 = LocalTransform1 * ParticleWorldTransform1;
 				UpdateCapsuleHeightFieldConstraintSwept(Particle0, *Object0, WorldTransformX0, *Object1, WorldTransform1, Dir, Length, Dt, *Constraint);
@@ -1182,7 +1182,7 @@ namespace Chaos
 			{
 				if (const TImplicitObjectScaled<FTriangleMeshImplicitObject>* ScaledTriangleMesh = Implicit1->template GetObject<const TImplicitObjectScaled<FTriangleMeshImplicitObject>>())
 				{
-					FPBDCollisionConstraint* Constraint = Context.MultiShapeCollisionDetector->FindOrCreateConstraint(Particle0, Implicit0, Shape0, nullptr, LocalTransform0, Particle1, Implicit1, Shape1,  nullptr, LocalTransform1, CullDistance, EContactShapesType::CapsuleTriMesh, Context.GetSettings().bAllowManifolds);
+					FPBDCollisionConstraint* Constraint = Context.MultiShapeCollisionDetector->FindOrCreateConstraint(Particle0, Implicit0, Shape0, nullptr, LocalTransform0, Particle1, Implicit1, Shape1,  nullptr, LocalTransform1, CullDistance, EContactShapesType::CapsuleTriMesh, Context.GetSettings().bAllowManifolds, Context);
 					if (T_TRAITS::bImmediateUpdate && (Constraint != nullptr))
 					{
 						FRigidTransform3 WorldTransform0 = LocalTransform0 * ParticleWorldTransform0;
@@ -1192,7 +1192,7 @@ namespace Chaos
 				}
 				else if (const FTriangleMeshImplicitObject* TriangleMesh = Implicit1->template GetObject<const FTriangleMeshImplicitObject>())
 				{
-					FPBDCollisionConstraint* Constraint = Context.MultiShapeCollisionDetector->FindOrCreateConstraint(Particle0, Implicit0, Shape0, nullptr, LocalTransform0, Particle1, Implicit1, Shape1, nullptr, LocalTransform1, CullDistance, EContactShapesType::CapsuleTriMesh, Context.GetSettings().bAllowManifolds);
+					FPBDCollisionConstraint* Constraint = Context.MultiShapeCollisionDetector->FindOrCreateConstraint(Particle0, Implicit0, Shape0, nullptr, LocalTransform0, Particle1, Implicit1, Shape1, nullptr, LocalTransform1, CullDistance, EContactShapesType::CapsuleTriMesh, Context.GetSettings().bAllowManifolds, Context);
 					if (T_TRAITS::bImmediateUpdate && (Constraint != nullptr))
 					{
 						FRigidTransform3 WorldTransform0 = LocalTransform0 * ParticleWorldTransform0;
@@ -1216,14 +1216,14 @@ namespace Chaos
 			{
 				if (const TImplicitObjectScaled<FTriangleMeshImplicitObject>* ScaledTriangleMesh = Implicit1->template GetObject<const TImplicitObjectScaled<FTriangleMeshImplicitObject>>())
 				{
-					FPBDCollisionConstraint* Constraint = Context.MultiShapeCollisionDetector->FindOrCreateSweptConstraint(Particle0, Implicit0, Shape0, nullptr, LocalTransform0, Particle1, Implicit1, Shape1, nullptr, LocalTransform1, CullDistance, EContactShapesType::CapsuleTriMesh);
+					FPBDCollisionConstraint* Constraint = Context.MultiShapeCollisionDetector->FindOrCreateSweptConstraint(Particle0, Implicit0, Shape0, nullptr, LocalTransform0, Particle1, Implicit1, Shape1, nullptr, LocalTransform1, CullDistance, EContactShapesType::CapsuleTriMesh, Context);
 					FRigidTransform3 WorldTransformX0 = LocalTransform0 * FRigidTransform3(Particle0->X(), Particle0->R());
 					FRigidTransform3 WorldTransform1 = LocalTransform1 * ParticleWorldTransform1;
 					UpdateCapsuleTriangleMeshConstraintSwept(Particle0, *Object0, WorldTransformX0, *ScaledTriangleMesh, WorldTransform1, Dir, Length, Dt, *Constraint);
 				}
 				else if (const FTriangleMeshImplicitObject* TriangleMesh = Implicit1->template GetObject<const FTriangleMeshImplicitObject>())
 				{
-					FPBDCollisionConstraint* Constraint = Context.MultiShapeCollisionDetector->FindOrCreateSweptConstraint(Particle0, Implicit0, Shape0, nullptr, LocalTransform0, Particle1, Implicit1, Shape1, nullptr, LocalTransform1, CullDistance, EContactShapesType::CapsuleTriMesh);
+					FPBDCollisionConstraint* Constraint = Context.MultiShapeCollisionDetector->FindOrCreateSweptConstraint(Particle0, Implicit0, Shape0, nullptr, LocalTransform0, Particle1, Implicit1, Shape1, nullptr, LocalTransform1, CullDistance, EContactShapesType::CapsuleTriMesh, Context);
 					FRigidTransform3 WorldTransformX0 = LocalTransform0 * FRigidTransform3(Particle0->X(), Particle0->R());
 					FRigidTransform3 WorldTransform1 = LocalTransform1 * ParticleWorldTransform1;
 					UpdateCapsuleTriangleMeshConstraintSwept(Particle0, *Object0, WorldTransformX0, *TriangleMesh, WorldTransform1, Dir, Length, Dt, *Constraint);
@@ -1264,7 +1264,7 @@ namespace Chaos
 			EImplicitObjectType Implicit0Type = Particle0->Geometry()->GetType();
 			EImplicitObjectType Implicit1Type = Particle1->Geometry()->GetType();
 
-			FPBDCollisionConstraint* Constraint = Context.MultiShapeCollisionDetector->FindOrCreateConstraint(Particle0, Implicit0, Shape0, nullptr, LocalTransform0, Particle1, Implicit1, Shape1, nullptr, LocalTransform1, CullDistance, EContactShapesType::GenericConvexConvex, Context.GetSettings().bAllowManifolds);
+			FPBDCollisionConstraint* Constraint = Context.MultiShapeCollisionDetector->FindOrCreateConstraint(Particle0, Implicit0, Shape0, nullptr, LocalTransform0, Particle1, Implicit1, Shape1, nullptr, LocalTransform1, CullDistance, EContactShapesType::GenericConvexConvex, Context.GetSettings().bAllowManifolds, Context);
 			if (T_TRAITS::bImmediateUpdate && (Constraint != nullptr))
 			{
 				FRigidTransform3 WorldTransform0 = LocalTransform0 * ParticleWorldTransform0;
@@ -1278,7 +1278,7 @@ namespace Chaos
 			CONDITIONAL_SCOPE_CYCLE_COUNTER(STAT_Collisions_ConstructGenericConvexConvexConstraintsSwept, ConstraintsDetailedStats);
 			FGenericParticleHandle RigidParticle0 = FGenericParticleHandle(Particle0);
 			FGenericParticleHandle RigidParticle1 = FGenericParticleHandle(Particle1);
-			FPBDCollisionConstraint* Constraint = Context.MultiShapeCollisionDetector->FindOrCreateSweptConstraint(Particle0, Implicit0, Shape0, nullptr, LocalTransform0, Particle1, Implicit1, Shape1, nullptr, LocalTransform1, CullDistance, EContactShapesType::GenericConvexConvex);
+			FPBDCollisionConstraint* Constraint = Context.MultiShapeCollisionDetector->FindOrCreateSweptConstraint(Particle0, Implicit0, Shape0, nullptr, LocalTransform0, Particle1, Implicit1, Shape1, nullptr, LocalTransform1, CullDistance, EContactShapesType::GenericConvexConvex, Context);
 			// Rotations of WorldTransformXRx are not actually used in UpdateGenericConvexConvexConstraintSwept.
 			FRigidTransform3 WorldTransformXR0 = LocalTransform0 * FRigidTransform3(StartX0, RigidParticle0->R());
 			FRigidTransform3 WorldTransformXR1 = LocalTransform1 * FRigidTransform3(StartX1, RigidParticle1->R());
@@ -1325,7 +1325,7 @@ namespace Chaos
 			const FHeightField* Object1 = Implicit1->template GetObject<const FHeightField >();
 			if (ensure(Implicit0->IsConvex() && Object1))
 			{
-				FPBDCollisionConstraint* Constraint = Context.MultiShapeCollisionDetector->FindOrCreateConstraint(Particle0, Implicit0, Shape0, nullptr, LocalTransform0, Particle1, Implicit1, Shape1, nullptr, LocalTransform1, CullDistance, EContactShapesType::ConvexHeightField, Context.GetSettings().bAllowManifolds);
+				FPBDCollisionConstraint* Constraint = Context.MultiShapeCollisionDetector->FindOrCreateConstraint(Particle0, Implicit0, Shape0, nullptr, LocalTransform0, Particle1, Implicit1, Shape1, nullptr, LocalTransform1, CullDistance, EContactShapesType::ConvexHeightField, Context.GetSettings().bAllowManifolds, Context);
 				if (T_TRAITS::bImmediateUpdate && (Constraint != nullptr))
 				{
 					FRigidTransform3 WorldTransform0 = LocalTransform0 * ParticleWorldTransform0;
@@ -1341,7 +1341,7 @@ namespace Chaos
 			const FHeightField* Object1 = Implicit1->template GetObject<const FHeightField >();
 			if (ensure(Implicit0->IsConvex() && Object1))
 			{
-				FPBDCollisionConstraint* Constraint = Context.MultiShapeCollisionDetector->FindOrCreateSweptConstraint(Particle0, Implicit0, Shape0, nullptr, LocalTransform0, Particle1, Implicit1, Shape1, nullptr, LocalTransform1, CullDistance, EContactShapesType::ConvexHeightField);
+				FPBDCollisionConstraint* Constraint = Context.MultiShapeCollisionDetector->FindOrCreateSweptConstraint(Particle0, Implicit0, Shape0, nullptr, LocalTransform0, Particle1, Implicit1, Shape1, nullptr, LocalTransform1, CullDistance, EContactShapesType::ConvexHeightField, Context);
 				FRigidTransform3 WorldTransformX0 = LocalTransform0 * FRigidTransform3(Particle0->X(), Particle0->R());
 				FRigidTransform3 WorldTransform1 = LocalTransform1 * ParticleWorldTransform1;
 				UpdateConvexHeightFieldConstraintSwept(Particle0, *Implicit0, WorldTransformX0, *Object1, WorldTransform1, Dir, Length, Dt, *Constraint);
@@ -1381,7 +1381,7 @@ namespace Chaos
 			CONDITIONAL_SCOPE_CYCLE_COUNTER(STAT_Collisions_ConstructConvexTriangleMeshConstraints, ConstraintsDetailedStats);
 			if (ensure(Implicit0->IsConvex()))
 			{
-				FPBDCollisionConstraint* Constraint = Context.MultiShapeCollisionDetector->FindOrCreateConstraint(Particle0, Implicit0, Shape0, nullptr, LocalTransform0, Particle1, Implicit1, Shape1, nullptr, LocalTransform1, CullDistance, EContactShapesType::ConvexTriMesh, Context.GetSettings().bAllowManifolds);
+				FPBDCollisionConstraint* Constraint = Context.MultiShapeCollisionDetector->FindOrCreateConstraint(Particle0, Implicit0, Shape0, nullptr, LocalTransform0, Particle1, Implicit1, Shape1, nullptr, LocalTransform1, CullDistance, EContactShapesType::ConvexTriMesh, Context.GetSettings().bAllowManifolds, Context);
 				if (T_TRAITS::bImmediateUpdate && (Constraint != nullptr))
 				{
 					FRigidTransform3 WorldTransform0 = LocalTransform0 * ParticleWorldTransform0;
@@ -1398,14 +1398,14 @@ namespace Chaos
 			{
 				if (const TImplicitObjectScaled<FTriangleMeshImplicitObject>* ScaledTriangleMesh = Implicit1->template GetObject<const TImplicitObjectScaled<FTriangleMeshImplicitObject>>())
 				{
-					FPBDCollisionConstraint* Constraint = Context.MultiShapeCollisionDetector->FindOrCreateSweptConstraint(Particle0, Implicit0, Shape0, nullptr, LocalTransform0, Particle1, Implicit1, Shape1, nullptr, LocalTransform1, CullDistance, EContactShapesType::ConvexTriMesh);
+					FPBDCollisionConstraint* Constraint = Context.MultiShapeCollisionDetector->FindOrCreateSweptConstraint(Particle0, Implicit0, Shape0, nullptr, LocalTransform0, Particle1, Implicit1, Shape1, nullptr, LocalTransform1, CullDistance, EContactShapesType::ConvexTriMesh, Context);
 					FRigidTransform3 WorldTransformX0 = LocalTransform0 * FRigidTransform3(Particle0->X(), Particle0->R());
 					FRigidTransform3 WorldTransform1 = LocalTransform1 * ParticleWorldTransform1;
 					UpdateConvexTriangleMeshConstraintSwept(Particle0, *Implicit0, WorldTransformX0, *ScaledTriangleMesh, WorldTransform1, Dir, Length, Dt, *Constraint);
 				}
 				else if (const FTriangleMeshImplicitObject* TriangleMesh = Implicit1->template GetObject<const FTriangleMeshImplicitObject>())
 				{
-					FPBDCollisionConstraint* Constraint = Context.MultiShapeCollisionDetector->FindOrCreateSweptConstraint(Particle0, Implicit0, Shape0, nullptr, LocalTransform0, Particle1, Implicit1, Shape1, nullptr, LocalTransform1, CullDistance, EContactShapesType::ConvexTriMesh);
+					FPBDCollisionConstraint* Constraint = Context.MultiShapeCollisionDetector->FindOrCreateSweptConstraint(Particle0, Implicit0, Shape0, nullptr, LocalTransform0, Particle1, Implicit1, Shape1, nullptr, LocalTransform1, CullDistance, EContactShapesType::ConvexTriMesh, Context);
 					FRigidTransform3 WorldTransformX0 = LocalTransform0 * FRigidTransform3(Particle0->X(), Particle0->R());
 					FRigidTransform3 WorldTransform1 = LocalTransform1 * ParticleWorldTransform1;
 					UpdateConvexTriangleMeshConstraintSwept(Particle0, *Implicit0, WorldTransformX0, *TriangleMesh, WorldTransform1, Dir, Length, Dt, *Constraint);
@@ -1489,11 +1489,11 @@ namespace Chaos
 			int32 P0NumCollisionParticles = Simplicial0 ? Simplicial0->Size() : Particle0->CastToRigidParticle()->CollisionParticlesSize();
 			if (!Particle1->Geometry() || (bIsParticleDynamic0 && !P0NumCollisionParticles && Particle0->Geometry() && !Particle0->Geometry()->IsUnderlyingUnion()))
 			{
-				Constraint = Context.MultiShapeCollisionDetector->FindOrCreateConstraint(Particle1, Implicit1, Shape1, Simplicial1, LocalTransform1, Particle0, Implicit0, Shape0, Simplicial0, LocalTransform0, CullDistance, EContactShapesType::LevelSetLevelSet, bAllowManifolds);
+				Constraint = Context.MultiShapeCollisionDetector->FindOrCreateConstraint(Particle1, Implicit1, Shape1, Simplicial1, LocalTransform1, Particle0, Implicit0, Shape0, Simplicial0, LocalTransform0, CullDistance, EContactShapesType::LevelSetLevelSet, bAllowManifolds, Context);
 			}
 			else
 			{
-				Constraint = Context.MultiShapeCollisionDetector->FindOrCreateConstraint(Particle0, Implicit0, Shape0, Simplicial0, LocalTransform0, Particle1, Implicit1, Shape1, Simplicial1, LocalTransform1, CullDistance, EContactShapesType::LevelSetLevelSet, bAllowManifolds);
+				Constraint = Context.MultiShapeCollisionDetector->FindOrCreateConstraint(Particle0, Implicit0, Shape0, Simplicial0, LocalTransform0, Particle1, Implicit1, Shape1, Simplicial1, LocalTransform1, CullDistance, EContactShapesType::LevelSetLevelSet, bAllowManifolds, Context);
 			}
 
 			if (Constraint != nullptr)
diff --git a/Engine/Source/Runtime/Experimental/Chaos/Private/Chaos/Evolution/PBDMinEvolution.cpp b/Engine/Source/Runtime/Experimental/Chaos/Private/Chaos/Evolution/PBDMinEvolution.cpp
--- a/Engine/Source/Runtime/Experimental/Chaos/Private/Chaos/Evolution/PBDMinEvolution.cpp
+++ b/Engine/Source/Runtime/Experimental/Chaos/Private/Chaos/Evolution/PBDMinEvolution.cpp
@@ -1,7 +1,6 @@
 // Copyright Epic Games, Inc. All Rights Reserved.
 
 #include "Chaos/Evolution/PBDMinEvolution.h"
-#include "Chaos/Collision/NarrowPhase.h"
 #include "Chaos/Collision/ParticlePairCollisionDetector.h"
 #include "Chaos/Evolution/SolverConstraintContainer.h"
 #include "Chaos/PBDCollisionConstraints.h"
@@ -17,11 +16,6 @@
 #include "Chaos/PerParticlePBDUpdateFromDeltaPosition.h"
 #include "ChaosStats.h"
 
-#if INTEL_ISPC
-#include "PBDMinEvolution.ispc.generated.h"
-#endif
-
-
 //PRAGMA_DISABLE_OPTIMIZATION
 
 namespace Chaos
@@ -293,7 +287,7 @@ namespace Chaos
 				// Update cached world space state, including bounds. We use the Swept bounds update so that the bounds includes P,Q and X,Q.
 				// This is because when we have joints, they often pull bodies back to their original positions, so we need to know if there
 				// are contacts at that location.
-				Particle.UpdateWorldSpaceStateSwept(FRigidTransform3(Particle.P(), Particle.Q()), FVec3(CollisionDetector.GetNarrowPhase().GetBoundsExpansion()), -V * Dt);
+				Particle.UpdateWorldSpaceStateSwept(FRigidTransform3(Particle.P(), Particle.Q()), FVec3(CollisionDetector.GetSettings().BoundsExpansion), -V * Dt);
 			}
 		}
 	}
diff --git a/Engine/Source/Runtime/Experimental/Chaos/Private/Chaos/Framework/Parallel.cpp b/Engine/Source/Runtime/Experimental/Chaos/Private/Chaos/Framework/Parallel.cpp
--- a/Engine/Source/Runtime/Experimental/Chaos/Private/Chaos/Framework/Parallel.cpp
+++ b/Engine/Source/Runtime/Experimental/Chaos/Private/Chaos/Framework/Parallel.cpp
@@ -144,6 +144,42 @@ void Chaos::PhysicsParallelForRange(int32 InNum, TFunctionRef<void(int32, int32)
 	::ParallelFor(RangeIndex.Num() - 1, PassThrough, !!GSingleThreadedPhysics || bDisablePhysicsParallelFor || bForceSingleThreaded);
 }
 
+void Chaos::PhysicsParallelForWithContext(int32 InNum, TFunctionRef<int32 (int32, int32)> InContextCreator, TFunctionRef<void(int32, int32)> InCallable, bool bForceSingleThreaded)
+{
+	TRACE_CPUPROFILER_EVENT_SCOPE(Chaos_PhysicsParallelFor);
+	using namespace Chaos;
+
+	// Passthrough for now, except with global flag to disable parallel
+#if PHYSICS_THREAD_CONTEXT
+	const bool bIsInPhysicsSimContext = IsInPhysicsThreadContext();
+	const bool bIsInGameThreadContext = IsInGameThreadContext();
+#else
+	const bool bIsInPhysicsSimContext = false;
+	const bool bIsInGameThreadContext = false;
+#endif
+
+	auto PassThrough = [InCallable, bIsInPhysicsSimContext, bIsInGameThreadContext](int32 ContextIndex, int32 Idx)
+	{
+#if PHYSICS_THREAD_CONTEXT
+		FPhysicsThreadContextScope PTScope(bIsInPhysicsSimContext);
+		FGameThreadContextScope GTScope(bIsInGameThreadContext);
+#endif
+		InCallable(Idx, ContextIndex);
+	};
+
+	const bool bSingleThreaded = !!GSingleThreadedPhysics || bDisablePhysicsParallelFor || bForceSingleThreaded;
+	const EParallelForFlags Flags = bSingleThreaded ? (EParallelForFlags::ForceSingleThread) : (EParallelForFlags::None);
+
+	// Unfortunately ParallelForWithTaskContext takes an array of context objects - we don't use it and in our case
+	// it ends up being an array where array[index] = index.
+	// The reason we don't need it is that our ContextCreator returns the context index we want to use on a given
+	// worker thread, and this is passed to the user function. The user function can just captures its array of
+	// contexts and use the context indeex to get its context from it.
+	TArray<int32, TInlineAllocator<16>> Contexts;
+
+	::ParallelForWithTaskContext(Contexts, InNum, InContextCreator, PassThrough, Flags);
+}
+
 
 //class FRecursiveDivideTask
 //{
diff --git a/Engine/Source/Runtime/Experimental/Chaos/Private/Chaos/PBDRigidsEvolutionGBF.cpp b/Engine/Source/Runtime/Experimental/Chaos/Private/Chaos/PBDRigidsEvolutionGBF.cpp
--- a/Engine/Source/Runtime/Experimental/Chaos/Private/Chaos/PBDRigidsEvolutionGBF.cpp
+++ b/Engine/Source/Runtime/Experimental/Chaos/Private/Chaos/PBDRigidsEvolutionGBF.cpp
@@ -141,10 +141,30 @@ namespace Chaos
 		// @chaos(todo): remove this when no longer needed
 		bool bChaos_Solver_TestMode  = false;
 		FAutoConsoleVariableRef CVarChaosSolverTestMode(TEXT("p.Chaos.Solver.TestMode"), bChaos_Solver_TestMode, TEXT(""));
-	}
 
+	}
 	using namespace CVars;
 
+	// We want to have largish blocks, but when we are running on multiple cores we don't want to waste too much memory
+	int32 CalculateNumCollisionsPerBlock()
+	{
+		const int32 NumCollisionsPerBlock = 2000;
+		const int32 MinCollisionsPerBlock = 100;
+
+		const bool bIsMultithreaded = (FApp::ShouldUseThreadingForPerformance() && !GSingleThreadedPhysics);
+		if (!bIsMultithreaded)
+		{
+			return NumCollisionsPerBlock;
+		}
+		else
+		{
+			const int32 NumWorkers = FMath::Min(FTaskGraphInterface::Get().GetNumWorkerThreads(), Chaos::MaxNumWorkers);
+			const int32 NumTasks = FMath::Max(NumWorkers, 1);
+			return FMath::Max(NumCollisionsPerBlock / NumTasks, MinCollisionsPerBlock);
+		}
+	}
+
+
 #if !UE_BUILD_SHIPPING
 template <typename TEvolution>
 void SerializeToDisk(TEvolution& Evolution)
@@ -598,10 +618,9 @@ void FPBDRigidsEvolutionGBF::TestModeResetCollisions()
 FPBDRigidsEvolutionGBF::FPBDRigidsEvolutionGBF(FPBDRigidsSOAs& InParticles,THandleArray<FChaosPhysicsMaterial>& SolverPhysicsMaterials, const TArray<ISimCallbackObject*>* InCollisionModifiers, bool InIsSingleThreaded)
 	: Base(InParticles, SolverPhysicsMaterials, InIsSingleThreaded)
 	, Clustering(*this, Particles.GetClusteredParticles())
-	, CollisionConstraints(InParticles, Collided, PhysicsMaterials, PerParticlePhysicsMaterials, &SolverPhysicsMaterials, DefaultNumCollisionsPerBlock, DefaultRestitutionThreshold)
+	, CollisionConstraints(InParticles, Collided, PhysicsMaterials, PerParticlePhysicsMaterials, &SolverPhysicsMaterials, CalculateNumCollisionsPerBlock(), DefaultRestitutionThreshold)
 	, BroadPhase(InParticles)
-	, NarrowPhase(DefaultCollisionCullDistance, BoundsThicknessVelocityMultiplier, CollisionConstraints.GetConstraintAllocator())
-	, CollisionDetector(BroadPhase, NarrowPhase, CollisionConstraints)
+	, CollisionDetector(BroadPhase, CollisionConstraints)
 	, PostIntegrateCallback(nullptr)
 	, PreApplyCallback(nullptr)
 	, CurrentStepResimCacheImp(nullptr)
@@ -615,6 +634,9 @@ FPBDRigidsEvolutionGBF::FPBDRigidsEvolutionGBF(FPBDRigidsSOAs& InParticles,THand
 
 	CollisionConstraints.SetCanDisableContacts(!!CollisionDisableCulledContacts);
 
+	CollisionDetector.SetBoundsExpansion(DefaultCollisionCullDistance);
+	CollisionDetector.SetBoundsVelocityInflation(BoundsThicknessVelocityMultiplier);
+
 	SetParticleUpdatePositionFunction([this](const TParticleView<FPBDRigidParticles>& ParticlesInput, const FReal Dt)
 	{
 		ParticlesInput.ParallelFor([&](auto& Particle, int32 Index)
diff --git a/Engine/Source/Runtime/Experimental/Chaos/Private/PBDRigidsSolver.cpp b/Engine/Source/Runtime/Experimental/Chaos/Private/PBDRigidsSolver.cpp
--- a/Engine/Source/Runtime/Experimental/Chaos/Private/PBDRigidsSolver.cpp
+++ b/Engine/Source/Runtime/Experimental/Chaos/Private/PBDRigidsSolver.cpp
@@ -852,13 +852,12 @@ namespace Chaos
 	{
 		MEvolution->GetCollisionConstraints().SetCollisionsEnabled(bChaosSolverCollisionEnabled);
 
-		FCollisionDetectorSettings CollisionDetectorSettings = MEvolution->GetCollisionConstraints().GetDetectorSettings();
+		FCollisionDetectorSettings CollisionDetectorSettings = MEvolution->GetCollisionDetector().GetSettings();
 		CollisionDetectorSettings.bAllowManifoldReuse = (ChaosSolverCollisionAllowManifoldUpdate != 0);
 		CollisionDetectorSettings.bDeferNarrowPhase = (ChaosSolverCollisionDeferNarrowPhase != 0);
 		CollisionDetectorSettings.bAllowManifolds = (ChaosSolverCollisionUseManifolds != 0);
 		CollisionDetectorSettings.bAllowCCD = bChaosUseCCD;
-		MEvolution->GetCollisionConstraints().SetDetectorSettings(CollisionDetectorSettings);
-		MEvolution->GetCollisionDetector().GetNarrowPhase().GetContext().SetSettings(CollisionDetectorSettings);
+		MEvolution->GetCollisionDetector().SetSettings(CollisionDetectorSettings);
 		
 		FPBDJointSolverSettings JointsSettings = MEvolution->GetJointConstraints().GetSettings();
 		JointsSettings.MinSolverStiffness = ChaosSolverJointMinSolverStiffness;
diff --git a/Engine/Source/Runtime/Experimental/Chaos/Public/Chaos/Collision/CollisionConstraintAllocator.h b/Engine/Source/Runtime/Experimental/Chaos/Public/Chaos/Collision/CollisionConstraintAllocator.h
--- a/Engine/Source/Runtime/Experimental/Chaos/Public/Chaos/Collision/CollisionConstraintAllocator.h
+++ b/Engine/Source/Runtime/Experimental/Chaos/Public/Chaos/Collision/CollisionConstraintAllocator.h
@@ -5,15 +5,193 @@
 
 #include "Chaos/Core.h"
 #include "Chaos/CollisionResolutionTypes.h"
+#include "Chaos/Collision/CollisionContext.h"
 #include "Chaos/Collision/CollisionKeys.h"
-#include "Chaos/Collision/PBDCollisionConstraint.h"
 #include "Chaos/Collision/ParticlePairMidPhase.h"
 #include "Chaos/ObjectPool.h"
 #include "Chaos/ParticleHandle.h"
 #include "ChaosStats.h"
 
+
 namespace Chaos
 {
+	/**
+	 * Container the storage for the FCollisionConstraintAllocator, as well as the API to create new midphases and collision constraints.
+	 * We have one of these objects per thread on which collisions detection is performed to get lock-free allocations and lists.
+	 */
+	class CHAOS_API FCollisionContextAllocator
+	{
+	public:
+		FCollisionContextAllocator(const int32 InNumCollisionsPerBlock, const int32 InCurrentEpoch)
+			: CurrentEpoch(InCurrentEpoch)
+#if CHAOS_COLLISION_OBJECTPOOL_ENABLED
+			, ConstraintPool(InNumCollisionsPerBlock, 0)
+			, MidPhasePool(InNumCollisionsPerBlock, 0)
+#endif
+		{
+		}
+
+		/**
+		 * The current epoch used to determine if a collision is up to date
+		 */
+		int32 GetCurrentEpoch() const
+		{
+			return CurrentEpoch;
+		}
+
+		/**
+		 * Create a constraint (called by the MidPhase)
+		 */
+		FPBDCollisionConstraintPtr CreateConstraint(
+			FGeometryParticleHandle* Particle0,
+			const FImplicitObject* Implicit0,
+			const FPerShapeData* Shape0,
+			const FBVHParticles* Simplicial0,
+			const FRigidTransform3& ShapeRelativeTransform0,
+			FGeometryParticleHandle* Particle1,
+			const FImplicitObject* Implicit1,
+			const FPerShapeData* Shape1,
+			const FBVHParticles* Simplicial1,
+			const FRigidTransform3& ShapeRelativeTransform1,
+			const FReal CullDistance,
+			const bool bUseManifold,
+			const EContactShapesType ShapePairType)
+		{
+			FPBDCollisionConstraintPtr Constraint = CreateConstraint();
+
+			if (Constraint.IsValid())
+			{
+				FPBDCollisionConstraint::Make(Particle0, Implicit0, Shape0, Simplicial0, ShapeRelativeTransform0, Particle1, Implicit1, Shape1, Simplicial1, ShapeRelativeTransform1, CullDistance, bUseManifold, ShapePairType, *Constraint);
+			}
+
+			return Constraint;
+		}
+
+		/**
+		 * Create an uninitialized collision constraint (public only for use by Resim which overwrites it with a saved constraint)
+		 */
+		FPBDCollisionConstraintPtr CreateConstraint()
+		{
+#if CHAOS_COLLISION_OBJECTPOOL_ENABLED 
+			FPBDCollisionConstraint* Constraint = ConstraintPool.Alloc();
+			return FPBDCollisionConstraintPtr(Constraint, FPBDCollisionConstraintDeleter(ConstraintPool));
+#else
+			return MakeUnique<FPBDCollisionConstraint>();
+#endif
+		}
+
+		/**
+		 * Set the constraint as Active - it will be added to the graph and solved this tick.
+		 * NOTE: This actually just enqueues the constraint - actual activation happens after the parallel phase is done.
+		 */
+		bool ActivateConstraint(FPBDCollisionConstraint* Constraint)
+		{
+			// When we wake an Island, we reactivate all constraints for all dynamic particles in the island. This
+			// results in duplicate calls to active for constraints involving two dynamic particles, hence the check for CurrentEpoch.
+			// @todo(chaos): fix duplicate calls from island wake. See UpdateSleepState in IslandManager.cpp
+			FPBDCollisionConstraintContainerCookie& Cookie = Constraint->GetContainerCookie();
+			if (Cookie.LastUsedEpoch != CurrentEpoch)
+			{
+				Cookie.LastUsedEpoch = CurrentEpoch;
+
+				NewActiveConstraints.Push(Constraint);
+
+				return true;
+			}
+			return false;
+		}
+
+		/**
+		 * Return a midphase for a particle pair.
+		 * This wil create a new midphase if the particle pairs were not recently overlapping, or return an
+		 * existing one if they were.
+		 * @note Nothing outside of thie allocator should hold a pointer to the midphase, or any constraints
+		 * it creates for more than the duration of the tick. Except the IslandManager :|
+		*/
+		FParticlePairMidPhase* GetMidPhase(FGeometryParticleHandle* Particle0, FGeometryParticleHandle* Particle1, FGeometryParticleHandle* SearchParticlePerformanceHint)
+		{
+			FParticlePairMidPhase* MidPhase = FindMidPhaseImpl(Particle0, Particle1, SearchParticlePerformanceHint);
+			if (MidPhase == nullptr)
+			{
+				MidPhase = CreateMidPhase(Particle0, Particle1);
+			}
+
+			return MidPhase;
+		}
+
+		/**
+		 * Return a midphase for a particle pair if it already exists, otherwise return null
+		*/
+		FParticlePairMidPhase* FindMidPhase(FGeometryParticleHandle* Particle0, FGeometryParticleHandle* Particle1, FGeometryParticleHandle* SearchParticlePerformanceHint)
+		{
+			return FindMidPhaseImpl(Particle0, Particle1, SearchParticlePerformanceHint);
+		}
+
+	private:
+		friend class FCollisionConstraintAllocator;
+
+		void Reset()
+		{
+#if CHAOS_COLLISION_OBJECTPOOL_ENABLED
+			ConstraintPool.Reset();
+			MidPhasePool.Reset();
+#endif
+		}
+
+		void BeginDetectCollisions(const int32 InEpoch)
+		{
+			check(NewActiveConstraints.IsEmpty());
+			check(NewMidPhases.IsEmpty());
+
+			CurrentEpoch = InEpoch;
+		}
+
+		// Process all the new midphases from the parallel collision detection
+		void ProcessNewMidPhases(FCollisionConstraintAllocator* Allocator);
+
+		// Process all the activated collisions from the parallel collision detection
+		void ProcessNewConstraints(FCollisionConstraintAllocator* Allocator);
+
+		// Find the midphase for the particle pair if it exists. Every particle holds a list of its midphases. We search "SearchParticle" which should be
+		// ideally be the one with fewer midphases on it.
+		FParticlePairMidPhase* FindMidPhaseImpl(FGeometryParticleHandle* Particle0, FGeometryParticleHandle* Particle1, FGeometryParticleHandle* SearchParticle)
+		{
+			check((SearchParticle == Particle0) || (SearchParticle == Particle1));
+
+			const FCollisionParticlePairKey Key = FCollisionParticlePairKey(Particle0, Particle1);
+			return SearchParticle->ParticleCollisions().FindMidPhase(Key.GetKey());
+		}
+
+		// Create and initialize a midphase for the particle pair. Adds it to the particles' lists of midphases.
+		FParticlePairMidPhase* CreateMidPhase(FGeometryParticleHandle* Particle0, FGeometryParticleHandle* Particle1)
+		{
+			const FCollisionParticlePairKey Key = FCollisionParticlePairKey(Particle0, Particle1);
+
+			// We temporarily hold new midphases as raw pointers and wrap them in a unique ptr later in ProcessNewMidPhases
+#if CHAOS_COLLISION_OBJECTPOOL_ENABLED 
+			FParticlePairMidPhase* MidPhase = MidPhasePool.Alloc();
+#else
+			FParticlePairMidPhase* MidPhase = new FParticlePairMidPhase();
+#endif
+
+			NewMidPhases.Add(MidPhase);
+
+			MidPhase->Init(Particle0, Particle1, Key);
+			return MidPhase;
+		}
+
+		int32 CurrentEpoch;
+
+		TArray<FPBDCollisionConstraint*> NewActiveConstraints;
+		TArray<FParticlePairMidPhase*> NewMidPhases;
+
+#if CHAOS_COLLISION_OBJECTPOOL_ENABLED
+		FPBDCollisionConstraintPool ConstraintPool;
+		FParticlePairMidPhasePool MidPhasePool;
+#endif
+	};
+
+
 	/**
 	 * @brief An allocator and container of collision constraints that supports reuse of constraints from the previous tick
 	 * 
@@ -48,13 +226,16 @@ namespace Chaos
 	class CHAOS_API FCollisionConstraintAllocator
 	{
 	public:
-		FCollisionConstraintAllocator(const int32 NumCollisionsPerBlock = 1000)
-			: ParticlePairMidPhases()
+		UE_NONCOPYABLE(FCollisionConstraintAllocator);
+
+		FCollisionConstraintAllocator(const int32 InNumCollisionsPerBlock = 1000)
+			: ContextAllocators()
+			, NumCollisionsPerBlock(InNumCollisionsPerBlock)
+			, ParticlePairMidPhases()
 			, ActiveConstraints()
 			, ActiveSweptConstraints()
 			, CurrentEpoch(0)
 			, bInCollisionDetectionPhase(false)
-			, ConstraintPool(NumCollisionsPerBlock, 0)
 		{
 		}
 
@@ -62,47 +243,25 @@ namespace Chaos
 		{
 			// Explicit clear so we don't have to worry about destruction order of pool and midphases
 			Reset();
-			ensure(ConstraintPool.GetNumAllocated() == 0);
 		}
 
-		FPBDCollisionConstraintPtr CreateCollisionConstraint()
+		void SetMaxContexts(const int32 MaxContexts)
 		{
-			// NOTE: Called from the CollisionDetection parellel-for loop.
-			FPBDCollisionConstraint* Constraint;
+			const int32 NumNewContexts = MaxContexts - ContextAllocators.Num();
+			if (NumNewContexts > 0)
 			{
-				FScopeLock Lock(&ConstraintPoolCS);
-				Constraint = ConstraintPool.Alloc();
+				ContextAllocators.Reserve(MaxContexts);
+				for (int32 NewIndex = 0; NewIndex < NumNewContexts; ++NewIndex)
+				{
+					TUniquePtr<FCollisionContextAllocator> ContextAllocator = MakeUnique<FCollisionContextAllocator>(NumCollisionsPerBlock, CurrentEpoch);
+					ContextAllocators.Emplace(MoveTemp(ContextAllocator));
+				}
 			}
-
-			// NOTE: Collision constraints are never destroyed in parallel, so the deleter does not require a lock
-			// (See PruneExpiredItems which destroys midphases that own collisions - it is called synchronously).
-			return FPBDCollisionConstraintPtr(Constraint, FPBDCollisionConstraintDeleter(ConstraintPool));
 		}
 
-		FPBDCollisionConstraintPtr CreateCollisionConstraint(
-			FGeometryParticleHandle* Particle0,
-			const FImplicitObject* Implicit0,
-			const FPerShapeData* Shape0,
-			const FBVHParticles* Simplicial0,
-			const FRigidTransform3& ShapeRelativeTransform0,
-			FGeometryParticleHandle* Particle1,
-			const FImplicitObject* Implicit1,
-			const FPerShapeData* Shape1,
-			const FBVHParticles* Simplicial1,
-			const FRigidTransform3& ShapeRelativeTransform1,
-			const FReal CullDistance,
-			const bool bUseManifold,
-			const EContactShapesType ShapePairType)
+		FCollisionContextAllocator* GetContextAllocator(const int32 Index)
 		{
-			// NOTE: Called from the CollisionDetection parellel-for loop.
-			FPBDCollisionConstraintPtr Constraint = CreateCollisionConstraint();
-
-			if (Constraint.IsValid())
-			{
-				FPBDCollisionConstraint::Make(Particle0, Implicit0, Shape0, Simplicial0, ShapeRelativeTransform0, Particle1, Implicit1, Shape1, Simplicial1, ShapeRelativeTransform1, CullDistance, bUseManifold, ShapePairType, *Constraint);
-			}
-
-			return Constraint;
+			return ContextAllocators[Index].Get();
 		}
 		
 		/**
@@ -164,6 +323,9 @@ namespace Chaos
 			ActiveConstraints.Reset();
 			ActiveSweptConstraints.Reset();
 			ParticlePairMidPhases.Reset();
+
+			// NOTE: this deletes all the storage for constraints and midphases - must be last
+			ContextAllocators.Reset();
 		}
 
 		/**
@@ -174,10 +336,6 @@ namespace Chaos
 		{
 			ActiveConstraints.Reset();
 			ActiveSweptConstraints.Reset();
-
-			// If we hit this we Activated constraints without calling ProcessNewItems
-			check(NewParticlePairMidPhases.IsEmpty());
-			check(NewConstraints.IsEmpty());
 		}
 
 		/**
@@ -189,14 +347,6 @@ namespace Chaos
 			check(!bInCollisionDetectionPhase);
 			bInCollisionDetectionPhase = true;
 
-			// If we hit this we Activated constraints without calling ProcessNewItems
-			check(NewParticlePairMidPhases.IsEmpty());
-			check(NewConstraints.IsEmpty());
-
-			// Clear the collision list for this tick - we are about to rebuild them
-			ActiveConstraints.Reset();
-			ActiveSweptConstraints.Reset();
-
 			// Update the tick counter
 			// NOTE: We do this here rather than in EndDetectionCollisions so that any contacts injected
 			// before collision detection count as the previous frame's collisions, e.g., from Islands
@@ -204,6 +354,15 @@ namespace Chaos
 			// done where we reset the Constraints array so that we can tell we have a valid index from
 			// the Epoch.
 			++CurrentEpoch;
+
+			for (const TUniquePtr<FCollisionContextAllocator>& ContextAllocator : ContextAllocators)
+			{
+				ContextAllocator->BeginDetectCollisions(CurrentEpoch);
+			}
+
+			// Clear the collision list for this tick - we are about to rebuild them
+			ActiveConstraints.Reset();
+			ActiveSweptConstraints.Reset();
 		}
 
 		/**
@@ -226,58 +385,6 @@ namespace Chaos
 			PruneExpiredMidPhases();
 		}
 
-		/**
-		 * @brief Return a midphase for a particle pair.
-		 * This wil create a new midphase if the particle pairs were not recently overlapping, or return an
-		 * existing one if they were.
-		 * @note Nothing outside of thie allocator should hold a pointer to the midphase, or any constraints 
-		 * it creates for more than the duration of the tick. Except the IslandManager :| 
-		*/
-		FParticlePairMidPhase* GetParticlePairMidPhase(FGeometryParticleHandle* Particle0, FGeometryParticleHandle* Particle1, FGeometryParticleHandle* SearchParticlePerformanceHint)
-		{
-			// NOTE: Called from the CollisionDetection parellel-for loop.
-
-			FParticlePairMidPhase* MidPhase = FindParticlePairMidPhaseImpl(Particle0, Particle1, SearchParticlePerformanceHint);
-			if (MidPhase == nullptr)
-			{
-				MidPhase = CreateParticlePairMidPhase(Particle0, Particle1);
-			}
-
-			return MidPhase;
-		}
-
-		/**
-		 * @brief Return a midphase for a particle pair only if it already exists
-		*/
-		FParticlePairMidPhase* FindParticlePairMidPhase(FGeometryParticleHandle* Particle0, FGeometryParticleHandle* Particle1, FGeometryParticleHandle* SearchParticlePerformanceHint)
-		{
-			return FindParticlePairMidPhaseImpl(Particle0, Particle1, SearchParticlePerformanceHint);
-		}
-
-		/**
-		 * @brief Called each tick when a constraint should be processed on that tick (i.e., the shapes are within CullDistance of each other)
-		*/
-		bool ActivateConstraint(FPBDCollisionConstraint* Constraint)
-		{
-			// NOTE: Called from the CollisionDetection parellel-for loop.
-			// We need to lock the arrays, but we can freely read/write to the constraint without the lock
-			// because each constraint is only processed once and not accessed by any other collision detection threads
-
-			// When we wake an Island, we reactivate all constraints for all dynamic particles in the island. This
-			// results in duplicate calls to active for constraints involving two dynamic particles, hence the check for CurrentEpoch.
-			// @todo(chaos): fix duplicate calls from island wake. See UpdateSleepState in IslandManager.cpp
-			FPBDCollisionConstraintContainerCookie& Cookie = Constraint->GetContainerCookie();
-			if (Cookie.LastUsedEpoch != CurrentEpoch)
-			{
-				Cookie.LastUsedEpoch = CurrentEpoch;
-
-				NewConstraints.Push(Constraint);
-
-				return true;
-			}
-			return false;
-		}
-
 		/**
 		 * @brief If we add new constraints after collision detection, do what needs to be done to add them to the system
 		*/
@@ -290,7 +397,7 @@ namespace Chaos
 		 * @brief Add a set of pre-built constraints and build required internal mapping data
 		 * This is used by the resim cache when restoring constraints after a desync
 		*/
-		void AddResimConstraints(const TArray<FPBDCollisionConstraint>& InConstraints)
+		void AddResimConstraints(const TArray<FPBDCollisionConstraint>& InConstraints, const FCollisionContext Context)
 		{
 			for (const FPBDCollisionConstraint& SourceConstraint : InConstraints)
 			{
@@ -304,15 +411,15 @@ namespace Chaos
 					Swap(Particle0, Particle1);
 				}
 
-				FParticlePairMidPhase* MidPhase = GetParticlePairMidPhase(Particle0, Particle1, SourceConstraint.Particle[0]);
+				FParticlePairMidPhase* MidPhase = Context.GetAllocator()->GetMidPhase(Particle0, Particle1, SourceConstraint.Particle[0]);
 
 				// We may be adding multiple constraints for the same particle pair, so we need
 				// to make sure the map is up to date in the case where we just created a new MidPhase
-				ProcessNewParticlePairMidPhases();
+				ProcessNewMidPhases();
 				
 				if (MidPhase != nullptr)
 				{
-					MidPhase->InjectCollision(SourceConstraint);
+					MidPhase->InjectCollision(SourceConstraint, Context);
 				}
 			}
 
@@ -362,7 +469,7 @@ namespace Chaos
 		template<typename TLambda>
 		void VisitConstCollisions(const TLambda& Visitor) const
 		{
-			for (const auto& MidPhase : ParticlePairMidPhases)
+			for (const FParticlePairMidPhasePtr& MidPhase : ParticlePairMidPhases)
 			{
 				if (MidPhase->VisitConstCollisions(Visitor) == ECollisionVisitorResult::Stop)
 				{
@@ -372,54 +479,35 @@ namespace Chaos
 		}
 		
 	private:
+		friend class FCollisionContextAllocator;
 
+		// Gather all the newly created items and finalize their setup
 		void ProcessNewItems()
 		{
-			ProcessNewParticlePairMidPhases();
-			ProcessNewConstraints();
-		}
-
-		FParticlePairMidPhase* FindParticlePairMidPhaseImpl(FGeometryParticleHandle* Particle0, FGeometryParticleHandle* Particle1, FGeometryParticleHandle* SearchParticle)
-		{
-			check((SearchParticle == Particle0) || (SearchParticle == Particle1));
+			ProcessNewMidPhases();
 
-			const FCollisionParticlePairKey Key = FCollisionParticlePairKey(Particle0, Particle1);
-			return SearchParticle->ParticleCollisions().FindMidPhase(Key.GetKey());
+			ProcessNewConstraints();
 		}
 
-		FParticlePairMidPhase* CreateParticlePairMidPhase(FGeometryParticleHandle* Particle0, FGeometryParticleHandle* Particle1)
+		// Collect all the midphases created on the context allocators (probably on multiple threads) and register them
+		void ProcessNewMidPhases()
 		{
-			const FCollisionParticlePairKey Key = FCollisionParticlePairKey(Particle0, Particle1);
-
-			// We enqueue a raw pointer so that we can use a lock-free queue. We wrap it in a UniquePtr later 
-			FParticlePairMidPhase* MidPhase = new FParticlePairMidPhase();
-			MidPhase->Init(Particle0, Particle1, Key, *this);
-
-			NewParticlePairMidPhases.Push(MidPhase);
-
-			return MidPhase;
+			for (TUniquePtr<FCollisionContextAllocator>& ContextAllocator : ContextAllocators)
+			{
+				ContextAllocator->ProcessNewMidPhases(this);
+			}
 		}
 
-		void ProcessNewParticlePairMidPhases()
+		void AddMidPhase(FParticlePairMidPhasePtr&& MidPhase)
 		{
-			while (FParticlePairMidPhase* MidPhase = NewParticlePairMidPhases.Pop())
-			{
-				FGeometryParticleHandle* Particle0 = MidPhase->GetParticle0();
-				FGeometryParticleHandle* Particle1 = MidPhase->GetParticle1();
-
-				FCollisionParticlePairKey Key = FCollisionParticlePairKey(Particle0, Particle1);
-
-				ParticlePairMidPhases.Emplace(TUniquePtr<FParticlePairMidPhase>(MidPhase));
-				Particle0->ParticleCollisions().AddMidPhase(Particle0, MidPhase);
-				Particle1->ParticleCollisions().AddMidPhase(Particle1, MidPhase);
-			}
+			ParticlePairMidPhases.Emplace(MoveTemp(MidPhase));
 		}
 
+		// Remove ParticlePairMidPhase from each Particles list of collisions
+		// NOTE: One or both particles may have been destroyed in which case it will
+		// have been set to null on the midphase.
 		void DetachParticlePairMidPhase(FParticlePairMidPhase* MidPhase)
 		{
-			// Remove ParticlePairMidPhase from each Particles list of collisions
-			// NOTE: One or both particles may have been destroyed in which case it will
-			// have been set to null on the midphase.
 			if (FGeometryParticleHandle* Particle0 = MidPhase->GetParticle0())
 			{
 				Particle0->ParticleCollisions().RemoveMidPhase(Particle0, MidPhase);
@@ -431,52 +519,19 @@ namespace Chaos
 			}
 		}
 
-		void PruneExpiredMidPhases()
-		{
-			check(NewParticlePairMidPhases.IsEmpty());
-			
-			// NOTE: Called from the physics thread. No need for locks.
-
-			// Determine which particle pairs are no longer overlapping
-			// Prune all pairs which were not updated this tick as part of the collision
-			// detection loop and are not asleep
-			int32 NumParticlePairMidPhases = ParticlePairMidPhases.Num();
-			for (int32 Index = 0; Index < NumParticlePairMidPhases; ++Index)
-			{
-				TUniquePtr<FParticlePairMidPhase>& MidPhase = ParticlePairMidPhases[Index];
-
-				// We could also add !MidPhase->IsInConstraintGraph() here, but we know that we will not be in the graph if we were
-				// not active this tick and were not asleep. The constraint graph ejects all non-sleeping constraints each tick.
-				// (There is a check in the collision destructor that verified this).
-				if ((MidPhase != nullptr) && !MidPhase->IsUsedSince(CurrentEpoch) && !MidPhase->IsSleeping())
-				{
-					// Remove from the particles' lists of contacts
-					DetachParticlePairMidPhase(MidPhase.Get());
-
-					// Destroy the midphase. ParticlePairMidPhases can get large so we allow it to shrink from time to time
-					const int32 MaxSlack = 1000;
-					const int32 Slack = ParticlePairMidPhases.Max() - ParticlePairMidPhases.Num();
-					const bool bAllowShrink = (Slack > MaxSlack);
-					ParticlePairMidPhases.RemoveAtSwap(Index, 1, bAllowShrink);
-
-					--Index;
-					--NumParticlePairMidPhases;
-				}
-			}
-		}
+		// Find all midphases that have not been updated this tick (i.e., bounds no longer overlap) and destroy them
+		void PruneExpiredMidPhases();
 
+		// Collect all the constraints activated in the collision tasks and register them (calls ActivateConstraintImp on each)
 		void ProcessNewConstraints()
 		{
-			if(!NewConstraints.IsEmpty())
+			for (TUniquePtr<FCollisionContextAllocator>& ContextAllocator : ContextAllocators)
 			{
-				while (FPBDCollisionConstraint* NewConstraint = NewConstraints.Pop())
-				{
-					// Add the constraint to the active list
-					ActivateConstraintImp(NewConstraint);
-				}
+				ContextAllocator->ProcessNewConstraints(this);
 			}
 		}
 
+		// Register an activated constraint
 		void ActivateConstraintImp(FPBDCollisionConstraint* CollisionConstraint)
 		{
 			FPBDCollisionConstraintContainerCookie& Cookie = CollisionConstraint->GetContainerCookie();
@@ -494,8 +549,12 @@ namespace Chaos
 			Cookie.LastUsedEpoch = CurrentEpoch;
 		}
 
+		// Storage for collisins and midphases, one for each thread on which we detect collisions
+		TArray<TUniquePtr<FCollisionContextAllocator>> ContextAllocators;
+		int32 NumCollisionsPerBlock;
+
 		// All of the overlapping particle pairs in the scene
-		TArray<TUniquePtr<FParticlePairMidPhase>> ParticlePairMidPhases;
+		TArray<FParticlePairMidPhasePtr> ParticlePairMidPhases;
 
 		// The active constraints (added or recovered this tick)
 		TArray<FPBDCollisionConstraint*> ActiveConstraints;
@@ -509,15 +568,6 @@ namespace Chaos
 
 		// For assertions
 		bool bInCollisionDetectionPhase;
-
-		// The set of constraints created or restored this tick during collision detection
-		mutable TLockFreePointerListLIFO<FPBDCollisionConstraint> NewConstraints;
-		
-		// The set of mid phases created this tick (i.e., for particle pairs that were not in the map)
-		mutable TLockFreePointerListLIFO<FParticlePairMidPhase> NewParticlePairMidPhases;
-
-		FPBDCollisionConstraintPool ConstraintPool;
-		mutable FCriticalSection ConstraintPoolCS;
 	};
 
 
@@ -531,14 +581,14 @@ namespace Chaos
 	///////////////////////////////////////////////////////////////////////////////////////////////
 
 	template<typename TLambda>
-	inline ECollisionVisitorResult FMultiShapePairCollisionDetector::VisitCollisions(const int32 LastEpoch, const TLambda& Visitor, const bool bOnlyActive)
+	inline ECollisionVisitorResult FMultiShapePairCollisionDetector::VisitCollisions(const TLambda& Visitor, const bool bOnlyActive)
 	{
 		for (auto& KVP : Constraints)
 		{
 			const FPBDCollisionConstraintPtr& Constraint = KVP.Value;
 
-			// If we only want active constraints, check the timestamp
-			if (Constraint.IsValid() && (!bOnlyActive || (Constraint->GetContainerCookie().LastUsedEpoch >= LastEpoch)))
+			// If we only want active constraints, check the we're in the graph
+			if (Constraint.IsValid() && (!bOnlyActive || Constraint->IsInConstraintGraph()))
 			{
 				if (Visitor(*Constraint) == ECollisionVisitorResult::Stop)
 				{
@@ -550,14 +600,14 @@ namespace Chaos
 	}
 
 	template<typename TLambda>
-	inline ECollisionVisitorResult FMultiShapePairCollisionDetector::VisitConstCollisions(const int32 LastEpoch, const TLambda& Visitor, const bool bOnlyActive) const
+	inline ECollisionVisitorResult FMultiShapePairCollisionDetector::VisitConstCollisions(const TLambda& Visitor, const bool bOnlyActive) const
 	{
 		for (auto& KVP : Constraints)
 		{
 			const FPBDCollisionConstraintPtr& Constraint = KVP.Value;
 
 			// If we only want active constraints, check the timestamp
-			if (Constraint.IsValid() && (!bOnlyActive || (Constraint->GetContainerCookie().LastUsedEpoch >= LastEpoch)))
+			if (Constraint.IsValid() && (!bOnlyActive || Constraint->IsInConstraintGraph()))
 			{
 				if (!bOnlyActive || (Visitor(*Constraint) == ECollisionVisitorResult::Stop))
 				{
@@ -571,11 +621,10 @@ namespace Chaos
 	template<typename TLambda>
 	inline ECollisionVisitorResult FParticlePairMidPhase::VisitCollisions(const TLambda& Visitor, const bool bOnlyActive)
 	{
-		const int32 LastEpoch = IsSleeping() ? LastUsedEpoch : GetCurrentEpoch();
 		for (FSingleShapePairCollisionDetector& ShapePair : ShapePairDetectors)
 		{
-			// If we only want active constraints, check the timestamp
-			if ((ShapePair.GetConstraint() != nullptr) && (!bOnlyActive || ShapePair.IsUsedSince(LastEpoch)))
+			// If we only want active constraints, check the we're in the graph
+			if ((ShapePair.GetConstraint() != nullptr) && (!bOnlyActive || ShapePair.GetConstraint()->IsInConstraintGraph()))
 			{
 				if (Visitor(*ShapePair.GetConstraint()) == ECollisionVisitorResult::Stop)
 				{
@@ -586,7 +635,7 @@ namespace Chaos
 
 		for (FMultiShapePairCollisionDetector& MultiShapePair : MultiShapePairDetectors)
 		{
-			if (MultiShapePair.VisitCollisions(LastEpoch, Visitor, bOnlyActive) == ECollisionVisitorResult::Stop)
+			if (MultiShapePair.VisitCollisions(Visitor, bOnlyActive) == ECollisionVisitorResult::Stop)
 			{
 				return ECollisionVisitorResult::Stop;
 			}
@@ -599,11 +648,10 @@ namespace Chaos
 	template<typename TLambda>
 	inline ECollisionVisitorResult FParticlePairMidPhase::VisitConstCollisions(const TLambda& Visitor, const bool bOnlyActive) const
 	{
-		const int32 LastEpoch = IsSleeping() ? LastUsedEpoch : GetCurrentEpoch();
 		for (const FSingleShapePairCollisionDetector& ShapePair : ShapePairDetectors)
 		{
 			// If we only want active constraints, check the timestamp
-			if ((ShapePair.GetConstraint() != nullptr) && (!bOnlyActive || ShapePair.IsUsedSince(LastEpoch)))
+			if ((ShapePair.GetConstraint() != nullptr) && (!bOnlyActive || ShapePair.GetConstraint()->IsInConstraintGraph()))
 			{
 				if (Visitor(*ShapePair.GetConstraint()) == ECollisionVisitorResult::Stop)
 				{
@@ -614,7 +662,7 @@ namespace Chaos
 
 		for (const FMultiShapePairCollisionDetector& MultiShapePair : MultiShapePairDetectors)
 		{
-			if (MultiShapePair.VisitConstCollisions(LastEpoch, Visitor, bOnlyActive) == ECollisionVisitorResult::Stop)
+			if (MultiShapePair.VisitConstCollisions(Visitor, bOnlyActive) == ECollisionVisitorResult::Stop)
 			{
 				return ECollisionVisitorResult::Stop;
 			}
diff --git a/Engine/Source/Runtime/Experimental/Chaos/Public/Chaos/Collision/CollisionContext.h b/Engine/Source/Runtime/Experimental/Chaos/Public/Chaos/Collision/CollisionContext.h
--- a/Engine/Source/Runtime/Experimental/Chaos/Public/Chaos/Collision/CollisionContext.h
+++ b/Engine/Source/Runtime/Experimental/Chaos/Public/Chaos/Collision/CollisionContext.h
@@ -5,13 +5,16 @@
 
 namespace Chaos
 {
+	class FCollisionContextAllocator;
 	class FMultiShapePairCollisionDetector;
 
 	class FCollisionDetectorSettings
 	{
 	public:
 		FCollisionDetectorSettings()
-			: bFilteringEnabled(true)
+			: BoundsExpansion(0)
+			, BoundsVelocityInflation(0)
+			, bFilteringEnabled(true)
 			, bDeferNarrowPhase(false)
 			, bAllowManifolds(true)
 			, bAllowManifoldReuse(true)
@@ -19,6 +22,12 @@ namespace Chaos
 		{
 		}
 
+		// Shape bounds are expanded by this in all 3 directions. This is also used as the CullDistance for collision response.
+		FReal BoundsExpansion;
+		
+		// Shape bounds in the broadphase are expanded by this multiple of velocity
+		FReal BoundsVelocityInflation;
+
 		// Whether to check the shape query flags in the narrow phase (e.g., Rigid Body nodes have already performed filtering prior to collision detection)
 		bool bFilteringEnabled;
 
@@ -45,16 +54,25 @@ namespace Chaos
 	{
 	public:
 		FCollisionContext()
-			: Settings()
+			: Settings(nullptr)
+			, Allocator(nullptr)
 			, MultiShapeCollisionDetector(nullptr)
 		{
 		}
 
-		const FCollisionDetectorSettings& GetSettings() const { return Settings; }
-		FCollisionDetectorSettings& GetSettings() { return Settings; }
-		void SetSettings(const FCollisionDetectorSettings& InSettings) { Settings = InSettings; }
+		const FCollisionDetectorSettings& GetSettings() const { return *Settings; }
+		void SetSettings(const FCollisionDetectorSettings& InSettings) { Settings = &InSettings; }
+
+		FCollisionContextAllocator* GetAllocator() const { return Allocator; }
+		void SetAllocator(FCollisionContextAllocator* InAllocator) { Allocator = InAllocator; }
+
+		FMultiShapePairCollisionDetector* GetMultiShapePairCollisionDetector() const { return MultiShapeCollisionDetector; }
+		void SetMultiShapePairCollisionDetector(FMultiShapePairCollisionDetector* InDetector) { MultiShapeCollisionDetector = InDetector; }
+
+
+		const FCollisionDetectorSettings* Settings;
 
-		FCollisionDetectorSettings Settings;
+		FCollisionContextAllocator* Allocator;
 
 		// This is used in the older collision detection path which is still used for particles that do not flatten their implicit hierrarchies
 		// into the Particle's ShapesArray. Currently this is only Clusters.
diff --git a/Engine/Source/Runtime/Experimental/Chaos/Public/Chaos/Collision/CollisionDetector.h b/Engine/Source/Runtime/Experimental/Chaos/Public/Chaos/Collision/CollisionDetector.h
--- a/Engine/Source/Runtime/Experimental/Chaos/Public/Chaos/Collision/CollisionDetector.h
+++ b/Engine/Source/Runtime/Experimental/Chaos/Public/Chaos/Collision/CollisionDetector.h
@@ -8,24 +8,45 @@
 namespace Chaos
 {
 	class FPBDCollisionConstraints;
-	class FNarrowPhase;
 	class FEvolutionResimCache;
 
 	class CHAOS_API FCollisionDetector
 	{
 	public:
 		FCollisionDetector(FPBDCollisionConstraints& InCollisionContainer)
-			: CollisionContainer(InCollisionContainer)
+			: Settings()
+			, CollisionContainer(InCollisionContainer)
 		{
 		}
 
 		virtual ~FCollisionDetector() {}
 
+		const FCollisionDetectorSettings& GetSettings() const
+		{
+			return Settings;
+		}
+
+		void SetSettings(const FCollisionDetectorSettings& InSettings)
+		{
+			Settings = InSettings;
+		}
+
+		void SetBoundsExpansion(const FReal InBoundsExpansion)
+		{
+			Settings.BoundsExpansion = InBoundsExpansion;
+		}
+
+		void SetBoundsVelocityInflation(const FReal InBoundsVelocityInflation)
+		{
+			Settings.BoundsVelocityInflation = InBoundsVelocityInflation;
+		}
+
 		FPBDCollisionConstraints& GetCollisionContainer() { return CollisionContainer; }
 
 		virtual void DetectCollisions(const FReal Dt, FEvolutionResimCache* ResimCache) = 0;
 
 	protected:
+		FCollisionDetectorSettings Settings;
 		FPBDCollisionConstraints& CollisionContainer;
 	};
 
diff --git a/Engine/Source/Runtime/Experimental/Chaos/Public/Chaos/Collision/NarrowPhase.h b/Engine/Source/Runtime/Experimental/Chaos/Public/Chaos/Collision/NarrowPhase.h
--- a/Engine/Source/Runtime/Experimental/Chaos/Public/Chaos/Collision/NarrowPhase.h
+++ b/Engine/Source/Runtime/Experimental/Chaos/Public/Chaos/Collision/NarrowPhase.h
@@ -1,62 +1,6 @@
 // Copyright Epic Games, Inc. All Rights Reserved.
 #pragma once
 
-#include "Chaos/Core.h"
-#include "Chaos/Collision/CollisionContext.h"
-#include "Chaos/Collision/CollisionConstraintAllocator.h"
-#include "Chaos/Collision/ParticlePairMidPhase.h"
-#include "Chaos/ParticleHandle.h"
-
-namespace Chaos
-{
-	class FCollisionContext;
-	class FCollisionConstraintAllocator;
-
-	/**
-	 * Generate contact manifolds for particle pairs.
-	 * @todo(chaos): Rename FNarrowPhase to FMidPhase
-	 */
-	class CHAOS_API FNarrowPhase
-	{
-	public:
-		FNarrowPhase(const FReal InBoundsExpansion, const FReal InBoundsVelocityInflation, FCollisionConstraintAllocator& InCollisionAllocator)
-			: CollisionAllocator(&InCollisionAllocator)
-			, BoundsExpansion(InBoundsExpansion)
-			, BoundsVelocityInflation(InBoundsVelocityInflation)
-		{
-		}
-
-		FCollisionContext& GetContext() { return Context; }
-
-		FReal GetBoundsExpansion() const { return BoundsExpansion; }
-		void SetBoundsExpansion(const FReal InBoundsExpansion) { BoundsExpansion = InBoundsExpansion; }
-		
-		FReal GetBoundsVelocityInflation() const { return BoundsVelocityInflation; }
-		void SetBoundsVelocityInflation(const FReal InBoundsVelocityInflation) { BoundsVelocityInflation = InBoundsVelocityInflation; }
-
-		void GenerateCollisions(FReal Dt, TGeometryParticleHandle<FReal, 3>* Particle0, TGeometryParticleHandle<FReal, 3>* Particle1, TGeometryParticleHandle<FReal, 3>* SearchParticlePerfHint)
-		{
-			FParticlePairMidPhase* MidPhase = CollisionAllocator->GetParticlePairMidPhase(Particle0, Particle1, SearchParticlePerfHint);
-			if (MidPhase != nullptr)
-			{
-				MidPhase->GenerateCollisions(GetBoundsExpansion(), Dt, GetContext());
-			}
-		}
-
-		FParticlePairMidPhase* GetParticlePairMidPhase(TGeometryParticleHandle<FReal, 3>* Particle0, TGeometryParticleHandle<FReal, 3>* Particle1, TGeometryParticleHandle<FReal, 3>* SearchParticlePerfHint )
-		{
-			return CollisionAllocator->GetParticlePairMidPhase(Particle0, Particle1, SearchParticlePerfHint);
-		}
-
-		// Use this function if a Mid phase pair is already allocated
-		void GenerateCollisions(FReal Dt, FParticlePairMidPhase* MidPhase)
-		{
-			MidPhase->GenerateCollisions(GetBoundsExpansion(), Dt, GetContext());
-		}
-	private:
-		FCollisionContext Context;
-		FCollisionConstraintAllocator* CollisionAllocator;
-		FReal BoundsExpansion;
-		FReal BoundsVelocityInflation;
-	};
-}
+//
+// TO BE REMOVED
+//
diff --git a/Engine/Source/Runtime/Experimental/Chaos/Public/Chaos/Collision/PBDCollisionConstraint.h b/Engine/Source/Runtime/Experimental/Chaos/Public/Chaos/Collision/PBDCollisionConstraint.h
--- a/Engine/Source/Runtime/Experimental/Chaos/Public/Chaos/Collision/PBDCollisionConstraint.h
+++ b/Engine/Source/Runtime/Experimental/Chaos/Public/Chaos/Collision/PBDCollisionConstraint.h
@@ -18,6 +18,7 @@
 namespace Chaos
 {
 	class FCollisionConstraintAllocator;
+	class FCollisionContextAllocator;
 	class FConstGenericParticleHandle;
 	class FImplicitObject;
 	class FParticlePairMidPhase;
@@ -220,6 +221,7 @@ namespace Chaos
 	class CHAOS_API FPBDCollisionConstraint final : public FPBDCollisionConstraintHandle
 	{
 		friend class FCollisionConstraintAllocator;
+		friend class FCollisionContextAllocator;
 		friend class FMultiShapePairCollisionDetector;
 		friend class FParticlePairMidPhase;
 		friend class FPBDCollisionConstraints;
diff --git a/Engine/Source/Runtime/Experimental/Chaos/Public/Chaos/Collision/ParticlePairBroadPhase.h b/Engine/Source/Runtime/Experimental/Chaos/Public/Chaos/Collision/ParticlePairBroadPhase.h
--- a/Engine/Source/Runtime/Experimental/Chaos/Public/Chaos/Collision/ParticlePairBroadPhase.h
+++ b/Engine/Source/Runtime/Experimental/Chaos/Public/Chaos/Collision/ParticlePairBroadPhase.h
@@ -2,7 +2,7 @@
 #pragma once
 
 #include "Chaos/BoundingVolumeUtilities.h"
-#include "Chaos/Collision/NarrowPhase.h"
+#include "Chaos/Collision/ParticlePairMidPhase.h"
 #include "Chaos/ParticleHandle.h"
 #include "Chaos/PBDCollisionConstraints.h"
 #include "Chaos/PBDRigidsSOAs.h"
@@ -32,10 +32,19 @@ namespace Chaos
 		/**
 		 *
 		 */
-		void ProduceOverlaps(FReal Dt, FNarrowPhase& NarrowPhase)
+		void ProduceOverlaps(
+			FReal Dt, 
+			FCollisionConstraintAllocator* Allocator,
+			const FCollisionDetectorSettings& Settings)
 		{
 			SCOPE_CYCLE_COUNTER(STAT_Collisions_ParticlePairBroadPhase);
 
+			Allocator->SetMaxContexts(1);
+
+			FCollisionContext Context;
+			Context.SetSettings(Settings);
+			Context.Allocator = Allocator->GetContextAllocator(0);
+
 			if (ParticlePairs != nullptr)
 			{
 				for (const FParticlePair& ParticlePair : *ParticlePairs)
@@ -46,7 +55,7 @@ namespace Chaos
 
 					if ((ParticleA != nullptr) && (ParticleB != nullptr))
 					{
-						ProduceOverlaps(Dt, NarrowPhase, ParticleA, ParticleB, ParticleA);
+						ProduceOverlaps(Dt, Context, ParticleA, ParticleB, ParticleA);
 					}
 				}
 			}
@@ -61,7 +70,7 @@ namespace Chaos
 						{
 							if (ParticleB != nullptr)
 							{
-								ProduceOverlaps(Dt, NarrowPhase, ParticleA, ParticleB, ParticleA);
+								ProduceOverlaps(Dt, Context, ParticleA, ParticleB, ParticleA);
 							}
 						}
 					}
@@ -72,7 +81,7 @@ namespace Chaos
 	private:
 		inline void ProduceOverlaps(
 			FReal Dt,
-			FNarrowPhase& NarrowPhase,
+			const FCollisionContext& Context,
 			FParticleHandle* ParticleA,
 			FParticleHandle* ParticleB,
 			FParticleHandle* SearchParticle)
@@ -81,7 +90,11 @@ namespace Chaos
 			const FAABB3& Box1 = ParticleB->WorldSpaceInflatedBounds();
 			if (Box0.Intersects(Box1))
 			{
-				NarrowPhase.GenerateCollisions(Dt, ParticleA, ParticleB, SearchParticle);
+				FParticlePairMidPhase* MidPhase = Context.GetAllocator()->GetMidPhase(ParticleA, ParticleB, SearchParticle);
+				if (MidPhase != nullptr)
+				{
+					MidPhase->GenerateCollisions(Context.GetSettings().BoundsExpansion, Dt, Context);
+				}
 			}
 		}
 
diff --git a/Engine/Source/Runtime/Experimental/Chaos/Public/Chaos/Collision/ParticlePairCollisionDetector.h b/Engine/Source/Runtime/Experimental/Chaos/Public/Chaos/Collision/ParticlePairCollisionDetector.h
--- a/Engine/Source/Runtime/Experimental/Chaos/Public/Chaos/Collision/ParticlePairCollisionDetector.h
+++ b/Engine/Source/Runtime/Experimental/Chaos/Public/Chaos/Collision/ParticlePairCollisionDetector.h
@@ -13,15 +13,13 @@ namespace Chaos
 	class CHAOS_API FBasicCollisionDetector : public FCollisionDetector
 	{
 	public:
-		FBasicCollisionDetector(FBasicBroadPhase& InBroadPhase, FNarrowPhase& InNarrowPhase, FPBDCollisionConstraints& InCollisionContainer)
+		FBasicCollisionDetector(FBasicBroadPhase& InBroadPhase, FPBDCollisionConstraints& InCollisionContainer)
 			: FCollisionDetector(InCollisionContainer)
 			, BroadPhase(InBroadPhase)
-			, NarrowPhase(InNarrowPhase)
 		{
 		}
 
 		FBasicBroadPhase& GetBroadPhase() { return BroadPhase; }
-		FNarrowPhase& GetNarrowPhase() { return NarrowPhase; }
 
 		virtual void DetectCollisions(const FReal Dt, FEvolutionResimCache* Unused) override
 		{
@@ -36,15 +34,16 @@ namespace Chaos
 
 			CollisionContainer.BeginDetectCollisions();
 
-			// Collision detection pipeline: BroadPhase -> NarrowPhase -> Container
-			BroadPhase.ProduceOverlaps(Dt, NarrowPhase);
+			FCollisionConstraintAllocator& Allocator = GetCollisionContainer().GetConstraintAllocator();
+
+			// Collision detection pipeline: BroadPhase -> MidPhase ->NarrowPhase -> Container
+			BroadPhase.ProduceOverlaps(Dt, &Allocator, GetSettings());
 
 			CollisionContainer.EndDetectCollisions();
 		}
 
 	private:
 		FBasicBroadPhase& BroadPhase;
-		FNarrowPhase& NarrowPhase;
 	};
 
 }
\ No newline at end of file
diff --git a/Engine/Source/Runtime/Experimental/Chaos/Public/Chaos/Collision/ParticlePairMidPhase.h b/Engine/Source/Runtime/Experimental/Chaos/Public/Chaos/Collision/ParticlePairMidPhase.h
--- a/Engine/Source/Runtime/Experimental/Chaos/Public/Chaos/Collision/ParticlePairMidPhase.h
+++ b/Engine/Source/Runtime/Experimental/Chaos/Public/Chaos/Collision/ParticlePairMidPhase.h
@@ -7,25 +7,35 @@
 #include "Chaos/CollisionResolutionTypes.h"
 #include "Chaos/Collision/CollisionKeys.h"
 #include "Chaos/Collision/CollisionVisitor.h"
+#include "Chaos/Collision/PBDCollisionConstraint.h"
 #include "Chaos/ObjectPool.h"
 #include "Chaos/ParticleHandleFwd.h"
-
 #include "ProfilingDebugging/CsvProfiler.h"
 
+	// Whether to use a pool for MidPhases
+#define CHAOS_COLLISION_OBJECTPOOL_ENABLED 1
+
 namespace Chaos
 {
 	class FCollisionContext;
 	class FCollisionConstraintAllocator;
 	class FParticlePairMidPhase;
-	class FPBDCollisionConstraint;
 	class FPBDCollisionConstraints;
 	class FPerShapeData;
 	class FSingleShapePairCollisionDetector;
 
-	// Collision constraints are stored in an ObjectPool. @see FCollisionConstraintAllocator
+#if CHAOS_COLLISION_OBJECTPOOL_ENABLED 
+	// Collision and MidPhases are stored in an ObjectPool (if CHAOS_COLLISION_OBJECTPOOL_ENABLED is set). @see FCollisionConstraintAllocator
 	using FPBDCollisionConstraintPool = TObjectPool<FPBDCollisionConstraint>;
-	using FPBDCollisionConstraintDeleter = TObjectPoolDeleter<FPBDCollisionConstraint>;
+	using FParticlePairMidPhasePool = TObjectPool<FParticlePairMidPhase>;
+	using FPBDCollisionConstraintDeleter = TObjectPoolDeleter<FPBDCollisionConstraintPool>;
 	using FPBDCollisionConstraintPtr = TUniquePtr<FPBDCollisionConstraint, FPBDCollisionConstraintDeleter>;
+	using FParticlePairMidPhaseDeleter = TObjectPoolDeleter<FParticlePairMidPhasePool>;
+	using FParticlePairMidPhasePtr = TUniquePtr<FParticlePairMidPhase, FParticlePairMidPhaseDeleter>;
+#else 
+	using FPBDCollisionConstraintPtr = TUniquePtr<FPBDCollisionConstraint>;
+	using FParticlePairMidPhasePtr = TUniquePtr<FParticlePairMidPhase>;
+#endif
 
 	/**
 	 * @brief Handles collision detection for a pair of simple shapes (i.e., not compound shapes)
@@ -74,7 +84,7 @@ namespace Chaos
 		int32 GenerateCollision(
 			const FReal CullDistance,
 			const FReal Dt,
-			FCollisionContext& Context);
+			const FCollisionContext& Context);
 
 		/**
 		 * @brief Generate a SweptConstraint as long as AABBs overlap
@@ -83,47 +93,47 @@ namespace Chaos
 		int32 GenerateCollisionCCD(
 			const FReal CullDistance,
 			const FReal Dt,
-			FCollisionContext& Context);
+			const FCollisionContext& Context);
 
 		/**
 		 * @brief Reactivate the constraint
 		 * @parame SleepEpoch The tick on which the particle went to sleep.
 		 * Only constraints that were active when the particle went to sleep should be reactivated.
 		*/
-		void WakeCollision(const int32 SleepEpoch);
+		void WakeCollision(const int32 SleepEpoch, const int32 CurrentEpoch);
 
 		/**
 		 * @brief Set the collision from the parameter and activate it
 		 * This is used by the Resim restore functionality
 		*/
-		void SetCollision(const FPBDCollisionConstraint& Constraint);
+		void SetCollision(const FPBDCollisionConstraint& Constraint, const FCollisionContext& Context);
 
 	private:
 		int32 GenerateCollisionImpl(
 			const FReal CullDistance, 
 			const FReal Dt,
-			FCollisionContext& Context);
+			const FCollisionContext& Context);
 			
 		int32 GenerateCollisionCCDImpl(
 			const FReal CullDistance, 
 			const FReal Dt,
-			FCollisionContext& Context);
+			const FCollisionContext& Context);
 
 		int32 GenerateCollisionProbeImpl(
 			const FReal CullDistance, 
 			const FReal Dt,
-			FCollisionContext& Context);
+			const FCollisionContext& Context);
 
 		/**
 		 * @brief Whether the two shapes are separated by less than CullDistance (i.e., we should run the narrow phase).
 		 * Also returns true if bounds checking is disabled (globally or for this pair)
 		*/
-		bool DoBoundsOverlap(const FReal CullDistance);
+		bool DoBoundsOverlap(const FReal CullDistance, const int32 CurrentEpoch);
 
 		/**
 		 * @brief Create a constraint
 		*/
-		void CreateConstraint(const FReal CullDistance, FCollisionContext& Context);
+		void CreateConstraint(const FReal CullDistance, const FCollisionContext& Context);
 
 		FParticlePairMidPhase& MidPhase;
 		FPBDCollisionConstraintPtr Constraint;
@@ -176,7 +186,7 @@ namespace Chaos
 		int32 GenerateCollisions(
 			const FReal CullDistance,
 			const FReal Dt,
-			FCollisionContext& Context);
+			const FCollisionContext& Context);
 
 		/**
 		 * @brief Callback from the narrow phase to create a collision constraint for this particle pair.
@@ -196,7 +206,8 @@ namespace Chaos
 			const FRigidTransform3& ShapeRelativeTransform1,
 			const FReal CullDistance,
 			const EContactShapesType ShapePairType,
-			const bool bInUseManifold);
+			const bool bUseManifold,
+			const FCollisionContext& Context);
 
 		/**
 		 * @brief FindOrCreateConstraint for swept constraints 
@@ -213,26 +224,27 @@ namespace Chaos
 			const FBVHParticles* BVHParticles1,
 			const FRigidTransform3& ShapeRelativeTransform1,
 			const FReal CullDistance,
-			EContactShapesType ShapePairType);
+			const EContactShapesType ShapePairType,
+			const FCollisionContext& Context);
 
 		/**
 		 * @brief Reactivate the constraint
 		 * @parame SleepEpoch The tick on which the particle went to sleep.
 		 * Only constraints that were active when the particle went to sleep should be reactivated.
 		*/
-		void WakeCollisions(const int32 SleepEpoch);
+		void WakeCollisions(const int32 SleepEpoch, const int32 CurrentEpoch);
 
 		/**
 		 * @brief Visit all the collisions
 		*/
 		template<typename TLambda>
-		ECollisionVisitorResult VisitCollisions(const int32 LastEpoch, const TLambda& Visitor, const bool bOnlyActive = true);
+		ECollisionVisitorResult VisitCollisions(const TLambda& Visitor, const bool bOnlyActive = true);
 
 		/**
 		 * @brief Visit all the collisions
 		*/
 		template<typename TLambda>
-		ECollisionVisitorResult VisitConstCollisions(const int32 LastEpoch, const TLambda& Visitor, const bool bOnlyActive = true) const;
+		ECollisionVisitorResult VisitConstCollisions(const TLambda& Visitor, const bool bOnlyActive = true) const;
 
 	private:
 		FPBDCollisionConstraint* FindConstraint(const FCollisionParticlePairConstraintKey& Key);
@@ -251,10 +263,11 @@ namespace Chaos
 			const FReal CullDistance,
 			const EContactShapesType ShapePairType,
 			const bool bInUseManifold,
-			const FCollisionParticlePairConstraintKey& Key);
+			const FCollisionParticlePairConstraintKey& Key,
+			const FCollisionContext& Context);
 
-		int32 ProcessNewConstraints();
-		void PruneConstraints();
+		int32 ProcessNewConstraints(const FCollisionContext& Context);
+		void PruneConstraints(const int32 CurrentEpoch);
 
 		FParticlePairMidPhase& MidPhase;
 		TMap<uint32, FPBDCollisionConstraintPtr> Constraints;
@@ -301,8 +314,7 @@ namespace Chaos
 		void Init(
 			FGeometryParticleHandle* InParticle0,
 			FGeometryParticleHandle* InParticle1,
-			const FCollisionParticlePairKey& InKey,
-			FCollisionConstraintAllocator& InCollisionAllocator);
+			const FCollisionParticlePairKey& InKey);
 
 		inline FGeometryParticleHandle* GetParticle0() { return Particle0; }
 
@@ -310,8 +322,6 @@ namespace Chaos
 
 		inline const FCollisionParticlePairKey& GetKey() const { return Key; }
 
-		inline FCollisionConstraintAllocator& GetCollisionAllocator() { return *CollisionAllocator; }
-
 		inline bool IsValid() const { return (Particle0 != nullptr) && (Particle1 != nullptr); }
 
 		// Prefetch the memory in this class for later use
@@ -339,7 +349,7 @@ namespace Chaos
 		 * If this switches the state to Awake, it will reactivate any collisions between the particle pair that
 		 * were active when they went to sleep.
 		*/
-		void SetIsSleeping(const bool bInIsSleeping);
+		void SetIsSleeping(const bool bInIsSleeping, const int32 CurrentEpoch);
 
 		bool IsInConstraintGraph() const;
 
@@ -361,16 +371,14 @@ namespace Chaos
 		void GenerateCollisions(
 			const FReal CullDistance,
 			const FReal Dt,
-			FCollisionContext& Context);
+			const FCollisionContext& Context);
 
 		/**
 		 * @brief Copy a collision and activate it
 		 * This is used by the Resim system to restore saved colisions. If there is already a matching constraint
 		 * it will be overwritten, otherwise a new constraint will be added.
 		*/
-		void InjectCollision(const FPBDCollisionConstraint& Constraint);
-
-		int32 GetCurrentEpoch() const;
+		void InjectCollision(const FPBDCollisionConstraint& Constraint, const FCollisionContext& Context);
 
 		/**
 		 * @brief Call a lambda on each active collision constraint
@@ -460,7 +468,6 @@ namespace Chaos
 		} Flags; // 4 bytes
 
 		FCollisionParticlePairKey Key; //8 bytes
-		FCollisionConstraintAllocator* CollisionAllocator; // 8 bytes
 
 		int32 LastUsedEpoch;  //4 bytes
 		int32 NumActiveConstraints; // 4 bytes
diff --git a/Engine/Source/Runtime/Experimental/Chaos/Public/Chaos/Collision/SpatialAccelerationBroadPhase.h b/Engine/Source/Runtime/Experimental/Chaos/Public/Chaos/Collision/SpatialAccelerationBroadPhase.h
--- a/Engine/Source/Runtime/Experimental/Chaos/Public/Chaos/Collision/SpatialAccelerationBroadPhase.h
+++ b/Engine/Source/Runtime/Experimental/Chaos/Public/Chaos/Collision/SpatialAccelerationBroadPhase.h
@@ -1,9 +1,10 @@
 // Copyright Epic Games, Inc. All Rights Reserved.
 #pragma once
 
+#include "Chaos/Collision/CollisionConstraintAllocator.h"
 #include "Chaos/Collision/CollisionConstraintFlags.h"
+#include "Chaos/Collision/CollisionContext.h"
 #include "Chaos/Collision/StatsData.h"
-#include "Chaos/Collision/NarrowPhase.h"
 #include "Chaos/ISpatialAccelerationCollection.h"
 #include "Chaos/ParticleHandleFwd.h"
 #include "Chaos/ParticleHandle.h"
@@ -121,7 +122,8 @@ namespace Chaos
 		 */
 		void ProduceOverlaps(
 			FReal Dt, 
-			FNarrowPhase& NarrowPhase, 
+			FCollisionConstraintAllocator* Allocator,
+			const FCollisionDetectorSettings& Settings,
 			IResimCacheBase* ResimCache
 			)
 		{
@@ -135,19 +137,19 @@ namespace Chaos
 
 			if (const auto AABBTree = SpatialAcceleration->template As<TAABBTree<FAccelerationStructureHandle, TAABBTreeLeafArray<FAccelerationStructureHandle>>>())
 			{
-				ProduceOverlaps(Dt, *AABBTree, NarrowPhase, ResimCache);
+				ProduceOverlaps(Dt, *AABBTree, Allocator, Settings, ResimCache);
 			}
 			else if (const auto BV = SpatialAcceleration->template As<TBoundingVolume<FAccelerationStructureHandle>>())
 			{
-				ProduceOverlaps(Dt, *BV, NarrowPhase, ResimCache);
+				ProduceOverlaps(Dt, *BV, Allocator, Settings, ResimCache);
 			}
 			else if (const auto AABBTreeBV = SpatialAcceleration->template As<TAABBTree<FAccelerationStructureHandle, TBoundingVolume<FAccelerationStructureHandle>>>())
 			{
-				ProduceOverlaps(Dt, *AABBTreeBV, NarrowPhase, ResimCache);
+				ProduceOverlaps(Dt, *AABBTreeBV, Allocator, Settings, ResimCache);
 			}
 			else if (const auto Collection = SpatialAcceleration->template As<ISpatialAccelerationCollection<FAccelerationStructureHandle, FReal, 3>>())
 			{
-				Collection->PBDComputeConstraintsLowLevel(Dt, *this, NarrowPhase, ResimCache);
+				Collection->PBDComputeConstraintsLowLevel(Dt, *this, Allocator, Settings, ResimCache);
 			}
 			else
 			{
@@ -164,23 +166,59 @@ namespace Chaos
 		 * */
 
 		template<bool bNeedsResim, bool bOnlyRigid, typename ViewType, typename SpatialAccelerationType>
-		void ComputeParticlesOverlaps(ViewType& OverlapView, FReal Dt,
-			const SpatialAccelerationType& InSpatialAcceleration, FNarrowPhase& NarrowPhase)
+		void ComputeParticlesOverlaps(
+			ViewType& OverlapView, 
+			FReal Dt,
+			const SpatialAccelerationType& InSpatialAcceleration, 
+			FCollisionConstraintAllocator* Allocator,
+			const FCollisionDetectorSettings& Settings)
 		{
-			 OverlapView.ParallelFor([&](auto& Particle1,int32 ActiveIdxIdx)
-			 {
-			 	ProduceParticleOverlaps<bNeedsResim,bOnlyRigid>(Dt, Particle1.Handle(), InSpatialAcceleration,NarrowPhase,ActiveIdxIdx);
-			},bDisableCollisionParallelFor);
+			// A set of contexts, one for each worker thread, containing the allocation buffers etc.
+			// NOTE: Ideally contexts would not resize after the initial creation, but actually it might because
+			// of the way the ParticleView parallel-for is implemented. If the view is a set of sub-views, we
+			// will get a parallel-for per sub-view which will have different numbers of particles and may
+			// therefore require a different number of workers.
+			// Fortunately, we only get calls to ContextCreator for a sub-view (with a possibly different NumWorkers)
+			// after the previous sub-view's parallel-for has completed, so we don't need to worry about context addresses 
+			// changing (we could use TChunkedArray, or allocate contexts on the heap if we ever need to handle that).
+			TArray<FCollisionContext> Contexts;
+			auto ContextCreator = [Allocator, &Settings, &Contexts](const int32 WorkerIndex, const int32 NumWorkers) -> int32
+			{
+				Contexts.SetNum(NumWorkers);
+
+				// Retrieve the midphase/collision allocator for this context
+				Allocator->SetMaxContexts(NumWorkers);
+				FCollisionContextAllocator* ContextAllocator = Allocator->GetContextAllocator(WorkerIndex);
+
+				// Build the context for the worker
+				Contexts[WorkerIndex].SetSettings(Settings);
+				Contexts[WorkerIndex].SetAllocator(ContextAllocator);
+
+				return WorkerIndex;
+			};
+
+			// Detect collisions for a single particle versus all others
+			auto ParticleCollisionGenerator = [this, Dt, &InSpatialAcceleration, &Contexts](auto& Particle1, const int32 ActiveIdxIdx, const int32 ContextIndex)
+			{
+				ProduceParticleOverlaps<bNeedsResim, bOnlyRigid>(Dt, Particle1.Handle(), InSpatialAcceleration, Contexts[ContextIndex]);
+			};
+
+			OverlapView.ParallelFor(ContextCreator, ParticleCollisionGenerator, bDisableCollisionParallelFor);
 		}
 
 		template<typename T_SPATIALACCELERATION>
 		void ProduceOverlaps(
 			FReal Dt, 
 			const T_SPATIALACCELERATION& InSpatialAcceleration, 
-			FNarrowPhase& NarrowPhase, 
+			FCollisionConstraintAllocator* Allocator,
+			const FCollisionDetectorSettings& Settings,
 			IResimCacheBase* ResimCache
 			)
 		{
+			// Select the set of particles that we loop over in the outer collision detection loop.
+			// The goal is to detection all required collisions (dynamic-vs-everything) while not
+			// visiting pairs that cannot collide (e.g., kinemtic-kinematic, or kinematic-sleeping for
+			// stationary kinematics)
 			const bool bDisableParallelFor = bDisableCollisionParallelFor;
 			auto* EvolutionResimCache = static_cast<FEvolutionResimCache*>(ResimCache);
 			const bool bResimSkipCollision = ResimCache && ResimCache->IsResimming();
@@ -189,20 +227,23 @@ namespace Chaos
 				const TParticleView<TPBDRigidParticles<FReal, 3>>& DynamicSleepingView = Particles.GetNonDisabledDynamicView();
 				const TParticleView<TKinematicGeometryParticles<FReal, 3>>& DynamicMovingKinematicView = Particles.GetActiveDynamicMovingKinematicParticlesView();
 
+				// Usually we ignore sleeping particles in the outer loop and iterate over awake-dynamics and moving-kinematics. 
+				// However, for scenes with a very large number of moving kinematics, it is faster to loop over awake-dynamics
+				// and sleeping-dynamics, even though this means we visit sleeping pairs.
 				if(DynamicSleepingView.Num() < DynamicMovingKinematicView.Num())
 				{
-					ComputeParticlesOverlaps<false, true>(DynamicSleepingView, Dt, InSpatialAcceleration, NarrowPhase);
+					ComputeParticlesOverlaps<false, true>(DynamicSleepingView, Dt, InSpatialAcceleration, Allocator, Settings);
 				}
 				else
 				{
-					ComputeParticlesOverlaps<false, false>(DynamicMovingKinematicView, Dt, InSpatialAcceleration, NarrowPhase);
+					ComputeParticlesOverlaps<false, false>(DynamicMovingKinematicView, Dt, InSpatialAcceleration, Allocator, Settings);
 				}
 			}
 			else
 			{
 				const TParticleView<TGeometryParticles<FReal, 3>>& DesyncedView = ResimCache->GetDesyncedView();
 				
-				ComputeParticlesOverlaps<true,false>(DesyncedView, Dt, InSpatialAcceleration, NarrowPhase);
+				ComputeParticlesOverlaps<true,false>(DesyncedView, Dt, InSpatialAcceleration, Allocator, Settings);
 			}
 		}
 
@@ -215,8 +256,7 @@ namespace Chaos
 		    FReal Dt,
 		    FGeometryParticleHandle* Particle1,
 		    const T_SPATIALACCELERATION& InSpatialAcceleration,
-		    FNarrowPhase& NarrowPhase,
-			int32 EntryIndex)
+			const FCollisionContext& Context)
 		{
 			TArray<FSimOverlapVisitorEntry> PotentialIntersections;
 
@@ -267,7 +307,6 @@ namespace Chaos
 					FSimOverlapVisitor OverlapVisitor(Particle1, ParticleSimData, PotentialIntersections);
 					InSpatialAcceleration.Overlap(Box1, OverlapVisitor);
 				}
-						
 			}
 			else
 			{
@@ -430,7 +469,7 @@ namespace Chaos
 				// NOTE: This searches the array of existing midphases on the particle specified in the last parameter. The assumption
 				// here is that Particle1 is more likely to have less collision on it because it is dynamic. Kinematics can often be
 				// large and may collide with many objects. Maybe we should consider using size instead.
-				FParticlePairMidPhase* MidPhase = NarrowPhase.GetParticlePairMidPhase(ParticleA, ParticleB, Particle1);
+				FParticlePairMidPhase* MidPhase = Context.GetAllocator()->GetMidPhase(ParticleA, ParticleB, Particle1);
 				if (MidPhase != nullptr)
 				{
 					MidPhases.Add(MidPhase);
@@ -454,7 +493,7 @@ namespace Chaos
 				}
 
 				// Run MidPhase + NarrowPhase
-				NarrowPhase.GenerateCollisions(Dt, MidPhases[Index]);
+				MidPhases[Index]->GenerateCollisions(Context.GetSettings().BoundsExpansion, Dt, Context);
 			}
 
 			PHYSICS_CSV_CUSTOM_EXPENSIVE(PhysicsCounters, NumFromBroadphase, NumPotentials, ECsvCustomStatOp::Accumulate);
diff --git a/Engine/Source/Runtime/Experimental/Chaos/Public/Chaos/Collision/SpatialAccelerationCollisionDetector.h b/Engine/Source/Runtime/Experimental/Chaos/Public/Chaos/Collision/SpatialAccelerationCollisionDetector.h
--- a/Engine/Source/Runtime/Experimental/Chaos/Public/Chaos/Collision/SpatialAccelerationCollisionDetector.h
+++ b/Engine/Source/Runtime/Experimental/Chaos/Public/Chaos/Collision/SpatialAccelerationCollisionDetector.h
@@ -3,7 +3,6 @@
 
 #include "Chaos/Collision/CollisionDetector.h"
 #include "Chaos/Collision/CollisionContext.h"
-#include "Chaos/Collision/NarrowPhase.h"
 #include "Chaos/Collision/SpatialAccelerationBroadPhase.h"
 #include "Chaos/EvolutionResimCache.h"
 #include "Chaos/PBDCollisionConstraints.h"
@@ -14,15 +13,16 @@ namespace Chaos
 	class CHAOS_API FSpatialAccelerationCollisionDetector : public FCollisionDetector
 	{
 	public:
-		FSpatialAccelerationCollisionDetector(FSpatialAccelerationBroadPhase& InBroadPhase, FNarrowPhase& InNarrowPhase, FPBDCollisionConstraints& InCollisionContainer)
+		FSpatialAccelerationCollisionDetector(FSpatialAccelerationBroadPhase& InBroadPhase, FPBDCollisionConstraints& InCollisionContainer)
 			: FCollisionDetector(InCollisionContainer)
 			, BroadPhase(InBroadPhase)
-			, NarrowPhase(InNarrowPhase)
 		{
 		}
 
-		FSpatialAccelerationBroadPhase& GetBroadPhase() { return BroadPhase; }
-		FNarrowPhase& GetNarrowPhase() { return NarrowPhase; }
+		FSpatialAccelerationBroadPhase& GetBroadPhase()
+		{ 
+			return BroadPhase;
+		}
 
 		virtual void DetectCollisions(const FReal Dt, FEvolutionResimCache* ResimCache) override
 		{
@@ -38,14 +38,21 @@ namespace Chaos
 			CollisionContainer.BeginDetectCollisions();
 
 			// Collision detection pipeline: BroadPhase -[parallel]-> NarrowPhase -[parallel]-> CollisionAllocator -[serial]-> Container
-			BroadPhase.ProduceOverlaps(Dt, NarrowPhase, ResimCache);
+			BroadPhase.ProduceOverlaps(Dt, &CollisionContainer.GetConstraintAllocator(), Settings, ResimCache);
 
 			CollisionContainer.EndDetectCollisions();
 
 			// If we have a resim cache restore and save contacts
 			if(ResimCache)
 			{
-				CollisionContainer.GetConstraintAllocator().AddResimConstraints(ResimCache->GetAndSanitizeConstraints());
+				// Ensure we have at least one allocator
+				CollisionContainer.GetConstraintAllocator().SetMaxContexts(1);
+
+				FCollisionContext Context;
+				Context.SetSettings(Settings);
+				Context.SetAllocator(CollisionContainer.GetConstraintAllocator().GetContextAllocator(0));
+
+				CollisionContainer.GetConstraintAllocator().AddResimConstraints(ResimCache->GetAndSanitizeConstraints(), Context);
 
 				ResimCache->SaveConstraints(CollisionContainer.GetConstraints());
 			}
@@ -53,6 +60,5 @@ namespace Chaos
 
 	private:
 		FSpatialAccelerationBroadPhase& BroadPhase;
-		FNarrowPhase& NarrowPhase;
 	};
 }
diff --git a/Engine/Source/Runtime/Experimental/Chaos/Public/Chaos/Framework/Parallel.h b/Engine/Source/Runtime/Experimental/Chaos/Public/Chaos/Framework/Parallel.h
--- a/Engine/Source/Runtime/Experimental/Chaos/Public/Chaos/Framework/Parallel.h
+++ b/Engine/Source/Runtime/Experimental/Chaos/Public/Chaos/Framework/Parallel.h
@@ -10,6 +10,7 @@ namespace Chaos
 	void CHAOS_API PhysicsParallelFor(int32 InNum, TFunctionRef<void(int32)> InCallable, bool bForceSingleThreaded = false);
 	void CHAOS_API InnerPhysicsParallelForRange(int32 InNum, TFunctionRef<void(int32, int32)> InCallable, const int32 MinBatchSize, bool bForceSingleThreaded = false);
 	void CHAOS_API InnerPhysicsParallelFor(int32 InNum, TFunctionRef<void(int32)> InCallable, bool bForceSingleThreaded = false);
+	void CHAOS_API PhysicsParallelForWithContext(int32 InNum, TFunctionRef<int32 (int32, int32)> InContextCreator, TFunctionRef<void(int32, int32)> InCallable, bool bForceSingleThreaded = false);
 	//void CHAOS_API PhysicsParallelFor_RecursiveDivide(int32 InNum, TFunctionRef<void(int32)> InCallable, bool bForceSingleThreaded = false);
 
 
diff --git a/Engine/Source/Runtime/Experimental/Chaos/Public/Chaos/ISpatialAccelerationCollection.h b/Engine/Source/Runtime/Experimental/Chaos/Public/Chaos/ISpatialAccelerationCollection.h
--- a/Engine/Source/Runtime/Experimental/Chaos/Public/Chaos/ISpatialAccelerationCollection.h
+++ b/Engine/Source/Runtime/Experimental/Chaos/Public/Chaos/ISpatialAccelerationCollection.h
@@ -10,7 +10,8 @@
 namespace Chaos
 {
 	class FAsyncCollisionReceiver;
-	class FNarrowPhase;
+	class FCollisionConstraintAllocator;
+	class FCollisionDetectorSettings;
 	class FSpatialAccelerationBroadPhase;
 	class IResimCacheBase;
 
@@ -30,7 +31,7 @@ public:
 	virtual void SwapSubstructure(ISpatialAccelerationCollection<TPayloadType, T, d>& Other, FSpatialAccelerationIdx Idx) = 0;	
 
 	/** This is kind of a hack to avoid virtuals. We simply route calls into templated functions */
-	virtual void PBDComputeConstraintsLowLevel(T Dt, FSpatialAccelerationBroadPhase& BroadPhase, FNarrowPhase& NarrowPhase, IResimCacheBase* ResimCache) const = 0;
+	virtual void PBDComputeConstraintsLowLevel(T Dt, FSpatialAccelerationBroadPhase& BroadPhase, FCollisionConstraintAllocator* Allocator, const FCollisionDetectorSettings& Settings, IResimCacheBase* ResimCache) const = 0;
 	virtual TArray<FSpatialAccelerationIdx> GetAllSpatialIndices() const = 0;
 
 	bool IsBucketActive(uint8 BucketIdx) const
diff --git a/Engine/Source/Runtime/Experimental/Chaos/Public/Chaos/PBDCollisionConstraints.h b/Engine/Source/Runtime/Experimental/Chaos/Public/Chaos/PBDCollisionConstraints.h
--- a/Engine/Source/Runtime/Experimental/Chaos/Public/Chaos/PBDCollisionConstraints.h
+++ b/Engine/Source/Runtime/Experimental/Chaos/Public/Chaos/PBDCollisionConstraints.h
@@ -251,10 +251,6 @@ public:
 
 	void UpdateConstraintMaterialProperties(FPBDCollisionConstraint& Contact);
 
-	// @todo(chaos): these settings are duplicated from the NarrowPhase context - fix this
-	const FCollisionDetectorSettings& GetDetectorSettings() const { return DetectorSettings; }
-	void SetDetectorSettings(const FCollisionDetectorSettings& InSettings) { DetectorSettings = InSettings; }
-
 	const FPBDCollisionSolverSettings& GetSolverSettings() const { return SolverSettings; }
 
 protected:
@@ -292,9 +288,6 @@ private:
 	FVec3 GravityDirection;
 	FReal GravitySize;
 
-	// Settings for the collision detector
-	FCollisionDetectorSettings DetectorSettings;
-
 	// Settings for the low-level collision solvers
 	FPBDCollisionSolverSettings SolverSettings;
 };
diff --git a/Engine/Source/Runtime/Experimental/Chaos/Public/Chaos/PBDRigidsEvolutionGBF.h b/Engine/Source/Runtime/Experimental/Chaos/Public/Chaos/PBDRigidsEvolutionGBF.h
--- a/Engine/Source/Runtime/Experimental/Chaos/Public/Chaos/PBDRigidsEvolutionGBF.h
+++ b/Engine/Source/Runtime/Experimental/Chaos/Public/Chaos/PBDRigidsEvolutionGBF.h
@@ -2,7 +2,6 @@
 #pragma once
 
 #include "Chaos/ChaosPerfTest.h"
-#include "Chaos/Collision/NarrowPhase.h"
 #include "Chaos/Collision/SpatialAccelerationBroadPhase.h"
 #include "Chaos/Collision/SpatialAccelerationCollisionDetector.h"
 #include "Chaos/Evolution/SolverBodyContainer.h"
@@ -66,7 +65,6 @@ namespace Chaos
 		static constexpr FRealSingle DefaultCollisionCullDistance = 3.0f;
 		static constexpr FRealSingle DefaultCollisionMaxPushOutVelocity = 1000.0f;
 		static constexpr int32 DefaultRestitutionThreshold = 1000;
-		static constexpr int32 DefaultNumCollisionsPerBlock = 1000;
 
 		CHAOS_API FPBDRigidsEvolutionGBF(FPBDRigidsSOAs& InParticles, THandleArray<FChaosPhysicsMaterial>& SolverPhysicsMaterials, const TArray<ISimCallbackObject*>* InCollisionModifiers = nullptr, bool InIsSingleThreaded = false);
 		CHAOS_API ~FPBDRigidsEvolutionGBF();
@@ -156,11 +154,11 @@ namespace Chaos
 			//SCOPE_CYCLE_COUNTER(STAT_Integrate);
 			CHAOS_SCOPED_TIMER(Integrate);
 
-			const FReal BoundsThickness = GetNarrowPhase().GetBoundsExpansion();
+			const FReal BoundsThickness = GetCollisionDetector().GetSettings().BoundsExpansion;
 			const FReal MaxAngularSpeedSq = CVars::HackMaxAngularVelocity * CVars::HackMaxAngularVelocity;
 			const FReal MaxSpeedSq = CVars::HackMaxVelocity * CVars::HackMaxVelocity;
-			InParticles.ParallelFor([&](auto& GeomParticle, int32 Index) {
-	
+			InParticles.ParallelFor([&](auto& GeomParticle, int32 Index) 
+			{
 				//question: can we enforce this at the API layer? Right now islands contain non dynamic which makes this hard
 				auto PBDParticle = GeomParticle.CastToRigidParticle();
 				if (PBDParticle && PBDParticle->ObjectState() == EObjectStateType::Dynamic)
@@ -291,7 +289,6 @@ namespace Chaos
 		CHAOS_API void SetCurrentStepResimCache(IResimCacheBase* InCurrentStepResimCache);
 
 		CHAOS_API FSpatialAccelerationBroadPhase& GetBroadPhase() { return BroadPhase; }
-		CHAOS_API FNarrowPhase& GetNarrowPhase() { return NarrowPhase; }
 
 		CHAOS_API void TransferJointConstraintCollisions();
 
@@ -335,7 +332,6 @@ namespace Chaos
 		FGravityForces GravityForces;
 		FCollisionConstraints CollisionConstraints;
 		FSpatialAccelerationBroadPhase BroadPhase;
-		FNarrowPhase NarrowPhase;
 		FSpatialAccelerationCollisionDetector CollisionDetector;
 
 		FPBDRigidsEvolutionCallback PostIntegrateCallback;
diff --git a/Engine/Source/Runtime/Experimental/Chaos/Public/Chaos/ParticleIterator.h b/Engine/Source/Runtime/Experimental/Chaos/Public/Chaos/ParticleIterator.h
--- a/Engine/Source/Runtime/Experimental/Chaos/Public/Chaos/ParticleIterator.h
+++ b/Engine/Source/Runtime/Experimental/Chaos/Public/Chaos/ParticleIterator.h
@@ -28,39 +28,49 @@ class THandleView;
 // The function ParticlesParallelFor may be called on ParticleViews, HandleViews,
 // or plain old manually curated arrays of either. In each case, the implementation
 // can differ. The following set of templates will select for the right case.
-
-template <typename TParticleView, typename Lambda>
-void ParticleViewParallelForImp(const TParticleView& Particles, const Lambda& Func);
-
-template <typename THandleView, typename Lambda>
-void HandleViewParallelForImp(const THandleView& HandleView, const Lambda& Func);
-
-template <typename TSOA, typename Lambda>
-void ParticlesParallelForImp(const TConstHandleView<TSOA>& Particles, const Lambda& Func)
+//
+// Signatures:
+// Lambda:				void(auto& Particle, int32 Index, const int32 ContextIndex)
+// ContextCreator:		int32 (int32 WorkerIndex, int32 NumWorkers)
+//
+// There will be one context creation call per worker thread and they are all created prior
+// to starting the work. The ContextCreator should return the context index (which is
+// usually just the WorkerIndex but, for example, it can also be used to generate a single context
+// for use by all workers, assuming the context has appropriate locks in place.)
+//
+
+template <typename TParticleView, typename ContextCreatorType, typename Lambda>
+void ParticleViewParallelForImp(const TParticleView& Particles, const ContextCreatorType& ContextCreator, const Lambda& Func);
+
+template <typename THandleView, typename ContextCreatorType, typename Lambda>
+void HandleViewParallelForImp(const THandleView& HandleView, const ContextCreatorType& ContextCreator, const Lambda& Func);
+
+template <typename TSOA, typename ContextCreatorType, typename Lambda>
+void ParticlesParallelForImp(const TConstHandleView<TSOA>& Particles, const ContextCreatorType& ContextCreator, const Lambda& Func)
 {
-	Chaos::HandleViewParallelForImp(Particles, Func);
+	Chaos::HandleViewParallelForImp(Particles, ContextCreator, Func);
 }
 
-template <typename TSOA, typename Lambda>
-void ParticlesParallelForImp(const THandleView<TSOA>& Particles, const Lambda& Func)
+template <typename TSOA, typename ContextCreatorType, typename Lambda>
+void ParticlesParallelForImp(const THandleView<TSOA>& Particles, const ContextCreatorType& ContextCreator, const Lambda& Func)
 {
-	Chaos::HandleViewParallelForImp(Particles, Func);
+	Chaos::HandleViewParallelForImp(Particles, ContextCreator, Func);
 }
 
-template <typename TSOA, typename Lambda>
-void ParticlesParallelForImp(const TConstParticleView<TSOA>& Particles, const Lambda& Func)
+template <typename TSOA, typename ContextCreatorType, typename Lambda>
+void ParticlesParallelForImp(const TConstParticleView<TSOA>& Particles, const ContextCreatorType& ContextCreator, const Lambda& Func)
 {
-	Chaos::ParticleViewParallelForImp(Particles, Func);
+	Chaos::ParticleViewParallelForImp(Particles, ContextCreator, Func);
 }
 
-template <typename TSOA, typename Lambda>
-void ParticlesParallelForImp(const TParticleView<TSOA>& Particles, const Lambda& Func)
+template <typename TSOA, typename ContextCreatorType, typename Lambda>
+void ParticlesParallelForImp(const TParticleView<TSOA>& Particles, const ContextCreatorType& ContextCreator, const Lambda& Func)
 {
-	Chaos::ParticleViewParallelForImp(Particles, Func);
+	Chaos::ParticleViewParallelForImp(Particles, ContextCreator, Func);
 }
 
-template <typename TParticle, typename Lambda>
-void ParticlesParallelForImp(const TArray<TParticle>& Particles, const Lambda& Func)
+template <typename TParticle, typename ContextCreatorType, typename Lambda>
+void ParticlesParallelForImp(const TArray<TParticle>& Particles, const ContextCreatorType& ContextCreator, const Lambda& Func)
 {
 	// When ParticlesParallelFor is called with a plain old TArray,
 	// just do normal parallelization.
@@ -71,30 +81,32 @@ void ParticlesParallelForImp(const TArray<TParticle>& Particles, const Lambda& F
 	});
 }
 
-template <typename TView, typename Lambda>
-void ParticlesSequentialFor(const TView& Particles, const Lambda& Func)
+
+template <typename TView, typename ContextCreatorType, typename Lambda>
+void ParticlesSequentialFor(const TView& Particles, const ContextCreatorType& ContextCreator, const Lambda& Func)
 {
 	SCOPE_CYCLE_COUNTER(STAT_ParticlesSequentialFor);
 
+	const int32 ContextIndex = ContextCreator(0, 1);
 	int32 Index = 0;
 	for (auto& Particle : Particles)
 	{
-		Func(Particle, Index++);
+		Func(Particle, Index++, ContextIndex);
 	}
 }
 
-template <typename TView, typename Lambda>
-void ParticlesParallelFor(const TView& Particles, const Lambda& Func, bool bForceSingleThreaded = false)
+template <typename TView, typename ContextCreatorType, typename Lambda>
+void ParticlesParallelFor(const TView& Particles, const ContextCreatorType& ContextCreator, const Lambda& Func, bool bForceSingleThreaded = false)
 {
 	SCOPE_CYCLE_COUNTER(STAT_ParticlesParallelFor);
 
 	if (!bForceSingleThreaded && !bDisableParticleParallelFor)
 	{
-		Chaos::ParticlesParallelForImp(Particles, Func);
+		Chaos::ParticlesParallelForImp(Particles, ContextCreator, Func);
 	}
 	else
 	{
-		Chaos::ParticlesSequentialFor(Particles, Func);
+		Chaos::ParticlesSequentialFor(Particles, ContextCreator, Func);
 	}
 }
 
@@ -230,13 +242,22 @@ public:
 	TConstHandleIterator<TSOA> End() const { return TConstHandleIterator<TSOA>(); }
 	TConstHandleIterator<TSOA> end() const { return End(); }
 
-	template <typename TParticleView, typename Lambda>
-	friend void HandleViewParallelForImp(const TParticleView& Particles, const Lambda& Func);
+	template <typename TParticleView, typename ContextCreatorType, typename Lambda>
+	friend void HandleViewParallelForImp(const TParticleView& Particles, const ContextCreatorType& ContextCreator, const Lambda& Func);
+
+	template <typename ContextCreatorType, typename Lambda>
+	void ParallelFor(const ContextCreatorType& ContextCreator, const Lambda& Func) const
+	{
+		ParticlesParallelFor(*this, ContextCreator, Func);
+	}
 
 	template <typename Lambda>
 	void ParallelFor(const Lambda& Func) const
 	{
-		ParticlesParallelFor(*this, Func);
+		// Dummy context creator and a context-stripping function wrapper
+		auto EmptyContextCreator = [](int32 WorkerIndex, int32 NumWorkers) -> int32 { return WorkerIndex; };
+		auto ContextFunc = [&Func](auto& Particle, int32 Index, int32 WorkerIndex) { return Func(Particle, Index); };
+		ParticlesParallelFor(*this, EmptyContextCreator, ContextFunc);
 	}
 
 protected:
@@ -269,10 +290,19 @@ public:
 	THandleIterator<TSOA> End() const { return THandleIterator<TSOA>(); }
 	THandleIterator<TSOA> end() const { return End(); }
 
+	template <typename ContextCreatorType, typename Lambda>
+	void ParallelFor(const ContextCreatorType& ContextCreator, const Lambda& Func) const
+	{
+		ParticlesParallelFor(*this, ContextCreator, Func);
+	}
+
 	template <typename Lambda>
 	void ParallelFor(const Lambda& Func) const
 	{
-		ParticlesParallelFor(*this, Func);
+		// Dummy context creator and a context-stripping function wrapper
+		auto EmptyContextCreator = [](int32 WorkerIndex, int32 NumWorkers) -> int32 { return WorkerIndex; };
+		auto ContextFunc = [&Func](auto& Particle, int32 Index, int32 WorkerIndex) { return Func(Particle, Index); };
+		ParticlesParallelFor(*this, EmptyContextCreator, ContextFunc);
 	}
 };
 
@@ -579,14 +609,24 @@ public:
 	TConstParticleIterator<TSOA> End() const { return TConstParticleIterator<TSOA>(); }
 	TConstParticleIterator<TSOA> end() const { return End(); }
 
+	template <typename ContextCreatorType, typename Lambda>
+	void ParallelFor(const ContextCreatorType& ContextCreator, const Lambda& Func, bool bForceSingleThreaded = false) const
+	{
+		ParticlesParallelFor(*this, ContextCreator, Func, bForceSingleThreaded);
+	}
+
 	template <typename Lambda>
 	void ParallelFor(const Lambda& Func, bool bForceSingleThreaded = false) const
 	{
-		ParticlesParallelFor(*this, Func, bForceSingleThreaded);
+		// Dummy context creator and a context-stripping function wrapper
+		auto EmptyContextCreator = [](int32 WorkerIndex, int32 NumWorkers) -> int32 { return WorkerIndex; };
+		auto ContextFunc = [&Func](auto& Particle, int32 Index, int32 WorkerIndex) { return Func(Particle, Index); };
+
+		ParticlesParallelFor(*this, EmptyContextCreator, ContextFunc, bForceSingleThreaded);
 	}
 
-	template <typename TParticleView, typename Lambda>
-	friend void ParticleViewParallelForImp(const TParticleView& Particles, const Lambda& Func);
+	template <typename TParticleView, typename ContextCreatorType, typename Lambda>
+	friend void ParticleViewParallelForImp(const TParticleView& Particles, const ContextCreatorType& ContextCreator, const Lambda& Func);
 
 protected:
 
@@ -620,10 +660,20 @@ public:
 	TParticleIterator<TSOA> End() const { return TParticleIterator<TSOA>(); }
 	TParticleIterator<TSOA> end() const { return End(); }
 
+	template <typename ContextCreatorType, typename Lambda>
+	void ParallelFor(const ContextCreatorType& ContextCreator, const Lambda& Func, bool bForceSingleThreaded = false) const
+	{
+		ParticlesParallelFor(*this, ContextCreator, Func, bForceSingleThreaded);
+	}
+
 	template <typename Lambda>
 	void ParallelFor(const Lambda& Func, bool bForceSingleThreaded = false) const
 	{
-		ParticlesParallelFor(*this, Func, bForceSingleThreaded);
+		// Dummy context creator and a context-stripping function wrapper
+		auto EmptyContextCreator = [](int32 WorkerIndex, int32 NumWorkers) -> int32 { return WorkerIndex; };
+		auto ContextFunc = [&Func](auto& Particle, int32 Index, int32 WorkerIndex) { return Func(Particle, Index); };
+
+		ParticlesParallelFor(*this, EmptyContextCreator, ContextFunc, bForceSingleThreaded);
 	}
 };
 
@@ -663,8 +713,8 @@ TParticleView<TSOA> MakeParticleView(TSOA* SOA)
 
 // The non-parallel implementation of iteration should not deviate in behavior from
 // this parallel implementation.  They must be kept in sync.
-template <typename TParticleView, typename Lambda>
-void ParticleViewParallelForImp(const TParticleView& Particles, const Lambda& UserFunc)
+template <typename TParticleView, typename ContextCreatorType, typename Lambda>
+void ParticleViewParallelForImp(const TParticleView& Particles, const ContextCreatorType& ContextCreator, const Lambda& UserFunc)
 {
 	SCOPE_CYCLE_COUNTER(STAT_ParticleViewParallelForImp);
 
@@ -673,11 +723,11 @@ void ParticleViewParallelForImp(const TParticleView& Particles, const Lambda& Us
 	using THandleBase = typename THandle::THandleBase;
 	using TTransientHandle = typename THandle::TTransientHandle;
 
-	auto Func = [&UserFunc](auto& Handle, const int32 Idx)
+	auto Func = [&UserFunc](auto& Handle, const int32 Idx, const int32 ContextIndex)
 	{
 		if (!Handle.LightWeightDisabled())
 		{
-			UserFunc(Handle, Idx);
+			UserFunc(Handle, Idx, ContextIndex);
 		}
 	};
 
@@ -697,45 +747,45 @@ void ParticleViewParallelForImp(const TParticleView& Particles, const Lambda& Us
 		{
 			// Do a regular parallel for over the handles in this SOA view
 			const int32 HandleCount = CurHandlesArray->Num();
-			PhysicsParallelFor(HandleCount, [&Func, CurHandlesArray, ParticleIdxOff](const int32 HandleIdx)
+			PhysicsParallelForWithContext(HandleCount, ContextCreator, [&Func, CurHandlesArray, ParticleIdxOff](const int32 HandleIdx, const int32 ContextIndex)
 			{
 				// Reconstruct the TransientHandle so that it has a chance to update
 				// other data members, like the particle type in the case of geometry 
 				// particles.
 				THandle* HandlePtr = (*CurHandlesArray)[HandleIdx];
 				THandleBase Handle(HandlePtr->GeometryParticles, HandlePtr->ParticleIdx);
-				Func(static_cast<TTransientHandle&>(Handle), ParticleIdxOff + HandleIdx);
+				Func(static_cast<TTransientHandle&>(Handle), ParticleIdxOff + HandleIdx, ContextIndex);
 			});
 			ParticleIdxOff += HandleCount;
 		}
 		else
 		{
 			// Do a regular parallel for over the particles in this SOA view
-			PhysicsParallelFor(ParticleCount, [&Func, &SOAView, ParticleIdxOff](const int32 ParticleIdx)
+			PhysicsParallelForWithContext(ParticleCount, ContextCreator, [&Func, &SOAView, ParticleIdxOff](const int32 ParticleIdx, const int32 ContextIndex)
 			{
 				// Reconstruct the TransientHandle so that it has a chance to update
 				// other data members, like the particle type in the case of geometry 
 				// particles.
 				THandleBase Handle(SOAView.SOA, ParticleIdx);
-				Func(static_cast<TTransientHandle&>(Handle), ParticleIdxOff + ParticleIdx);
+				Func(static_cast<TTransientHandle&>(Handle), ParticleIdxOff + ParticleIdx, ContextIndex);
 			});
 			ParticleIdxOff += ParticleCount;
 		}
 	}
 }
 
-template <typename THandleView, typename Lambda>
-void HandleViewParallelForImp(const THandleView& HandleView, const Lambda& Func)
+template <typename THandleView, typename ContextCreatorType, typename Lambda>
+void HandleViewParallelForImp(const THandleView& HandleView, const ContextCreatorType& ContextCreator, const Lambda& Func)
 {
 	SCOPE_CYCLE_COUNTER(STAT_HandleViewParallelForImp);
 
 	using THandle = typename THandleView::THandle;
 	const int32 HandleCount = HandleView.Handles.Num();
-	PhysicsParallelFor(HandleCount, [&HandleView, &Func](const int32 Index)
+	PhysicsParallelForWithContext(HandleCount, ContextCreator, [&HandleView, &Func](const int32 Index, const int32 ContextIndex)
 	{
 		if (!HandleView.Handles[Index]->LightWeightDisabled())
 		{
-			Func(static_cast<THandle&>(*HandleView.Handles[Index]), Index);
+			Func(static_cast<THandle&>(*HandleView.Handles[Index]), Index, ContextIndex);
 		}
 	});
 }
diff --git a/Engine/Source/Runtime/Experimental/Chaos/Public/Chaos/SpatialAccelerationCollection.h b/Engine/Source/Runtime/Experimental/Chaos/Public/Chaos/SpatialAccelerationCollection.h
--- a/Engine/Source/Runtime/Experimental/Chaos/Public/Chaos/SpatialAccelerationCollection.h
+++ b/Engine/Source/Runtime/Experimental/Chaos/Public/Chaos/SpatialAccelerationCollection.h
@@ -354,13 +354,13 @@ struct TSpatialAccelerationCollectionHelper
 };
 
 template <typename SpatialAccelerationCollection>
-typename TEnableIf<TIsSame<typename SpatialAccelerationCollection::TPayloadType, FAccelerationStructureHandle>::Value, void>::Type PBDComputeConstraintsLowLevel_Helper(FReal Dt, const SpatialAccelerationCollection& Accel, FSpatialAccelerationBroadPhase& BroadPhase, FNarrowPhase& NarrowPhase, IResimCacheBase* ResimCache)
+typename TEnableIf<TIsSame<typename SpatialAccelerationCollection::TPayloadType, FAccelerationStructureHandle>::Value, void>::Type PBDComputeConstraintsLowLevel_Helper(FReal Dt, const SpatialAccelerationCollection& Accel, FSpatialAccelerationBroadPhase& BroadPhase, FCollisionConstraintAllocator* Allocator, const FCollisionDetectorSettings& Settings, IResimCacheBase* ResimCache)
 {
-	BroadPhase.ProduceOverlaps(Dt, Accel, NarrowPhase, ResimCache);
+	BroadPhase.ProduceOverlaps(Dt, Accel, Allocator, Settings, ResimCache);
 }
 
 template <typename SpatialAccelerationCollection>
-typename TEnableIf<!TIsSame<typename SpatialAccelerationCollection::TPayloadType, FAccelerationStructureHandle>::Value, void>::Type PBDComputeConstraintsLowLevel_Helper(FReal Dt, const SpatialAccelerationCollection& Accel, FSpatialAccelerationBroadPhase& BroadPhase, FNarrowPhase& NarrowPhase, IResimCacheBase* ResimCache)
+typename TEnableIf<!TIsSame<typename SpatialAccelerationCollection::TPayloadType, FAccelerationStructureHandle>::Value, void>::Type PBDComputeConstraintsLowLevel_Helper(FReal Dt, const SpatialAccelerationCollection& Accel, FSpatialAccelerationBroadPhase& BroadPhase, FCollisionConstraintAllocator* Allocator, const FCollisionDetectorSettings& Settings, IResimCacheBase* ResimCache)
 {
 }
 
@@ -562,9 +562,9 @@ public:
 		return TUniquePtr<ISpatialAcceleration<TPayloadType, T, d>>(new TSpatialAccelerationCollection<TSpatialAccelerationTypes...>(*this));
 	}
 
-	virtual void PBDComputeConstraintsLowLevel(T Dt, FSpatialAccelerationBroadPhase& BroadPhase, FNarrowPhase& NarrowPhase, IResimCacheBase* ResimCache) const override
+	virtual void PBDComputeConstraintsLowLevel(T Dt, FSpatialAccelerationBroadPhase& BroadPhase, FCollisionConstraintAllocator* Allocator, const FCollisionDetectorSettings& Settings, IResimCacheBase* ResimCache) const override
 	{
-		PBDComputeConstraintsLowLevel_Helper(Dt, *this, BroadPhase, NarrowPhase, ResimCache);
+		PBDComputeConstraintsLowLevel_Helper(Dt, *this, BroadPhase, Allocator, Settings, ResimCache);
 	}
 
 #if !UE_BUILD_SHIPPING
diff --git a/Engine/Source/Runtime/Experimental/Chaos/Public/PBDRigidsSolver.h b/Engine/Source/Runtime/Experimental/Chaos/Public/PBDRigidsSolver.h
--- a/Engine/Source/Runtime/Experimental/Chaos/Public/PBDRigidsSolver.h
+++ b/Engine/Source/Runtime/Experimental/Chaos/Public/PBDRigidsSolver.h
@@ -171,7 +171,7 @@ namespace Chaos
 		void SetPositionIterations(const int32 InNumIterations) { GetEvolution()->SetNumPositionIterations(InNumIterations); }
 		void SetVelocityIterations(const int32 InNumIterations) { GetEvolution()->SetNumVelocityIterations(InNumIterations); }
 		void SetProjectionIterations(const int32 InNumIterations) { GetEvolution()->SetNumProjectionIterations(InNumIterations); }
-		void SetCollisionCullDistance(const FReal InCullDistance) { GetEvolution()->GetNarrowPhase().SetBoundsExpansion(InCullDistance); }
+		void SetCollisionCullDistance(const FReal InCullDistance) { GetEvolution()->GetCollisionDetector().SetBoundsExpansion(InCullDistance); }
 		void SetCollisionMaxPushOutVelocity(const FReal InMaxPushOutVelocity) { GetEvolution()->GetCollisionConstraints().SetMaxPushOutVelocity(InMaxPushOutVelocity); }
 
 		/**/
diff --git a/Engine/Source/Runtime/Experimental/ChaosCore/Public/Chaos/ObjectPool.h b/Engine/Source/Runtime/Experimental/ChaosCore/Public/Chaos/ObjectPool.h
--- a/Engine/Source/Runtime/Experimental/ChaosCore/Public/Chaos/ObjectPool.h
+++ b/Engine/Source/Runtime/Experimental/ChaosCore/Public/Chaos/ObjectPool.h
@@ -26,7 +26,8 @@ namespace Chaos
 		// Alignment of the stored type
 		constexpr static int32 ItemAlign = alignof(ObjectType);
 
-		using Ptr = ObjectType*;
+		using FObject = ObjectType;
+		using FPtr = ObjectType*;
 
 		explicit TObjectPool(int32 InNumPerBlock, int32 InitialBlocks = 1)
 			: NumPerBlock(InNumPerBlock)
@@ -58,7 +59,7 @@ namespace Chaos
 		 * @return The newly allocated object
 		 */
 		template<typename... TArgs>
-		Ptr Alloc(TArgs&&... Args)
+		FPtr Alloc(TArgs&&... Args)
 		{
 			// Need a new Block
 			if(FreeCount == 0)
@@ -75,23 +76,22 @@ namespace Chaos
 
 			if (Blocks[0].IsFull())
 			{
-				// Set BackIndex to the first empty block from the end
-				int32 BackIndex = Blocks.Num() - 1;
-				while(BackIndex > 0)
+				// Find the fullest block that is not completely full
+				int32 SelectedIndex = INDEX_NONE;
+				int32 SelectedNumFree = TNumericLimits<int32>::Max();
+				for (int32 Index = 0; Index < Blocks.Num(); ++Index)
 				{
-					if(Blocks[BackIndex].IsFull())
+					const int32 BlockNumFree = Blocks[Index].NumFree;
+					if ((BlockNumFree > 0) && (BlockNumFree < SelectedNumFree))
 					{
-						--BackIndex;
-					}
-					else
-					{
-						break;
+						SelectedNumFree = BlockNumFree;
+						SelectedIndex = Index;
 					}
 				}
 
-				checkf(!Blocks[BackIndex].IsFull(), TEXT("Could not find an empty block"));
-
-				Swap(Blocks[0], Blocks[BackIndex]);
+				// Move the selected block to the front
+				checkf(SelectedIndex != INDEX_NONE, TEXT("Could not find an empty block"));
+				Swap(Blocks[0], Blocks[SelectedIndex]);
 			}
 
 			// Ensure any writes to Blocks[0] cannot be reordered so that the GetNextFree
@@ -114,7 +114,7 @@ namespace Chaos
 		 * this will not be asserted
 		 * @param Object The object to free
 		 */
-		void Free(Ptr Object)
+		void Free(FPtr Object)
 		{
 			int32 BlockIndex = FindBlock(Object);
 			Blocks[BlockIndex].Free(Object);
@@ -334,7 +334,7 @@ namespace Chaos
 				FMemory::Free(Begin);
 			}
 
-			Ptr GetNextFree()
+			FPtr GetNextFree()
 			{
 				if(FreeList != INDEX_NONE)
 				{
@@ -559,21 +559,24 @@ namespace Chaos
 	/**
 	 * A deleter for use with TUniquePtr and a TObjectPool item
 	 */
-	template<typename ObjectType>
+	template<typename ObjectPoolType>
 	class TObjectPoolDeleter
 	{
 	public:
+		using FObjectPool = ObjectPoolType;
+		using FObject = typename FObjectPool::FObject;
+
 		TObjectPoolDeleter()
 			: Pool(nullptr)
 		{
 		}
 
-		TObjectPoolDeleter(TObjectPool<ObjectType>& InPool)
+		TObjectPoolDeleter(FObjectPool& InPool)
 			: Pool(&InPool)
 		{
 		}
 
-		void operator()(ObjectType* Object)
+		void operator()(FObject* Object)
 		{
 			if (Object != nullptr)
 			{
@@ -583,6 +586,6 @@ namespace Chaos
 		}
 
 	private:
-		TObjectPool<ObjectType>* Pool;
+		FObjectPool* Pool;
 	};
 }
\ No newline at end of file
