GeometryCollection : Introspection.
-Use virtual IsA to identify derived classes.
-Fixed typo in api.

diff --git a/Engine/Source/Runtime/Experimental/Chaos/Private/GeometryCollection/GeometryCollection.cpp b/Engine/Source/Runtime/Experimental/Chaos/Private/GeometryCollection/GeometryCollection.cpp
--- a/Engine/Source/Runtime/Experimental/Chaos/Private/GeometryCollection/GeometryCollection.cpp
+++ b/Engine/Source/Runtime/Experimental/Chaos/Private/GeometryCollection/GeometryCollection.cpp
@@ -32,14 +32,14 @@ const FName FGeometryCollection::SimulatableParticlesAttribute("SimulatableParti
 const FName FGeometryCollection::SimulationTypeAttribute("SimulationType");
 const FName FGeometryCollection::StatusFlagsAttribute("StatusFlags");
 
-FGeometryCollection::FGeometryCollection(FClassType InType)
-	: FTransformCollection(InType.Union({ FGeometryCollection::StaticType() }))
+FGeometryCollection::FGeometryCollection()
+	: FTransformCollection()
 	, FGeometryCollectionConvexPropertiesInterface(this)
 {
 	Construct();
 }
 
-void FGeometryCollection::DefineGeometryScheam(FManagedArrayCollection& InCollection)
+void FGeometryCollection::DefineGeometrySchema(FManagedArrayCollection& InCollection)
 {
 	FTransformCollection::DefineTransformSchema(InCollection);
 
diff --git a/Engine/Source/Runtime/Experimental/Chaos/Private/GeometryCollection/ManagedArrayCollection.cpp b/Engine/Source/Runtime/Experimental/Chaos/Private/GeometryCollection/ManagedArrayCollection.cpp
--- a/Engine/Source/Runtime/Experimental/Chaos/Private/GeometryCollection/ManagedArrayCollection.cpp
+++ b/Engine/Source/Runtime/Experimental/Chaos/Private/GeometryCollection/ManagedArrayCollection.cpp
@@ -11,10 +11,8 @@ DEFINE_LOG_CATEGORY_STATIC(FManagedArrayCollectionLogging, NoLogging, All);
 int8 FManagedArrayCollection::Invalid = INDEX_NONE;
 
 
-FManagedArrayCollection::FManagedArrayCollection(FClassType InType)
-	: ClassType(InType)
+FManagedArrayCollection::FManagedArrayCollection()
 {
-	ClassType.Add(FManagedArrayCollection::StaticType());
 	Version = 9;
 }
 
diff --git a/Engine/Source/Runtime/Experimental/Chaos/Private/GeometryCollection/TransformCollection.cpp b/Engine/Source/Runtime/Experimental/Chaos/Private/GeometryCollection/TransformCollection.cpp
--- a/Engine/Source/Runtime/Experimental/Chaos/Private/GeometryCollection/TransformCollection.cpp
+++ b/Engine/Source/Runtime/Experimental/Chaos/Private/GeometryCollection/TransformCollection.cpp
@@ -13,8 +13,8 @@ const FName FTransformCollection::ParentAttribute = "Parent";
 const FName FTransformCollection::ChildrenAttribute = "Children";
 const FName FTransformCollection::ParticlesAttribute = "Particles";
 
-FTransformCollection::FTransformCollection(FClassType InType)
-	: FManagedArrayCollection(InType.Union({ FTransformCollection::StaticType() }))
+FTransformCollection::FTransformCollection()
+	: FManagedArrayCollection()
 {
 	Construct();
 }
diff --git a/Engine/Source/Runtime/Experimental/Chaos/Public/GeometryCollection/GeometryCollection.h b/Engine/Source/Runtime/Experimental/Chaos/Public/GeometryCollection/GeometryCollection.h
--- a/Engine/Source/Runtime/Experimental/Chaos/Public/GeometryCollection/GeometryCollection.h
+++ b/Engine/Source/Runtime/Experimental/Chaos/Public/GeometryCollection/GeometryCollection.h
@@ -31,13 +31,15 @@ class CHAOS_API FGeometryCollection : public FTransformCollection,
 {
 
 public:
-	FGeometryCollection(FClassType InType = {});
+	typedef FTransformCollection Super;
+
+	FGeometryCollection();
 	FGeometryCollection(FGeometryCollection &) = delete;
 	FGeometryCollection& operator=(const FGeometryCollection &) = delete;
 	FGeometryCollection(FGeometryCollection &&) = default;
 	FGeometryCollection& operator=(FGeometryCollection &&) = default;
+	MANAGED_ARRAY_COLLECTION_INTERNAL(FGeometryCollection);
 
-	typedef FTransformCollection Super;
 		/***
 		*  Attribute Groups
 		*
@@ -122,19 +124,14 @@ public:
 		FS_IgnoreCollisionInParentCluster = 0x00000008
 
 	};
-	//
-	//
-	//
 
-	/** Type name for this class. */
-	static FName StaticType() { return FName("FGeometryCollection"); }
 
 	/**
 	 * Create a GeometryCollection from Vertex and Indices arrays
 	 */
 	static FGeometryCollection* NewGeometryCollection(const TArray<float>& RawVertexArray, const TArray<int32>& RawIndicesArray, bool ReverseVertexOrder = true);
 	static void Init(FGeometryCollection* Collection, const TArray<float>& RawVertexArray, const TArray<int32>& RawIndicesArray, bool ReverseVertexOrder = true);
-	static void DefineGeometryScheam(FManagedArrayCollection&);
+	static void DefineGeometrySchema(FManagedArrayCollection&);
 
 	/**
 	* Create a GeometryCollection from Vertex, Indices, BoneMap, Transform, BoneHierarchy arrays
diff --git a/Engine/Source/Runtime/Experimental/Chaos/Public/GeometryCollection/ManagedArrayCollection.h b/Engine/Source/Runtime/Experimental/Chaos/Public/GeometryCollection/ManagedArrayCollection.h
--- a/Engine/Source/Runtime/Experimental/Chaos/Public/GeometryCollection/ManagedArrayCollection.h
+++ b/Engine/Source/Runtime/Experimental/Chaos/Public/GeometryCollection/ManagedArrayCollection.h
@@ -47,10 +47,9 @@ struct CHAOS_API FManagedArrayCollection
 	friend FSimulationProperties;
 
 public:
-	typedef TSet<FName> FClassType;
 
 
-	FManagedArrayCollection(FClassType InType = {});
+	FManagedArrayCollection();
 	virtual ~FManagedArrayCollection() {}
 
 	FManagedArrayCollection(const FManagedArrayCollection& In) { In.CopyTo(this); }
@@ -104,14 +103,16 @@ public:
 	*/
 	static FName StaticType() { return FName("FManagedArrayCollection"); }
 
-	template<class T>
-	bool IsA() { return ClassType.Contains(T::StaticType())?true:false; }
+	virtual bool IsAType(FName InTypeName) const { return InTypeName.IsEqual(FManagedArrayCollection::StaticType()); }
+
+	template<typename T>
+	bool IsA() { return IsAType(T::StaticType()); }
 
 	template<class T>
-	T* Cast() { return ClassType.Contains(T::StaticType())?static_cast<T*>(this):nullptr; }
+	T* Cast() { return IsAType(T::StaticType())?static_cast<T*>(this):nullptr; }
 
 	template<class T>
-	const T* Cast() const { return ClassType.Contains(T::StaticType()) ? static_cast<const T*>(this) : nullptr; }
+	const T* Cast() const { return IsAType(T::StaticType())? static_cast<const T*>(this) : nullptr; }
 
 	/**
 	* Add an attribute of Type(T) to the group
@@ -624,7 +625,6 @@ private:
 
 	virtual void SetDefaults(FName Group, uint32 StartSize, uint32 NumElements) {};
 
-	FClassType ClassType;
 	TMap< FKeyType, FValueType> Map;	//data is owned by the map explicitly
 	TMap< FName, FGroupInfo> GroupInfo;
 	bool bDirty;
@@ -675,3 +675,11 @@ protected:
 	FManagedArrayCollection* ManagedCollection;
 
 };
+
+
+#define MANAGED_ARRAY_COLLECTION_INTERNAL(TYPE_NAME)					\
+	static FName StaticType() { return FName(#TYPE_NAME); }				\
+	virtual bool IsAType(FName InTypeName) const override {				\
+		return InTypeName.IsEqual(TYPE_NAME::StaticType())				\
+				|| Super::IsAType(InTypeName);							\
+	}
diff --git a/Engine/Source/Runtime/Experimental/Chaos/Public/GeometryCollection/TransformCollection.h b/Engine/Source/Runtime/Experimental/Chaos/Public/GeometryCollection/TransformCollection.h
--- a/Engine/Source/Runtime/Experimental/Chaos/Public/GeometryCollection/TransformCollection.h
+++ b/Engine/Source/Runtime/Experimental/Chaos/Public/GeometryCollection/TransformCollection.h
@@ -18,11 +18,12 @@ class CHAOS_API FTransformCollection : public FManagedArrayCollection
 public:
 	typedef FManagedArrayCollection Super;
 
-	FTransformCollection(FClassType InType = {});
+	FTransformCollection();
 	FTransformCollection(FTransformCollection &) = delete;
 	FTransformCollection& operator=(const FTransformCollection&) = delete;
 	FTransformCollection(FTransformCollection&&) = default;
 	FTransformCollection& operator=(FTransformCollection&&) = default;
+	MANAGED_ARRAY_COLLECTION_INTERNAL(FTransformCollection);
 
 	/***
 	*  Attribute Groups
@@ -50,8 +51,6 @@ public:
 	static const FName ChildrenAttribute;
 	static const FName ParticlesAttribute;
 
-	/** Type name for this class. */
-	static FName StaticType() { return FName("FTransformCollection"); }
 
 	/** Schema */
 	static void DefineTransformSchema(FManagedArrayCollection&);
diff --git a/Engine/Source/Runtime/Experimental/GeometryCollectionEngine/Private/GeometryCollection/GeometryCollectionEngineConversion.cpp b/Engine/Source/Runtime/Experimental/GeometryCollectionEngine/Private/GeometryCollection/GeometryCollectionEngineConversion.cpp
--- a/Engine/Source/Runtime/Experimental/GeometryCollectionEngine/Private/GeometryCollection/GeometryCollectionEngineConversion.cpp
+++ b/Engine/Source/Runtime/Experimental/GeometryCollectionEngine/Private/GeometryCollection/GeometryCollectionEngineConversion.cpp
@@ -919,7 +919,7 @@ bool FGeometryCollectionEngineConversion::AppendSkeletalMesh(const USkeletalMesh
 	{
 		return false;
 	}
-	FGeometryCollection::DefineGeometryScheam(*InCollection);
+	FGeometryCollection::DefineGeometrySchema(*InCollection);
 
 	const FSkeletalMeshLODRenderData* MeshLODData = GetSkeletalMeshLOD(SkeletalMesh, 0);
 	if (!MeshLODData)
