GeometryCollection - Skeleton Import
- Added utility function to import a skeleton into a FTransformCollection

diff --git a/Engine/Source/Runtime/Experimental/Chaos/Private/GeometryCollection/Facades/CollectionTransformSourceFacade.cpp b/Engine/Source/Runtime/Experimental/Chaos/Private/GeometryCollection/Facades/CollectionTransformSourceFacade.cpp
new file mode 100644
--- /dev/null
+++ b/Engine/Source/Runtime/Experimental/Chaos/Private/GeometryCollection/Facades/CollectionTransformSourceFacade.cpp
@@ -0,0 +1,94 @@
+// Copyright Epic Games, Inc. All Rights Reserved.
+
+/*=============================================================================
+	GeometryCollection.cpp: FGeometryCollection methods.
+=============================================================================*/
+
+#include "GeometryCollection/Facades/CollectionTransformSourceFacade.h"
+#include "GeometryCollection/TransformCollection.h"
+
+namespace GeometryCollection::Facades
+{
+
+	// Groups 
+	const FName FTransformSource::TransformSourceGroup = "TransformSource";
+
+
+	// Attributes
+	const FName FTransformSource::SourceNameAttribute = "Name";
+	const FName FTransformSource::SourceGuidAttribute = "Guid";
+	const FName FTransformSource::SourceRootsAttribute = "Roots";
+
+	FTransformSource::FTransformSource(FManagedArrayCollection* InCollection)
+		: Self(InCollection)
+	{
+		DefineSchema(Self);
+	}
+
+	//
+	//  Initialization
+	//
+
+	void FTransformSource::DefineSchema(FManagedArrayCollection* Collection)
+	{
+		FManagedArrayCollection::FConstructionParameters TransformDependency(FTransformCollection::TransformGroup);
+
+		if (!Collection->HasGroup(TransformSourceGroup))
+		{
+			Collection->AddGroup(TransformSourceGroup);
+		}
+		Collection->AddAttribute<FString>(SourceNameAttribute, TransformSourceGroup);
+		Collection->AddAttribute<FGuid>(SourceGuidAttribute, TransformSourceGroup);
+		Collection->AddAttribute<TSet<int32>>(SourceRootsAttribute, TransformSourceGroup, TransformDependency);
+	}
+
+	bool FTransformSource::HasFacade(FManagedArrayCollection* Collection)
+	{
+		return Collection->HasGroup(TransformSourceGroup) &&
+			Collection->FindAttribute<FString>(SourceNameAttribute, TransformSourceGroup) &&
+			Collection->FindAttribute<FGuid>(SourceGuidAttribute, TransformSourceGroup) &&
+			Collection->FindAttribute<TSet<int32>>(SourceRootsAttribute, TransformSourceGroup);
+	}
+
+	//
+	//  Add Data
+	//
+	void FTransformSource::AddTransformSource(FManagedArrayCollection* Collection, const FString& InName, const FGuid& InGuid, const TSet<int32>& InRoots)
+	{
+		DefineSchema(Collection);
+
+		int Idx = Collection->AddElements(1, TransformSourceGroup);
+		Collection->ModifyAttribute<FString>(FTransformSource::SourceNameAttribute, TransformSourceGroup)[Idx] = InName;
+		Collection->ModifyAttribute<FGuid>(FTransformSource::SourceGuidAttribute, TransformSourceGroup)[Idx] = InGuid;
+		Collection->ModifyAttribute<TSet<int32>>(FTransformSource::SourceRootsAttribute, TransformSourceGroup)[Idx] = InRoots;
+	}
+
+	//
+	//  Get Data
+	//
+
+
+	TSet<int32> FTransformSource::GetTransformSource(FManagedArrayCollection* Collection, const FString& InName, const FGuid& InGuid)
+	{
+		if (Collection->HasGroup(TransformSourceGroup))
+		{
+			const TManagedArray<TSet<int32>>* Roots = Collection->FindAttribute<TSet<int32>>(SourceRootsAttribute, TransformSourceGroup);
+			const TManagedArray<FGuid>* Guids = Collection->FindAttribute<FGuid>(SourceGuidAttribute, TransformSourceGroup);
+			const TManagedArray<FString>* Names = Collection->FindAttribute<FString>(SourceNameAttribute, TransformSourceGroup);
+			if (Roots && Guids && Names)
+			{
+				int32 GroupNum = Collection->NumElements(TransformSourceGroup);
+				for (int i = 0; i < GroupNum; i++)
+				{
+					if ((*Guids)[i]==InGuid && (*Names)[i].Equals(InName))
+					{
+						return (*Roots)[i];
+					}
+				}
+			}
+		}
+		return TSet<int32>();
+	}
+};
+
+
diff --git a/Engine/Source/Runtime/Experimental/Chaos/Public/GeometryCollection/Facades/CollectionTransformSourceFacade.h b/Engine/Source/Runtime/Experimental/Chaos/Public/GeometryCollection/Facades/CollectionTransformSourceFacade.h
new file mode 100644
--- /dev/null
+++ b/Engine/Source/Runtime/Experimental/Chaos/Public/GeometryCollection/Facades/CollectionTransformSourceFacade.h
@@ -0,0 +1,79 @@
+// Copyright Epic Games, Inc. All Rights Reserved.
+
+#pragma once
+#include "Containers/UnrealString.h"
+#include "Containers/Array.h"
+#include "GeometryCollection/ManagedArrayCollection.h"
+
+namespace GeometryCollection::Facades
+{
+
+	/**
+	* FTransformSource
+	* 
+	* Defines common API for storing owners of a transform hierarchy. For example, to store
+	* the component that generated a transform hierirchy, use the function :
+	*    Key = this->AddTransformSource(<Skeleton Asset Name>, <Skeleton Asset Guid>, {<Root transform indices in collection>});
+	* 
+	*    Key = this->AddTransformSource("", GUID, {1,5,7});
+	*
+	* The source root indices can be queries later using the name and guid:
+	*	TSet<int32> SourceRoots = this->GetSourceTransformRoots("", GUID);
+	*
+	* The following groups are created on the collection based on which API is called. 
+	* 
+	*	<Group> = FTransformSource::TransformSourceGroup
+	*	- FindAttribute<FString>(FTransformSource::SourceNameAttribute, <Group>);
+	*	- FindAttribute<Guid>(FTransformSource::SourceGuidAttribute, <Group>);
+	*	- FindAttribute<TSet<int32>>(FTransformSource::SourceRootsAttribute, <Group>);
+	* 
+	*/
+	class CHAOS_API FTransformSource
+ 	{
+		FManagedArrayCollection* Self;
+
+	public:
+
+		// groups
+		static const FName TransformSourceGroup;
+
+		// Attributes
+		static const FName SourceNameAttribute;
+		static const FName SourceGuidAttribute;
+		static const FName SourceRootsAttribute;
+
+		/**
+		* FSelectionFacade Constuctor
+		* @param VertixDependencyGroup : GroupName the index attribute is dependent on. 
+		*/
+		FTransformSource(FManagedArrayCollection* InSelf);
+
+		/**
+		*  Create the facade.
+		*/
+		static void DefineSchema(FManagedArrayCollection* Self);
+
+		/**
+		*  Is the Facade defined on the collection?
+		*/
+		static bool HasFacade(FManagedArrayCollection* Collection);
+
+		/**
+		* Add a transform root mapping.
+		* @param Name : Name of the owner of the transform set.
+		* @param Guid : Guid of the owner of the transform set.
+		* @param Roots : Root indices of the transform set.
+		*/
+		static void AddTransformSource(FManagedArrayCollection* Collection, const FString& Name, const FGuid& Guid, const TSet<int32>& Roots);
+
+		/**
+		* Query for root indices.  
+		* @param Name : Name of the owner of the transform set.  
+		* @param Guid : Guid of the owner of the transform set.
+		*/
+		static TSet<int32> GetTransformSource(FManagedArrayCollection* Collection, const FString& Name, const FGuid& Guid);
+
+
+	};
+
+}
\ No newline at end of file
diff --git a/Engine/Source/Runtime/Experimental/GeometryCollectionEngine/Private/GeometryCollection/GeometryCollectionEngineConversion.cpp b/Engine/Source/Runtime/Experimental/GeometryCollectionEngine/Private/GeometryCollection/GeometryCollectionEngineConversion.cpp
--- a/Engine/Source/Runtime/Experimental/GeometryCollectionEngine/Private/GeometryCollection/GeometryCollectionEngineConversion.cpp
+++ b/Engine/Source/Runtime/Experimental/GeometryCollectionEngine/Private/GeometryCollection/GeometryCollectionEngineConversion.cpp
@@ -10,6 +10,7 @@
 #include "Engine/Selection.h"
 #include "Engine/StaticMesh.h"
 #include "GameFramework/Actor.h"
+#include "GeometryCollection/Facades/CollectionTransformSourceFacade.h"
 #include "GeometryCollection/GeometryCollection.h"
 #include "GeometryCollection/GeometryCollectionActor.h"
 #include "GeometryCollection/GeometryCollectionAlgo.h"
@@ -1108,6 +1109,52 @@ bool FGeometryCollectionEngineConversion::AppendSkeletalMesh(const USkeletalMesh
 	return true;
 }
 
+void FGeometryCollectionEngineConversion::AppendSkeleton(const USkeleton* InSkeleton, const FTransform& SkeletalMeshTransform, FManagedArrayCollection* InCollection)
+{
+	//UE_LOG(UGeometryCollectionConversionLogging, Log, TEXT("FGeometryCollectionEngineConversion::AppendSkeletalMesh()"));
+	if (!InCollection || !InSkeleton)
+	{
+		return;
+	}
+	FGeometryCollection::DefineTransformSchema(*InCollection);
+	GeometryCollection::Facades::FTransformSource TransformSourceFacade(InCollection);
+
+	TManagedArray<FTransform>& Transform = InCollection->ModifyAttribute<FTransform>(FTransformCollection::TransformAttribute, FTransformCollection::TransformGroup);
+	TManagedArray<FLinearColor>& BoneColor = InCollection->ModifyAttribute<FLinearColor>("BoneColor", FTransformCollection::TransformGroup);
+	TManagedArray<FString>& BoneName = InCollection->ModifyAttribute<FString>("BoneName", FTransformCollection::TransformGroup);
+	TManagedArray<int32>& Parent = InCollection->ModifyAttribute<int32>(FTransformCollection::ParentAttribute, FTransformCollection::TransformGroup);
+	TManagedArray< TSet<int32> >& Child = InCollection->ModifyAttribute< TSet<int32> >(FTransformCollection::ChildrenAttribute, FTransformCollection::TransformGroup);
+
+	const FReferenceSkeleton& Skeleton = InSkeleton->GetReferenceSkeleton();
+	int32 NumBones = Skeleton.GetNum();
+	if (NumBones)
+	{
+		const TArray<FTransform>& RestTransform = Skeleton.GetRefBonePose();
+		const TArray<FMeshBoneInfo>& BoneInfo = Skeleton.GetRefBoneInfo();
+
+		TSet<int32> Roots;
+		int32 TransformBaseIndex = InCollection->AddElements(NumBones, FGeometryCollection::TransformGroup);
+		for (int i = 0, Idx = TransformBaseIndex; i < NumBones; i++, Idx++)
+		{
+			Transform[Idx] = RestTransform[i];
+			BoneColor[Idx] = FLinearColor(FColor(FMath::Rand() % 100 + 5, FMath::Rand() % 100 + 5, FMath::Rand() % 100 + 5, 255));
+			BoneName[Idx] = BoneInfo[i].Name.ToString();
+			Parent[Idx] = BoneInfo[i].ParentIndex;
+			if (Parent[Idx] != INDEX_NONE)
+			{
+				Child[Parent[Idx]].Add(Idx);
+			}
+			else
+			{
+				Roots.Add(Idx);
+			}
+		}
+
+		ensure(Roots.Num());
+		TransformSourceFacade.AddTransformSource(InCollection, InSkeleton->GetName(), InSkeleton->GetGuid(), Roots);
+	}
+}
+
 const FSkeletalMeshLODRenderData* FGeometryCollectionEngineConversion::GetSkeletalMeshLOD(const USkeletalMesh* SkeletalMesh, int32 LOD)
 {
 	if (const USkeleton* Skeleton = SkeletalMesh->GetSkeleton())
diff --git a/Engine/Source/Runtime/Experimental/GeometryCollectionEngine/Public/GeometryCollection/GeometryCollectionEngineConversion.h b/Engine/Source/Runtime/Experimental/GeometryCollectionEngine/Public/GeometryCollection/GeometryCollectionEngineConversion.h
--- a/Engine/Source/Runtime/Experimental/GeometryCollectionEngine/Public/GeometryCollection/GeometryCollectionEngineConversion.h
+++ b/Engine/Source/Runtime/Experimental/GeometryCollectionEngine/Public/GeometryCollection/GeometryCollectionEngineConversion.h
@@ -5,6 +5,7 @@
 
 class UStaticMesh;
 class USkeletalMesh;
+class USkeleton;
 class USkeletalMeshComponent;
 class UGeometryCollection;
 class UMaterialInterface;
@@ -106,6 +107,14 @@ public:
 	*/
 	static bool AppendSkeletalMesh(const USkeletalMesh* SkeletalMesh, int32 MaterialStartIndex, const FTransform& SkeletalMeshTransform, FManagedArrayCollection* InCollection, bool bReindexMaterials = true);
 
+	/**
+	*  Appends a skeleton mesh to a GeometryCollection.
+	*  @param USkeleton : Const mesh to read vertex/normals/index data from
+	*  @param SkeletalMeshTransform : Mesh transform.
+	*  @param GeometryCollection    : Collection to append the mesh into.
+	*/
+	static void AppendSkeleton(const USkeleton* Skeleton, const FTransform& SkeletalMeshTransform, FManagedArrayCollection* InCollection);
+
 	/**
 	*  Appends a skeletal mesh to a GeometryCollectionComponent.
 	*  @param SkeletalMeshComponent : Const mesh to read vertex/normals/index data from
