GeometryCollection - Fixed cloth construct assert

diff --git a/Engine/Source/Runtime/Experimental/Chaos/Private/GeometryCollection/ManagedArrayCollection.cpp b/Engine/Source/Runtime/Experimental/Chaos/Private/GeometryCollection/ManagedArrayCollection.cpp
--- a/Engine/Source/Runtime/Experimental/Chaos/Private/GeometryCollection/ManagedArrayCollection.cpp
+++ b/Engine/Source/Runtime/Experimental/Chaos/Private/GeometryCollection/ManagedArrayCollection.cpp
@@ -499,8 +499,11 @@ bool FManagedArrayCollection::IsConnected(FName StartingNode, FName TargetNode)
 				if (!Visited.Contains(Curr))
 				{
 					Visited.Add(Curr);
-					if (!DMap[Curr].IsEmpty())
-						SearchSet.Append(DMap[Curr]);
+					if (DMap.Contains(Curr))
+					{
+						if (!DMap[Curr].IsEmpty())
+							SearchSet.Append(DMap[Curr]);
+					}
 				}
 			}
 		}
