ChaosFlesh : Frame updates
- Split out DeformablePhysicsComponent

diff --git a/Engine/Source/Runtime/Experimental/Chaos/Private/GeometryCollection/Facades/CollectionTransformSourceFacade.cpp b/Engine/Source/Runtime/Experimental/Chaos/Private/GeometryCollection/Facades/CollectionTransformSourceFacade.cpp
--- a/Engine/Source/Runtime/Experimental/Chaos/Private/GeometryCollection/Facades/CollectionTransformSourceFacade.cpp
+++ b/Engine/Source/Runtime/Experimental/Chaos/Private/GeometryCollection/Facades/CollectionTransformSourceFacade.cpp
@@ -42,7 +42,7 @@ namespace GeometryCollection::Facades
 		Collection->AddAttribute<TSet<int32>>(SourceRootsAttribute, TransformSourceGroup, TransformDependency);
 	}
 
-	bool FTransformSource::HasFacade(FManagedArrayCollection* Collection)
+	bool FTransformSource::HasFacade(const FManagedArrayCollection* Collection)
 	{
 		return Collection->HasGroup(TransformSourceGroup) &&
 			Collection->FindAttribute<FString>(SourceNameAttribute, TransformSourceGroup) &&
@@ -68,7 +68,7 @@ namespace GeometryCollection::Facades
 	//
 
 
-	TSet<int32> FTransformSource::GetTransformSource(FManagedArrayCollection* Collection, const FString& InName, const FGuid& InGuid)
+	TSet<int32> FTransformSource::GetTransformSource(const FManagedArrayCollection* Collection, const FString& InName, const FGuid& InGuid)
 	{
 		if (Collection->HasGroup(TransformSourceGroup))
 		{
diff --git a/Engine/Source/Runtime/Experimental/Chaos/Public/GeometryCollection/Facades/CollectionTransformSourceFacade.h b/Engine/Source/Runtime/Experimental/Chaos/Public/GeometryCollection/Facades/CollectionTransformSourceFacade.h
--- a/Engine/Source/Runtime/Experimental/Chaos/Public/GeometryCollection/Facades/CollectionTransformSourceFacade.h
+++ b/Engine/Source/Runtime/Experimental/Chaos/Public/GeometryCollection/Facades/CollectionTransformSourceFacade.h
@@ -56,7 +56,7 @@ namespace GeometryCollection::Facades
 		/**
 		*  Is the Facade defined on the collection?
 		*/
-		static bool HasFacade(FManagedArrayCollection* Collection);
+		static bool HasFacade(const FManagedArrayCollection* Collection);
 
 		/**
 		* Add a transform root mapping.
@@ -71,7 +71,7 @@ namespace GeometryCollection::Facades
 		* @param Name : Name of the owner of the transform set.  
 		* @param Guid : Guid of the owner of the transform set.
 		*/
-		static TSet<int32> GetTransformSource(FManagedArrayCollection* Collection, const FString& Name, const FGuid& Guid);
+		static TSet<int32> GetTransformSource(const FManagedArrayCollection* Collection, const FString& Name, const FGuid& Guid);
 
 
 	};
