Improve Capsule vs Triangle overlap performance by 50% for our test cases

diff --git a/Engine/Source/Runtime/Experimental/Chaos/Private/Chaos/TriangleMeshImplicitObject.cpp b/Engine/Source/Runtime/Experimental/Chaos/Private/Chaos/TriangleMeshImplicitObject.cpp
--- a/Engine/Source/Runtime/Experimental/Chaos/Private/Chaos/TriangleMeshImplicitObject.cpp
+++ b/Engine/Source/Runtime/Experimental/Chaos/Private/Chaos/TriangleMeshImplicitObject.cpp
@@ -381,6 +381,42 @@ struct FTriangleMeshOverlapVisitorNoMTD
 		const UE::Math::TVector<FReal>& TranslationDouble = WorldScaleQueryTM.GetTranslation();
 		Translation = MakeVectorRegisterFloatFromDouble(MakeVectorRegister(TranslationDouble.X, TranslationDouble.Y, TranslationDouble.Z, 0.0));
 	}
+
+	// Default shape type specific culling
+	// Type specialization to follow
+	template <typename ShapeType>
+	bool ShapeTypeAdditionalCulling(const VectorRegister4Float& A, const VectorRegister4Float& B, const VectorRegister4Float& C)
+	{
+		return false;
+	}
+
+	template <>
+	bool ShapeTypeAdditionalCulling<Chaos::FCapsule>(const VectorRegister4Float& A, const VectorRegister4Float& B, const VectorRegister4Float& C)
+	{
+		const VectorRegister4Float InvRotation = VectorQuaternionInverse(Rotation);
+		const VectorRegister4Float NegTranslation = VectorNegate(Translation);
+		const VectorRegister4Float ATxSimd = VectorAdd(VectorQuaternionRotateVector(InvRotation, A), NegTranslation);
+		const VectorRegister4Float BTxSimd = VectorAdd(VectorQuaternionRotateVector(InvRotation, B), NegTranslation);
+		const VectorRegister4Float CTxSimd = VectorAdd(VectorQuaternionRotateVector(InvRotation, C), NegTranslation);
+
+		const VectorRegister4Float MinBounds = VectorMin(VectorMin(ATxSimd, BTxSimd), CTxSimd);
+		const VectorRegister4Float MaxBounds = VectorMax(VectorMax(ATxSimd, BTxSimd), CTxSimd);
+		FAABBVectorized GeometrySpaceAABB(MinBounds, MaxBounds);
+		FAABB3 GeometryAABB = QueryGeom.BoundingBox();
+		FAABBVectorized VecGeomAABB(GeometryAABB);
+
+		if (!VecGeomAABB.Intersects(GeometrySpaceAABB))
+		{
+			return true;
+		}
+		return false;
+	}
+
+	template <>
+	bool ShapeTypeAdditionalCulling<Chaos::TImplicitObjectScaled<Chaos::FCapsule, 1>>(const VectorRegister4Float& A, const VectorRegister4Float& B, const VectorRegister4Float& C)
+	{
+		return ShapeTypeAdditionalCulling<Chaos::FCapsule>(A,B,C);
+	}
 	
 	bool VisitOverlap(int32 TriIdx)
 	{
@@ -398,6 +434,11 @@ struct FTriangleMeshOverlapVisitorNoMTD
 		const VectorRegister4Float BSimd = MakeVectorRegisterFloatFromDouble(MakeVectorRegister(B.X, B.Y, B.Z, 0.0));
 		const VectorRegister4Float CSimd = MakeVectorRegisterFloatFromDouble(MakeVectorRegister(C.X, C.Y, C.Z, 0.0));
 
+		if (ShapeTypeAdditionalCulling<QueryGeomType>(ASimd, BSimd, CSimd))
+		{
+			return true;
+		}
+
 		const VectorRegister4Float AB = VectorSubtract(BSimd, ASimd);
 		const VectorRegister4Float AC = VectorSubtract(CSimd, ASimd);
 
