Constraints: Support spawnables. Transformable handles now hold onto the sequencer binding if they are in sequencer, which then get resolved during evaluation similar to how attach tracks handle this. When an actor with a bound constraint is deleted we don't remove it, but it is not considered fully active so it ticks but doesn't evaluate, this simplifed things greatly.

diff --git a/Engine/Plugins/Animation/ControlRig/Source/ControlRig/Private/Constraints/ControlRigTransformableHandle.cpp b/Engine/Plugins/Animation/ControlRig/Source/ControlRig/Private/Constraints/ControlRigTransformableHandle.cpp
--- a/Engine/Plugins/Animation/ControlRig/Source/ControlRig/Private/Constraints/ControlRigTransformableHandle.cpp
+++ b/Engine/Plugins/Animation/ControlRig/Source/ControlRig/Private/Constraints/ControlRigTransformableHandle.cpp
@@ -514,8 +514,10 @@ bool UTransformableControlHandle::AddTransformKeys(const TArray<FFrameNumber>& I
 	return true;
 }
 
-
-
+void UTransformableControlHandle::ResolveBoundObjects(FMovieSceneSequenceID LocalSequenceID, IMovieScenePlayer& Player)
+{
+	//don't do anything for ControlRigs since they don't get removed when it's actor is destroyed or respawned since the track owns it. In the template
+}
 
 #if WITH_EDITOR
 
diff --git a/Engine/Plugins/Animation/ControlRig/Source/ControlRig/Private/Sequencer/MovieSceneControlRigParameterSection.cpp b/Engine/Plugins/Animation/ControlRig/Source/ControlRig/Private/Sequencer/MovieSceneControlRigParameterSection.cpp
--- a/Engine/Plugins/Animation/ControlRig/Source/ControlRig/Private/Sequencer/MovieSceneControlRigParameterSection.cpp
+++ b/Engine/Plugins/Animation/ControlRig/Source/ControlRig/Private/Sequencer/MovieSceneControlRigParameterSection.cpp
@@ -19,6 +19,8 @@
 #include "Components/SkeletalMeshComponent.h"
 #include "Constraints/ControlRigTransformableHandle.h"
 #include "UObject/UE5MainStreamObjectVersion.h"
+#include "TransformConstraint.h"
+#include "TransformableHandle.h"
 
 #if WITH_EDITOR
 #include "AnimPose.h"
@@ -720,6 +722,42 @@ UMovieSceneControlRigParameterSection::UMovieSceneControlRigParameterSection() :
 #endif
 }
 
+void UMovieSceneControlRigParameterSection::OnBindingIDsUpdated(const TMap<UE::MovieScene::FFixedObjectBindingID, UE::MovieScene::FFixedObjectBindingID>& OldFixedToNewFixedMap, FMovieSceneSequenceID LocalSequenceID, const FMovieSceneSequenceHierarchy* Hierarchy, IMovieScenePlayer& Player)
+{
+	for (FConstraintAndActiveChannel& ConstraintChannel : ConstraintsChannels)
+	{
+		if (UTickableTransformConstraint* TransformConstraint = Cast< UTickableTransformConstraint>(ConstraintChannel.Constraint.Get()))
+		{
+			if (TransformConstraint->ChildTRSHandle)
+			{
+				TransformConstraint->ChildTRSHandle->OnBindingIDsUpdated(OldFixedToNewFixedMap, LocalSequenceID, Hierarchy, Player);
+			}
+			if (TransformConstraint->ParentTRSHandle)
+			{
+				TransformConstraint->ParentTRSHandle->OnBindingIDsUpdated(OldFixedToNewFixedMap, LocalSequenceID, Hierarchy, Player);
+			}
+		}
+	}
+}
+
+void UMovieSceneControlRigParameterSection::GetReferencedBindings(TArray<FGuid>& OutBindings)
+{
+	for (FConstraintAndActiveChannel& ConstraintChannel : ConstraintsChannels)
+	{
+		if (UTickableTransformConstraint* TransformConstraint = Cast< UTickableTransformConstraint>(ConstraintChannel.Constraint.Get()))
+		{
+			if (TransformConstraint->ChildTRSHandle && TransformConstraint->ChildTRSHandle->ConstraintBindingID.IsValid())
+			{
+				OutBindings.Add(TransformConstraint->ChildTRSHandle->ConstraintBindingID.GetGuid());
+			}
+			if (TransformConstraint->ParentTRSHandle && TransformConstraint->ParentTRSHandle->ConstraintBindingID.IsValid())
+			{
+				OutBindings.Add(TransformConstraint->ParentTRSHandle->ConstraintBindingID.GetGuid());
+			}
+		}
+	}
+}
+
 bool UMovieSceneControlRigParameterSection::RenameParameterName(const FName& OldParameterName, const FName& NewParameterName)
 {
 	bool bWasReplaced = false;
diff --git a/Engine/Plugins/Animation/ControlRig/Source/ControlRig/Public/Constraints/ControlRigTransformableHandle.h b/Engine/Plugins/Animation/ControlRig/Source/ControlRig/Public/Constraints/ControlRigTransformableHandle.h
--- a/Engine/Plugins/Animation/ControlRig/Source/ControlRig/Public/Constraints/ControlRigTransformableHandle.h
+++ b/Engine/Plugins/Animation/ControlRig/Source/ControlRig/Public/Constraints/ControlRigTransformableHandle.h
@@ -65,6 +65,9 @@ public:
 		UMovieSceneSection* InSection,
 		const bool bLocal = true) const override;
 
+	/** Resolve the bound objects so that any object it references are resolved and correctly set up*/
+	virtual void ResolveBoundObjects(FMovieSceneSequenceID LocalSequenceID, IMovieScenePlayer& Player) override;
+
 	/** Returns the skeletal mesh bound to ControlRig. */
 	USkeletalMeshComponent* GetSkeletalMesh() const;
 
diff --git a/Engine/Plugins/Animation/ControlRig/Source/ControlRig/Public/Sequencer/MovieSceneControlRigParameterSection.h b/Engine/Plugins/Animation/ControlRig/Source/ControlRig/Public/Sequencer/MovieSceneControlRigParameterSection.h
--- a/Engine/Plugins/Animation/ControlRig/Source/ControlRig/Public/Sequencer/MovieSceneControlRigParameterSection.h
+++ b/Engine/Plugins/Animation/ControlRig/Source/ControlRig/Public/Sequencer/MovieSceneControlRigParameterSection.h
@@ -489,6 +489,8 @@ protected:
 	virtual void PostEditImport() override;
 	virtual void PostLoad() override;
 	virtual float GetTotalWeightValue(FFrameTime InTime) const override;
+	virtual void OnBindingIDsUpdated(const TMap<UE::MovieScene::FFixedObjectBindingID, UE::MovieScene::FFixedObjectBindingID>& OldFixedToNewFixedMap, FMovieSceneSequenceID LocalSequenceID, const FMovieSceneSequenceHierarchy* Hierarchy, IMovieScenePlayer& Player) override;
+	virtual void GetReferencedBindings(TArray<FGuid>& OutBindings) override;
 
 
 	// When true we do not set a key on the section, since it will be set because we changed the value
diff --git a/Engine/Plugins/Animation/ControlRig/Source/ControlRigEditor/Private/EditMode/ControlRigEditMode.cpp b/Engine/Plugins/Animation/ControlRig/Source/ControlRigEditor/Private/EditMode/ControlRigEditMode.cpp
--- a/Engine/Plugins/Animation/ControlRig/Source/ControlRigEditor/Private/EditMode/ControlRigEditMode.cpp
+++ b/Engine/Plugins/Animation/ControlRig/Source/ControlRigEditor/Private/EditMode/ControlRigEditMode.cpp
@@ -3391,7 +3391,14 @@ bool FControlRigEditMode::AreRigElementSelectedAndMovable(UControlRig* ControlRi
 	if (!IsInLevelEditor())
 	{
 		if (!IsTransformDelegateAvailable())
+		{
+			return false;
+		}
+	}
+	else //do check for the binding though
 	{
+		if (GetHostingSceneComponent(ControlRig) == nullptr)
+		{
 			return false;
 		}
 	}
diff --git a/Engine/Plugins/Animation/ControlRig/Source/ControlRigEditor/Private/Editor/Constraints/SConstraintsWidget.cpp b/Engine/Plugins/Animation/ControlRig/Source/ControlRigEditor/Private/Editor/Constraints/SConstraintsWidget.cpp
--- a/Engine/Plugins/Animation/ControlRig/Source/ControlRigEditor/Private/Editor/Constraints/SConstraintsWidget.cpp
+++ b/Engine/Plugins/Animation/ControlRig/Source/ControlRigEditor/Private/Editor/Constraints/SConstraintsWidget.cpp
@@ -458,7 +458,7 @@ void SEditableConstraintItem::Construct(
 				{
 					return FStyleColors::Transparent;
 				}
-				return Constraint->Active ? FStyleColors::Select : FStyleColors::Transparent;
+				return Constraint->IsFullyActive() ? FStyleColors::Select : FStyleColors::Transparent;
 			})
 			.Content()
 			[
diff --git a/Engine/Plugins/Animation/ControlRig/Source/ControlRigEditor/Private/Sequencer/ConstraintChannelHelper.cpp b/Engine/Plugins/Animation/ControlRig/Source/ControlRigEditor/Private/Sequencer/ConstraintChannelHelper.cpp
--- a/Engine/Plugins/Animation/ControlRig/Source/ControlRigEditor/Private/Sequencer/ConstraintChannelHelper.cpp
+++ b/Engine/Plugins/Animation/ControlRig/Source/ControlRigEditor/Private/Sequencer/ConstraintChannelHelper.cpp
@@ -19,6 +19,7 @@
 #include "Channels/MovieSceneFloatChannel.h"
 #include "Channels/MovieSceneDoubleChannel.h"
 #include "MovieSceneConstraintChannelHelper.h"
+#include "MovieSceneSpawnableAnnotation.h"
 
 
 #define LOCTEXT_NAMESPACE "Constraints"
@@ -82,16 +83,47 @@ void FConstraintChannelHelper::SmartConstraintKey(UTickableTransformConstraint*
 
 	if (const UTransformableComponentHandle* ComponentHandle = Cast<UTransformableComponentHandle>(InConstraint->ChildTRSHandle))
 	{
-		return SmartComponentConstraintKey(InConstraint, WeakSequencer.Pin());
+		SmartComponentConstraintKey(InConstraint, WeakSequencer.Pin());
 	}
 	
-	if (const UTransformableControlHandle* ControlHandle = Cast<UTransformableControlHandle>(InConstraint->ChildTRSHandle))
+	else if (const UTransformableControlHandle* ControlHandle = Cast<UTransformableControlHandle>(InConstraint->ChildTRSHandle))
 	{
-		return SmartControlConstraintKey(InConstraint, WeakSequencer.Pin());
+		SmartControlConstraintKey(InConstraint, WeakSequencer.Pin());
 	}
+
+	FConstraintChannelHelper::CreateBindingIDForHandle(InConstraint->ChildTRSHandle);
+	FConstraintChannelHelper::CreateBindingIDForHandle(InConstraint->ParentTRSHandle);
 }
 
 
+void FConstraintChannelHelper::CreateBindingIDForHandle(UTransformableHandle* InHandle)
+{
+	const TWeakPtr<ISequencer> WeakSequencer = FBakingHelper::GetSequencer();
+	if (InHandle == nullptr || WeakSequencer.IsValid() == false)
+	{
+		return;
+	}
+	if (USceneComponent* SceneComponent = Cast<USceneComponent>(InHandle->GetTarget().Get()))
+	{
+		AActor* Actor = SceneComponent->GetTypedOuter<AActor>();
+		if (Actor)
+		{
+			TOptional<FMovieSceneSpawnableAnnotation> Spawnable = FMovieSceneSpawnableAnnotation::Find(Actor);
+			if (Spawnable.IsSet())
+			{
+				// Check whether the spawnable is underneath the current sequence, if so, we can remap it to a local sequence ID
+				InHandle->ConstraintBindingID = UE::MovieScene::FRelativeObjectBindingID(WeakSequencer.Pin()->GetFocusedTemplateID(), Spawnable->SequenceID, Spawnable->ObjectBindingID,
+					*(WeakSequencer.Pin().Get()));
+			}
+			else
+			{
+				FGuid Guid = WeakSequencer.Pin()->GetHandleToObject(Actor, false); //don't create it???
+				InHandle->ConstraintBindingID = UE::MovieScene::FRelativeObjectBindingID(Guid);
+			}
+		}
+	}
+}
+
 UMovieSceneControlRigParameterSection* FConstraintChannelHelper::GetControlSection(
 	const UTransformableControlHandle* InHandle,
 	const TSharedPtr<ISequencer>& InSequencer)
@@ -288,7 +320,7 @@ void FConstraintChannelHelper::SmartComponentConstraintKey(
 	
 				Section->Modify();
 
-				//new for compensatio
+				//new for compensation
 
 				TGuardValue<bool> CompensateGuard(FMovieSceneConstraintChannelHelper::bDoNotCompensate, true);
 
diff --git a/Engine/Plugins/Animation/ControlRig/Source/ControlRigEditor/Public/ConstraintChannelHelper.h b/Engine/Plugins/Animation/ControlRig/Source/ControlRigEditor/Public/ConstraintChannelHelper.h
--- a/Engine/Plugins/Animation/ControlRig/Source/ControlRigEditor/Public/ConstraintChannelHelper.h
+++ b/Engine/Plugins/Animation/ControlRig/Source/ControlRigEditor/Public/ConstraintChannelHelper.h
@@ -75,5 +75,8 @@ private:
 
 	/** END COMPONENT SECTION */
 
+	/** For the given handle create any movie scene binding for it based upon the current sequencer that's open*/
+	static void CreateBindingIDForHandle(UTransformableHandle* InHandle);
+
 	
 };
diff --git a/Engine/Source/Runtime/Experimental/Animation/Constraints/Private/ConstraintsManager.cpp b/Engine/Source/Runtime/Experimental/Animation/Constraints/Private/ConstraintsManager.cpp
--- a/Engine/Source/Runtime/Experimental/Animation/Constraints/Private/ConstraintsManager.cpp
+++ b/Engine/Source/Runtime/Experimental/Animation/Constraints/Private/ConstraintsManager.cpp
@@ -1,7 +1,6 @@
 ﻿// Copyright Epic Games, Inc. All Rights Reserved.
 
 #include "ConstraintsManager.h"
-
 #include "ConstraintsActor.h"
 
 #include "Algo/Copy.h"
@@ -71,6 +70,11 @@ void UTickableConstraint::SetActive(const bool bIsActive)
 	ConstraintTick.SetTickFunctionEnable(bIsActive);
 }
 
+bool UTickableConstraint::IsFullyActive() const
+{
+	return Active;
+}
+
 void UTickableConstraint::Evaluate(bool bTickHandlesAlso) const
 {
 	ConstraintTick.EvaluateFunctions();
@@ -151,7 +155,8 @@ void UConstraintsManager::OnActorDestroyed(AActor* InActor)
 		TArray< int32 > IndicesToRemove;
 		for (int32 Index = 0; Index < Constraints.Num(); ++Index)
 		{
-			if (IsValid(Constraints[Index]) && Constraints[Index]->ReferencesObject(SceneComponent))
+			//if the constraint has a bound object(in sequencer) we don't remove the constraint, it could be a spawnable
+			if (IsValid(Constraints[Index]) &&  Constraints[Index]->HasBoundObjects() == false && Constraints[Index]->ReferencesObject(SceneComponent))
 			{
 				IndicesToRemove.Add(Index);
 			}
diff --git a/Engine/Source/Runtime/Experimental/Animation/Constraints/Private/TransformConstraint.cpp b/Engine/Source/Runtime/Experimental/Animation/Constraints/Private/TransformConstraint.cpp
--- a/Engine/Source/Runtime/Experimental/Animation/Constraints/Private/TransformConstraint.cpp
+++ b/Engine/Source/Runtime/Experimental/Animation/Constraints/Private/TransformConstraint.cpp
@@ -226,6 +226,31 @@ bool UTickableTransformConstraint::ReferencesObject(TWeakObjectPtr<UObject> InOb
 	return false;
 }
 
+bool UTickableTransformConstraint::HasBoundObjects() const
+{
+	if (ChildTRSHandle && ChildTRSHandle->HasBoundObjects())
+	{
+		return true;
+	}
+	if (ParentTRSHandle && ParentTRSHandle->HasBoundObjects())
+	{
+		return true;
+	}
+	return false;
+}
+
+void UTickableTransformConstraint::ResolveBoundObjects(FMovieSceneSequenceID LocalSequenceID, IMovieScenePlayer& Player)
+{
+	if (ChildTRSHandle && ChildTRSHandle->HasBoundObjects())
+	{
+		ChildTRSHandle->ResolveBoundObjects(LocalSequenceID, Player);
+	}
+	if (ParentTRSHandle && ParentTRSHandle->HasBoundObjects())
+	{
+		ParentTRSHandle->ResolveBoundObjects(LocalSequenceID, Player);
+	}
+}
+
 void UTickableTransformConstraint::SetChildGlobalTransform(const FTransform& InGlobal) const
 {
 	if(ChildTRSHandle && ChildTRSHandle->IsValid())
@@ -279,6 +304,12 @@ void UTickableTransformConstraint::OnHandleModified(UTransformableHandle* InHand
 	}
 }
 
+bool UTickableTransformConstraint::IsFullyActive() const
+{
+	return (Active && IsValid(ChildTRSHandle) && ChildTRSHandle->IsValid() 
+		&& IsValid(ParentTRSHandle) && ParentTRSHandle->IsValid());
+}
+
 /** 
  * UTickableTranslationConstraint
  **/
@@ -329,7 +360,7 @@ FConstraintTickFunction::ConstraintFunction UTickableTranslationConstraint::GetF
 {
 	return [this]()
 	{
-		if (!Active)
+		if (!IsFullyActive())
 		{
 			return;
 		}
@@ -357,7 +388,7 @@ void UTickableTranslationConstraint::OnHandleModified(UTransformableHandle* InHa
 {
 	Super::OnHandleModified(InHandle, InEvent);
 	
-	if (!Active || !bDynamicOffset)
+	if (!IsFullyActive() || !bDynamicOffset)
 	{
 		return;
 	}
@@ -461,7 +492,7 @@ FConstraintTickFunction::ConstraintFunction UTickableRotationConstraint::GetFunc
 {
 	return [this]()
 	{
-		if (!Active)
+		if (!IsFullyActive())
 		{
 			return;
 		}
@@ -490,7 +521,7 @@ void UTickableRotationConstraint::OnHandleModified(UTransformableHandle* InHandl
 {
 	Super::OnHandleModified(InHandle, InEvent);
 	
-	if (!Active || !bDynamicOffset)
+	if (!IsFullyActive() || !bDynamicOffset)
 	{
 		return;
 	}
@@ -598,7 +629,7 @@ FConstraintTickFunction::ConstraintFunction UTickableScaleConstraint::GetFunctio
 {
 	return [this]()
 	{
-		if (!Active)
+		if (!IsFullyActive())
 		{
 			return;
 		}
@@ -626,7 +657,7 @@ void UTickableScaleConstraint::OnHandleModified(UTransformableHandle* InHandle,
 {
 	Super::OnHandleModified(InHandle, InEvent);
 	
-	if (!Active || !bDynamicOffset)
+	if (!IsFullyActive() || !bDynamicOffset)
 	{
 		return;
 	}
@@ -723,7 +754,7 @@ FConstraintTickFunction::ConstraintFunction UTickableParentConstraint::GetFuncti
 {
 	return [this]()
 	{
-		if (!Active)
+		if (!IsFullyActive())
 		{
 			return;
 		}
@@ -768,7 +799,7 @@ void UTickableParentConstraint::OnHandleModified(UTransformableHandle* InHandle,
 {
 	Super::OnHandleModified(InHandle, InEvent);
 	
-	if (!Active || !bDynamicOffset)
+	if (!IsFullyActive() || !bDynamicOffset)
 	{
 		return;
 	}
@@ -852,7 +883,7 @@ FConstraintTickFunction::ConstraintFunction UTickableLookAtConstraint::GetFuncti
 	// @todo handle weight here
 	return [this]()
 	{
-		if (!Active)
+		if (!IsFullyActive())
 		{
 			return;
 		}
@@ -1334,7 +1365,7 @@ TOptional<FTransform> FTransformConstraintUtils::GetConstraintRelativeTransform(
     {
     	if (const UTickableTransformConstraint* TransformConstraint = Cast<UTickableTransformConstraint>(InConstraint.Get()))
     	{
-    		return InConstraint->Active && TransformConstraint->bDynamicOffset;
+    		return InConstraint->IsFullyActive() && TransformConstraint->bDynamicOffset;
     	}
     	return false;
     });
diff --git a/Engine/Source/Runtime/Experimental/Animation/Constraints/Private/TransformableHandle.cpp b/Engine/Source/Runtime/Experimental/Animation/Constraints/Private/TransformableHandle.cpp
--- a/Engine/Source/Runtime/Experimental/Animation/Constraints/Private/TransformableHandle.cpp
+++ b/Engine/Source/Runtime/Experimental/Animation/Constraints/Private/TransformableHandle.cpp
@@ -32,6 +32,21 @@ UTransformableHandle::FHandleModifiedEvent& UTransformableHandle::HandleModified
 	return OnHandleModified;
 }
 
+bool UTransformableHandle::HasBoundObjects() const
+{
+	return ConstraintBindingID.IsValid();
+}
+
+void UTransformableHandle::OnBindingIDsUpdated(const TMap<UE::MovieScene::FFixedObjectBindingID, UE::MovieScene::FFixedObjectBindingID>& OldFixedToNewFixedMap, FMovieSceneSequenceID LocalSequenceID, const FMovieSceneSequenceHierarchy* Hierarchy, IMovieScenePlayer& Player)
+{
+	if (ConstraintBindingID.IsValid())
+	{
+		UE::MovieScene::FFixedObjectBindingID FixedBindingID = ConstraintBindingID.ResolveToFixed(LocalSequenceID, Player);
+		Modify();
+		ConstraintBindingID = OldFixedToNewFixedMap[FixedBindingID].ConvertToRelative(LocalSequenceID, Hierarchy);
+	}
+}
+
 /**
  * UTransformableComponentHandle
  */
@@ -472,6 +487,21 @@ bool UTransformableComponentHandle::AddTransformKeys(const TArray<FFrameNumber>&
 	return true;
 }
 
+void UTransformableComponentHandle::ResolveBoundObjects(FMovieSceneSequenceID LocalSequenceID, IMovieScenePlayer& Player)
+{
+	for (TWeakObjectPtr<> ParentObject : ConstraintBindingID.ResolveBoundObjects(LocalSequenceID, Player))
+	{
+		if (AActor* Actor = Cast<AActor>(ParentObject.Get()))
+		{
+			Component = Actor->GetRootComponent();
+		}
+		else if (USceneComponent* Comp = Cast<USceneComponent>(ParentObject.Get()))
+		{
+			Component = Comp;
+		}
+		break; //just do one
+	}
+}
 
 #if WITH_EDITOR
 FString UTransformableComponentHandle::GetLabel() const
diff --git a/Engine/Source/Runtime/Experimental/Animation/Constraints/Public/ConstraintsManager.h b/Engine/Source/Runtime/Experimental/Animation/Constraints/Public/ConstraintsManager.h
--- a/Engine/Source/Runtime/Experimental/Animation/Constraints/Public/ConstraintsManager.h
+++ b/Engine/Source/Runtime/Experimental/Animation/Constraints/Public/ConstraintsManager.h
@@ -4,9 +4,10 @@
 
 #include "UObject/Object.h"
 #include "Engine/EngineBaseTypes.h"
-
+#include "MovieSceneSequenceID.h"
 #include "ConstraintsManager.generated.h"
 
+class IMovieScenePlayer;
 class UTickableConstraint;
 class ULevel;
 
@@ -87,9 +88,16 @@ public:
 	/** Sets the Active value and enable/disable the tick function. */
 	void SetActive(const bool bIsActive);
 	
+	/** Get whether or not it's fully active, it's set to active and all pieces are set up,e.g. objects really exist that are being constrainted*/
+	virtual bool IsFullyActive() const;
+	/** If true it contains objects bound to an external system, like sequencer so we don't do certain things, like remove constraints when they don't resolve*/
+	virtual bool HasBoundObjects() const PURE_VIRTUAL(HasBoundObjects, return false;);
+	/** Resolve the bound objects so that any object it references are resovled and correctly set up*/
+	virtual void ResolveBoundObjects(FMovieSceneSequenceID LocalSequenceID, IMovieScenePlayer& Player)  PURE_VIRTUAL(ResolveBoundObjects);
+
 	/** @todo document */
 	virtual uint32 GetTargetHash() const PURE_VIRTUAL(GetTargetHash, return 0;);
-	/** @todo document */
+	/** Whether or not this object references this object. If it HasBoundObjects you should ResolveBoundObjects if you expect this to be up to date */
 	virtual bool ReferencesObject(TWeakObjectPtr<UObject> InObject) const PURE_VIRTUAL(ReferencesObject, return false;);
 
 #if WITH_EDITOR
diff --git a/Engine/Source/Runtime/Experimental/Animation/Constraints/Public/TransformConstraint.h b/Engine/Source/Runtime/Experimental/Animation/Constraints/Public/TransformConstraint.h
--- a/Engine/Source/Runtime/Experimental/Animation/Constraints/Public/TransformConstraint.h
+++ b/Engine/Source/Runtime/Experimental/Animation/Constraints/Public/TransformConstraint.h
@@ -41,6 +41,15 @@ public:
 	/** Test whether an InObject is referenced by that constraint. (i.e. is it's parent or child). */
 	virtual bool ReferencesObject(TWeakObjectPtr<UObject> InObject) const override;
 	
+	/** If true it contains objects bound to an external system, like sequencer so we don't do certain things, like remove constraints when they don't resolve*/
+	virtual bool HasBoundObjects() const override;
+	
+	/** Resolve the bound objects so that any object it references are resovled and correctly set up*/
+	virtual void ResolveBoundObjects(FMovieSceneSequenceID LocalSequenceID, IMovieScenePlayer& Player) override;
+
+	/** If Active and the handles and targets are valid*/
+	virtual bool IsFullyActive() const override;
+
 	/** The transformable handle representing the parent of that constraint. */
 	UPROPERTY(BlueprintReadWrite, Category = "Handle")
 	TObjectPtr<UTransformableHandle> ParentTRSHandle;
diff --git a/Engine/Source/Runtime/Experimental/Animation/Constraints/Public/TransformableHandle.h b/Engine/Source/Runtime/Experimental/Animation/Constraints/Public/TransformableHandle.h
--- a/Engine/Source/Runtime/Experimental/Animation/Constraints/Public/TransformableHandle.h
+++ b/Engine/Source/Runtime/Experimental/Animation/Constraints/Public/TransformableHandle.h
@@ -4,7 +4,7 @@
 
 #include "UObject/Object.h"
 #include "Engine/EngineBaseTypes.h"
-
+#include "MovieSceneObjectBindingID.h"
 #include "TransformableHandle.generated.h"
 
 struct FTickFunction;
@@ -15,6 +15,10 @@ class UMovieSceneSection;
 struct FFRameNumber;
 struct FFrameRate;
 enum class EMovieSceneTransformChannel : uint32;
+class IMovieScenePlayer;
+namespace UE { namespace MovieScene {struct FFixedObjectBindingID;} }
+struct FMovieSceneSequenceHierarchy;
+struct FMovieSceneSequenceID;
 
 UENUM()
 enum class EHandleEvent : uint8
@@ -56,8 +60,16 @@ public:
 	/** Gets the local transform of the underlying transformable object in it's parent space. */
 	virtual FTransform GetLocalTransform() const PURE_VIRTUAL(GetLocalTransform, return FTransform::Identity;);
 
-	/** Performan any special ticking needed for this handle, by default it does nothing
-	todo need to see if we need to tick control rig also*/
+	/** If true it contains objects bound to an external system, like sequencer so we don't do certain things, like remove constraints when they don't resolve*/
+	virtual bool HasBoundObjects() const;
+
+	/** Resolve the bound objects so that any object it references are resovled and correctly set up*/
+	virtual void ResolveBoundObjects(FMovieSceneSequenceID LocalSequenceID, IMovieScenePlayer& Player) PURE_VIRTUAL(ResolveBoundObjects);
+
+	/** Fix up Binding in case it has changed*/
+	void OnBindingIDsUpdated(const TMap<UE::MovieScene::FFixedObjectBindingID, UE::MovieScene::FFixedObjectBindingID>& OldFixedToNewFixedMap, FMovieSceneSequenceID LocalSequenceID, const FMovieSceneSequenceHierarchy* Hierarchy, IMovieScenePlayer& Player);
+	
+	/** Perform any special ticking needed for this handle, by default it does nothing, todo need to see if we need to tick control rig also*/
 	virtual void TickForBaking() {};
 	/** Get the array of float channels for the specified section*/
 	virtual TArrayView<FMovieSceneFloatChannel*>  GetFloatChannels(const UMovieSceneSection* InSection) const PURE_VIRTUAL(GetFloatChannels, return TArrayView<FMovieSceneFloatChannel*>(); );
@@ -101,6 +113,9 @@ public:
 	virtual FString GetFullLabel() const PURE_VIRTUAL(GetFullLabel, return FString(););
 #endif
 
+	//possible bindingID
+	UPROPERTY(EditAnywhere, Category = "Binding")
+	FMovieSceneObjectBindingID ConstraintBindingID;
 protected:
 	FHandleModifiedEvent OnHandleModified;
 };
@@ -162,6 +177,9 @@ public:
 		UMovieSceneSection* InSection,
 		const bool bLocal = true) const override;
 
+	/** Resolve the bound objects so that any object it references are resovled and correctly set up*/
+	virtual void ResolveBoundObjects(FMovieSceneSequenceID LocalSequenceID, IMovieScenePlayer& Player) override;
+
 #if WITH_EDITOR
 	/** Returns labels used for UI. */
 	virtual FString GetLabel() const override;
diff --git a/Engine/Source/Runtime/MovieSceneTracks/Private/Sections/MovieScene3DTransformSection.cpp b/Engine/Source/Runtime/MovieSceneTracks/Private/Sections/MovieScene3DTransformSection.cpp
--- a/Engine/Source/Runtime/MovieSceneTracks/Private/Sections/MovieScene3DTransformSection.cpp
+++ b/Engine/Source/Runtime/MovieSceneTracks/Private/Sections/MovieScene3DTransformSection.cpp
@@ -17,6 +17,7 @@
 #include "MovieSceneTracksComponentTypes.h"
 #include "Algo/AnyOf.h"
 #include "TransformConstraint.h"
+#include "TransformableHandle.h"
 
 #if WITH_EDITOR
 
@@ -347,6 +348,48 @@ UMovieScene3DTransformSection::UMovieScene3DTransformSection(const FObjectInitia
 	Scale[2].SetDefault(1.f);
 }
 
+void UMovieScene3DTransformSection::OnBindingIDsUpdated(const TMap<UE::MovieScene::FFixedObjectBindingID, UE::MovieScene::FFixedObjectBindingID>& OldFixedToNewFixedMap, FMovieSceneSequenceID LocalSequenceID, const FMovieSceneSequenceHierarchy* Hierarchy, IMovieScenePlayer& Player)
+{
+	if (Constraints)
+	{
+		for (FConstraintAndActiveChannel& ConstraintChannel : Constraints->ConstraintsChannels)
+		{
+			if (UTickableTransformConstraint* TransformConstraint = Cast< UTickableTransformConstraint>(ConstraintChannel.Constraint.Get()))
+			{
+				if (TransformConstraint->ChildTRSHandle)
+				{
+					TransformConstraint->ChildTRSHandle->OnBindingIDsUpdated(OldFixedToNewFixedMap, LocalSequenceID, Hierarchy, Player);
+				}
+				if (TransformConstraint->ParentTRSHandle)
+				{
+					TransformConstraint->ParentTRSHandle->OnBindingIDsUpdated(OldFixedToNewFixedMap, LocalSequenceID, Hierarchy, Player);
+				}
+			}
+		}
+	}
+}
+
+void UMovieScene3DTransformSection::GetReferencedBindings(TArray<FGuid>& OutBindings)
+{
+	if (Constraints)
+	{
+		for (FConstraintAndActiveChannel& ConstraintChannel : Constraints->ConstraintsChannels)
+		{
+			if (UTickableTransformConstraint* TransformConstraint = Cast< UTickableTransformConstraint>(ConstraintChannel.Constraint.Get()))
+			{
+				if (TransformConstraint->ChildTRSHandle && TransformConstraint->ChildTRSHandle->ConstraintBindingID.IsValid())
+				{
+					OutBindings.Add(TransformConstraint->ChildTRSHandle->ConstraintBindingID.GetGuid());
+				}
+				if (TransformConstraint->ParentTRSHandle && TransformConstraint->ParentTRSHandle->ConstraintBindingID.IsValid())
+				{
+					OutBindings.Add(TransformConstraint->ParentTRSHandle->ConstraintBindingID.GetGuid());
+				}
+			}
+		}
+	}
+}
+
 template<typename BaseBuilderType>
 void UMovieScene3DTransformSection::BuildEntity(BaseBuilderType& InBaseBuilder, UMovieSceneEntitySystemLinker* Linker, const FEntityImportParams& Params, FImportedEntity* OutImportedEntity)
 {
diff --git a/Engine/Source/Runtime/MovieSceneTracks/Private/Systems/MovieSceneConstraintSystem.cpp b/Engine/Source/Runtime/MovieSceneTracks/Private/Systems/MovieSceneConstraintSystem.cpp
--- a/Engine/Source/Runtime/MovieSceneTracks/Private/Systems/MovieSceneConstraintSystem.cpp
+++ b/Engine/Source/Runtime/MovieSceneTracks/Private/Systems/MovieSceneConstraintSystem.cpp
@@ -70,7 +70,6 @@ UMovieSceneConstraintSystem::UMovieSceneConstraintSystem(const FObjectInitialize
 
 	FMovieSceneTracksComponentTypes* TracksComponents = FMovieSceneTracksComponentTypes::Get();
 
-	//RelevantComponent = TracksComponents->ConstraintName;
 	RelevantComponent = TracksComponents->ConstraintChannel;
 
 	// Run constraints during instantiation or evaluation
@@ -111,15 +110,18 @@ void UMovieSceneConstraintSystem::OnRun(FSystemTaskPrerequisites& InPrerequisite
 
 		struct FEvaluateConstraintChannels
 		{
-			FEvaluateConstraintChannels(FConstraintsManagerController* InController, UMovieSceneConstraintSystem* InSystem) :Controller(InController), System(InSystem) { check(Controller); check(System) };
-			void ForEachEntity(UObject* BoundObject, const FConstraintComponentData& ConstraintChannel, FFrameTime FrameTime)
+			FEvaluateConstraintChannels(FInstanceRegistry* InInstanceRegistry, FConstraintsManagerController* InController, UMovieSceneConstraintSystem* InSystem) : InstanceRegistry(InInstanceRegistry), Controller(InController), System(InSystem) { check(InstanceRegistry); check(Controller); check(System) };
+			void ForEachEntity(UObject* BoundObject, UE::MovieScene::FInstanceHandle InstanceHandle, const FConstraintComponentData& ConstraintChannel, FFrameTime FrameTime)
 			{
+				const FSequenceInstance& TargetInstance = InstanceRegistry->GetInstance(InstanceHandle);
+
 				UTickableConstraint* Constraint = Controller->GetConstraint(ConstraintChannel.ConstraintName);
 				if (Constraint)
 				{
 					bool Result = false;
 					ConstraintChannel.Channel->Evaluate(FrameTime, Result);
 					Constraint->SetActive(Result);
+					Constraint->ResolveBoundObjects(TargetInstance.GetSequenceID(), *TargetInstance.GetPlayer());
 					if (UTickableTransformConstraint* TransformConstraint = Cast< UTickableTransformConstraint>(Constraint))
 					{
 						if (UTransformableComponentHandle* ComponentHandle = Cast<UTransformableComponentHandle>(TransformConstraint->ChildTRSHandle))
@@ -132,6 +134,7 @@ void UMovieSceneConstraintSystem::OnRun(FSystemTaskPrerequisites& InPrerequisite
 					}
 				}
 			}
+			FInstanceRegistry* InstanceRegistry;
 			FConstraintsManagerController* Controller;
 			UMovieSceneConstraintSystem* System;
 		};
@@ -140,9 +143,10 @@ void UMovieSceneConstraintSystem::OnRun(FSystemTaskPrerequisites& InPrerequisite
 		FEntityTaskBuilder()
 			.SetDesiredThread(Linker->EntityManager.GetGatherThread())
 			.Read(BuiltInComponents->BoundObject)
+			.Read(BuiltInComponents->InstanceHandle)
 			.Read(TracksComponents->ConstraintChannel)
 			.Read(BuiltInComponents->EvalTime)
-			.Dispatch_PerEntity<FEvaluateConstraintChannels>(&Linker->EntityManager, InPrerequisites, &Subsequents, &Controller, this);
+			.Dispatch_PerEntity<FEvaluateConstraintChannels>(&Linker->EntityManager, InPrerequisites, &Subsequents, Linker->GetInstanceRegistry(), &Controller, this);
 	}
 	else if (CurrentPhase == ESystemPhase::Finalization)
 	{
diff --git a/Engine/Source/Runtime/MovieSceneTracks/Public/Sections/MovieScene3DTransformSection.h b/Engine/Source/Runtime/MovieSceneTracks/Public/Sections/MovieScene3DTransformSection.h
--- a/Engine/Source/Runtime/MovieSceneTracks/Public/Sections/MovieScene3DTransformSection.h
+++ b/Engine/Source/Runtime/MovieSceneTracks/Public/Sections/MovieScene3DTransformSection.h
@@ -226,6 +226,8 @@ public:
 	
 	MOVIESCENETRACKS_API virtual bool ShowCurveForChannel(const void *Channel) const override;
 	MOVIESCENETRACKS_API virtual void SetBlendType(EMovieSceneBlendType InBlendType) override;
+	MOVIESCENETRACKS_API virtual void OnBindingIDsUpdated(const TMap<UE::MovieScene::FFixedObjectBindingID, UE::MovieScene::FFixedObjectBindingID>& OldFixedToNewFixedMap, FMovieSceneSequenceID LocalSequenceID, const FMovieSceneSequenceHierarchy* Hierarchy, IMovieScenePlayer& Player) override;
+	MOVIESCENETRACKS_API virtual void GetReferencedBindings(TArray<FGuid>& OutBindings) override;
 
 public:
 
