Constraints: Make sure channel shows up on redo after removal
#trivial

diff --git a/Engine/Source/Editor/MovieSceneTools/Private/MovieSceneConstraintChannelHelper.cpp b/Engine/Source/Editor/MovieSceneTools/Private/MovieSceneConstraintChannelHelper.cpp
--- a/Engine/Source/Editor/MovieSceneTools/Private/MovieSceneConstraintChannelHelper.cpp
+++ b/Engine/Source/Editor/MovieSceneTools/Private/MovieSceneConstraintChannelHelper.cpp
@@ -519,7 +519,7 @@ void FMovieSceneConstraintChannelHelper::HandleConstraintRemoved(
 	{
 		return;
 	}
-
+	InSection->Modify();
 	TGuardValue<bool> CompensateGuard(bDoNotCompensate, true);
 	if (const UTransformableHandle* ControlHandle =Constraint->ChildTRSHandle)
 	{
diff --git a/Engine/Source/Runtime/Experimental/Animation/Constraints/Private/ConstraintsManager.cpp b/Engine/Source/Runtime/Experimental/Animation/Constraints/Private/ConstraintsManager.cpp
--- a/Engine/Source/Runtime/Experimental/Animation/Constraints/Private/ConstraintsManager.cpp
+++ b/Engine/Source/Runtime/Experimental/Animation/Constraints/Private/ConstraintsManager.cpp
@@ -445,10 +445,11 @@ bool FConstraintsManagerController::RemoveConstraint(const int32 InConstraintInd
 	UTickableConstraint* Constraint = Manager->Constraints[InConstraintIndex];
 	
 	// notify deletion
-	Manager->Modify();
-	Manager->Constraints[InConstraintIndex]->Modify();
+
 	ConstraintRemoved.Broadcast(ConstraintName);
 	Manager->OnConstraintRemoved_BP.Broadcast(Manager, Constraint);
+	Manager->Constraints[InConstraintIndex]->Modify();
+	Manager->Modify();
 
 	Manager->Constraints[InConstraintIndex]->SetActive(false);
 	Manager->Constraints.RemoveAt(InConstraintIndex);
