Remove MUserData array from internal geometry particle

diff --git a/Engine/Source/Runtime/Experimental/Chaos/Public/Chaos/GeometryParticles.h b/Engine/Source/Runtime/Experimental/Chaos/Public/Chaos/GeometryParticles.h
--- a/Engine/Source/Runtime/Experimental/Chaos/Public/Chaos/GeometryParticles.h
+++ b/Engine/Source/Runtime/Experimental/Chaos/Public/Chaos/GeometryParticles.h
@@ -453,7 +453,6 @@ namespace Chaos
 			TArrayCollection::AddArray(&MWorldSpaceInflatedBounds);
 			TArrayCollection::AddArray(&MHasBounds);
 			TArrayCollection::AddArray(&MSpatialIdx);
-			TArrayCollection::AddArray(&MUserData);
 			TArrayCollection::AddArray(&MSyncState);
 			TArrayCollection::AddArray(&MWeakParticleHandle);
 			TArrayCollection::AddArray(&MParticleConstraints);
@@ -494,7 +493,6 @@ namespace Chaos
 			, MWorldSpaceInflatedBounds(MoveTemp(Other.MWorldSpaceInflatedBounds))
 			, MHasBounds(MoveTemp(Other.MHasBounds))
 			, MSpatialIdx(MoveTemp(Other.MSpatialIdx))
-			, MUserData(MoveTemp(Other.MUserData))
 			, MSyncState(MoveTemp(Other.MSyncState))
 			, MWeakParticleHandle(MoveTemp(Other.MWeakParticleHandle))
 			, MParticleConstraints(MoveTemp(Other.MParticleConstraints))
@@ -527,7 +525,6 @@ namespace Chaos
 			TArrayCollection::AddArray(&MWorldSpaceInflatedBounds);
 			TArrayCollection::AddArray(&MHasBounds);
 			TArrayCollection::AddArray(&MSpatialIdx);
-			TArrayCollection::AddArray(&MUserData);
 			TArrayCollection::AddArray(&MSyncState);
 			TArrayCollection::AddArray(&MWeakParticleHandle);
 			TArrayCollection::AddArray(&MParticleConstraints);
@@ -570,7 +567,6 @@ namespace Chaos
 			TArrayCollection::AddArray(&MWorldSpaceInflatedBounds);
 			TArrayCollection::AddArray(&MHasBounds);
 			TArrayCollection::AddArray(&MSpatialIdx);
-			TArrayCollection::AddArray(&MUserData);
 			TArrayCollection::AddArray(&MSyncState);
 			TArrayCollection::AddArray(&MWeakParticleHandle);
 			TArrayCollection::AddArray(&MParticleConstraints);
@@ -964,7 +960,6 @@ public:
 		TArrayCollectionArray<TAABB<T, d>> MWorldSpaceInflatedBounds;
 		TArrayCollectionArray<bool> MHasBounds;
 		TArrayCollectionArray<FSpatialAccelerationIdx> MSpatialIdx;
-		TArrayCollectionArray<void*> MUserData;
 		TArrayCollectionArray<FSyncState> MSyncState;
 		TArrayCollectionArray<FWeakParticleHandle> MWeakParticleHandle;
 		TArrayCollectionArray<FConstraintHandleArray> MParticleConstraints;
