Chaos : Add GC component blue print functions
- CrumbleActiveClusters
- RemoveAllAnchors
- Also fixed a potential crash in fracture editor

diff --git a/Engine/Plugins/Experimental/ChaosEditor/Source/FractureEditor/Private/FractureEditorModeToolkit.cpp b/Engine/Plugins/Experimental/ChaosEditor/Source/FractureEditor/Private/FractureEditorModeToolkit.cpp
--- a/Engine/Plugins/Experimental/ChaosEditor/Source/FractureEditor/Private/FractureEditorModeToolkit.cpp
+++ b/Engine/Plugins/Experimental/ChaosEditor/Source/FractureEditor/Private/FractureEditorModeToolkit.cpp
@@ -1894,15 +1894,16 @@ void FFractureEditorModeToolkit::GetStatisticsSummary(FGeometryCollectionStatist
 		{
 			if (AGeometryCollectionActor* Actor = Cast<AGeometryCollectionActor>(*Iter))
 			{
-				const UGeometryCollection* RestCollection = Actor->GetGeometryCollectionComponent()->GetRestCollection();
-
-				if(RestCollection)
+				if (UGeometryCollectionComponent* Component = Actor->GetGeometryCollectionComponent())
 				{
-					const FGeometryCollection* GeometryCollection = RestCollection->GetGeometryCollection().Get();
-
-					if(GeometryCollection != nullptr)
+					if (const UGeometryCollection* RestCollection = Component->GetRestCollection())
 					{
-						GeometryCollectionArray.Add(GeometryCollection);
+						const FGeometryCollection* GeometryCollection = RestCollection->GetGeometryCollection().Get();
+
+						if (GeometryCollection != nullptr)
+						{
+							GeometryCollectionArray.Add(GeometryCollection);
+						}
 					}
 				}
 			}
diff --git a/Engine/Source/Runtime/Experimental/Chaos/Private/Chaos/PBDRigidClustering.cpp b/Engine/Source/Runtime/Experimental/Chaos/Private/Chaos/PBDRigidClustering.cpp
--- a/Engine/Source/Runtime/Experimental/Chaos/Private/Chaos/PBDRigidClustering.cpp
+++ b/Engine/Source/Runtime/Experimental/Chaos/Private/Chaos/PBDRigidClustering.cpp
@@ -1340,6 +1340,35 @@ namespace Chaos
 		return false;
 	}
 
+	bool FRigidClustering::BreakClustersByProxy(const IPhysicsProxyBase* Proxy)
+	{
+		bool bCrumbledAnyCluster = false;
+		// max strain will allow to unconditionally release the children when strain is evaluated
+		constexpr FReal MaxStrain = TNumericLimits<FReal>::Max();
+
+		// we should probably have a way to retrieve all the active clusters per proxy instead of going through the map 
+		for (FPBDRigidClusteredParticleHandle* ClusteredHandle : GetTopLevelClusterParents())
+		{
+			if (ClusteredHandle && ClusteredHandle->PhysicsProxy() == Proxy)
+			{
+				const TArray<FPBDRigidParticleHandle*>& Children = MChildren[ClusteredHandle];
+				for (FPBDRigidParticleHandle* ChildHandle : Children)
+				{
+					if (Chaos::FPBDRigidClusteredParticleHandle* ClusteredChildHandle = ChildHandle->CastToClustered())
+					{
+						ClusteredChildHandle->SetExternalStrain(MaxStrain);
+					}
+				}
+				if (Children.Num() > 0)
+				{
+					SendCrumblingEvent(ClusteredHandle);
+				}
+				bCrumbledAnyCluster = true;
+			}
+		}
+		return bCrumbledAnyCluster;
+	}
+
 	static void ConnectClusteredNodes(FPBDRigidClusteredParticleHandle* ClusteredChild1, FPBDRigidClusteredParticleHandle* ClusteredChild2)
 	{
 		check(ClusteredChild1 && ClusteredChild2);
diff --git a/Engine/Source/Runtime/Experimental/Chaos/Private/PhysicsProxy/GeometryCollectionPhysicsProxy.cpp b/Engine/Source/Runtime/Experimental/Chaos/Private/PhysicsProxy/GeometryCollectionPhysicsProxy.cpp
--- a/Engine/Source/Runtime/Experimental/Chaos/Private/PhysicsProxy/GeometryCollectionPhysicsProxy.cpp
+++ b/Engine/Source/Runtime/Experimental/Chaos/Private/PhysicsProxy/GeometryCollectionPhysicsProxy.cpp
@@ -1606,6 +1606,49 @@ void FGeometryCollectionPhysicsProxy::BreakClusters_External(TArray<FGeometryCol
 	}
 }
 
+
+void FGeometryCollectionPhysicsProxy::BreakActiveClusters_External()
+{
+	check(IsInGameThread());
+
+	if (Chaos::FPhysicsSolver* RBDSolver = GetSolver<Chaos::FPhysicsSolver>())
+	{
+		RBDSolver->EnqueueCommandImmediate([this, RBDSolver]()
+		{
+			Chaos::FRigidClustering& Clustering = RBDSolver->GetEvolution()->GetRigidClustering();
+			Clustering.BreakClustersByProxy(this);
+		});
+	}
+}
+
+void FGeometryCollectionPhysicsProxy::RemoveAllAnchors_External()
+{
+	check(IsInGameThread());
+
+	if (Chaos::FPhysicsSolver* RBDSolver = GetSolver<Chaos::FPhysicsSolver>())
+	{
+		RBDSolver->EnqueueCommandImmediate([this, RBDSolver]()
+			{
+				Chaos::FPBDRigidsEvolution* Evolution = RBDSolver->GetEvolution();
+
+				// disable anchoring where relevant and collect the nodes to update
+				Chaos::FKinematicTarget NoKinematicTarget;
+				for (FClusterHandle* ParticleHandle : GetSolverParticleHandles())
+				{
+					if (ParticleHandle)
+					{
+						ParticleHandle->SetIsAnchored(false);
+						if (!ParticleHandle->IsDynamic())
+						{
+							Evolution->SetParticleObjectState(ParticleHandle, Chaos::EObjectStateType::Dynamic);
+							Evolution->SetParticleKinematicTarget(ParticleHandle, NoKinematicTarget);
+						}
+					}
+				}
+			});
+	}
+}
+
 template <typename TAction>
 static void ApplyToChildrenAtPointWithRadiusAndPropagation_Internal(
 	Chaos::FRigidClustering& Clustering,
diff --git a/Engine/Source/Runtime/Experimental/Chaos/Public/Chaos/PBDRigidClustering.h b/Engine/Source/Runtime/Experimental/Chaos/Public/Chaos/PBDRigidClustering.h
--- a/Engine/Source/Runtime/Experimental/Chaos/Public/Chaos/PBDRigidClustering.h
+++ b/Engine/Source/Runtime/Experimental/Chaos/Public/Chaos/PBDRigidClustering.h
@@ -154,6 +154,15 @@ public:
 	*    @return true if the cluster was successfully found and act upon 
 	*/
 	bool BreakCluster(FPBDRigidClusteredParticleHandle* ClusteredParticle);
+
+	/*
+	*  BreakClustersByProxy
+	*    Breaks clusters (internal or not) own by a specific proxy by applying max external strain to all its children
+	*
+	*    @param Proxy proxy owning the clusters to break
+	*    @return true if any cluster was successfully found and act upon
+	*/
+	bool BreakClustersByProxy(const IPhysicsProxyBase* Proxy);
 	
 	//
 	// Operational 
diff --git a/Engine/Source/Runtime/Experimental/Chaos/Public/PhysicsProxy/GeometryCollectionPhysicsProxy.h b/Engine/Source/Runtime/Experimental/Chaos/Public/PhysicsProxy/GeometryCollectionPhysicsProxy.h
--- a/Engine/Source/Runtime/Experimental/Chaos/Public/PhysicsProxy/GeometryCollectionPhysicsProxy.h
+++ b/Engine/Source/Runtime/Experimental/Chaos/Public/PhysicsProxy/GeometryCollectionPhysicsProxy.h
@@ -373,6 +373,8 @@ public:
 	void DisableParticles_External(TArray<int32>&& TransformGroupIndices);
 
 	void BreakClusters_External(TArray<FGeometryCollectionItemIndex>&& ItemIndices);
+	void BreakActiveClusters_External();
+	void RemoveAllAnchors_External();
 	void ApplyExternalStrain_External(FGeometryCollectionItemIndex ItemIndex, const FVector& WorldLocation, float Radius, int32 PropagationDepth, float PropagationFactor, float StrainValue);
 	void ApplyInternalStrain_External(FGeometryCollectionItemIndex ItemIndex, const FVector& WorldLocation, float Radius, int32 PropagationDepth, float PropagationFactor, float StrainValue);
 	void ApplyBreakingLinearVelocity_External(FGeometryCollectionItemIndex ItemIndex, const FVector& LinearVelocity);
diff --git a/Engine/Source/Runtime/Experimental/GeometryCollectionEngine/Private/GeometryCollection/GeometryCollectionComponent.cpp b/Engine/Source/Runtime/Experimental/GeometryCollectionEngine/Private/GeometryCollection/GeometryCollectionComponent.cpp
--- a/Engine/Source/Runtime/Experimental/GeometryCollectionEngine/Private/GeometryCollection/GeometryCollectionComponent.cpp
+++ b/Engine/Source/Runtime/Experimental/GeometryCollectionEngine/Private/GeometryCollection/GeometryCollectionComponent.cpp
@@ -4034,6 +4034,22 @@ void UGeometryCollectionComponent::CrumbleCluster(int32 ItemIndex)
 	}
 }
 
+void UGeometryCollectionComponent::CrumbleActiveClusters()
+{
+	if (PhysicsProxy)
+	{
+		PhysicsProxy->BreakActiveClusters_External();
+	}
+}
+
+void UGeometryCollectionComponent::RemoveAllAnchors()
+{
+	if (PhysicsProxy)
+	{
+		PhysicsProxy->RemoveAllAnchors_External();
+	}
+}
+
 void UGeometryCollectionComponent::ApplyBreakingLinearVelocity(int32 ItemIndex, const FVector& LinearVelocity)
 {
 	if (PhysicsProxy)
diff --git a/Engine/Source/Runtime/Experimental/GeometryCollectionEngine/Public/GeometryCollection/GeometryCollectionComponent.h b/Engine/Source/Runtime/Experimental/GeometryCollectionEngine/Public/GeometryCollection/GeometryCollectionComponent.h
--- a/Engine/Source/Runtime/Experimental/GeometryCollectionEngine/Public/GeometryCollection/GeometryCollectionComponent.h
+++ b/Engine/Source/Runtime/Experimental/GeometryCollectionEngine/Public/GeometryCollection/GeometryCollectionComponent.h
@@ -534,6 +534,19 @@ public:
 	UFUNCTION(BlueprintCallable, Category = "ChaosPhysics")
 	void CrumbleCluster(int32 ItemIndex);
 
+	/**
+	* Crumbe active clusters for this entire geometry collection
+	* this will apply to internal and regular clusters
+	*/
+	UFUNCTION(BlueprintCallable, Category = "ChaosPhysics")
+	void CrumbleActiveClusters();
+
+	/**
+	* this will remove anchors on all the pieces ( inlcuding the static and kinematic initial states ones ) of the geometry colection
+	*/
+	UFUNCTION(BlueprintCallable, Category = "ChaosPhysics")
+	void RemoveAllAnchors();
+
 	/**
 	 * Apply linear velocity on breaking pieces for a specific cluster
 	 * If ItemIndex does not represent a cluster this will do nothing  
