Chaos : fix potential crash when using CrumbleActiveClusters

diff --git a/Engine/Source/Runtime/Experimental/Chaos/Private/Chaos/PBDRigidClustering.cpp b/Engine/Source/Runtime/Experimental/Chaos/Private/Chaos/PBDRigidClustering.cpp
--- a/Engine/Source/Runtime/Experimental/Chaos/Private/Chaos/PBDRigidClustering.cpp
+++ b/Engine/Source/Runtime/Experimental/Chaos/Private/Chaos/PBDRigidClustering.cpp
@@ -1351,17 +1351,19 @@ namespace Chaos
 		{
 			if (ClusteredHandle && ClusteredHandle->PhysicsProxy() == Proxy)
 			{
-				const TArray<FPBDRigidParticleHandle*>& Children = MChildren[ClusteredHandle];
-				for (FPBDRigidParticleHandle* ChildHandle : Children)
+				if (const TArray<FPBDRigidParticleHandle*>* Children = MChildren.Find(ClusteredHandle))
 				{
-					if (Chaos::FPBDRigidClusteredParticleHandle* ClusteredChildHandle = ChildHandle->CastToClustered())
+					for (FPBDRigidParticleHandle* ChildHandle : *Children)
 					{
-						ClusteredChildHandle->SetExternalStrain(MaxStrain);
+						if (Chaos::FPBDRigidClusteredParticleHandle* ClusteredChildHandle = ChildHandle->CastToClustered())
+						{
+							ClusteredChildHandle->SetExternalStrain(MaxStrain);
+						}
+					}
+					if (Children->Num() > 0)
+					{
+						SendCrumblingEvent(ClusteredHandle);
 					}
-				}
-				if (Children.Num() > 0)
-				{
-					SendCrumblingEvent(ClusteredHandle);
 				}
 				bCrumbledAnyCluster = true;
 			}
