Remove ensure for degenerate triangles.

diff --git a/Engine/Source/Runtime/Experimental/Chaos/Private/Chaos/TriangleMeshImplicitObject.cpp b/Engine/Source/Runtime/Experimental/Chaos/Private/Chaos/TriangleMeshImplicitObject.cpp
--- a/Engine/Source/Runtime/Experimental/Chaos/Private/Chaos/TriangleMeshImplicitObject.cpp
+++ b/Engine/Source/Runtime/Experimental/Chaos/Private/Chaos/TriangleMeshImplicitObject.cpp
@@ -1556,7 +1556,7 @@ FVec3 FTriangleMeshImplicitObject::GetFaceNormal(const int32 FaceIdx) const
 			if(Normal.SafeNormalize() < UE_SMALL_NUMBER)
 			{
 				UE_LOG(LogChaos, Warning, TEXT("Degenerate triangle %d: (%f %f %f) (%f %f %f) (%f %f %f)"), FaceIdx, A.X, A.Y, A.Z, B.X, B.Y, B.Z, C.X, C.Y, C.Z);
-				ensure(false);
+				CHAOS_ENSURE(false);
 				return FVec3(0, 0, 1);
 			}
 
