remove scale when computing the relative transform to minimize rotation errors with negative scale

diff --git a/Engine/Source/Runtime/Experimental/Animation/Constraints/Private/TransformConstraint.cpp b/Engine/Source/Runtime/Experimental/Animation/Constraints/Private/TransformConstraint.cpp
--- a/Engine/Source/Runtime/Experimental/Animation/Constraints/Private/TransformConstraint.cpp
+++ b/Engine/Source/Runtime/Experimental/Animation/Constraints/Private/TransformConstraint.cpp
@@ -749,11 +749,15 @@ UTickableParentConstraint::UTickableParentConstraint()
 void UTickableParentConstraint::ComputeOffset()
 {
 	const FTransform InitParentTransform = GetParentGlobalTransform();
-	const FTransform InitChildTransform = GetChildGlobalTransform();
+	FTransform InitChildTransform = GetChildGlobalTransform();
 	
 	OffsetTransform = FTransform::Identity;
 	if (bMaintainOffset || bDynamicOffset)
 	{
+		if (!bScaling)
+		{
+			InitChildTransform.RemoveScaling();
+		}
 		OffsetTransform = InitChildTransform.GetRelativeTransform(InitParentTransform); 
 	}
 }
@@ -851,7 +855,11 @@ void UTickableParentConstraint::OnHandleModified(UTransformableHandle* InHandle,
 		if (bUpdateFromGlobal)
 		{
 			const FTransform ParentWorldTransform = GetParentGlobalTransform();
-			const FTransform ChildGlobalTransform = GetChildGlobalTransform();
+			FTransform ChildGlobalTransform = GetChildGlobalTransform();
+			if (!bScaling)
+			{
+				ChildGlobalTransform.RemoveScaling();
+			}
 			OffsetTransform = ChildGlobalTransform.GetRelativeTransform(ParentWorldTransform);
 		}
 		else
@@ -874,7 +882,11 @@ void UTickableParentConstraint::PostEditChangeProperty(FPropertyChangedEvent& Pr
 		{
 			Cache.CachedInputHash = CalculateInputHash();
 			
-			const FTransform ChildGlobalTransform = GetChildGlobalTransform();
+			FTransform ChildGlobalTransform = GetChildGlobalTransform();
+			if (!bScaling)
+			{
+				ChildGlobalTransform.RemoveScaling();
+			}
 			const FTransform ParentWorldTransform = GetParentGlobalTransform();
 			OffsetTransform = ChildGlobalTransform.GetRelativeTransform(ParentWorldTransform);
 			
@@ -1351,9 +1363,17 @@ FTransform FTransformConstraintUtils::ComputeRelativeTransform(
 		}
 	case ETransformConstraintType::Parent:
 		{
-			FTransform RelativeTransform = InChildWorld.GetRelativeTransform(InSpaceWorld);
 			const UTickableParentConstraint* ParentConstraint = Cast<UTickableParentConstraint>(InConstraint);
 			const bool bScale = ParentConstraint ? ParentConstraint->IsScalingEnabled() : true;
+			
+			FTransform ChildTransform = InChildWorld;
+			if (!bScale)
+			{
+				ChildTransform.RemoveScaling();
+			}
+
+			FTransform RelativeTransform = ChildTransform.GetRelativeTransform(InSpaceWorld);
+			
 			if (!bScale)
 			{
 				RelativeTransform.SetScale3D(InChildLocal.GetScale3D());
