Re-enable limiting of mass ratios in cached joint solver

diff --git a/Engine/Source/Runtime/Experimental/Chaos/Private/Chaos/Joint/PBDJointCachedSolverGaussSeidel.cpp b/Engine/Source/Runtime/Experimental/Chaos/Private/Chaos/Joint/PBDJointCachedSolverGaussSeidel.cpp
--- a/Engine/Source/Runtime/Experimental/Chaos/Private/Chaos/Joint/PBDJointCachedSolverGaussSeidel.cpp
+++ b/Engine/Source/Runtime/Experimental/Chaos/Private/Chaos/Joint/PBDJointCachedSolverGaussSeidel.cpp
@@ -90,12 +90,12 @@ void FPBDJointCachedSolver::Update(
 	//UpdateIsActive();
 }
 
-void FPBDJointCachedSolver::UpdateMass0()
+void FPBDJointCachedSolver::UpdateMass0(const FReal& InInvM, const FVec3& InInvIL)
 {
 	if (Body0().IsDynamic())
 	{
-		InvMs[0] = Body0().InvM();
-		InvIs[0] = Utilities::ComputeWorldSpaceInertia(CurrentQs[0], Body0().InvILocal());
+		InvMs[0] = InInvM;
+		InvIs[0] = Utilities::ComputeWorldSpaceInertia(CurrentQs[0], InInvIL);
 	}
 	else
 	{
@@ -104,12 +104,12 @@ void FPBDJointCachedSolver::UpdateMass0()
 	}
 }
 
-void FPBDJointCachedSolver::UpdateMass1()
+void FPBDJointCachedSolver::UpdateMass1(const FReal& InInvM, const FVec3& InInvIL)
 {
 	if (Body1().IsDynamic())
 	{
-		InvMs[1] = Body1().InvM();
-		InvIs[1] = Utilities::ComputeWorldSpaceInertia(CurrentQs[1], Body1().InvILocal());
+		InvMs[1] = InInvM;
+		InvIs[1] = Utilities::ComputeWorldSpaceInertia(CurrentQs[1], InInvIL);
 	}
 	else
 	{
@@ -121,16 +121,20 @@ void FPBDJointCachedSolver::UpdateMass1()
 void FPBDJointCachedSolver::SetShockPropagationScales(const FReal InvMScale0, const FReal InvMScale1, const FReal Dt)
 {
 	bool bNeedsUpdate = false;
-	if (Body0().ShockPropagationScale() != InvMScale0)
+	if (Body0().ShockPropagationScale() != InvMScale0 && Body0().ShockPropagationScale() > 0.0f)
 	{
+		const FReal Mult = InvMScale0 / Body0().ShockPropagationScale();
+		InvMs[0] *= Mult;
+		InvIs[0] *= Mult;
 		Body0().SetShockPropagationScale(InvMScale0);
-		UpdateMass0();
 		bNeedsUpdate = true;
 	}
-	if (Body1().ShockPropagationScale() != InvMScale1)
+	if (Body1().ShockPropagationScale() != InvMScale1 && Body1().ShockPropagationScale() > 0.0f)
 	{
+		const FReal Mult = InvMScale1 / Body1().ShockPropagationScale();
+		InvMs[1] *= Mult;
+		InvIs[1] *= Mult;
 		Body1().SetShockPropagationScale(InvMScale1);
-		UpdateMass1();
 		bNeedsUpdate = true;
 	}
 	if(bNeedsUpdate)
@@ -198,8 +202,12 @@ void FPBDJointCachedSolver::Init(
 
 	InitDerivedState();
 
-	UpdateMass0();
-	UpdateMass1();
+	FReal ConditionedInvMs[2];
+	FVec3 ConditionedInvILs[2];
+	FPBDJointUtilities::ConditionInverseMassAndInertia(Body0().InvM(), Body1().InvM(), Body0().InvILocal(), Body1().InvILocal(), SolverSettings.MinParentMassRatio, SolverSettings.MaxInertiaRatio, ConditionedInvMs[0], ConditionedInvMs[1], ConditionedInvILs[0], ConditionedInvILs[1]);
+
+	UpdateMass0(ConditionedInvMs[0], ConditionedInvILs[0]);
+	UpdateMass1(ConditionedInvMs[1], ConditionedInvILs[1]);
 
 	// Cache all the informations for the position and rotation constraints
 	InitPositionConstraints(Dt, SolverSettings, JointSettings);
diff --git a/Engine/Source/Runtime/Experimental/Chaos/Public/Chaos/Joint/PBDJointCachedSolverGaussSeidel.h b/Engine/Source/Runtime/Experimental/Chaos/Public/Chaos/Joint/PBDJointCachedSolverGaussSeidel.h
--- a/Engine/Source/Runtime/Experimental/Chaos/Public/Chaos/Joint/PBDJointCachedSolverGaussSeidel.h
+++ b/Engine/Source/Runtime/Experimental/Chaos/Public/Chaos/Joint/PBDJointCachedSolverGaussSeidel.h
@@ -281,8 +281,8 @@ struct FAxisConstraintDatas
 		void UpdateDerivedState(const int32 BodyIndex);
 		void UpdateDerivedState();
 
-		void UpdateMass0();
-		void UpdateMass1();
+		void UpdateMass0(const FReal& InInvM, const FVec3& InInvIL);
+		void UpdateMass1(const FReal& InInvM, const FVec3& InInvIL);
 
 		// Check to see if this constraint still needs further solved
 		// @todo(chaos): the term "active" is used inconsistently with the meaning elsewhere. Active should
