Turned on unity files for Chaos for every platform. Tweaked the unity file sizes for low to high end PCs.

diff --git a/Engine/Source/Runtime/Experimental/Chaos/Chaos.Build.cs b/Engine/Source/Runtime/Experimental/Chaos/Chaos.Build.cs
--- a/Engine/Source/Runtime/Experimental/Chaos/Chaos.Build.cs
+++ b/Engine/Source/Runtime/Experimental/Chaos/Chaos.Build.cs
@@ -9,6 +9,8 @@ namespace UnrealBuildTool.Rules
     {
         public Chaos(ReadOnlyTargetRules Target) : base(Target)
         {
+			NumIncludedBytesPerUnityCPPOverride = 192 * 1024; // This is half of the default unity size(NumIncludedBytesPerUnityCPP) specified in TargetRules.cs
+
             PublicIncludePaths.Add("Runtime/Experimental/Chaos/Public");
 
             PublicDependencyModuleNames.AddRange(
