Chaos : fix character movement interaction with geometry collections

diff --git a/Engine/Source/Runtime/Engine/Engine.Build.cs b/Engine/Source/Runtime/Engine/Engine.Build.cs
--- a/Engine/Source/Runtime/Engine/Engine.Build.cs
+++ b/Engine/Source/Runtime/Engine/Engine.Build.cs
@@ -143,7 +143,8 @@ public class Engine : ModuleRules
 				"TraceLog",
 				"ColorManagement",
 				"Icmp",
-				"XmlParser"
+				"XmlParser",
+				"GeometryCollectionEngine"
 			}
 		);
 
diff --git a/Engine/Source/Runtime/Engine/Private/Components/CharacterMovementComponent.cpp b/Engine/Source/Runtime/Engine/Private/Components/CharacterMovementComponent.cpp
--- a/Engine/Source/Runtime/Engine/Private/Components/CharacterMovementComponent.cpp
+++ b/Engine/Source/Runtime/Engine/Private/Components/CharacterMovementComponent.cpp
@@ -8,6 +8,7 @@
 #include "GameFramework/CharacterMovementComponent.h"
 #include "EngineStats.h"
 #include "Components/PrimitiveComponent.h"
+#include "GeometryCollection/GeometryCollectionComponent.h"
 #include "AI/NavigationSystemBase.h"
 #include "AI/Navigation/NavigationDataInterface.h"
 #include "UObject/Package.h"
@@ -7015,12 +7016,14 @@ void UCharacterMovementComponent::ApplyImpactPhysicsForces(const FHitResult& Imp
 	{
 		if (UPrimitiveComponent* ImpactComponent = Impact.GetComponent())
 		{
+			FVector ForcePoint = Impact.ImpactPoint;
+			float BodyMass = 0.0f;
+
+			bool bCanBePushed = true;
 			FBodyInstance* BI = ImpactComponent->GetBodyInstance(Impact.BoneName);
 			if(BI != nullptr && BI->IsInstanceSimulatingPhysics())
 			{
-				FVector ForcePoint = Impact.ImpactPoint;
-
-				const float BodyMass = FMath::Max(BI->GetBodyMass(), 1.0f);
+				BodyMass = FMath::Max(BI->GetBodyMass(), 1.0f);
 
 				if(bPushForceUsingZOffset)
 				{
@@ -7034,7 +7037,29 @@ void UCharacterMovementComponent::ApplyImpactPhysicsForces(const FHitResult& Imp
 						ForcePoint.Z = Center.Z + Extents.Z * PushForcePointZOffsetFactor;
 					}
 				}
+			}
+			else if (UGeometryCollectionComponent* GCComp = Cast<UGeometryCollectionComponent>(ImpactComponent))
+			{
+				FBox Bounds;
+				GCComp->GetMassAndExtents(Impact.Item, BodyMass, Bounds);
+				if (bPushForceUsingZOffset)
+				{
+					FVector Center, Extents;
+					Bounds.GetCenterAndExtents(Center, Extents);
 
+					if (!Extents.IsNearlyZero())
+					{
+						ForcePoint.Z = Center.Z + Extents.Z * PushForcePointZOffsetFactor;
+					}
+				}
+			}
+			else
+			{
+				// no body instance, not a GC, not supported scenario
+				bCanBePushed = false;
+			}
+			if (bCanBePushed)
+			{
 				FVector Force = Impact.ImpactNormal * -1.0f;
 
 				float PushForceModificator = 1.0f;
@@ -7045,7 +7070,7 @@ void UCharacterMovementComponent::ApplyImpactPhysicsForces(const FHitResult& Imp
 				float Dot = 0.0f;
 
 				if (bScalePushForceToVelocity && !ComponentVelocity.IsNearlyZero())
-				{			
+				{
 					Dot = ComponentVelocity | VirtualVelocity;
 
 					if (Dot > 0.0f && Dot < 1.0f)
diff --git a/Engine/Source/Runtime/Experimental/Chaos/Private/PhysicsProxy/GeometryCollectionPhysicsProxy.cpp b/Engine/Source/Runtime/Experimental/Chaos/Private/PhysicsProxy/GeometryCollectionPhysicsProxy.cpp
--- a/Engine/Source/Runtime/Experimental/Chaos/Private/PhysicsProxy/GeometryCollectionPhysicsProxy.cpp
+++ b/Engine/Source/Runtime/Experimental/Chaos/Private/PhysicsProxy/GeometryCollectionPhysicsProxy.cpp
@@ -450,6 +450,8 @@ void FGeometryCollectionPhysicsProxy::Initialize(Chaos::FPBDRigidsEvolutionBase
 	GameThreadCollection.AddExternalAttribute<TUniquePtr<Chaos::FGeometryParticle>>(FGeometryCollection::ParticlesAttribute, FTransformCollection::TransformGroup, GTParticles);
 	
 
+	const FVector Scale = Parameters.WorldTransform.GetScale3D();
+
 	TArray<int32> ChildrenToCheckForParentFix;
 	if(ensure(NumTransforms == GameThreadCollection.Implicits.Num() && NumTransforms == GTParticles.Num())) // Implicits are in the transform group so this invariant should always hold
 	{
@@ -466,7 +468,14 @@ void FGeometryCollectionPhysicsProxy::Initialize(Chaos::FPBDRigidsEvolutionBase
 			P->SetR(T.GetRotation(), false);
 			P->SetUserData(Parameters.UserData);
 			P->SetProxy(this);
-			P->SetGeometry(GameThreadCollection.Implicits[Index]);
+
+			FGeometryDynamicCollection::FSharedImplicit ImplicitGeometry = GameThreadCollection.Implicits[Index];
+			if (!Scale.Equals(FVector::OneVector))
+			{
+				TUniquePtr<Chaos::FImplicitObject> ScaledImplicit = ImplicitGeometry->CopyWithScale(Scale);
+				ImplicitGeometry = FGeometryDynamicCollection::FSharedImplicit(ScaledImplicit.Release());
+			}
+			P->SetGeometry(ImplicitGeometry);
 
 			// this step is necessary for Phase 2 where we need to walk back the hierarchy from children to parent 
 			if (bGeometryCollectionAlwaysGenerateGTCollisionForClusters && GameThreadCollection.Children[Index].Num() == 0)
@@ -1586,6 +1595,84 @@ void FGeometryCollectionPhysicsProxy::DisableParticles_External(TArray<int32>&&
 	}
 }
 
+void FGeometryCollectionPhysicsProxy::ApplyForceAt_External(FVector Force, FVector WorldLocation)
+{
+	check(IsInGameThread());
+
+	if (Chaos::FPhysicsSolver* RBDSolver = GetSolver<Chaos::FPhysicsSolver>())
+	{
+		RBDSolver->EnqueueCommandImmediate([this, Force, WorldLocation, RBDSolver]()
+			{
+				Chaos::FReal ClosestDistanceSquared = TNumericLimits<Chaos::FReal>::Max();
+				Chaos::FPBDRigidClusteredParticleHandle* ClosestHandle = nullptr;
+
+				Chaos::FRigidClustering& Clustering = RBDSolver->GetEvolution()->GetRigidClustering();
+				for (Chaos::FPBDRigidClusteredParticleHandle* ClusteredHandle : Clustering.GetTopLevelClusterParents())
+				{
+					if (ClusteredHandle && ClusteredHandle->PhysicsProxy() == this)
+					{
+						if (ClusteredHandle->IsDynamic())
+						{
+							const Chaos::FReal DistanceSquared = (WorldLocation - ClusteredHandle->X()).SquaredLength();
+							if (DistanceSquared < ClosestDistanceSquared)
+							{
+								ClosestDistanceSquared = DistanceSquared;
+								ClosestHandle = ClusteredHandle;
+							}
+						}
+					}
+				}
+				if (ClosestHandle)
+				{
+					const Chaos::FVec3 WorldCOM = Chaos::FParticleUtilitiesXR::GetCoMWorldPosition(ClosestHandle);
+					const Chaos::FVec3 WorldTorque = Chaos::FVec3::CrossProduct(WorldLocation - WorldCOM, Force);
+					ClosestHandle->AddForce(Force);
+					ClosestHandle->AddTorque(WorldTorque);
+				}
+			});
+	}
+}
+
+void FGeometryCollectionPhysicsProxy::ApplyImpulseAt_External(FVector Impulse, FVector WorldLocation)
+{
+	check(IsInGameThread());
+
+	if (Chaos::FPhysicsSolver* RBDSolver = GetSolver<Chaos::FPhysicsSolver>())
+	{
+		RBDSolver->EnqueueCommandImmediate([this, Impulse, WorldLocation, RBDSolver]()
+			{
+				Chaos::FReal ClosestDistanceSquared = TNumericLimits<Chaos::FReal>::Max();
+				Chaos::FPBDRigidClusteredParticleHandle* ClosestHandle = nullptr;
+
+				Chaos::FRigidClustering& Clustering = RBDSolver->GetEvolution()->GetRigidClustering();
+				for (Chaos::FPBDRigidClusteredParticleHandle* ClusteredHandle : Clustering.GetTopLevelClusterParents())
+				{
+					if (ClusteredHandle && ClusteredHandle->PhysicsProxy() == this)
+					{
+						if (ClusteredHandle->IsDynamic())
+						{
+							const Chaos::FReal DistanceSquared = (WorldLocation - ClusteredHandle->X()).SquaredLength();
+							if (DistanceSquared < ClosestDistanceSquared)
+							{
+								ClosestDistanceSquared = DistanceSquared;
+								ClosestHandle = ClusteredHandle;
+							}
+						}
+					}
+				}
+				if (ClosestHandle)
+				{
+					const Chaos::FVec3 WorldCOM = Chaos::FParticleUtilitiesXR::GetCoMWorldPosition(ClosestHandle);
+					ClosestHandle->SetLinearImpulseVelocity(ClosestHandle->LinearImpulseVelocity() + (Impulse * ClosestHandle->InvM()), false);
+
+					const Chaos::FMatrix33 WorldInvI = Chaos::Utilities::ComputeWorldSpaceInertia(ClosestHandle->R() * ClosestHandle->RotationOfMass(), ClosestHandle->InvI());
+					const Chaos::FVec3 AngularImpulse = Chaos::FVec3::CrossProduct(WorldLocation - WorldCOM, Impulse);
+					ClosestHandle->SetAngularImpulseVelocity(ClosestHandle->AngularImpulseVelocity() + WorldInvI * AngularImpulse, false);
+				}
+			});
+	}
+}
+
 void FGeometryCollectionPhysicsProxy::BreakClusters_External(TArray<FGeometryCollectionItemIndex>&& ItemIndices)
 {
 	check(IsInGameThread());
diff --git a/Engine/Source/Runtime/Experimental/Chaos/Public/PhysicsProxy/GeometryCollectionPhysicsProxy.h b/Engine/Source/Runtime/Experimental/Chaos/Public/PhysicsProxy/GeometryCollectionPhysicsProxy.h
--- a/Engine/Source/Runtime/Experimental/Chaos/Public/PhysicsProxy/GeometryCollectionPhysicsProxy.h
+++ b/Engine/Source/Runtime/Experimental/Chaos/Public/PhysicsProxy/GeometryCollectionPhysicsProxy.h
@@ -372,6 +372,8 @@ public:
 
 	void DisableParticles_External(TArray<int32>&& TransformGroupIndices);
 
+	void ApplyForceAt_External(FVector Force, FVector WorldLocation);
+	void ApplyImpulseAt_External(FVector Force, FVector WorldLocation);
 	void BreakClusters_External(TArray<FGeometryCollectionItemIndex>&& ItemIndices);
 	void BreakActiveClusters_External();
 	void RemoveAllAnchors_External();
diff --git a/Engine/Source/Runtime/Experimental/GeometryCollectionEngine/Private/GeometryCollection/GeometryCollectionComponent.cpp b/Engine/Source/Runtime/Experimental/GeometryCollectionEngine/Private/GeometryCollection/GeometryCollectionComponent.cpp
--- a/Engine/Source/Runtime/Experimental/GeometryCollectionEngine/Private/GeometryCollection/GeometryCollectionComponent.cpp
+++ b/Engine/Source/Runtime/Experimental/GeometryCollectionEngine/Private/GeometryCollection/GeometryCollectionComponent.cpp
@@ -773,6 +773,15 @@ void UGeometryCollectionComponent::AddForce(FVector Force, FName BoneName, bool
 	DispatchFieldCommand(Command);
 }
 
+
+void UGeometryCollectionComponent::AddForceAtLocation(FVector Force, FVector WorldLocation, FName BoneName)
+{
+	if (PhysicsProxy)
+	{
+		PhysicsProxy->ApplyForceAt_External(Force, WorldLocation);
+	}
+}
+
 void UGeometryCollectionComponent::AddImpulse(FVector Impulse, FName BoneName, bool bVelChange)
 {
 	const FVector Direction = Impulse.GetSafeNormal(); 
@@ -783,6 +792,15 @@ void UGeometryCollectionComponent::AddImpulse(FVector Impulse, FName BoneName, b
 	DispatchFieldCommand(Command);
 }
 
+void UGeometryCollectionComponent::AddImpulseAtLocation(FVector Impulse, FVector WorldLocation, FName BoneName)
+{
+	if (PhysicsProxy)
+	{
+		PhysicsProxy->ApplyImpulseAt_External(Impulse, WorldLocation);
+	}
+	
+}
+
 TUniquePtr<FFieldNodeBase> MakeRadialField(const FVector& Origin, float Radius, float Strength, ERadialImpulseFalloff Falloff)
 {
 	TUniquePtr<FFieldNodeBase> Field;
@@ -4129,6 +4147,61 @@ int32 UGeometryCollectionComponent::GetInitialLevel(int32 ItemIndex)
 	return Level;
 }
 
+void UGeometryCollectionComponent::GetMassAndExtents(int32 ItemIndex, float& OutMass, FBox& OutExtents)
+{
+	using FGeometryCollectionPtr = const TSharedPtr<FGeometryCollection, ESPMode::ThreadSafe>;
+
+	OutMass = 0.0f;
+	OutExtents = FBox(EForceInit::ForceInitToZero);
+
+	int32 Level = INDEX_NONE;
+	if (RestCollection && RestCollection->GetGeometryCollection())
+	{
+		const FGeometryCollection& Collection = *RestCollection->GetGeometryCollection();
+		if (const TManagedArray<float>* CollectionMass = Collection.FindAttribute<float>(TEXT("Mass"), FTransformCollection::TransformGroup))
+		{
+			const TManagedArray<FBox>* TransformBoundingBoxes = Collection.FindAttribute<FBox>(TEXT("BoundingBox"), FTransformCollection::TransformGroup);
+			const TManagedArray<FBox>* GeoBoundingBoxes = Collection.FindAttribute<FBox>(TEXT("BoundingBox"), FGeometryCollection::GeometryGroup);
+
+			int32 TransformIndex = INDEX_NONE;
+			FGeometryCollectionItemIndex GCItemIndex = FGeometryCollectionItemIndex::CreateFromExistingItemIndex(ItemIndex);
+
+			if (GCItemIndex.IsInternalCluster())
+			{
+				const TArray<int32>* Children = PhysicsProxy->FindInternalClusterChildrenTransformIndices_External(GCItemIndex);
+				if (Children)
+				{
+					for (const int32 ChildTramsformIndex : *Children)
+					{
+						OutMass += (*CollectionMass)[ChildTramsformIndex];
+						if (TransformBoundingBoxes)
+						{
+							OutExtents += (*TransformBoundingBoxes)[ChildTramsformIndex];
+						}
+						else if (GeoBoundingBoxes)
+						{
+							OutExtents += (*GeoBoundingBoxes)[Collection.TransformToGeometryIndex[ChildTramsformIndex]];
+						}
+					}
+				}
+			}
+			else
+			{
+				TransformIndex = GCItemIndex.GetTransformIndex();
+				OutMass = (*CollectionMass)[TransformIndex];
+				if (TransformBoundingBoxes)
+				{
+					OutExtents = (*TransformBoundingBoxes)[TransformIndex];
+				}
+				else
+				{
+					OutExtents = (*GeoBoundingBoxes)[Collection.TransformToGeometryIndex[TransformIndex]];
+				}
+			}
+		}
+	}
+}
+
 bool UGeometryCollectionComponent::CalculateInnerSphere(int32 TransformIndex, UE::Math::TSphere<double>& SphereOut) const
 {
 	// Approximates the inscribed sphere. Returns false if no such sphere exists, if for instance the index is to an embedded geometry. 
diff --git a/Engine/Source/Runtime/Experimental/GeometryCollectionEngine/Public/GeometryCollection/GeometryCollectionComponent.h b/Engine/Source/Runtime/Experimental/GeometryCollectionEngine/Public/GeometryCollection/GeometryCollectionComponent.h
--- a/Engine/Source/Runtime/Experimental/GeometryCollectionEngine/Public/GeometryCollection/GeometryCollectionComponent.h
+++ b/Engine/Source/Runtime/Experimental/GeometryCollectionEngine/Public/GeometryCollection/GeometryCollectionComponent.h
@@ -486,7 +486,9 @@ public:
 	virtual bool CanEditSimulatePhysics() override;
 	virtual void SetSimulatePhysics(bool bEnabled) override;
 	virtual void AddForce(FVector Force, FName BoneName = NAME_None, bool bAccelChange = false) override;
+	virtual void AddForceAtLocation(FVector Force, FVector WorldLocation, FName BoneName = NAME_None) override;
 	virtual void AddImpulse(FVector Impulse, FName BoneName = NAME_None, bool bVelChange = false) override;
+	virtual void AddImpulseAtLocation(FVector Impulse, FVector WorldLocation, FName BoneName = NAME_None) override;
 	virtual void AddRadialForce(FVector Origin, float Radius, float Strength, ERadialImpulseFalloff Falloff, bool bAccelChange = false) override;
 	virtual void AddRadialImpulse(FVector Origin, float Radius, float Strength, enum ERadialImpulseFalloff Falloff, bool bVelChange = false) override;
 	virtual void AddTorqueInRadians(FVector Torque, FName BoneName = NAME_None, bool bAccelChange = false) override;
@@ -590,6 +592,14 @@ public:
 	UFUNCTION(BlueprintCallable, Category = "ChaosPhysics")
 	int32 GetInitialLevel(int32 ItemIndex);
 	
+	/**
+	* Get mass and extent of a specific piece
+	* @param ItemIndex item index ( from HitResult) of the cluster to get level from
+	* @param Level of the piece ( 0 for root level and positive for the rest )
+	*/
+	UFUNCTION(BlueprintCallable, Category = "ChaosPhysics")
+	void GetMassAndExtents(int32 ItemIndex, float& OutMass, FBox& OutExtents);
+
 	/** RestCollection */
 	UFUNCTION(BlueprintCallable, Category = "ChaosPhysics")
 	void SetRestCollection(const UGeometryCollection * RestCollectionIn);
