Chaos - fix body order in broad phase

diff --git a/Engine/Source/Runtime/Experimental/Chaos/Public/Chaos/Collision/CollisionKeys.h b/Engine/Source/Runtime/Experimental/Chaos/Public/Chaos/Collision/CollisionKeys.h
--- a/Engine/Source/Runtime/Experimental/Chaos/Public/Chaos/Collision/CollisionKeys.h
+++ b/Engine/Source/Runtime/Experimental/Chaos/Public/Chaos/Collision/CollisionKeys.h
@@ -34,15 +34,17 @@ namespace Chaos
 	 * Used to order particles in a consistent way in Broadphase and Resim. Returns whether the partcile order should be reversed.
 	 * We want the particle with the lower ID first, unless only one is dynamic in which case that one is first.
 	 */
-	inline bool ShouldSwapParticleOrder(const bool bIsDynamicOrSleeping0, const bool bIsParticle0Preferred)
+	inline bool ShouldSwapParticleOrder(const bool bIsDynamicOrSleeping0, const bool bIsDynamicOrSleeping1, const bool bIsParticle0Preferred)
 	{
-		return !bIsDynamicOrSleeping0 || !bIsParticle0Preferred;
+		return !bIsDynamicOrSleeping0 || (bIsDynamicOrSleeping1 && !bIsParticle0Preferred);
 	}
 	inline bool ShouldSwapParticleOrder(const FGeometryParticleHandle* Particle0, const FGeometryParticleHandle* Particle1)
 	{
 		const EObjectStateType ObjectState0 = Particle0->ObjectState();
-		const bool bIsDynamic0 = (ObjectState0 == EObjectStateType::Dynamic) || (ObjectState0 == EObjectStateType::Sleeping);
-		return ShouldSwapParticleOrder(bIsDynamic0, AreParticlesInPreferredOrder(Particle0, Particle1));
+		const EObjectStateType ObjectState1 = Particle1->ObjectState();
+		const bool bIsDynamicOrSleeping0 = (ObjectState0 == EObjectStateType::Dynamic) || (ObjectState0 == EObjectStateType::Sleeping);
+		const bool bIsDynamicOrSleeping1 = (ObjectState1 == EObjectStateType::Dynamic) || (ObjectState1 == EObjectStateType::Sleeping);
+		return ShouldSwapParticleOrder(bIsDynamicOrSleeping0, bIsDynamicOrSleeping1, AreParticlesInPreferredOrder(Particle0, Particle1));
 	}
 
 	/**
diff --git a/Engine/Source/Runtime/Experimental/Chaos/Public/Chaos/Collision/SpatialAccelerationBroadPhase.h b/Engine/Source/Runtime/Experimental/Chaos/Public/Chaos/Collision/SpatialAccelerationBroadPhase.h
--- a/Engine/Source/Runtime/Experimental/Chaos/Public/Chaos/Collision/SpatialAccelerationBroadPhase.h
+++ b/Engine/Source/Runtime/Experimental/Chaos/Public/Chaos/Collision/SpatialAccelerationBroadPhase.h
@@ -462,7 +462,7 @@ namespace Chaos
 				// Since particles can change type (between Kinematic and Dynamic) we may visit them in different orders at different times, but if we allow
 				// that it would break Resim and constraint re-use. Also, if only one particle is Dynamic, we want it in first position. This isn't a strtct 
 				// requirement but some downstream systems assume this is true (e.g., CCD, TriMesh collision).
-				const bool bSwapOrder = ShouldSwapParticleOrder((bIsDynamicAwake1 || bIsDynamicAsleep1), bIsParticle1Preferred);
+				const bool bSwapOrder = ShouldSwapParticleOrder((bIsDynamicAwake1 || bIsDynamicAsleep1), (bIsDynamicAwake2 || bIsDynamicAsleep2), bIsParticle1Preferred);
 				FGeometryParticleHandle* ParticleA = bSwapOrder ? Particle2 : Particle1;
 				FGeometryParticleHandle* ParticleB = bSwapOrder ? Particle1 : Particle2;
 
