ensure the map contains the binding id before coverting it

diff --git a/Engine/Source/Runtime/Experimental/Animation/Constraints/Private/TransformableHandle.cpp b/Engine/Source/Runtime/Experimental/Animation/Constraints/Private/TransformableHandle.cpp
--- a/Engine/Source/Runtime/Experimental/Animation/Constraints/Private/TransformableHandle.cpp
+++ b/Engine/Source/Runtime/Experimental/Animation/Constraints/Private/TransformableHandle.cpp
@@ -44,8 +44,11 @@ void UTransformableHandle::OnBindingIDsUpdated(const TMap<UE::MovieScene::FFixed
 	if (ConstraintBindingID.IsValid())
 	{
 		UE::MovieScene::FFixedObjectBindingID FixedBindingID = ConstraintBindingID.ResolveToFixed(LocalSequenceID, Player);
-		Modify();
-		ConstraintBindingID = OldFixedToNewFixedMap[FixedBindingID].ConvertToRelative(LocalSequenceID, Hierarchy);
+		if (OldFixedToNewFixedMap.Contains(FixedBindingID))
+		{
+			Modify();
+			ConstraintBindingID = OldFixedToNewFixedMap[FixedBindingID].ConvertToRelative(LocalSequenceID, Hierarchy);
+		}
 	}
 }
 
