Chaos - remove some dynamic casts
- clean up

diff --git a/Engine/Source/Runtime/Experimental/Chaos/Private/Chaos/Collision/PBDCollisionConstraint.cpp b/Engine/Source/Runtime/Experimental/Chaos/Private/Chaos/Collision/PBDCollisionConstraint.cpp
--- a/Engine/Source/Runtime/Experimental/Chaos/Private/Chaos/Collision/PBDCollisionConstraint.cpp
+++ b/Engine/Source/Runtime/Experimental/Chaos/Private/Chaos/Collision/PBDCollisionConstraint.cpp
@@ -644,8 +644,7 @@ namespace Chaos
 		ClosestManifoldPointIndex = INDEX_NONE;
 
 		// How many manifold points we expect. E.g., for Box-box this will be 4 or 1 depending on whether
-		// we have a face or edge contact. We don't reuse the manifold if we lose points after culling
-		// here and potentially adding the new narrow phase result (See TryAddManifoldContact).
+		// we have a face or edge contact. We don't reuse the manifold if we lose points after culling here
 		ExpectedNumManifoldPoints = ManifoldPoints.Num();
 		Flags.bWasManifoldRestored = false;
 
diff --git a/Engine/Source/Runtime/Experimental/Chaos/Private/Chaos/Collision/PBDCollisionSolver.h b/Engine/Source/Runtime/Experimental/Chaos/Private/Chaos/Collision/PBDCollisionSolver.h
--- a/Engine/Source/Runtime/Experimental/Chaos/Private/Chaos/Collision/PBDCollisionSolver.h
+++ b/Engine/Source/Runtime/Experimental/Chaos/Private/Chaos/Collision/PBDCollisionSolver.h
@@ -1,9 +1,6 @@
 // Copyright Epic Games, Inc. All Rights Reserved.
 #pragma once
 
-#include "CoreMinimal.h"
-#include "Logging/LogMacros.h"
-
 #include "Chaos/Core.h"
 #include "Chaos/Defines.h"
 #include "Chaos/Evolution/SolverBody.h"
@@ -884,7 +881,6 @@ namespace Chaos
 		}
 	}
 
-
 	FORCEINLINE_DEBUGGABLE void FPBDCollisionSolverManifoldPoint::CalculateContactPositionErrorTangential(const FConstraintSolverBody& Body0, const FConstraintSolverBody& Body1, FSolverReal& OutContactDeltaTangentU, FSolverReal& OutContactDeltaTangentV) const
 	{
 		// Linear version: calculate the contact delta assuming linear motion after applying a positional impulse at the contact point. There will be an error that depends on the size of the rotation.
diff --git a/Engine/Source/Runtime/Experimental/Chaos/Private/Chaos/Collision/SolverCollisionContainer.cpp b/Engine/Source/Runtime/Experimental/Chaos/Private/Chaos/Collision/SolverCollisionContainer.cpp
--- a/Engine/Source/Runtime/Experimental/Chaos/Private/Chaos/Collision/SolverCollisionContainer.cpp
+++ b/Engine/Source/Runtime/Experimental/Chaos/Private/Chaos/Collision/SolverCollisionContainer.cpp
@@ -423,8 +423,7 @@ namespace Chaos
 	{
 		for (FPBDIslandConstraint& IslandConstraint : IslandConstraints)
 		{
-			// @todo(chaos): unsafe version of As
-			FPBDCollisionConstraint& Constraint = IslandConstraint.GetConstraint()->As<FPBDCollisionConstraintHandle>()->GetContact();
+			FPBDCollisionConstraint& Constraint = IslandConstraint.GetConstraint()->AsUnsafe<FPBDCollisionConstraintHandle>()->GetContact();
 
 			AddConstraint(Constraint);
 		}
diff --git a/Engine/Source/Runtime/Experimental/Chaos/Private/Chaos/PBDCollisionConstraints.cpp b/Engine/Source/Runtime/Experimental/Chaos/Private/Chaos/PBDCollisionConstraints.cpp
--- a/Engine/Source/Runtime/Experimental/Chaos/Private/Chaos/PBDCollisionConstraints.cpp
+++ b/Engine/Source/Runtime/Experimental/Chaos/Private/Chaos/PBDCollisionConstraints.cpp
@@ -388,12 +388,10 @@ namespace Chaos
 			IslandManager.VisitConstraintsInAwakeIslands(GetContainerId(),
 				[this](FConstraintHandle* ConstraintHandle)
 				{
-					if (FPBDCollisionConstraintHandle* CollisionHandle = ConstraintHandle->As<FPBDCollisionConstraintHandle>())
+					FPBDCollisionConstraintHandle* CollisionHandle = ConstraintHandle->AsUnsafe<FPBDCollisionConstraintHandle>();
+					if (!CollisionHandle->IsEnabled() || ConstraintAllocator.IsConstraintExpired(CollisionHandle->GetContact()))
 					{
-						if (!CollisionHandle->IsEnabled() || ConstraintAllocator.IsConstraintExpired(CollisionHandle->GetContact()))
-						{
-							TempCollisions.Add(CollisionHandle);
-						}
+						TempCollisions.Add(CollisionHandle);
 					}
 				});
 
diff --git a/Engine/Source/Runtime/Experimental/Chaos/Public/Chaos/Collision/ContactPoint.h b/Engine/Source/Runtime/Experimental/Chaos/Public/Chaos/Collision/ContactPoint.h
--- a/Engine/Source/Runtime/Experimental/Chaos/Public/Chaos/Collision/ContactPoint.h
+++ b/Engine/Source/Runtime/Experimental/Chaos/Public/Chaos/Collision/ContactPoint.h
@@ -1,12 +1,7 @@
 // Copyright Epic Games, Inc.All Rights Reserved.
 #pragma once
 
-#include "CoreMinimal.h"
-#include "Templates/UnrealTemplate.h"
-
 #include "Chaos/Core.h"
-#include "Chaos/Vector.h"
-
 
 namespace Chaos
 {
diff --git a/Engine/Source/Runtime/Experimental/Chaos/Public/Chaos/Collision/PBDCollisionConstraint.h b/Engine/Source/Runtime/Experimental/Chaos/Public/Chaos/Collision/PBDCollisionConstraint.h
--- a/Engine/Source/Runtime/Experimental/Chaos/Public/Chaos/Collision/PBDCollisionConstraint.h
+++ b/Engine/Source/Runtime/Experimental/Chaos/Public/Chaos/Collision/PBDCollisionConstraint.h
@@ -374,6 +374,8 @@ namespace Chaos
 
 		void SetModifierApplied() { Flags.bModifierApplied = true; }
 
+		const FPBDCollisionConstraintMaterial& GetCollisionMaterial() const { return Material; }
+
 		void SetInvMassScale0(const FReal InInvMassScale) { Material.InvMassScale0 = FRealSingle(InInvMassScale); }
 		FReal GetInvMassScale0() const { return Material.InvMassScale0; }
 
@@ -557,12 +559,10 @@ namespace Chaos
 		// The GJK warm-start data. This is updated directly in the narrow phase
 		FGJKSimplexData& GetGJKWarmStartData() { return GJKWarmStartData; }
 
-		FSolverBody* GetSolverBody0() { return SolverBodies[0]; }
 		const FSolverBody* GetSolverBody0() const { return SolverBodies[0]; }
-		FSolverBody* GetSolverBody1() { return SolverBodies[1]; }
 		const FSolverBody* GetSolverBody1() const { return SolverBodies[1]; }
 
-		void SetSolverBodies(FSolverBody* InSolverBody0, FSolverBody* InSolverBody1)
+		void SetSolverBodies(const FSolverBody* InSolverBody0, const FSolverBody* InSolverBody1)
 		{
 			SolverBodies[0] = InSolverBody0;
 			SolverBodies[1] = InSolverBody1;
@@ -785,7 +785,7 @@ namespace Chaos
 				uint32 bIsProbeUnmodified : 1;  // Is this constraint a probe pre-contact-modification
 				uint32 bIsProbe : 1;            // Is this constraint currently a probe
 				uint32 bCCDEnabled : 1;			// Is CCD enabled for the current tick
-				uint32 bModifierApplied;		// Was a constraint modifier applied this tick
+				uint32 bModifierApplied : 1;	// Was a constraint modifier applied this tick
 			};
 			uint32 Bits;
 		} Flags;
@@ -795,8 +795,8 @@ namespace Chaos
 		FVec3f LastShapeWorldPositionDelta;
 		FRotation3f LastShapeWorldRotationDelta;
 
-		// These are only needed here while we still have the legacy solvers (not QuasiPBD)
-		FSolverBody* SolverBodies[2];
+		// @todo(chaos): These are only needed here to support incremental collision detection. We should pass them to the functions instead
+		const FSolverBody* SolverBodies[2];
 
 		// Simplex data from the last call to GJK, used to warm-start GJK
 		FGJKSimplexData GJKWarmStartData;
diff --git a/Engine/Source/Runtime/Experimental/Chaos/Public/Chaos/Collision/PBDCollisionConstraintHandle.h b/Engine/Source/Runtime/Experimental/Chaos/Public/Chaos/Collision/PBDCollisionConstraintHandle.h
--- a/Engine/Source/Runtime/Experimental/Chaos/Public/Chaos/Collision/PBDCollisionConstraintHandle.h
+++ b/Engine/Source/Runtime/Experimental/Chaos/Public/Chaos/Collision/PBDCollisionConstraintHandle.h
@@ -86,9 +86,6 @@ namespace Chaos
 		// Declared final so that TPBDConstraintGraphRuleImpl::AddToGraph() does not need to hit vtable
 		virtual FParticlePair GetConstrainedParticles() const override final;
 
-		FSolverBody* GetSolverBody0();
-		FSolverBody* GetSolverBody1();
-
 		const FPBDCollisionConstraints* ConcreteContainer() const;
 		FPBDCollisionConstraints* ConcreteContainer();
 
diff --git a/Engine/Source/Runtime/Experimental/Chaos/Public/Chaos/ConstraintHandle.h b/Engine/Source/Runtime/Experimental/Chaos/Public/Chaos/ConstraintHandle.h
--- a/Engine/Source/Runtime/Experimental/Chaos/Public/Chaos/ConstraintHandle.h
+++ b/Engine/Source/Runtime/Experimental/Chaos/Public/Chaos/ConstraintHandle.h
@@ -163,6 +163,10 @@ namespace Chaos
 		template<typename T>  T* As();
 		template<typename T>  const T* As() const;
 
+		// For use when you absolutely know the type (asserted in non-shipping)
+		template<typename T>  T* AsUnsafe() { check(As<T>() != nullptr); return static_cast<T*>(this); }
+		template<typename T>  const T* AsUnsafe() const { check(As<T>() != nullptr); return static_cast<const T*>(this); }
+
 		const FConstraintHandleTypeID& GetType() const;
 
 		static const FConstraintHandleTypeID& StaticType()
diff --git a/Engine/Source/Runtime/Experimental/Chaos/Public/Chaos/Evolution/IndexedConstraintContainer.h b/Engine/Source/Runtime/Experimental/Chaos/Public/Chaos/Evolution/IndexedConstraintContainer.h
--- a/Engine/Source/Runtime/Experimental/Chaos/Public/Chaos/Evolution/IndexedConstraintContainer.h
+++ b/Engine/Source/Runtime/Experimental/Chaos/Public/Chaos/Evolution/IndexedConstraintContainer.h
@@ -152,9 +152,8 @@ namespace Chaos
 		{
 			for (const FPBDIslandConstraint& Constraint : Constraints)
 			{
-				// @todo(chaos): unchecked test/shipping version of As() - we know the type here
-				const FIndexedConstraintHandle* IndexedConstraintHandle = Constraint.GetConstraint()->As<FIndexedConstraintHandle>();
-				check(IndexedConstraintHandle != nullptr);
+				// We will only ever be given constraints from our container (asserts in non-shipping)
+				const FIndexedConstraintHandle* IndexedConstraintHandle = Constraint.GetConstraint()->AsUnsafe<FIndexedConstraintHandle>();
 
 				int32 ConstraintIndex = IndexedConstraintHandle->GetConstraintIndex();
 
diff --git a/Engine/Source/Runtime/Experimental/Chaos/Public/Chaos/Framework/UncheckedArray.h b/Engine/Source/Runtime/Experimental/Chaos/Public/Chaos/Framework/UncheckedArray.h
--- a/Engine/Source/Runtime/Experimental/Chaos/Public/Chaos/Framework/UncheckedArray.h
+++ b/Engine/Source/Runtime/Experimental/Chaos/Public/Chaos/Framework/UncheckedArray.h
@@ -58,7 +58,7 @@ namespace Chaos
 		*/
 		inline void SetNum(const int32 InNum)
 		{
-			check(InNum < MaxElements);
+			check(InNum <= MaxElements);
 			NumElements = InNum;
 		}
 
diff --git a/Engine/Source/Runtime/Experimental/Chaos/Public/Chaos/PBDCollisionConstraints.h b/Engine/Source/Runtime/Experimental/Chaos/Public/Chaos/PBDCollisionConstraints.h
--- a/Engine/Source/Runtime/Experimental/Chaos/Public/Chaos/PBDCollisionConstraints.h
+++ b/Engine/Source/Runtime/Experimental/Chaos/Public/Chaos/PBDCollisionConstraints.h
@@ -348,14 +348,4 @@ inline FParticlePair FPBDCollisionConstraintHandle::GetConstrainedParticles() co
 	return { GetContact().GetParticle0(), GetContact().GetParticle1() };
 }
 
-inline FSolverBody* FPBDCollisionConstraintHandle::GetSolverBody0()
-{
-	return GetContact().GetSolverBody0();
-}
-
-inline FSolverBody* FPBDCollisionConstraintHandle::GetSolverBody1()
-{
-	return GetContact().GetSolverBody1();
-}
-
 }
