Chaos - fix CitySample vehicle waking up
- doin't add particles that start disabled to the constraint graph
- don't add joint connected to a disabled particle to the graph

diff --git a/Engine/Source/Programs/HeadlessChaos/Private/GeometryCollection/GeometryCollectionTestSimulationSolver.cpp b/Engine/Source/Programs/HeadlessChaos/Private/GeometryCollection/GeometryCollectionTestSimulationSolver.cpp
--- a/Engine/Source/Programs/HeadlessChaos/Private/GeometryCollection/GeometryCollectionTestSimulationSolver.cpp
+++ b/Engine/Source/Programs/HeadlessChaos/Private/GeometryCollection/GeometryCollectionTestSimulationSolver.cpp
@@ -115,7 +115,7 @@ using namespace ChaosTest;
 			EXPECT_EQ(Parent[2],-1);
 
 			// Set the one cluster to disabled
-			Collection->PhysObject->GetSolverClusterHandles()[0]->SetDisabled(true);
+			UnitTest.Solver->GetEvolution()->DisableParticle(Collection->PhysObject->GetSolverClusterHandles()[0]);
 
 			Transform = &Collection->DynamicCollection->Transform;
 			StartingRigidDistance = ((*Transform)[1].GetTranslation() - (*Transform)[0].GetTranslation()).Size();
diff --git a/Engine/Source/Runtime/Experimental/Chaos/Private/Chaos/Island/IslandManager.cpp b/Engine/Source/Runtime/Experimental/Chaos/Private/Chaos/Island/IslandManager.cpp
--- a/Engine/Source/Runtime/Experimental/Chaos/Private/Chaos/Island/IslandManager.cpp
+++ b/Engine/Source/Runtime/Experimental/Chaos/Private/Chaos/Island/IslandManager.cpp
@@ -335,6 +335,10 @@ void FPBDIslandManager::GraphEdgeRemoved(const FConstraintHandleHolder& Constrai
 
 int32 FPBDIslandManager::AddParticle(FGeometryParticleHandle* ParticleHandle, const FGeometryParticleHandle* ParentParticleHandle)
 {
+#if CHAOS_CONSTRAINTHANDLE_DEBUG_ENABLED
+	ensure(!FConstGenericParticleHandle(ParticleHandle)->Disabled());
+#endif
+
 	int32 IslandIndex = INDEX_NONE;
 
 	// If we have a parent, assign the child to the same island in the same sleep state
@@ -916,6 +920,17 @@ bool FPBDIslandManager::DebugCheckParticlesInGraph(const TParticleView<FPBDRigid
 	// We used to loop over all particles every trick and re-add them, but now we only add/remove when necessary.
 	bool bSuccess = true;
 
+	// All particles in the graph should be enabled
+	for (const FGraphNode& Node : IslandGraph->GraphNodes)
+	{
+		FPBDRigidParticleHandle* Rigid = Node.NodeItem->CastToRigidParticle();
+		if (Rigid != nullptr)
+		{
+			bSuccess = bSuccess && !Rigid->Disabled();
+		}
+		ensure(bSuccess);
+	}
+
 	// All dynamic particles should be in the graph
 	for (const auto& ParticleHandle : Particles)
 	{
diff --git a/Engine/Source/Runtime/Experimental/Chaos/Private/Chaos/PBDJointConstraints.cpp b/Engine/Source/Runtime/Experimental/Chaos/Private/Chaos/PBDJointConstraints.cpp
--- a/Engine/Source/Runtime/Experimental/Chaos/Private/Chaos/PBDJointConstraints.cpp
+++ b/Engine/Source/Runtime/Experimental/Chaos/Private/Chaos/PBDJointConstraints.cpp
@@ -455,6 +455,10 @@ namespace Chaos
 		ConstraintSettings.AddDefaulted();
 		SetConstraintSettings(ConstraintIndex, InConstraintSettings);
 
+		// If our particle(s) are disabled, so is the constraint for now. It will get enabled if both particles get enabled.
+		const bool bStartDisabled = FConstGenericParticleHandle(InConstrainedParticles[0])->Disabled() || FConstGenericParticleHandle(InConstrainedParticles[1])->Disabled();
+		ConstraintStates[ConstraintIndex].bDisabled = bStartDisabled;
+
 		return Handles.Last();
 	}
 
diff --git a/Engine/Source/Runtime/Experimental/Chaos/Public/Chaos/PBDRigidsEvolution.h b/Engine/Source/Runtime/Experimental/Chaos/Public/Chaos/PBDRigidsEvolution.h
--- a/Engine/Source/Runtime/Experimental/Chaos/Public/Chaos/PBDRigidsEvolution.h
+++ b/Engine/Source/Runtime/Experimental/Chaos/Public/Chaos/PBDRigidsEvolution.h
@@ -429,8 +429,15 @@ public:
 	*/
 	CHAOS_API void RegisterParticle(FGeometryParticleHandle* Particle)
 	{
-		// Add to the graph if necessary
-		ConstraintGraph.AddParticle(Particle);
+		// Add to the graph if necessary. Only enabled dynamic particles are added at this stage. Kinematics
+		// and statics are ignored until referenced by a constraint.
+		if (FPBDRigidParticleHandle* Rigid = Particle->CastToRigidParticle())
+		{
+			if (Rigid->IsDynamic() && !Rigid->Disabled())
+			{
+				ConstraintGraph.AddParticle(Particle);
+			}
+		}
 
 		// Flag as dirty to update the spatial query acceleration structures
 		DirtyParticle(*Particle);
