Chaos - fix CCD TOI being initialized out of float range

diff --git a/Engine/Source/Runtime/Experimental/Chaos/Private/Chaos/CCDUtilities.cpp b/Engine/Source/Runtime/Experimental/Chaos/Private/Chaos/CCDUtilities.cpp
--- a/Engine/Source/Runtime/Experimental/Chaos/Private/Chaos/CCDUtilities.cpp
+++ b/Engine/Source/Runtime/Experimental/Chaos/Private/Chaos/CCDUtilities.cpp
@@ -539,7 +539,7 @@ namespace Chaos
 			// After applying impulse, constraint TOI need be updated to reflect the new velocities. 
 			// Usually the new velocities are separating, and therefore TOI should be infinity.
 			// See resweep below which (optionally) updates TOI for all other contacts as a result of handling this one
-			CCDConstraint->SweptConstraint->SetCCDTimeOfImpact(TNumericLimits<FReal>::Max());
+			CCDConstraint->SweptConstraint->ResetCCDTimeOfImpact();
 
 			bool bMovedParticle0 = false;
 			bool bMovedParticle1 = false;
diff --git a/Engine/Source/Runtime/Experimental/Chaos/Private/Chaos/CollisionResolution.cpp b/Engine/Source/Runtime/Experimental/Chaos/Private/Chaos/CollisionResolution.cpp
--- a/Engine/Source/Runtime/Experimental/Chaos/Private/Chaos/CollisionResolution.cpp
+++ b/Engine/Source/Runtime/Experimental/Chaos/Private/Chaos/CollisionResolution.cpp
@@ -1249,7 +1249,7 @@ namespace Chaos
 		void UpdateGenericConvexConvexConstraintSwept(TGeometryParticleHandle<FReal, 3>* Particle0, const FImplicitObject& Implicit0, const FRigidTransform3& StartWorldTransform0, TGeometryParticleHandle<FReal, 3>* Particle1, const FImplicitObject& Implicit1, const FRigidTransform3& StartWorldTransform1, const FVec3& Dir, const FReal Length, const FReal Dt, FPBDCollisionConstraint& Constraint)
 		{
 			CONDITIONAL_SCOPE_CYCLE_COUNTER(STAT_Collisions_UpdateGenericConvexConvexConstraintSwept, ConstraintsDetailedStats);
-			FReal TOI = TNumericLimits<FReal>::Max();
+			FReal TOI = FPBDCollisionConstraint::MaxTOI;
 			const FRigidTransform3& EndWorldTransform0 = Constraint.GetShapeWorldTransform0();
 			const FRigidTransform3& EndWorldTransform1 = Constraint.GetShapeWorldTransform1();
 			UpdateContactPointNoCull(Constraint, GenericConvexConvexContactPointSwept(Implicit0, StartWorldTransform0, EndWorldTransform0, Implicit1, StartWorldTransform1, EndWorldTransform1, Dir, Length, Constraint.GetCCDEnablePenetration(), Constraint.GetCCDTargetPenetration(), TOI), Dt);
@@ -1661,7 +1661,7 @@ namespace Chaos
 			const FVec3 DeltaX1 = Constraint.GetShapeWorldTransform1().GetTranslation() - WorldTransform1.GetLocation();
 			bool bUseCCD = ShouldUseCCD(Particle0, DeltaX0, Particle1, DeltaX1, DirCCD, LengthCCD);
 
-			Constraint.SetCCDTimeOfImpact(TNumericLimits<FReal>::Max());
+			Constraint.ResetCCDTimeOfImpact();
 
 			if (bUseCCD)
 			{
diff --git a/Engine/Source/Runtime/Experimental/Chaos/Public/Chaos/Collision/PBDCollisionConstraint.h b/Engine/Source/Runtime/Experimental/Chaos/Public/Chaos/Collision/PBDCollisionConstraint.h
--- a/Engine/Source/Runtime/Experimental/Chaos/Public/Chaos/Collision/PBDCollisionConstraint.h
+++ b/Engine/Source/Runtime/Experimental/Chaos/Public/Chaos/Collision/PBDCollisionConstraint.h
@@ -234,6 +234,8 @@ namespace Chaos
 
 		static const int32 MaxManifoldPoints = 4;
 
+		static constexpr FRealSingle MaxTOI = std::numeric_limits<FRealSingle>::max();
+
 		/**
 		 * @brief Create a contact constraint
 		 * Initializes a constraint stored inline in an object. Only intended to be called once right after construction.
@@ -604,7 +606,12 @@ namespace Chaos
 		/**
 		 * @brief Set the CCD TOI from the collision detection sweep
 		*/
-		void SetCCDTimeOfImpact(const FReal TOI) { CCDTimeOfImpact = FRealSingle(TOI); }
+		void SetCCDTimeOfImpact(const FReal TOI) { check(TOI <= FReal(MaxTOI)); CCDTimeOfImpact = FRealSingle(TOI); }
+
+		/**
+		 * Initialize the CCD TOI to the highest possible value
+		 */
+		void ResetCCDTimeOfImpact() { CCDTimeOfImpact = MaxTOI; }
 
 		/**
 		 * \brief Store the results of CCD contact resolution, if active
