Geometry: Voronoi - ue-truncation-fixes. Enable type cast warnings.

diff --git a/Engine/Source/Runtime/Experimental/Voronoi/Private/voro++/src/cell.cc b/Engine/Source/Runtime/Experimental/Voronoi/Private/voro++/src/cell.cc
--- a/Engine/Source/Runtime/Experimental/Voronoi/Private/voro++/src/cell.cc
+++ b/Engine/Source/Runtime/Experimental/Voronoi/Private/voro++/src/cell.cc
@@ -16,6 +16,10 @@
 #include "config.hh"
 #include "common.hh"
 #include "cell.hh"
+#include "HAL/PlatformMisc.h"
+
+
+PRAGMA_DISABLE_UNSAFE_TYPECAST_WARNINGS
 
 namespace voro {
 
@@ -2222,3 +2226,4 @@ template void voronoicell_base::check_memory_for_copy(voronoicell&,voronoicell_b
 template void voronoicell_base::check_memory_for_copy(voronoicell_neighbor&,voronoicell_base*);
 
 }
+PRAGMA_RESTORE_UNSAFE_TYPECAST_WARNINGS
\ No newline at end of file
diff --git a/Engine/Source/Runtime/Experimental/Voronoi/Voronoi.Build.cs b/Engine/Source/Runtime/Experimental/Voronoi/Voronoi.Build.cs
--- a/Engine/Source/Runtime/Experimental/Voronoi/Voronoi.Build.cs
+++ b/Engine/Source/Runtime/Experimental/Voronoi/Voronoi.Build.cs
@@ -14,8 +14,11 @@ namespace UnrealBuildTool.Rules
                 }
                 );
 
+			UnsafeTypeCastWarningLevel = WarningLevel.Warning;
+
 			// Disable all static analysis checkers for this module ( third party code ) 
 			StaticAnalyzerDisabledCheckers.Add("all");
+
 		}
 	}
 }
