Added ensures NetGuidCache::Supports and ObjectRefereceCache to not allow repliation of references to objects pointers that belong to another pie instance

Added guards in GameplayDebuggerLocalController to not forward calls to rpc via delegates from other pie instances.

diff --git a/Engine/Source/Developer/GameplayDebugger/Private/GameplayDebuggerLocalController.cpp b/Engine/Source/Developer/GameplayDebugger/Private/GameplayDebuggerLocalController.cpp
--- a/Engine/Source/Developer/GameplayDebugger/Private/GameplayDebuggerLocalController.cpp
+++ b/Engine/Source/Developer/GameplayDebugger/Private/GameplayDebuggerLocalController.cpp
@@ -769,6 +769,12 @@ void UGameplayDebuggerLocalController::OnSelectionChanged(UObject* Object)
 
 		if (SelectedActor)
 		{
+			// Since this is used from a delegate we need to filter out objects not belonging to the same pie instance/world
+			if (SelectedActor && SelectedActor->GetWorld() != GetWorld())	
+			{
+				return;
+			}
+
 			CachedReplicator->SetDebugActor(SelectedActor, false);
 			CachedReplicator->CollectCategoryData(/*bForce=*/true);
 		}
@@ -777,6 +783,12 @@ void UGameplayDebuggerLocalController::OnSelectionChanged(UObject* Object)
 
 void UGameplayDebuggerLocalController::OnSelectedObject(UObject* Object)
 {
+	// Since this is used from a delegate we need to filter out objects not belonging to the same pie instance/world
+	if (Object && Object->GetWorld() != GetWorld())	
+	{
+		return;
+	}
+
 	AController* SelectedController = Cast<AController>(Object);
 	APawn* SelectedPawn = SelectedController ? SelectedController->GetPawn() : Cast<APawn>(Object);
 	if (CachedReplicator && SelectedPawn && SelectedPawn->IsSelected())
diff --git a/Engine/Source/Runtime/Engine/Private/PackageMapClient.cpp b/Engine/Source/Runtime/Engine/Private/PackageMapClient.cpp
--- a/Engine/Source/Runtime/Engine/Private/PackageMapClient.cpp
+++ b/Engine/Source/Runtime/Engine/Private/PackageMapClient.cpp
@@ -2689,6 +2689,21 @@ bool FNetGUIDCache::SupportsObject( const UObject* Object, const TWeakObjectPtr<
 		return true;
 	}
 
+#if WITH_EDITOR
+	const UPackage* ObjectPackage = Object->GetPackage();
+	if (ObjectPackage->HasAnyPackageFlags(PKG_PlayInEditor))
+	{
+		const int32 DriverPIEInstanceID = Driver->GetWorld() ? Driver->GetWorld()->GetPackage()->GetPIEInstanceID() : INDEX_NONE;
+		const int32 ObjectPIEInstanceID = ObjectPackage->GetPIEInstanceID();
+
+		if (!ensureAlwaysMsgf(DriverPIEInstanceID == ObjectPIEInstanceID, TEXT("FNetGUIDCache::SupportsObject: Object %s is not supported since its PIE InstanceID: %d differs from the one of the NetDriver's world PIE InstanceID: %d, it will replicate as an invalid reference."), *GetPathNameSafe(Object), ObjectPIEInstanceID, DriverPIEInstanceID))
+		{
+			// Don't replicate references to objects owned by other PIE instances.
+			return false;
+		}
+	}
+#endif
+
 	if ( Object->IsFullNameStableForNetworking() )
 	{
 		// If object is fully net addressable, it's definitely supported
diff --git a/Engine/Source/Runtime/Experimental/Iris/Core/Private/Iris/ReplicationSystem/ObjectReferenceCache.cpp b/Engine/Source/Runtime/Experimental/Iris/Core/Private/Iris/ReplicationSystem/ObjectReferenceCache.cpp
--- a/Engine/Source/Runtime/Experimental/Iris/Core/Private/Iris/ReplicationSystem/ObjectReferenceCache.cpp
+++ b/Engine/Source/Runtime/Experimental/Iris/Core/Private/Iris/ReplicationSystem/ObjectReferenceCache.cpp
@@ -12,6 +12,7 @@
 #include "Iris/Serialization/NetBitStreamReader.h"
 #include "Iris/Serialization/ObjectNetSerializer.h"
 #include "UObject/Package.h"
+#include "Misc/CoreMiscDefines.h"
 #include "Net/Core/Trace/NetDebugName.h"
 #include "Net/Core/Trace/NetTrace.h"
 #include "Core/Public/Misc/StringBuilder.h"
@@ -220,6 +221,21 @@ bool FObjectReferenceCache::CreateObjectReferenceInternal(const UObject* Object,
 		}
 	}
 
+#if WITH_EDITOR
+	const UPackage* ObjectPackage = Object->GetPackage();
+	if (ObjectPackage->HasAnyPackageFlags(PKG_PlayInEditor))
+	{
+		const int32 ReplicationSystemPIEInstanceID = ReplicationSystem->GetPIEInstanceID();
+		const int32 ObjectPIEInstanceID = ObjectPackage->GetPIEInstanceID();
+
+		if (!ensureAlwaysMsgf(ReplicationSystemPIEInstanceID == ObjectPIEInstanceID, TEXT("FObjectRefereceCache::CreateObjectReferenceInternal: Object %s is not supported since its PIE InstanceID: %d differs from the one of the NetDriver's world PIE InstanceID: %d, it will replicate as an invalid reference."), *GetPathNameSafe(Object), ObjectPIEInstanceID, ReplicationSystemPIEInstanceID))
+		{
+			// Don't replicate references to objects owned by other PIE instances.
+			return false;
+		}
+	}
+#endif
+
 	// Translated from GuidCache, but kept as separate variable to avoid calling the same functions multiple times
 	const bool bIsFullNameStableForNetworking = Object->IsFullNameStableForNetworking();
 	const bool bIsDynamic = !bIsFullNameStableForNetworking;
diff --git a/Engine/Source/Runtime/Experimental/Iris/Core/Private/Iris/ReplicationSystem/ReplicationSystem.cpp b/Engine/Source/Runtime/Experimental/Iris/Core/Private/Iris/ReplicationSystem/ReplicationSystem.cpp
--- a/Engine/Source/Runtime/Experimental/Iris/Core/Private/Iris/ReplicationSystem/ReplicationSystem.cpp
+++ b/Engine/Source/Runtime/Experimental/Iris/Core/Private/Iris/ReplicationSystem/ReplicationSystem.cpp
@@ -437,6 +437,7 @@ UReplicationSystem::UReplicationSystem()
 : Super()
 , Impl(nullptr)
 , Id(~0U)
+, PIEInstanceID(INDEX_NONE)
 , bIsServer(0)
 , bAllowObjectReplication(0)
 , bDoCollectGarbage(0)
diff --git a/Engine/Source/Runtime/Experimental/Iris/Core/Public/Iris/ReplicationSystem/ReplicationSystem.h b/Engine/Source/Runtime/Experimental/Iris/Core/Public/Iris/ReplicationSystem/ReplicationSystem.h
--- a/Engine/Source/Runtime/Experimental/Iris/Core/Public/Iris/ReplicationSystem/ReplicationSystem.h
+++ b/Engine/Source/Runtime/Experimental/Iris/Core/Public/Iris/ReplicationSystem/ReplicationSystem.h
@@ -486,6 +486,8 @@ public:
 	 */
 	IRISCORE_API UObject* GetConnectionUserData(uint32 ConnectionId) const;
 
+	IRISCORE_API int32 GetPIEInstanceID() const { return PIEInstanceID; }
+
 public:
 	// For internal use and not exported.
 
@@ -510,6 +512,8 @@ private:
 	IRISCORE_API void ResetGameWorldState();
 	IRISCORE_API void NotifyStreamingLevelUnload(const UObject* Level);
 
+	IRISCORE_API void SetPIEInstanceID(int32 InPIEInstanceID) { PIEInstanceID = InPIEInstanceID; }
+
 private:
 
 	void PostGarbageCollection();
@@ -523,6 +527,7 @@ private:
 	TObjectPtr<UReplicationBridge> ReplicationBridge;
 
 	uint32 Id;
+	int32 PIEInstanceID;
 	uint32 bIsServer : 1;
 	uint32 bAllowObjectReplication : 1;
 	uint32 bDoCollectGarbage : 1;
