Clang Static Analysis Fix:
Non initialized variable in TriangleMeshImplicitObject

diff --git a/Engine/Source/Runtime/Experimental/Chaos/Private/Chaos/TriangleMeshImplicitObject.cpp b/Engine/Source/Runtime/Experimental/Chaos/Private/Chaos/TriangleMeshImplicitObject.cpp
--- a/Engine/Source/Runtime/Experimental/Chaos/Private/Chaos/TriangleMeshImplicitObject.cpp
+++ b/Engine/Source/Runtime/Experimental/Chaos/Private/Chaos/TriangleMeshImplicitObject.cpp
@@ -1171,6 +1171,10 @@ void ComputeScaledSweepInputs(FVec3 TriMeshScale, const FRigidTransform3& StartT
 		const FReal LengthScaleInv = 1.f / LengthScale;
 		OutScaledDirNormalized = UnscaledDirDenorm * LengthScaleInv;
 	}
+	else
+	{
+		OutScaledDirNormalized = FVec3(1.0, 0.0, 0.0);
+	}
 
 
 	OutLengthScale = LengthScale;
