Chaos - fix GenericParticleHandle accessors for sleeping particles

diff --git a/Engine/Source/Runtime/Experimental/Chaos/Public/Chaos/ParticleHandle.h b/Engine/Source/Runtime/Experimental/Chaos/Public/Chaos/ParticleHandle.h
--- a/Engine/Source/Runtime/Experimental/Chaos/Public/Chaos/ParticleHandle.h
+++ b/Engine/Source/Runtime/Experimental/Chaos/Public/Chaos/ParticleHandle.h
@@ -1423,7 +1423,7 @@ public:
 	// TODO: Make all of these check ObjectState to maintain current functionality
 	int32 CollisionParticlesSize() const
 	{
-		if (MHandle->CastToRigidParticle() && MHandle->ObjectState() == EObjectStateType::Dynamic)
+		if (MHandle->CastToRigidParticle())
 		{
 			return MHandle->CastToRigidParticle()->CollisionParticlesSize();
 		}
@@ -1433,7 +1433,7 @@ public:
 
 	const TUniquePtr<FBVHParticles>& CollisionParticles() const
 	{
-		if (MHandle->CastToRigidParticle() && MHandle->ObjectState() == EObjectStateType::Dynamic)
+		if (MHandle->CastToRigidParticle())
 		{
 			return MHandle->CastToRigidParticle()->CollisionParticles();
 		}
@@ -1463,7 +1463,7 @@ public:
 
 	bool HasCollisionConstraintFlag(const ECollisionConstraintFlags Flag)  const
 	{
-		if (MHandle->CastToRigidParticle() && (MHandle->ObjectState() == EObjectStateType::Dynamic || MHandle->ObjectState() == EObjectStateType::Sleeping))
+		if (MHandle->CastToRigidParticle())
 		{
 			return MHandle->CastToRigidParticle()->HasCollisionConstraintFlag(Flag);
 		}
@@ -1475,7 +1475,7 @@ public:
 	// @todo(ccaulfield): should be available on all types?
 	bool Disabled() const
 	{
-		if (MHandle->CastToRigidParticle() && MHandle->ObjectState() == EObjectStateType::Dynamic)
+		if (MHandle->CastToRigidParticle())
 		{
 			return MHandle->CastToRigidParticle()->Disabled();
 		}
@@ -1521,7 +1521,7 @@ public:
 
 	FVec3& P()
 	{
-		if (MHandle->CastToRigidParticle() && MHandle->ObjectState() == EObjectStateType::Dynamic)
+		if (IsDynamic())
 		{
 			return MHandle->CastToRigidParticle()->P();
 		}
@@ -1531,7 +1531,7 @@ public:
 
 	const FVec3& P() const
 	{
-		if (MHandle->CastToRigidParticle() && MHandle->ObjectState() == EObjectStateType::Dynamic)
+		if (IsDynamic())
 		{
 			return MHandle->CastToRigidParticle()->P();
 		}
@@ -1541,7 +1541,7 @@ public:
 
 	FRotation3& Q()
 	{
-		if (MHandle->CastToRigidParticle() && MHandle->ObjectState() == EObjectStateType::Dynamic)
+		if (IsDynamic())
 		{
 			return MHandle->CastToRigidParticle()->Q();
 		}
@@ -1551,7 +1551,7 @@ public:
 
 	const FRotation3& Q() const
 	{
-		if (MHandle->CastToRigidParticle() && MHandle->ObjectState() == EObjectStateType::Dynamic)
+		if (IsDynamic())
 		{
 			return MHandle->CastToRigidParticle()->Q();
 		}
@@ -1581,7 +1581,7 @@ public:
 
 	const FVec3& Acceleration() const
 	{ 
-		if (MHandle->CastToRigidParticle() && MHandle->ObjectState() == EObjectStateType::Dynamic)
+		if (IsDynamic())
 		{
 			return MHandle->CastToRigidParticle()->Acceleration();
 		}
@@ -1590,7 +1590,7 @@ public:
 	}
 	const FVec3& AngularAcceleration() const
 	{ 
-		if (MHandle->CastToRigidParticle() && MHandle->ObjectState() == EObjectStateType::Dynamic)
+		if (IsDynamic())
 		{
 			return MHandle->CastToRigidParticle()->AngularAcceleration();
 		}
