Chaos - fix CCD issue. CCD objects can't be pushed through walls by a heavy object.

diff --git a/Engine/Source/Runtime/Experimental/Chaos/Private/Chaos/CCDUtilities.cpp b/Engine/Source/Runtime/Experimental/Chaos/Private/Chaos/CCDUtilities.cpp
--- a/Engine/Source/Runtime/Experimental/Chaos/Private/Chaos/CCDUtilities.cpp
+++ b/Engine/Source/Runtime/Experimental/Chaos/Private/Chaos/CCDUtilities.cpp
@@ -52,6 +52,14 @@ namespace Chaos
 		bool bCCDSweepsUseProbeShapes = false;
 		FAutoConsoleVariableRef  CVarCCDSweepsUseProbeShapes(TEXT("p.Chaos.CCD.CCDSweepsUseProbeShapes"), bCCDSweepsUseProbeShapes , TEXT("When true, probe shapes can be swept for more accurate collision detection."));
 
+		// How many post-solve CCD correction iterations to run
+		int32 ChaosCollisionCCDCorrectionIterations = 4;
+		FAutoConsoleVariableRef CVarChaosCollisionCCDCorrectionIterations(TEXT("p.Chaos.Collision.CCD.CorrectionIterations"), ChaosCollisionCCDCorrectionIterations, TEXT("The number of post-solve CCD correction ietaryions to run."));
+
+		// A multiplier on the constraint CCD threshold that determines how much penetration we allow in the correction phase
+		FRealSingle ChaosCollisionCCDCorrectionPhiToleranceScale = 0.02f;
+		FAutoConsoleVariableRef CVarChaosCollisionCCDCorrectionPhiToleranceScale(TEXT("p.Chaos.Collision.CCD.CorrectionPhiToleranceScale"), ChaosCollisionCCDCorrectionPhiToleranceScale, TEXT("How much penetration we allow during the correction phase (multiplier on shape size)"));
+
 		extern int32 ChaosSolverDrawCCDInteractions;
 
 #if CHAOS_DEBUG_DRAW
@@ -87,7 +95,7 @@ namespace Chaos
 
 	void FCCDManager::ApplyConstraintsPhaseCCD(const FReal Dt, FCollisionConstraintAllocator *CollisionAllocator, const int32 NumDynamicParticles)
 	{
-		SweptConstraints = CollisionAllocator->GetSweptConstraints();
+		SweptConstraints = CollisionAllocator->GetCCDConstraints();
 		if (SweptConstraints.Num() > 0)
 		{
 			ApplySweptConstraints(Dt, SweptConstraints, NumDynamicParticles);
@@ -143,6 +151,12 @@ namespace Chaos
 				continue;
 			}
 
+			// If we don't need to sweep for this constraint, skip it
+			if (!Constraint->GetCCDSweepEnabled())
+			{
+				continue;
+			}
+
 			// For now, probe constraints do not work with CCD
 			if (Constraint->IsProbe())
 			{
@@ -871,6 +885,11 @@ namespace Chaos
 				continue;
 			}
 
+			if (!SweptConstraint->GetCCDSweepEnabled())
+			{
+				continue;
+			}
+
 			const FConstGenericParticleHandle P0 = FConstGenericParticleHandle(SweptConstraint->GetParticle0());
 			const FConstGenericParticleHandle P1 = FConstGenericParticleHandle(SweptConstraint->GetParticle1());
 			SweptConstraint->ResetManifold();
@@ -884,6 +903,146 @@ namespace Chaos
 		}
 	}
 
+	void FCCDManager::ApplyCorrections(const FReal Dt)
+	{
+		// Build a list of CCD constraints sorted so that the most important ones are last.
+		// The primary goal is to prevent CCD objects from leaving the world, so we process collisions
+		// with statics/kinematics last. The secondary goal is to prevent CCD objects passing through each 
+		// other. And finally we try to prevent CCD object passing through non-CCD dynamics but aren't
+		// too worried if that happens.
+		// Sort order: 
+		//		CCD - Dynamic
+		//		CCD - CCD
+		//		CCD - Kinematic
+		TArray<FPBDCollisionConstraint*> Constraints;
+		Constraints.Append(SweptConstraints);
+		Constraints.Sort([](const FPBDCollisionConstraint& L, const FPBDCollisionConstraint& R)
+			{
+				const FGenericParticleHandle LP0 = L.GetParticle0();
+				const FGenericParticleHandle LP1 = L.GetParticle1();
+				const FGenericParticleHandle RP0 = R.GetParticle0();
+				const FGenericParticleHandle RP1 = R.GetParticle1();
+
+				const bool bLCCD0 = LP0->CCDEnabled();
+				const bool bLCCD1 = LP1->CCDEnabled();
+				const bool bRCCD0 = RP0->CCDEnabled();
+				const bool bRCCD1 = RP1->CCDEnabled();
+				const bool bLDynamic0 = LP0->IsDynamic();
+				const bool bLDynamic1 = LP1->IsDynamic();
+				const bool bRDynamic0 = RP0->IsDynamic();
+				const bool bRDynamic1 = RP1->IsDynamic();
+
+				// If only one constraint has a non-ccd-dynamic it goes before the other
+				const bool bLHasNonCCDDynamic = (!bLCCD0 && bLDynamic0) || (!bLCCD1 && bLDynamic1);
+				const bool bRHasNonCCDDynamic = (!bRCCD0 && bRDynamic0) || (!bRCCD1 && bRDynamic1);
+				if (bLHasNonCCDDynamic != bRHasNonCCDDynamic)
+				{
+					return bLHasNonCCDDynamic;
+				}
+
+				// If only one constraint has a non-ccd-kinematic it goes after the other
+				const bool bLHasNonCCDKinematic = (!bLCCD0 && !bLDynamic0) || (!bLCCD1 && !bLDynamic1);
+				const bool bRHasNonCCDKinematic = (!bRCCD0 && !bRDynamic0) || (!bRCCD1 && !bRDynamic1);
+				if (bLHasNonCCDKinematic != bRHasNonCCDKinematic)
+				{
+					return bRHasNonCCDKinematic;
+				}
+
+				// Otherwise we don't care about the order
+				return false;
+			});
+
+
+		// Loop over all the CCD constraints and if the penetration depth is greater than some threshold,
+		// move the two bodies so that they are no longer penetrating.
+		// We iterate to handle stacks of CCD objects and extracting a CCD objects from a wedge.
+		const FReal PhiToleranceScale = CVars::ChaosCollisionCCDCorrectionPhiToleranceScale;
+		const int32 MaxIterations = CVars::ChaosCollisionCCDCorrectionIterations;
+		int32 NumIterations = 0;
+		bool bSolved = false;
+		while (!bSolved && (NumIterations < MaxIterations))
+		{
+			bSolved = true;
+			++NumIterations;
+
+			for (FPBDCollisionConstraint* Constraint : Constraints)
+			{
+				if (!Constraint->IsEnabled())
+				{
+					continue;
+				}
+
+				const FGenericParticleHandle P0 = Constraint->GetParticle0();
+				const FGenericParticleHandle P1 = Constraint->GetParticle1();
+				const bool bDynamic0 = P0->IsDynamic();
+				const bool bDynamic1 = P1->IsDynamic();
+				const bool bCCD0 = P0->CCDEnabled();
+				const bool bCCD1 = P1->CCDEnabled();
+
+				// Skip no CCD (should not happen)
+				if (!bCCD0 && !bCCD1)
+				{
+					continue;
+				}
+
+				// Skip two non-dynamics (should not happen)
+				if (!bDynamic0 && !bDynamic1)
+				{
+					continue;
+				}
+
+				// For two dynamics we move each body 50% (we know we have at least one dynamic here)
+				FReal Bias0 = FReal(0.5);
+				if (!bDynamic0)
+				{
+					// Only body1 moves
+					Bias0 = FReal(0);
+				}
+				else if (!bDynamic1)
+				{
+					// Only body0 moves
+					Bias0 = FReal(1);
+				}
+				const FReal Bias1 = FReal(1) - Bias0;
+
+				// This function is called after the solver phase - the shape transform and Phi are out of date
+				const FRigidTransform3 ShapeWorldTransform0 = Constraint->GetShapeRelativeTransform0() * FRigidTransform3(P0->P(), P0->Q());
+				const FRigidTransform3 ShapeWorldTransform1 = Constraint->GetShapeRelativeTransform1() * FRigidTransform3(P1->P(), P1->Q());
+				Constraint->SetShapeWorldTransforms(ShapeWorldTransform0, ShapeWorldTransform1);
+				Constraint->UpdateManifoldContacts();
+
+				const FReal Phi = Constraint->GetPhi();
+				const FReal PhiTolerance = PhiToleranceScale * Constraint->GetCCDTargetPenetration();
+
+				if (Phi < -PhiTolerance)
+				{
+					bSolved = false;
+
+					// Extract the two bodies along the normal and set the contact velocity to zero, 
+					// as if we have two point particles and zero restitution.
+					// NOTE: this is ignoring their relative masses and so does not conserve momentum
+					const FVec3 WorldNormal = Constraint->CalculateWorldContactNormal();
+					const FReal VelNormal = FMath::Min(FVec3::DotProduct(P0->V() - P1->V(), WorldNormal), FReal(0));
+
+					if (Bias0 > UE_SMALL_NUMBER)
+					{
+						const FVec3 Correction0 = -Bias0 * Phi * WorldNormal;
+						P0->P() = P0->P() + Correction0;
+						P0->SetV(P0->V() - Bias0 * VelNormal * WorldNormal);
+					}
+
+					if (Bias1 > UE_SMALL_NUMBER)
+					{
+						const FVec3 Correction1 = Bias1 * Phi * WorldNormal;
+						P1->P() = P1->P() + Correction1;
+						P1->SetV(P1->V() + Bias1 * VelNormal * WorldNormal);
+					}
+				}
+			}
+		}
+	}
+
+
 	void FCCDManager::OverwriteXUsingV(const FReal Dt)
 	{
 		// Overwriting X = P - V * Dt so that the implicit velocity step will give our velocity back.
diff --git a/Engine/Source/Runtime/Experimental/Chaos/Private/Chaos/Collision/ParticlePairMidPhase.cpp b/Engine/Source/Runtime/Experimental/Chaos/Private/Chaos/Collision/ParticlePairMidPhase.cpp
--- a/Engine/Source/Runtime/Experimental/Chaos/Private/Chaos/Collision/ParticlePairMidPhase.cpp
+++ b/Engine/Source/Runtime/Experimental/Chaos/Private/Chaos/Collision/ParticlePairMidPhase.cpp
@@ -252,11 +252,12 @@ namespace Chaos
 	}
 
 	int32 FSingleShapePairCollisionDetector::GenerateCollisionCCD(
+		const bool bEnableCCDSweep,
 		const FReal CullDistance,
 		const FReal Dt,
 		const FCollisionContext& Context)
 	{
-		return GenerateCollisionCCDImpl(CullDistance, Dt, Context);
+		return GenerateCollisionCCDImpl(bEnableCCDSweep, CullDistance, Dt, Context);
 	}
 
 	void FSingleShapePairCollisionDetector::CreateConstraint(const FReal CullDistance, const FCollisionContext& Context)
@@ -359,6 +360,7 @@ namespace Chaos
 	}
 
 	int32 FSingleShapePairCollisionDetector::GenerateCollisionCCDImpl(
+		const bool bEnableCCDSweep,
 		const FReal CullDistance, 
 		const FReal Dt,
 		const FCollisionContext& Context)
@@ -372,8 +374,19 @@ namespace Chaos
 		{
 			// Lazy creation of the constraint. 
 			CreateConstraint(CullDistance, Context);
+
+			// Flag this contact as requiring CCD
+			Constraint->SetCCDEnabled(true);
+		}
+
+		// Do we want to enable the CCD sweep? If not, we fall back to the standard collision detection for this tick
+		Constraint->SetCCDSweepEnabled(bEnableCCDSweep);
+		if (!bEnableCCDSweep)
+		{
+			return GenerateCollision(CullDistance, Dt, Context);
 		}
 
+		// Swept collision detection
 		if (Constraint.IsValid())
 		{
 			PHYSICS_CSV_SCOPED_EXPENSIVE(PhysicsVerbose, NarrowPhase_UpdateConstraintCCD);
@@ -389,7 +402,6 @@ namespace Chaos
 			Constraint->SetShapeWorldTransforms(ShapeWorldTransform0, ShapeWorldTransform1);
 
 			Constraint->ResetModifications();
-			Constraint->SetCCDEnabled(true);
 			Constraint->SetCullDistance(CullDistance);
 			Constraint->ResetManifold();
 			Constraint->ResetActiveManifoldContacts();
@@ -408,11 +420,11 @@ namespace Chaos
 			const FRigidTransform3 CCDShapeWorldTransform1 = Constraint->ImplicitTransform[1] * CCDParticleWorldTransform1;
 			Collisions::UpdateConstraintSwept(*Constraint.Get(), CCDShapeWorldTransform0, CCDShapeWorldTransform1, Dt);
 
-			// If we did get a hit but it's at TOI = 1, treat this constraint as a regular non-swept constraint.
+			// If we did get a hit but it's at TOI = 1, treat this constraint as a regular non-swept constraint (skip the rewind)
 			if (Constraint->GetCCDTimeOfImpact() == FReal(1))
 			{
+				Constraint->SetCCDSweepEnabled(false);
 				Collisions::UpdateConstraint(*Constraint.Get(), Constraint->GetShapeWorldTransform0(), Constraint->GetShapeWorldTransform1(), Dt);
-				Constraint->SetCCDEnabled(false);
 			}
 
 			Context.GetAllocator()->ActivateConstraint(Constraint.Get());
@@ -805,7 +817,8 @@ namespace Chaos
 	void FParticlePairMidPhase::Init(
 		FGeometryParticleHandle* InParticle0,
 		FGeometryParticleHandle* InParticle1,
-		const FCollisionParticlePairKey& InKey)
+		const FCollisionParticlePairKey& InKey,
+		const FCollisionContext& Context)
 	{
 		PHYSICS_CSV_SCOPED_EXPENSIVE(PhysicsVerbose, NarrowPhase_Filter);
 
@@ -813,7 +826,7 @@ namespace Chaos
 		Particle1 = InParticle1;
 		Key = InKey;
 
-		Flags.bIsCCD = FConstGenericParticleHandle(Particle0)->CCDEnabled() || FConstGenericParticleHandle(Particle1)->CCDEnabled();
+		Flags.bIsCCD = Context.GetSettings().bAllowCCD && (FConstGenericParticleHandle(Particle0)->CCDEnabled() || FConstGenericParticleHandle(Particle1)->CCDEnabled());
 
 		BuildDetectors();
 
@@ -933,21 +946,21 @@ namespace Chaos
 		// CullDistance is scaled by the size of the dynamic objects.
 		FReal CullDistance = InCullDistance * CullDistanceScale;
 
-		// Enable CCD?
-		const bool bUseCCD = Context.GetSettings().bAllowCCD && Flags.bIsCCD && ShouldEnableCCD(Dt);
-		if (bUseCCD)
+		// Extend cull distance for sweep-enabled CCD collision
+		const bool bUseSweep = Flags.bIsCCD && ShouldEnableCCD(Dt);
+		if (bUseSweep)
 		{
 			const FReal VMax = (FConstGenericParticleHandle(GetParticle0())->V() - FConstGenericParticleHandle(GetParticle1())->V()).GetAbsMax();
 			CullDistance += VMax * Dt;
 		}
-		
+
 		// Run collision detection on all potentially colliding shape pairs
 		NumActiveConstraints = 0;
-		if (bUseCCD)
+		if (Flags.bIsCCD)
 		{
 			for (FSingleShapePairCollisionDetector& ShapePair : ShapePairDetectors)
 			{
-				NumActiveConstraints += ShapePair.GenerateCollisionCCD(CullDistance, Dt, Context);
+				NumActiveConstraints += ShapePair.GenerateCollisionCCD(bUseSweep, CullDistance, Dt, Context);
 			}
 		}
 		else
diff --git a/Engine/Source/Runtime/Experimental/Chaos/Private/Chaos/PBDRigidsEvolutionGBF.cpp b/Engine/Source/Runtime/Experimental/Chaos/Private/Chaos/PBDRigidsEvolutionGBF.cpp
--- a/Engine/Source/Runtime/Experimental/Chaos/Private/Chaos/PBDRigidsEvolutionGBF.cpp
+++ b/Engine/Source/Runtime/Experimental/Chaos/Private/Chaos/PBDRigidsEvolutionGBF.cpp
@@ -112,6 +112,7 @@ namespace Chaos
 		DECLARE_CYCLE_STAT(TEXT("FPBDRigidsEvolutionGBF::PostIntegrateCallback"), STAT_Evolution_PostIntegrateCallback, STATGROUP_Chaos);
 		DECLARE_CYCLE_STAT(TEXT("FPBDRigidsEvolutionGBF::CollisionModifierCallback"), STAT_Evolution_CollisionModifierCallback, STATGROUP_Chaos);
 		DECLARE_CYCLE_STAT(TEXT("FPBDRigidsEvolutionGBF::CCD"), STAT_Evolution_CCD, STATGROUP_Chaos);
+		DECLARE_CYCLE_STAT(TEXT("FPBDRigidsEvolutionGBF::CCDCorrection"), STAT_Evolution_CCDCorrection, STATGROUP_Chaos);
 		DECLARE_CYCLE_STAT(TEXT("FPBDRigidsEvolutionGBF::GraphColor"), STAT_Evolution_GraphColor, STATGROUP_Chaos);
 		DECLARE_CYCLE_STAT(TEXT("FPBDRigidsEvolutionGBF::BuildGroups"), STAT_Evolution_BuildGroups, STATGROUP_Chaos);
 		DECLARE_CYCLE_STAT(TEXT("FPBDRigidsEvolutionGBF::DetectCollisions"), STAT_Evolution_DetectCollisions, STATGROUP_Chaos);
@@ -453,6 +454,12 @@ void FPBDRigidsEvolutionGBF::AdvanceOneTimeStepImpl(const FReal Dt, const FSubSt
 			IslandGroupManager.Solve(Dt);
 		}
 
+		{
+			SCOPE_CYCLE_COUNTER(STAT_Evolution_CCDCorrection);
+			CSV_SCOPED_TIMING_STAT(PhysicsVerbose, CCDCorrection);
+			CCDManager.ApplyCorrections(Dt);
+		}
+
 		// Determine which particles can be disabled and which islands are sleeping
 		{
 			SCOPE_CYCLE_COUNTER(STAT_Evolution_BuildDisabledParticles);
diff --git a/Engine/Source/Runtime/Experimental/Chaos/Public/Chaos/CCDUtilities.h b/Engine/Source/Runtime/Experimental/Chaos/Public/Chaos/CCDUtilities.h
--- a/Engine/Source/Runtime/Experimental/Chaos/Public/Chaos/CCDUtilities.h
+++ b/Engine/Source/Runtime/Experimental/Chaos/Public/Chaos/CCDUtilities.h
@@ -86,6 +86,9 @@ namespace Chaos
 		FCCDManager(){}
 		void ApplyConstraintsPhaseCCD(const FReal Dt, FCollisionConstraintAllocator *CollisionAllocator, const int32 NumDynamicParticles = TNumericLimits<int32>::Max());
 
+		// A post process on CCD contacts to ensure that CCD objects never pass through non-dynamic objects
+		void ApplyCorrections(const FReal Dt);
+
 	private:
 		void ApplySweptConstraints(const FReal Dt, TArrayView<FPBDCollisionConstraint* const> InSweptConstraints, const int32 NumDynamicParticles = TNumericLimits<int32>::Max());
 		bool UpdateParticleSweptConstraints(FCCDParticle* CCDParticle, const FReal IslandTOI, const FReal Dt);
diff --git a/Engine/Source/Runtime/Experimental/Chaos/Public/Chaos/Collision/CollisionConstraintAllocator.h b/Engine/Source/Runtime/Experimental/Chaos/Public/Chaos/Collision/CollisionConstraintAllocator.h
--- a/Engine/Source/Runtime/Experimental/Chaos/Public/Chaos/Collision/CollisionConstraintAllocator.h
+++ b/Engine/Source/Runtime/Experimental/Chaos/Public/Chaos/Collision/CollisionConstraintAllocator.h
@@ -108,12 +108,12 @@ namespace Chaos
 		 * @note Nothing outside of thie allocator should hold a pointer to the midphase, or any constraints
 		 * it creates for more than the duration of the tick. Except the IslandManager :|
 		*/
-		FParticlePairMidPhase* GetMidPhase(FGeometryParticleHandle* Particle0, FGeometryParticleHandle* Particle1, FGeometryParticleHandle* SearchParticlePerformanceHint)
+		FParticlePairMidPhase* GetMidPhase(FGeometryParticleHandle* Particle0, FGeometryParticleHandle* Particle1, FGeometryParticleHandle* SearchParticlePerformanceHint, const FCollisionContext& Context)
 		{
 			FParticlePairMidPhase* MidPhase = FindMidPhaseImpl(Particle0, Particle1, SearchParticlePerformanceHint);
 			if (MidPhase == nullptr)
 			{
-				MidPhase = CreateMidPhase(Particle0, Particle1);
+				MidPhase = CreateMidPhase(Particle0, Particle1, Context);
 			}
 
 			return MidPhase;
@@ -163,7 +163,7 @@ namespace Chaos
 		}
 
 		// Create and initialize a midphase for the particle pair. Adds it to the particles' lists of midphases.
-		FParticlePairMidPhase* CreateMidPhase(FGeometryParticleHandle* Particle0, FGeometryParticleHandle* Particle1)
+		FParticlePairMidPhase* CreateMidPhase(FGeometryParticleHandle* Particle0, FGeometryParticleHandle* Particle1, const FCollisionContext& Context)
 		{
 			const FCollisionParticlePairKey Key = FCollisionParticlePairKey(Particle0, Particle1);
 
@@ -176,7 +176,7 @@ namespace Chaos
 
 			NewMidPhases.Add(MidPhase);
 
-			MidPhase->Init(Particle0, Particle1, Key);
+			MidPhase->Init(Particle0, Particle1, Key, Context);
 			return MidPhase;
 		}
 
@@ -233,7 +233,7 @@ namespace Chaos
 			, NumCollisionsPerBlock(InNumCollisionsPerBlock)
 			, ParticlePairMidPhases()
 			, ActiveConstraints()
-			, ActiveSweptConstraints()
+			, ActiveCCDConstraints()
 			, CurrentEpoch(0)
 			, bInCollisionDetectionPhase(false)
 		{
@@ -281,9 +281,9 @@ namespace Chaos
 		 * @note Some elements may be null (constraints that have been explicitly deleted)
 		 * @note This is not thread-safe and should not be used during the collision detection phase (i.e., when the list is being built)
 		*/
-		TArrayView<FPBDCollisionConstraint* const> GetSweptConstraints() const
+		TArrayView<FPBDCollisionConstraint* const> GetCCDConstraints() const
 		{ 
-			return MakeArrayView(ActiveSweptConstraints);
+			return MakeArrayView(ActiveCCDConstraints);
 		}
 
 		/**
@@ -321,7 +321,7 @@ namespace Chaos
 		void Reset()
 		{
 			ActiveConstraints.Reset();
-			ActiveSweptConstraints.Reset();
+			ActiveCCDConstraints.Reset();
 			ParticlePairMidPhases.Reset();
 
 			// NOTE: this deletes all the storage for constraints and midphases - must be last
@@ -335,7 +335,7 @@ namespace Chaos
 		void BeginFrame()
 		{
 			ActiveConstraints.Reset();
-			ActiveSweptConstraints.Reset();
+			ActiveCCDConstraints.Reset();
 		}
 
 		/**
@@ -362,7 +362,7 @@ namespace Chaos
 
 			// Clear the collision list for this tick - we are about to rebuild them
 			ActiveConstraints.Reset();
-			ActiveSweptConstraints.Reset();
+			ActiveCCDConstraints.Reset();
 		}
 
 		/**
@@ -411,7 +411,7 @@ namespace Chaos
 					Swap(Particle0, Particle1);
 				}
 
-				FParticlePairMidPhase* MidPhase = Context.GetAllocator()->GetMidPhase(Particle0, Particle1, SourceConstraint.Particle[0]);
+				FParticlePairMidPhase* MidPhase = Context.GetAllocator()->GetMidPhase(Particle0, Particle1, SourceConstraint.Particle[0], Context);
 
 				// We may be adding multiple constraints for the same particle pair, so we need
 				// to make sure the map is up to date in the case where we just created a new MidPhase
@@ -540,10 +540,11 @@ namespace Chaos
 			checkSlow(ActiveConstraints.Find(CollisionConstraint) == INDEX_NONE);
 			Cookie.ConstraintIndex = ActiveConstraints.Add(CollisionConstraint);
 
+			// If the constraint uses CCD, keep it in another list so we don't have to search the full list for CCD contacts
 			if (CollisionConstraint->GetCCDEnabled())
 			{
-				checkSlow(ActiveSweptConstraints.Find(CollisionConstraint) == INDEX_NONE);
-				Cookie.SweptConstraintIndex = ActiveSweptConstraints.Add(CollisionConstraint);
+				checkSlow(ActiveCCDConstraints.Find(CollisionConstraint) == INDEX_NONE);
+				Cookie.CCDConstraintIndex = ActiveCCDConstraints.Add(CollisionConstraint);
 			}
 
 			Cookie.LastUsedEpoch = CurrentEpoch;
@@ -560,7 +561,7 @@ namespace Chaos
 		TArray<FPBDCollisionConstraint*> ActiveConstraints;
 
 		// The active sweep constraints (added or recovered this tick)
-		TArray<FPBDCollisionConstraint*> ActiveSweptConstraints;
+		TArray<FPBDCollisionConstraint*> ActiveCCDConstraints;
 
 		// The current epoch used to track out-of-date contacts. A constraint whose Epoch is
 		// older than the current Epoch at the end of the tick was not refreshed this tick.
diff --git a/Engine/Source/Runtime/Experimental/Chaos/Public/Chaos/Collision/PBDCollisionConstraint.h b/Engine/Source/Runtime/Experimental/Chaos/Public/Chaos/Collision/PBDCollisionConstraint.h
--- a/Engine/Source/Runtime/Experimental/Chaos/Public/Chaos/Collision/PBDCollisionConstraint.h
+++ b/Engine/Source/Runtime/Experimental/Chaos/Public/Chaos/Collision/PBDCollisionConstraint.h
@@ -166,7 +166,7 @@ namespace Chaos
 			, CreationEpoch(INDEX_NONE)
 			, LastUsedEpoch(INDEX_NONE)
 			, ConstraintIndex(INDEX_NONE)
-			, SweptConstraintIndex(INDEX_NONE)
+			, CCDConstraintIndex(INDEX_NONE)
 		{
 		}
 
@@ -179,7 +179,7 @@ namespace Chaos
 		{
 			MidPhase = nullptr;
 			ConstraintIndex = INDEX_NONE;
-			SweptConstraintIndex = INDEX_NONE;
+			CCDConstraintIndex = INDEX_NONE;
 		}
 
 		// The constraint owner - set when the constraint is created
@@ -196,11 +196,11 @@ namespace Chaos
 		// The Epoch when the constraint was last used
 		int32 LastUsedEpoch;
 
-		// The index in the container - this changes every tick
+		// The index in the container - this changes every tick (is valid for all constraints, including CCD)
 		int32 ConstraintIndex;
 
-		// The index in swept constraints - this changes every tick
-		int32 SweptConstraintIndex;
+		// The CCD index in the container - this changes every tick (is INDEX_NONE for non-CCD constraints)
+		int32 CCDConstraintIndex;
 	};
 
 
@@ -280,17 +280,18 @@ namespace Chaos
 		virtual ~FPBDCollisionConstraint();
 
 		/**
-		 * Whetehr CCD is enabled for this collision.
-		 * This may change from tick to tick as an object's velocity changes.
-		*/
+		 * Whether CCD is enabled for this collision.
+		 * This value depends only on the CCD requirements of the two particles. It does not change from tick to tick.
+		 */
 		bool GetCCDEnabled() const { return Flags.bCCDEnabled; }
 
 		/**
-		 * @brief Enable or disable CCD for this constraint
-		*/
+		 * Enable or disable CCD for this constraint. Called once right after the constraint is created.
+		 */
 		void SetCCDEnabled(const bool bCCDEnabled)
 		{ 
 			Flags.bCCDEnabled = bCCDEnabled;
+			Flags.bCCDSweepEnabled = bCCDEnabled;
 
 			// Initialize the CCD thresholds (one time only)
 			if (bCCDEnabled && (CCDEnablePenetration == FReal(0)))
@@ -299,6 +300,21 @@ namespace Chaos
 			}
 		}
 
+		/**
+		 * For CCD constraints, do we need to run the initial sweep/rewind step.
+		 * This may change from tick to tick as an object's velocity changes.
+		 */
+		bool GetCCDSweepEnabled() const { return Flags.bCCDSweepEnabled; }
+
+		/**
+		 * For CCD constraints, enable or disable the pre-solve sweep/rewind for this tick. Called
+		 * every tick for CCD constraints based on current velocity versus size etc.
+		 */
+		void SetCCDSweepEnabled(const bool bCCDSweepEnabled)
+		{
+			Flags.bCCDSweepEnabled = bCCDSweepEnabled;
+		}
+
 		/**
 		 * @brief If CCD is enabled, contacts deeper than this will be handled by CCD
 		*/
@@ -594,8 +610,6 @@ namespace Chaos
 		 */
 		void EndTick()
 		{
-			// CCD state gets set each tick based on speed, size, etc
-			SetCCDEnabled(false);
 		}
 
 		/**
@@ -792,6 +806,7 @@ namespace Chaos
 				uint32 bIsProbeUnmodified : 1;  // Is this constraint a probe pre-contact-modification
 				uint32 bIsProbe : 1;            // Is this constraint currently a probe
 				uint32 bCCDEnabled : 1;			// Is CCD enabled for the current tick
+				uint32 bCCDSweepEnabled: 1;		// If this is a CCD constraint, do we want to enable the sweep/rewind phase?
 				uint32 bModifierApplied : 1;	// Was a constraint modifier applied this tick
 			};
 			uint32 Bits;
diff --git a/Engine/Source/Runtime/Experimental/Chaos/Public/Chaos/Collision/ParticlePairBroadPhase.h b/Engine/Source/Runtime/Experimental/Chaos/Public/Chaos/Collision/ParticlePairBroadPhase.h
--- a/Engine/Source/Runtime/Experimental/Chaos/Public/Chaos/Collision/ParticlePairBroadPhase.h
+++ b/Engine/Source/Runtime/Experimental/Chaos/Public/Chaos/Collision/ParticlePairBroadPhase.h
@@ -90,7 +90,7 @@ namespace Chaos
 			const FAABB3& Box1 = ParticleB->WorldSpaceInflatedBounds();
 			if (Box0.Intersects(Box1))
 			{
-				FParticlePairMidPhase* MidPhase = Context.GetAllocator()->GetMidPhase(ParticleA, ParticleB, SearchParticle);
+				FParticlePairMidPhase* MidPhase = Context.GetAllocator()->GetMidPhase(ParticleA, ParticleB, SearchParticle, Context);
 				if (MidPhase != nullptr)
 				{
 					MidPhase->GenerateCollisions(Context.GetSettings().BoundsExpansion, Dt, Context);
diff --git a/Engine/Source/Runtime/Experimental/Chaos/Public/Chaos/Collision/ParticlePairMidPhase.h b/Engine/Source/Runtime/Experimental/Chaos/Public/Chaos/Collision/ParticlePairMidPhase.h
--- a/Engine/Source/Runtime/Experimental/Chaos/Public/Chaos/Collision/ParticlePairMidPhase.h
+++ b/Engine/Source/Runtime/Experimental/Chaos/Public/Chaos/Collision/ParticlePairMidPhase.h
@@ -91,6 +91,7 @@ namespace Chaos
 		 * @return The number of collisions constraints that were activated
 		*/
 		int32 GenerateCollisionCCD(
+			const bool bEnableCCDSweep,
 			const FReal CullDistance,
 			const FReal Dt,
 			const FCollisionContext& Context);
@@ -115,7 +116,8 @@ namespace Chaos
 			const FCollisionContext& Context);
 			
 		int32 GenerateCollisionCCDImpl(
-			const FReal CullDistance, 
+			const bool bEnableCCDSweep,
+			const FReal CullDistance,
 			const FReal Dt,
 			const FCollisionContext& Context);
 
@@ -314,7 +316,8 @@ namespace Chaos
 		void Init(
 			FGeometryParticleHandle* InParticle0,
 			FGeometryParticleHandle* InParticle1,
-			const FCollisionParticlePairKey& InKey);
+			const FCollisionParticlePairKey& InKey,
+			const FCollisionContext& Context);
 
 		inline FGeometryParticleHandle* GetParticle0() { return Particle0; }
 
diff --git a/Engine/Source/Runtime/Experimental/Chaos/Public/Chaos/Collision/SpatialAccelerationBroadPhase.h b/Engine/Source/Runtime/Experimental/Chaos/Public/Chaos/Collision/SpatialAccelerationBroadPhase.h
--- a/Engine/Source/Runtime/Experimental/Chaos/Public/Chaos/Collision/SpatialAccelerationBroadPhase.h
+++ b/Engine/Source/Runtime/Experimental/Chaos/Public/Chaos/Collision/SpatialAccelerationBroadPhase.h
@@ -469,7 +469,7 @@ namespace Chaos
 				// NOTE: This searches the array of existing midphases on the particle specified in the last parameter. The assumption
 				// here is that Particle1 is more likely to have less collision on it because it is dynamic. Kinematics can often be
 				// large and may collide with many objects. Maybe we should consider using size instead.
-				FParticlePairMidPhase* MidPhase = Context.GetAllocator()->GetMidPhase(ParticleA, ParticleB, Particle1);
+				FParticlePairMidPhase* MidPhase = Context.GetAllocator()->GetMidPhase(ParticleA, ParticleB, Particle1, Context);
 				if (MidPhase != nullptr)
 				{
 					MidPhases.Add(MidPhase);
