Chaos - fix default joint solve position/rotation solve order (was intended to be position last)

diff --git a/Engine/Source/Runtime/Experimental/Chaos/Private/PBDRigidsSolver.cpp b/Engine/Source/Runtime/Experimental/Chaos/Private/PBDRigidsSolver.cpp
--- a/Engine/Source/Runtime/Experimental/Chaos/Private/PBDRigidsSolver.cpp
+++ b/Engine/Source/Runtime/Experimental/Chaos/Private/PBDRigidsSolver.cpp
@@ -279,13 +279,13 @@ namespace Chaos
 		float ChaosSolverJointMinSolverStiffness = 1.0f;
 		float ChaosSolverJointMaxSolverStiffness = 1.0f;
 		int32 ChaosSolverJointNumIterationsAtMaxSolverStiffness = 1;
-		bool bChaosSolverJointSolvePositionFirst = false;
+		bool bChaosSolverJointSolvePositionLast = true;
 		int32 ChaosSolverJointNumShockProagationIterations = 0;
 		FRealSingle ChaosSolverJointShockPropagation = -1.0f;
 		FAutoConsoleVariableRef CVarChaosSolverJointMinSolverStiffness(TEXT("p.Chaos.Solver.Joint.MinSolverStiffness"), ChaosSolverJointMinSolverStiffness, TEXT("Solver stiffness on first iteration, increases each iteration toward MaxSolverStiffness."));
 		FAutoConsoleVariableRef CVarChaosSolverJointMaxSolverStiffness(TEXT("p.Chaos.Solver.Joint.MaxSolverStiffness"), ChaosSolverJointMaxSolverStiffness, TEXT("Solver stiffness on last iteration, increases each iteration from MinSolverStiffness."));
 		FAutoConsoleVariableRef CVarChaosSolverJointNumIterationsAtMaxSolverStiffness(TEXT("p.Chaos.Solver.Joint.NumIterationsAtMaxSolverStiffness"), ChaosSolverJointNumIterationsAtMaxSolverStiffness, TEXT("How many iterations we want at MaxSolverStiffness."));
-		FAutoConsoleVariableRef CVarChaosSolverJointSolvePositionFirst(TEXT("p.Chaos.Solver.Joint.SolvePositionFirst"), bChaosSolverJointSolvePositionFirst, TEXT("Should we solve joints in position-then-rotation order (true) rotation-then-position order (false, default)"));
+		FAutoConsoleVariableRef CVarChaosSolverJointSolvePositionFirst(TEXT("p.Chaos.Solver.Joint.SolvePositionLast"), bChaosSolverJointSolvePositionLast, TEXT("Should we solve joints in position-then-rotation order (false) or rotation-then-position order (true, default)"));
 		FAutoConsoleVariableRef CVarChaosSolverJointNumShockPropagationIterations(TEXT("p.Chaos.Solver.Joint.NumShockPropagationIterations"), ChaosSolverJointNumShockProagationIterations, TEXT("How many iterations to enable SHockProagation for."));
 		FAutoConsoleVariableRef CVarChaosSolverJointShockPropagation(TEXT("p.Chaos.Solver.Joint.ShockPropagation"), ChaosSolverJointShockPropagation, TEXT("6Dof joint shock propagation override (if >= 0)."));
 
@@ -867,7 +867,7 @@ namespace Chaos
 		JointsSettings.AngleTolerance = ChaosSolverJointAngleTolerance;
 		JointsSettings.MinParentMassRatio = ChaosSolverJointMinParentMassRatio;
 		JointsSettings.MaxInertiaRatio = ChaosSolverJointMaxInertiaRatio;
-		JointsSettings.bSolvePositionLast = bChaosSolverJointSolvePositionFirst;
+		JointsSettings.bSolvePositionLast = bChaosSolverJointSolvePositionLast;
 		JointsSettings.NumShockPropagationIterations = ChaosSolverJointNumShockProagationIterations;
 		JointsSettings.ShockPropagationOverride = ChaosSolverJointShockPropagation;
 		JointsSettings.bUseLinearSolver = bChaosSolverJointUseLinearSolver;
