Chaos - fix PhysicsHandleComponent

diff --git a/Engine/Source/Runtime/Engine/Private/PhysicsEngine/PhysicsHandleComponent.cpp b/Engine/Source/Runtime/Engine/Private/PhysicsEngine/PhysicsHandleComponent.cpp
--- a/Engine/Source/Runtime/Engine/Private/PhysicsEngine/PhysicsHandleComponent.cpp
+++ b/Engine/Source/Runtime/Engine/Private/PhysicsEngine/PhysicsHandleComponent.cpp
@@ -139,8 +139,7 @@ void UPhysicsHandleComponent::GrabComponentImp(UPrimitiveComponent* InComponent,
 	// set target and current, so we don't need another "Tick" call to have it right
 	TargetTransform = CurrentTransform = KinematicTransform;
 
-	KinematicHandle->GetGameThreadAPI().SetX(KinematicTransform.GetLocation());
-	KinematicHandle->GetGameThreadAPI().SetR(KinematicTransform.GetRotation());	
+	KinematicHandle->GetGameThreadAPI().SetKinematicTarget(KinematicTransform);
 	
 	FTransform GrabbedTransform(InHandle->GetGameThreadAPI().R(), InHandle->GetGameThreadAPI().X());
 	ConstraintLocalPosition = GrabbedTransform.InverseTransformPosition(Location);
@@ -244,8 +243,7 @@ void UPhysicsHandleComponent::UpdateHandleTransform(const FTransform& NewTransfo
 	{
 		FPhysicsCommand::ExecuteWrite(KinematicHandle, [&](const FPhysicsActorHandle& InKinematicHandle)
 			{
-				KinematicHandle->GetGameThreadAPI().SetX(CurrentTransform.GetTranslation());
-				KinematicHandle->GetGameThreadAPI().SetR(CurrentTransform.GetRotation());
+				KinematicHandle->GetGameThreadAPI().SetKinematicTarget(FTransform(CurrentTransform.GetRotation(), CurrentTransform.GetTranslation()));
 			});
 
 		PreviousTransform = CurrentTransform;
diff --git a/Engine/Source/Runtime/Experimental/Chaos/Public/Chaos/KinematicTargets.h b/Engine/Source/Runtime/Experimental/Chaos/Public/Chaos/KinematicTargets.h
--- a/Engine/Source/Runtime/Experimental/Chaos/Public/Chaos/KinematicTargets.h
+++ b/Engine/Source/Runtime/Experimental/Chaos/Public/Chaos/KinematicTargets.h
@@ -15,7 +15,7 @@ namespace Chaos
 	{
 		None,			/** Particle does not move and no data is changed */
 		Reset,			/** Particle does not move, velocity and angular velocity are zeroed, then mode is set to "None". */
-		Position,		/** Particle is moved to Kinematic Target transform, velocity and angular velocity updated to reflect the change, then mode is set to "Zero". */
+		Position,		/** Particle is moved to Kinematic Target transform, velocity and angular velocity updated to reflect the change, then mode is set to "Reset". */
 		Velocity,		/** Particle is moved based on velocity and angular velocity, mode remains as "Velocity" until changed. */
 	};
 
@@ -26,6 +26,12 @@ namespace Chaos
 	class TKinematicTarget
 	{
 	public:
+
+		static TKinematicTarget<T, d> MakePositionTarget(const TRigidTransform<T, d>& InTransform)
+		{
+			return TKinematicTarget<T, d>(InTransform);
+		}
+
 		TKinematicTarget()
 			: Position(0)
 			, Rotation(TRotation<T, d>::FromIdentity())
@@ -131,6 +137,13 @@ namespace Chaos
 		}
 
 	private:
+		explicit TKinematicTarget(const TRigidTransform<T, d>& InTransform)
+			: Position(InTransform.GetTranslation())
+			, Rotation(InTransform.GetRotation())
+			, Mode(EKinematicTargetMode::Position)
+		{
+		}
+
 		TVector<T, d> Position;
 		TRotation<T, d> Rotation;
 		EKinematicTargetMode Mode;
diff --git a/Engine/Source/Runtime/Experimental/Chaos/Public/PhysicsProxy/SingleParticlePhysicsProxy.h b/Engine/Source/Runtime/Experimental/Chaos/Public/PhysicsProxy/SingleParticlePhysicsProxy.h
--- a/Engine/Source/Runtime/Experimental/Chaos/Public/PhysicsProxy/SingleParticlePhysicsProxy.h
+++ b/Engine/Source/Runtime/Experimental/Chaos/Public/PhysicsProxy/SingleParticlePhysicsProxy.h
@@ -322,7 +322,12 @@ protected:
 	}
 public:
 
-	void SetKinematicTarget(const TKinematicTarget<FReal, 3>& InKinematicTarget, bool bInvalidate = true)
+	void SetKinematicTarget(const FRigidTransform3& InTargetTransform, bool bInvalidate = true)
+	{
+		SetKinematicTarget(FKinematicTarget::MakePositionTarget(InTargetTransform), bInvalidate);
+	}
+
+	void SetKinematicTarget(const FKinematicTarget& InKinematicTarget, bool bInvalidate = true)
 	{
 		Write([&InKinematicTarget, bInvalidate](auto* Ptr)
 		{
