Chaos - fix issue that moving a body does not work if the motion is below the static friction threhold

diff --git a/Engine/Source/Runtime/Experimental/Chaos/Private/Chaos/Collision/PBDCollisionConstraint.cpp b/Engine/Source/Runtime/Experimental/Chaos/Private/Chaos/Collision/PBDCollisionConstraint.cpp
--- a/Engine/Source/Runtime/Experimental/Chaos/Private/Chaos/Collision/PBDCollisionConstraint.cpp
+++ b/Engine/Source/Runtime/Experimental/Chaos/Private/Chaos/Collision/PBDCollisionConstraint.cpp
@@ -410,6 +410,13 @@ namespace Chaos
 		ContainerCookie.MidPhase->SetIsSleeping(bInIsSleeping, ConcreteContainer()->GetConstraintAllocator().GetCurrentEpoch());
 	}
 
+	void FPBDCollisionConstraint::UpdateParticleTransform(FGeometryParticleHandle* InParticle)
+	{
+		// We need to update the friction anchors based on how much the user moved the particle, but
+		// it's easier to just reset the friction state altogether, and possibly the expected behaviour anyway
+		ResetManifold();
+	}
+
 	FVec3 FPBDCollisionConstraint::CalculateWorldContactLocation() const
 	{
 		if (ClosestManifoldPointIndex != INDEX_NONE)
diff --git a/Engine/Source/Runtime/Experimental/Chaos/Private/Chaos/PBDRigidsEvolutionGBF.cpp b/Engine/Source/Runtime/Experimental/Chaos/Private/Chaos/PBDRigidsEvolutionGBF.cpp
--- a/Engine/Source/Runtime/Experimental/Chaos/Private/Chaos/PBDRigidsEvolutionGBF.cpp
+++ b/Engine/Source/Runtime/Experimental/Chaos/Private/Chaos/PBDRigidsEvolutionGBF.cpp
@@ -892,6 +892,39 @@ void FPBDRigidsEvolutionGBF::DestroyTransientConstraints(FGeometryParticleHandle
 	}
 }
 
+void FPBDRigidsEvolutionGBF::OnParticleMoved(FGeometryParticleHandle* Particle, const FVec3& PrevX, const FRotation3& PrevR)
+{
+	// When a particle is moved, we need to tell the collisions because it 
+	// caches friction anchors and it will attempt to undo small translations
+	Particle->ParticleCollisions().VisitCollisions([this, Particle](FPBDCollisionConstraint& Collision)
+		{
+			Collision.UpdateParticleTransform(Particle);
+			return ECollisionVisitorResult::Continue;
+		});
+
+	// Reset any sleep state so that we don't keep sleeping and waking (we have a zero velocity
+	// so the current sleep logic will think it can go to sleep). For now this is the best we can do
+	// @todo(chaos): better sleep system!
+	// @todo(chaos): we need to know if we are teleporting so we can reset stuff
+	if (FPBDRigidParticleHandle* Rigid = Particle->CastToRigidParticle())
+	{
+		// We use the length of DX to guess when we're teleporting. This is lame. We also have to guess a DT.
+		// The only thing that really matters is that VSmooth is above the material's sleep threshold
+		// and not super huge because we're being initially positioned with a PrevX of (0,0,0)
+		const FReal TeleportThreshold = FReal(100);
+		FVec3 DX = (PrevX - Particle->X());
+		const FReal DXLenSq = DX.SizeSquared();
+		if (DXLenSq > FMath::Square(TeleportThreshold))
+		{
+			DX = TeleportThreshold * DX * FMath::InvSqrt(DXLenSq);
+		}
+
+		FReal Dt = FReal(1.0 / 30.0);
+		const FReal SmoothRate = FMath::Clamp(CVars::SmoothedPositionLerpRate, 0.0f, 1.0f);
+		Rigid->VSmooth() = FMath::Lerp(Rigid->VSmooth(), DX / Dt, SmoothRate);
+	}
+}
+
 void FPBDRigidsEvolutionGBF::ParticleMaterialChanged(FGeometryParticleHandle* Particle)
 {
 	Particle->ParticleCollisions().VisitCollisions([this](FPBDCollisionConstraint& Collision)
diff --git a/Engine/Source/Runtime/Experimental/Chaos/Private/PhysicsProxy/SingleParticlePhysicsProxy.cpp b/Engine/Source/Runtime/Experimental/Chaos/Private/PhysicsProxy/SingleParticlePhysicsProxy.cpp
--- a/Engine/Source/Runtime/Experimental/Chaos/Private/PhysicsProxy/SingleParticlePhysicsProxy.cpp
+++ b/Engine/Source/Runtime/Experimental/Chaos/Private/PhysicsProxy/SingleParticlePhysicsProxy.cpp
@@ -61,7 +61,8 @@ void PushToPhysicsStateImp(const Chaos::FDirtyPropertiesManager& Manager, Chaos:
 
 		if(NewXR)
 		{
-			Handle->SetXR(*NewXR);
+			// @todo(chaos): we need to know if this is a teleport or not and pass that on
+			Evolution.SetParticleTransform(Handle, NewXR->X(), NewXR->R());
 		}
 
 		if(NewNonFrequentData)
diff --git a/Engine/Source/Runtime/Experimental/Chaos/Public/Chaos/Collision/PBDCollisionConstraint.h b/Engine/Source/Runtime/Experimental/Chaos/Public/Chaos/Collision/PBDCollisionConstraint.h
--- a/Engine/Source/Runtime/Experimental/Chaos/Public/Chaos/Collision/PBDCollisionConstraint.h
+++ b/Engine/Source/Runtime/Experimental/Chaos/Public/Chaos/Collision/PBDCollisionConstraint.h
@@ -367,6 +367,9 @@ namespace Chaos
 			SetDisabled(false);
 		}
 
+		// When a particle is moved under user control, we need to update some cached state to prevent friction from undoing the move
+		void UpdateParticleTransform(FGeometryParticleHandle* InParticle);
+
 		// @todo(chaos): half of this API is wrong for the new multi-point manifold constraints. Remove it
 
 		void ResetPhi(FReal InPhi) { ClosestManifoldPointIndex = INDEX_NONE; }
diff --git a/Engine/Source/Runtime/Experimental/Chaos/Public/Chaos/PBDRigidsEvolutionGBF.h b/Engine/Source/Runtime/Experimental/Chaos/Public/Chaos/PBDRigidsEvolutionGBF.h
--- a/Engine/Source/Runtime/Experimental/Chaos/Public/Chaos/PBDRigidsEvolutionGBF.h
+++ b/Engine/Source/Runtime/Experimental/Chaos/Public/Chaos/PBDRigidsEvolutionGBF.h
@@ -120,6 +120,27 @@ namespace Chaos
 		FORCEINLINE FPBDSuspensionConstraints& GetSuspensionConstraints() { return SuspensionConstraints; }
 		FORCEINLINE const FPBDSuspensionConstraints& GetSuspensionConstraints() const { return SuspensionConstraints; }
 
+
+		//
+		// Particle API (most of the particle API is in the base class)
+		//
+
+		/**
+		 * User has moved a particle
+		 */
+		CHAOS_API void SetParticleTransform(FGeometryParticleHandle* InParticle, const FVec3& InPos, const FRotation3& InRot)
+		{
+			const FVec3 PrevX = InParticle->X();
+			const FRotation3 PrevR = InParticle->R();
+
+			FGenericParticleHandle(InParticle)->SetTransform(InPos, InRot);
+
+			OnParticleMoved(InParticle, PrevX, PrevR);
+		}
+
+		CHAOS_API void OnParticleMoved(FGeometryParticleHandle* InParticle, const FVec3& PrevX, const FRotation3& PrevR);
+
+
 		/**
 		 * Reload the particles cache for all particles where appropriate
 		 */
diff --git a/Engine/Source/Runtime/Experimental/Chaos/Public/Chaos/ParticleHandle.h b/Engine/Source/Runtime/Experimental/Chaos/Public/Chaos/ParticleHandle.h
--- a/Engine/Source/Runtime/Experimental/Chaos/Public/Chaos/ParticleHandle.h
+++ b/Engine/Source/Runtime/Experimental/Chaos/Public/Chaos/ParticleHandle.h
@@ -1392,6 +1392,17 @@ public:
 	FGeometryParticleHandle* Handle() { return MHandle; }
 	int32 GetHandleIdx() const { return MHandle->GetHandleIdx(); }
 
+	void SetTransform(const FVec3& Pos, const FRotation3& Rot)
+	{
+		MHandle->X() = Pos;
+		MHandle->R() = Rot;
+		if (FPBDRigidParticleHandle* Dynamic = CastToRigidParticle())
+		{
+			Dynamic->P() = Pos;
+			Dynamic->Q() = Rot;
+		}
+	}
+
 	// Static Particles
 	FVec3& X() { return MHandle->X(); }
 	const FVec3& X() const { return MHandle->X(); }
