Chaos - Disable the non-teleport move handling pending further testing. (See UE-165746 for the issue that would fix)

diff --git a/Engine/Source/Runtime/Experimental/Chaos/Private/Chaos/PBDRigidsEvolutionGBF.cpp b/Engine/Source/Runtime/Experimental/Chaos/Private/Chaos/PBDRigidsEvolutionGBF.cpp
--- a/Engine/Source/Runtime/Experimental/Chaos/Private/Chaos/PBDRigidsEvolutionGBF.cpp
+++ b/Engine/Source/Runtime/Experimental/Chaos/Private/Chaos/PBDRigidsEvolutionGBF.cpp
@@ -892,36 +892,33 @@ void FPBDRigidsEvolutionGBF::DestroyTransientConstraints(FGeometryParticleHandle
 	}
 }
 
-void FPBDRigidsEvolutionGBF::OnParticleMoved(FGeometryParticleHandle* Particle, const FVec3& PrevX, const FRotation3& PrevR)
+void FPBDRigidsEvolutionGBF::OnParticleMoved(FGeometryParticleHandle* InParticle, const FVec3& PrevX, const FRotation3& PrevR, const bool bIsTeleport)
 {
-	// When a particle is moved, we need to tell the collisions because it 
-	// caches friction anchors and it will attempt to undo small translations
-	Particle->ParticleCollisions().VisitCollisions([this, Particle](FPBDCollisionConstraint& Collision)
+	// When a particle is moved, we need to tell the collisions because they cache friction state and 
+	// would attempt to undo small translations within the friction cone
+	InParticle->ParticleCollisions().VisitCollisions([this, InParticle](FPBDCollisionConstraint& Collision)
 		{
-			Collision.UpdateParticleTransform(Particle);
+			Collision.UpdateParticleTransform(InParticle);
 			return ECollisionVisitorResult::Continue;
 		});
 
-	// Reset any sleep state so that we don't keep sleeping and waking (we have a zero velocity
-	// so the current sleep logic will think it can go to sleep). For now this is the best we can do
-	// @todo(chaos): better sleep system!
-	// @todo(chaos): we need to know if we are teleporting so we can reset stuff
-	if (FPBDRigidParticleHandle* Rigid = Particle->CastToRigidParticle())
-	{
-		// We use the length of DX to guess when we're teleporting. This is lame. We also have to guess a DT.
-		// The only thing that really matters is that VSmooth is above the material's sleep threshold
-		// and not super huge because we're being initially positioned with a PrevX of (0,0,0)
-		const FReal TeleportThreshold = FReal(100);
-		FVec3 DX = (PrevX - Particle->X());
-		const FReal DXLenSq = DX.SizeSquared();
-		if (DXLenSq > FMath::Square(TeleportThreshold))
+	// If this is a teleport, we assume other state has been set separately (V, W, etc) if required, which 
+	// would include resetting any sleep-related state. If this is not a teleport, we want to prevent 
+	// sleeping if the user keeps moving the body without changing the velocity.
+	if (!bIsTeleport)
+	{
+		if (FPBDRigidParticleHandle* Rigid = InParticle->CastToRigidParticle())
 		{
-			DX = TeleportThreshold * DX * FMath::InvSqrt(DXLenSq);
+			// No need to do this for kinematics
+			if (Rigid->IsDynamic())
+			{
+				// @todo(chaos): better sleep system! We should probably have a sleep accumulator per particle
+				const FReal InvDt = FReal(30.0);
+				const FVec3 DV = (PrevX - Rigid->X()) * InvDt;
+				const FReal SmoothRate = FMath::Clamp(CVars::SmoothedPositionLerpRate, 0.0f, 1.0f);
+				Rigid->VSmooth() = FMath::Lerp(Rigid->VSmooth(), Rigid->V() + DV, SmoothRate);
+			}
 		}
-
-		FReal Dt = FReal(1.0 / 30.0);
-		const FReal SmoothRate = FMath::Clamp(CVars::SmoothedPositionLerpRate, 0.0f, 1.0f);
-		Rigid->VSmooth() = FMath::Lerp(Rigid->VSmooth(), DX / Dt, SmoothRate);
 	}
 }
 
diff --git a/Engine/Source/Runtime/Experimental/Chaos/Private/PhysicsProxy/SingleParticlePhysicsProxy.cpp b/Engine/Source/Runtime/Experimental/Chaos/Private/PhysicsProxy/SingleParticlePhysicsProxy.cpp
--- a/Engine/Source/Runtime/Experimental/Chaos/Private/PhysicsProxy/SingleParticlePhysicsProxy.cpp
+++ b/Engine/Source/Runtime/Experimental/Chaos/Private/PhysicsProxy/SingleParticlePhysicsProxy.cpp
@@ -61,8 +61,11 @@ void PushToPhysicsStateImp(const Chaos::FDirtyPropertiesManager& Manager, Chaos:
 
 		if(NewXR)
 		{
-			// @todo(chaos): we need to know if this is a teleport or not and pass that on
-			Evolution.SetParticleTransform(Handle, NewXR->X(), NewXR->R());
+			// @todo(chaos): we need to know if this is a teleport or not and pass that on. See UE-165746
+			// For now we just set bIsTeleport to true since that's the no-impact option for SetParticleTransform
+			// (there would be issues if we report a non-teleport move for an initial-position a long way from the origin)
+			const bool bIsTeleport = true;
+			Evolution.SetParticleTransform(Handle, NewXR->X(), NewXR->R(), bIsTeleport);
 		}
 
 		if(NewNonFrequentData)
diff --git a/Engine/Source/Runtime/Experimental/Chaos/Public/Chaos/PBDRigidsEvolutionGBF.h b/Engine/Source/Runtime/Experimental/Chaos/Public/Chaos/PBDRigidsEvolutionGBF.h
--- a/Engine/Source/Runtime/Experimental/Chaos/Public/Chaos/PBDRigidsEvolutionGBF.h
+++ b/Engine/Source/Runtime/Experimental/Chaos/Public/Chaos/PBDRigidsEvolutionGBF.h
@@ -128,17 +128,17 @@ namespace Chaos
 		/**
 		 * User has moved a particle
 		 */
-		CHAOS_API void SetParticleTransform(FGeometryParticleHandle* InParticle, const FVec3& InPos, const FRotation3& InRot)
+		CHAOS_API void SetParticleTransform(FGeometryParticleHandle* InParticle, const FVec3& InPos, const FRotation3& InRot, const bool bIsTeleport)
 		{
 			const FVec3 PrevX = InParticle->X();
 			const FRotation3 PrevR = InParticle->R();
 
 			FGenericParticleHandle(InParticle)->SetTransform(InPos, InRot);
 
-			OnParticleMoved(InParticle, PrevX, PrevR);
+			OnParticleMoved(InParticle, PrevX, PrevR, bIsTeleport);
 		}
 
-		CHAOS_API void OnParticleMoved(FGeometryParticleHandle* InParticle, const FVec3& PrevX, const FRotation3& PrevR);
+		CHAOS_API void OnParticleMoved(FGeometryParticleHandle* InParticle, const FVec3& PrevX, const FRotation3& PrevR, const bool bIsTeleport);
 
 
 		/**
