Geometry collection : fix bounds not being updated properly when scrubbing start time from caches in editor mode

diff --git a/Engine/Source/Runtime/Experimental/GeometryCollectionEngine/Private/GeometryCollection/GeometryCollectionComponent.cpp b/Engine/Source/Runtime/Experimental/GeometryCollectionEngine/Private/GeometryCollection/GeometryCollectionComponent.cpp
--- a/Engine/Source/Runtime/Experimental/GeometryCollectionEngine/Private/GeometryCollection/GeometryCollectionComponent.cpp
+++ b/Engine/Source/Runtime/Experimental/GeometryCollectionEngine/Private/GeometryCollection/GeometryCollectionComponent.cpp
@@ -1146,6 +1146,7 @@ void UGeometryCollectionComponent::SetRestState(TArray<FTransform>&& InRestTrans
 	{
 #if WITH_EDITOR
 			// We need to do this in case we're controlled by Sequencer in editor, which doesn't invoke PostEditChangeProperty
+			UpdateCachedBounds();
 			SendRenderTransform_Concurrent();
 #endif
 		if (SceneProxy->IsNaniteMesh())
