Geometry collection : make sure replicated break do trigger client side break events

diff --git a/Engine/Source/Runtime/Experimental/Chaos/Private/Chaos/PBDRigidClustering.cpp b/Engine/Source/Runtime/Experimental/Chaos/Private/Chaos/PBDRigidClustering.cpp
--- a/Engine/Source/Runtime/Experimental/Chaos/Private/Chaos/PBDRigidClustering.cpp
+++ b/Engine/Source/Runtime/Experimental/Chaos/Private/Chaos/PBDRigidClustering.cpp
@@ -773,19 +773,19 @@ namespace Chaos
 	DECLARE_CYCLE_STAT(TEXT("TPBDRigidClustering<>::ReleaseClusterParticles(LIST)"), STAT_ReleaseClusterParticles_LIST, STATGROUP_Chaos);
 	TSet<FPBDRigidParticleHandle*> 
 	FRigidClustering::ReleaseClusterParticles(
-		TArray<FPBDRigidParticleHandle*> ChildrenParticles)
+		TArray<FPBDRigidParticleHandle*> ChildrenParticles, bool bTriggerBreakEvents /* = false */)
 	{
 		SCOPE_CYCLE_COUNTER(STAT_ReleaseClusterParticles_LIST);
 		TSet<FPBDRigidParticleHandle*> ActivatedBodies;
 		if (ChildrenParticles.Num())
 		{
+			//for now just assume these all belong to same cluster
 			FPBDRigidParticleHandle* ClusterHandle = nullptr;
-			//todo(ocohen): refactor incoming, for now just assume these all belong to same cluster and hack strain array
 			
 			TMap<FGeometryParticleHandle*, FReal> FakeStrain;
 
 			bool bPreDoGenerateData = DoGenerateBreakingData;
-			DoGenerateBreakingData = false;
+			DoGenerateBreakingData = bTriggerBreakEvents;
 
 			for (FPBDRigidParticleHandle* ChildHandle : ChildrenParticles)
 			{
diff --git a/Engine/Source/Runtime/Experimental/Chaos/Public/Chaos/PBDRigidClustering.h b/Engine/Source/Runtime/Experimental/Chaos/Public/Chaos/PBDRigidClustering.h
--- a/Engine/Source/Runtime/Experimental/Chaos/Public/Chaos/PBDRigidClustering.h
+++ b/Engine/Source/Runtime/Experimental/Chaos/Public/Chaos/PBDRigidClustering.h
@@ -132,7 +132,7 @@ public:
 	*    Release all rigid body IDs passed,
 	*/
 	TSet<FPBDRigidParticleHandle*> ReleaseClusterParticles(
-		TArray<FPBDRigidParticleHandle*> ChildrenParticles);
+		TArray<FPBDRigidParticleHandle*> ChildrenParticles, bool bTriggerBreakEvents = false);
 
 	/*
 	*  DestroyClusterParticle
diff --git a/Engine/Source/Runtime/Experimental/GeometryCollectionEngine/Private/GeometryCollection/GeometryCollectionComponent.cpp b/Engine/Source/Runtime/Experimental/GeometryCollectionEngine/Private/GeometryCollection/GeometryCollectionComponent.cpp
--- a/Engine/Source/Runtime/Experimental/GeometryCollectionEngine/Private/GeometryCollection/GeometryCollectionComponent.cpp
+++ b/Engine/Source/Runtime/Experimental/GeometryCollectionEngine/Private/GeometryCollection/GeometryCollectionComponent.cpp
@@ -1471,7 +1471,6 @@ void UGeometryCollectionComponent::UpdateRepData()
 		//If the first child is enabled it means it's a one off and the cluster IS the first child
 		
 		//TODO: for now we have to iterate over all particles to find the clusters, would be better if we had the clusters and children already available
-		//Large refactor happening to this stuff so for now we just iterate
 		//We are relying on the fact that we fracture one level per step. This means we will see all one offs here
 
 		bool bClustersChanged = false;
@@ -1695,7 +1694,7 @@ void UGeometryCollectionComponent::ProcessRepData()
 			OneOff->SetW(ActivatedCluster.InitialAngularVelocity);
 		}
 
-		RigidClustering.ReleaseClusterParticles(TArray<FPBDRigidParticleHandle*>{ OneOff });
+		RigidClustering.ReleaseClusterParticles(TArray<FPBDRigidParticleHandle*>{ OneOff }, /* bTriggerBreakEvents */ true);
 	}
 
 	if(bHardSnap)
