Chaos - fix crash when Fields are used to break a geometry collection

diff --git a/Engine/Source/Runtime/Experimental/Chaos/Private/PhysicsProxy/GeometryCollectionPhysicsProxy.cpp b/Engine/Source/Runtime/Experimental/Chaos/Private/PhysicsProxy/GeometryCollectionPhysicsProxy.cpp
--- a/Engine/Source/Runtime/Experimental/Chaos/Private/PhysicsProxy/GeometryCollectionPhysicsProxy.cpp
+++ b/Engine/Source/Runtime/Experimental/Chaos/Private/PhysicsProxy/GeometryCollectionPhysicsProxy.cpp
@@ -939,12 +939,23 @@ void FGeometryCollectionPhysicsProxy::InitializeBodiesPT(Chaos::FPBDRigidsSolver
 				{
 					Shape->SetMaterial(Parameters.PhysicalMaterialHandle);
 				}
-
-				RigidsSolver->GetEvolution()->EnableParticle(Handle);
 			}
 		},true);
 
+		// Enable the particle (register with graph, add to broadphase, etc)
+		// NOTE: Cannot be called in parallel (do not move this to the above loop)
+		// @todo(chaos): can this be done at the end of the function once we know 
+		// whether we're asleep etc.? (but see calls to UpdateGeometryCollectionViews)
+		for (FClusterHandle* ClusterHandle : SolverParticleHandles)
+		{
+			if ((ClusterHandle != nullptr) && !ClusterHandle->Disabled())
+			{
+				RigidsSolver->GetEvolution()->EnableParticle(ClusterHandle);
+			}
+		}
+
 		// After population, the states of each particle could have changed
+		// @todo(Chaos) Do we really need to maintain the Views while running this function? We do it twice, but there's also a Dirty flag built-in
 		Particles.UpdateGeometryCollectionViews();
 
 		for (FFieldSystemCommand& Cmd : Parameters.InitializationCommands)
diff --git a/Engine/Source/Runtime/Experimental/Chaos/Public/PhysicsProxy/FieldSystemProxyHelper.h b/Engine/Source/Runtime/Experimental/Chaos/Public/PhysicsProxy/FieldSystemProxyHelper.h
--- a/Engine/Source/Runtime/Experimental/Chaos/Public/PhysicsProxy/FieldSystemProxyHelper.h
+++ b/Engine/Source/Runtime/Experimental/Chaos/Public/PhysicsProxy/FieldSystemProxyHelper.h
@@ -193,7 +193,15 @@ namespace Chaos
 		}
 		else
 		{
+			// @todo(chaos): this should also call Evolution.SetParticleObjectState, but to do that
+			// we need to change how the ActiveParticlesArray is handled in RigidParticleSOAs.
+			// Instead we manually do the parts of SetParticleObjectState that are still required 
+			// while avoiding the SOA management.
 			RigidHandle->SetObjectStateLowLevel(FieldState);
+			if (!RigidHandle->Disabled())
+			{
+				RigidSolver->GetEvolution()->GetConstraintGraph().AddParticle(RigidHandle);
+			}
 		}
 	}
 
