Add array of single particle physics proxies to pbd rigids solver - this is mainly for use as a PT utility for async callbacks, so for example if you did a PT scene query in an async callback object you could send proxies back to GT to do stuff with. Without PT proxy access there's no way to tell GT about particles we know about on PT.

diff --git a/Engine/Source/Runtime/Experimental/Chaos/Private/PBDRigidsSolver.cpp b/Engine/Source/Runtime/Experimental/Chaos/Private/PBDRigidsSolver.cpp
--- a/Engine/Source/Runtime/Experimental/Chaos/Private/PBDRigidsSolver.cpp
+++ b/Engine/Source/Runtime/Experimental/Chaos/Private/PBDRigidsSolver.cpp
@@ -585,8 +585,8 @@ namespace Chaos
 		{
 			int32 ExternalTimestamp = GetMarshallingManager().GetExternalTimestamp_External();
 			Chaos::FIgnoreCollisionManager::FDeactivationSet& PendingMap = CollisionManager.GetPendingDeactivationsForGameThread(ExternalTimestamp);
-				PendingMap.Add(UniqueIdx);
-			}
+			PendingMap.Add(UniqueIdx);
+		}
 
 		// Enqueue a command to remove the particle and delete the proxy
 		EnqueueCommandImmediate([Proxy, UniqueIdx, this]()
@@ -613,7 +613,7 @@ namespace Chaos
 			FGeometryParticleHandle* Handle = Proxy->GetHandle_LowLevel();
 			Proxy->SetHandle(nullptr);
 			const int32 OffsetForRewind = MRewindData ? MRewindData->Capacity() : 0;
-			PendingDestroyPhysicsProxy.Add(  FPendingDestroyInfo{Proxy, GetCurrentFrame() + OffsetForRewind, Handle, UniqueIdx});
+			PendingDestroyPhysicsProxy.Add(FPendingDestroyInfo{Proxy, GetCurrentFrame() + OffsetForRewind, Handle, UniqueIdx});
 			
 			//If particle was created and destroyed before commands were enqueued just skip. I suspect we can skip entire lambda, but too much code to verify right now
 			if(Handle)
@@ -621,6 +621,12 @@ namespace Chaos
 				//Disable until particle is finally destroyed
 				GetEvolution()->DisableParticle(Handle);
 			}
+
+			// Remove the proxy for this particle
+			if (SingleParticlePhysicsProxies_PT.IsValidIndex(UniqueIdx.Idx))
+			{
+				SingleParticlePhysicsProxies_PT.RemoveAt(UniqueIdx.Idx);
+			}
 		});
 	}
 
@@ -689,23 +695,22 @@ namespace Chaos
 
 		GTConstraint->ReleaseKinematicEndPoint(this);
 
-		FParticlesType* InParticles = &GetParticles();
 		JointProxy->DestroyOnGameThread();	//destroy the game thread portion of the proxy
 		
-		// Finish registration on the physics thread...
-		EnqueueCommandImmediate([InParticles, JointProxy, this]()
+		// Finish de-registration on the physics thread...
+		EnqueueCommandImmediate([JointProxy, this]()
+		{
+			//TODO: consider deferring this so that rewind can resim with joint until moment of destruction
+			//For now we assume this always comes from server update
+			if(FRewindData* RewindData = GetRewindData())
 			{
-				//TODO: consider deferring this so that rewind can resim with joint until moment of destruction
-				//For now we assume this always comes from server update
-				if(FRewindData* RewindData = GetRewindData())
-				{
-					RewindData->RemoveObject(JointProxy->GetHandle());
-				}
+				RewindData->RemoveObject(JointProxy->GetHandle());
+			}
 
-				JointProxy->DestroyOnPhysicsThread(this);
-				JointConstraintPhysicsProxies_Internal.RemoveSingle(JointProxy);
-				delete JointProxy;
-			});
+			JointProxy->DestroyOnPhysicsThread(this);
+			JointConstraintPhysicsProxies_Internal.RemoveSingle(JointProxy);
+			delete JointProxy;
+		});
 	}
 
 	void FPBDRigidsSolver::RegisterObject(Chaos::FSuspensionConstraint* GTConstraint)
@@ -1086,6 +1091,9 @@ namespace Chaos
 
 				auto Handle = Proxy->GetHandle_LowLevel();
 				Handle->GTGeometryParticle() = Proxy->GetParticle_LowLevel();
+
+				// Track proxies
+				SingleParticlePhysicsProxies_PT.EmplaceAt(UniqueIdx->Idx, Proxy);
 			}
 
 			if(Proxy->GetHandle_LowLevel())
@@ -1165,7 +1173,7 @@ namespace Chaos
 		});
 
 		//need to create new constraint handles
-		DirtyProxiesData->ForEachProxy([this, &ProcessProxyPT, Manager, RewindData](int32 DataIdx, FDirtyProxy& Dirty)
+		DirtyProxiesData->ForEachProxy([this, Manager, RewindData](int32 DataIdx, FDirtyProxy& Dirty)
 		{
 			if (Dirty.Proxy->GetIgnoreDataOnStep_Internal() != CurrentFrame)
 			{
@@ -1634,6 +1642,20 @@ CSV_CUSTOM_STAT(PhysicsCounters, Name, Value, ECsvCustomStatOp::Set);
 #endif
 	}
 
+	FSingleParticlePhysicsProxy* FPBDRigidsSolver::GetParticleProxy_PT(const FGeometryParticleHandle& Handle)
+	{
+		const FUniqueIdx UniqueIdx = Handle.UniqueIdx();
+		return SingleParticlePhysicsProxies_PT.IsValidIndex(UniqueIdx.Idx)
+			? SingleParticlePhysicsProxies_PT[UniqueIdx.Idx] : nullptr;
+	}
+
+	const FSingleParticlePhysicsProxy* FPBDRigidsSolver::GetParticleProxy_PT(const FGeometryParticleHandle& Handle) const
+	{
+		const FUniqueIdx UniqueIdx = Handle.UniqueIdx();
+		return SingleParticlePhysicsProxies_PT.IsValidIndex(UniqueIdx.Idx)
+			? SingleParticlePhysicsProxies_PT[UniqueIdx.Idx] : nullptr;
+	}
+
 	void FPBDRigidsSolver::PostEvolutionVDBPush() const
 	{
 #if CHAOS_VISUAL_DEBUGGER_ENABLED
diff --git a/Engine/Source/Runtime/Experimental/Chaos/Public/PBDRigidsSolver.h b/Engine/Source/Runtime/Experimental/Chaos/Public/PBDRigidsSolver.h
--- a/Engine/Source/Runtime/Experimental/Chaos/Public/PBDRigidsSolver.h
+++ b/Engine/Source/Runtime/Experimental/Chaos/Public/PBDRigidsSolver.h
@@ -301,6 +301,10 @@ namespace Chaos
 		void UpdateStatCounters() const;
 		void UpdateExpensiveStatCounters() const;
 
+		// Access particle proxy from physics thread useful for cross thread communication
+		FSingleParticlePhysicsProxy* GetParticleProxy_PT(const FGeometryParticleHandle& Handle);
+		const FSingleParticlePhysicsProxy* GetParticleProxy_PT(const FGeometryParticleHandle& Handle) const;
+
 	private:
 
 		/**/
@@ -334,6 +338,7 @@ namespace Chaos
 		// Proxies
 		//
 		TSharedPtr<FCriticalSection> MCurrentLock;
+		TSparseArray< FSingleParticlePhysicsProxy* > SingleParticlePhysicsProxies_PT;
 		TArray< FGeometryCollectionPhysicsProxy* > GeometryCollectionPhysicsProxies_Internal; // PT
 		TArray< FJointConstraintPhysicsProxy* > JointConstraintPhysicsProxies_Internal; // PT
 
