Chaos : fix server crash when replicating geometry collections that are missing InitialLevel attribute

diff --git a/Engine/Source/Runtime/Experimental/GeometryCollectionEngine/Private/GeometryCollection/GeometryCollectionComponent.cpp b/Engine/Source/Runtime/Experimental/GeometryCollectionEngine/Private/GeometryCollection/GeometryCollectionComponent.cpp
--- a/Engine/Source/Runtime/Experimental/GeometryCollectionEngine/Private/GeometryCollection/GeometryCollectionComponent.cpp
+++ b/Engine/Source/Runtime/Experimental/GeometryCollectionEngine/Private/GeometryCollection/GeometryCollectionComponent.cpp
@@ -1496,7 +1496,7 @@ void UGeometryCollectionComponent::UpdateRepData()
 		const FPBDRigidsSolver* Solver = PhysicsProxy->GetSolver<Chaos::FPBDRigidsSolver>();
 		const FRigidClustering& RigidClustering = Solver->GetEvolution()->GetRigidClustering();
 
-		const TManagedArray<int32>& InitialLevels = PhysicsProxy->GetPhysicsCollection().GetAttribute<int32>("InitialLevel", FGeometryCollection::TransformGroup);
+		const TManagedArray<int32>* InitialLevels = PhysicsProxy->GetPhysicsCollection().FindAttribute<int32>("InitialLevel", FGeometryCollection::TransformGroup);
 		const TManagedArray<TSet<int32>>& InitialChildren = PhysicsProxy->GetPhysicsCollection().Children;
 
 		//see if we have any new clusters that are enabled
@@ -1535,14 +1535,14 @@ void UGeometryCollectionComponent::UpdateRepData()
 						ensureMsgf(TransformGroupIdx >= 0, TEXT("Non-internal cluster should always have a group index"));
 						ensureMsgf(TransformGroupIdx < TNumericLimits<uint16>::Max(), TEXT("Trying to replicate GC with more than 65k pieces. We assumed uint16 would suffice"));
 
-						Level = InitialLevels[TransformGroupIdx];
+						Level = InitialLevels? (*InitialLevels)[TransformGroupIdx]: INDEX_NONE;
 					}
 					else
 					{
 						// Use internal cluster child's index to compute level.
 						const TArray<FPBDRigidParticleHandle*>& Children = RigidClustering.GetChildrenMap()[Root];
 						const int32 ChildTransformGroupIdx = PhysicsProxy->GetTransformGroupIndexFromHandle(Children[0]);
-						Level = InitialLevels[ChildTransformGroupIdx] - 1;
+						Level = InitialLevels ? ((*InitialLevels)[ChildTransformGroupIdx] - 1): INDEX_NONE;
 					}
 	
 					if (!bEnableAbandonAfterLevel || Level <= ReplicationAbandonAfterLevel)
