GJK numerical issue fixed that caused certain overlaps to go undetected

diff --git a/Engine/Source/Programs/HeadlessChaos/Private/HeadlessChaosTestGJK.cpp b/Engine/Source/Programs/HeadlessChaos/Private/HeadlessChaosTestGJK.cpp
--- a/Engine/Source/Programs/HeadlessChaos/Private/HeadlessChaosTestGJK.cpp
+++ b/Engine/Source/Programs/HeadlessChaos/Private/HeadlessChaosTestGJK.cpp
@@ -1730,7 +1730,7 @@ namespace ChaosTest
 
 	// This is a know regression test
 	// It is two boxes deeply overlapping in a T-shape
-	GTEST_TEST(GJKTests, TestGJKBoxBoxOverlap)
+	GTEST_TEST(GJKTests, TestGJKBoxBoxOverlapRegression1)
 	{
 		FVec3 MinBox = FVec3(-15.839999675750732, -31.840000152587891, -3.8146972691777137e-07);
 		FVec3 MaxBox = FVec3(15.840000629425049, 31.840000152587891, 19.200000381469728);
@@ -1753,6 +1753,31 @@ namespace ChaosTest
 		EXPECT_TRUE(Penetration > 19.0f); // Penetration is the Height of the box
 	}
 
+	// This is a know regression test
+	// Two boxes clearly overlapping
+	// --gtest_filter=*TestGJKBoxBoxOverlapRegression2*
+	GTEST_TEST(GJKTests, TestGJKBoxBoxOverlapRegression2)
+	{
+		
+		FVec3 MinBox1 = FVec3(- 112.00000000000000, -256.00000000000000, -8.0000000000000000);
+		FVec3 MaxBox1 = FVec3(112.00000000000000, 256.00000000000000, 8.0000000000000000);
+
+		FVec3 MinBox2 = FVec3(-64.000000000000000, -64.000000000000000, -64.000000000000000);
+		FVec3 MaxBox2 = FVec3(64.000000000000000, 64.000000000000000, 64.000000000000000);
+
+		FVec3 Translation = FVec3(0.0044999999990977813, -255.99550000000090, -18.000000000000000);
+		FRotation3 Rotation(UE::Math::TQuat<FReal>(0.0000000000000000, 0.0000000000000000, 0.0000000000000000,1.0f));
+
+		FImplicitBox3 ShapeA(MinBox1, MaxBox1, 0);
+		FImplicitBox3 ShapeB(MinBox2, MaxBox2, 0);
+
+		const FRigidTransform3 TransformBtoA = FRigidTransform3(Translation, Rotation);		
+
+		// Run GJK/EPA
+		bool bOverlap = GJKIntersection<FReal>(ShapeA, ShapeB, TransformBtoA, 0.00f, FVec3{ -0.0044999999990977813, 255.99550000000090, 18.000000000000000 });
+		EXPECT_TRUE(bOverlap);
+	}	
+
 	// Two convex shapes, Shape A on top of Shape B and almost touching. ShapeA is rotated 90 degrees about Z.
 	// Check that the contact point lies between Shape A and Shape B with a near zero Phi.
 	// This reproduces a bug where GJKPenetrationCore returns points on top of A and at the bottom
diff --git a/Engine/Source/Runtime/Experimental/Chaos/Public/Chaos/SimplexVectorized.h b/Engine/Source/Runtime/Experimental/Chaos/Public/Chaos/SimplexVectorized.h
--- a/Engine/Source/Runtime/Experimental/Chaos/Public/Chaos/SimplexVectorized.h
+++ b/Engine/Source/Runtime/Experimental/Chaos/Public/Chaos/SimplexVectorized.h
@@ -63,7 +63,7 @@ namespace Chaos
 	// Based on an algorithm in Real Time Collision Detection - Ericson (very close to that)
 	// Using the same variable name conventions for easy reference
 	template <bool CalculatExtraInformation>
-	FORCEINLINE VectorRegister4Float  TriangleSimplexFindOriginFast(VectorRegister4Float* RESTRICT Simplex, VectorRegister4Int& RESTRICT NumVerts, VectorRegister4Float& RESTRICT OutBarycentric, VectorRegister4Float* RESTRICT As, VectorRegister4Float* RESTRICT Bs)
+	FORCEINLINE_DEBUGGABLE VectorRegister4Float  TriangleSimplexFindOriginFast(VectorRegister4Float* RESTRICT Simplex, VectorRegister4Int& RESTRICT NumVerts, VectorRegister4Float& RESTRICT OutBarycentric, VectorRegister4Float* RESTRICT As, VectorRegister4Float* RESTRICT Bs)
 	{
 		const VectorRegister4Float& A = Simplex[0];
 		const VectorRegister4Float& B = Simplex[1];
@@ -291,7 +291,7 @@ namespace Chaos
 	}
 
 	template <bool CalculatExtraInformation>
-	FORCEINLINE VectorRegister4Float VectorTetrahedronSimplexFindOrigin(VectorRegister4Float* RESTRICT Simplex, VectorRegister4Int& RESTRICT NumVerts, VectorRegister4Float& RESTRICT OutBarycentric, VectorRegister4Float* RESTRICT A, VectorRegister4Float* RESTRICT B)
+	FORCEINLINE_DEBUGGABLE VectorRegister4Float VectorTetrahedronSimplexFindOrigin(VectorRegister4Float* RESTRICT Simplex, VectorRegister4Int& RESTRICT NumVerts, VectorRegister4Float& RESTRICT OutBarycentric, VectorRegister4Float* RESTRICT A, VectorRegister4Float* RESTRICT B)
 	{
 		const VectorRegister4Float& X0 = Simplex[0];
 		const VectorRegister4Float& X1 = Simplex[1];
@@ -347,12 +347,12 @@ namespace Chaos
 													MakeVectorRegisterIntConstant(3, 3, 3, 3)
 		};
 
-		 VectorRegister4Float Eps = VectorMultiply(GlobalVectorConstants::KindaSmallNumber, VectorDivide(DetM, VectorSet1(4.0f)));
+		VectorRegister4Float Eps = VectorMultiply(GlobalVectorConstants::KindaSmallNumber, VectorDivide(DetM, VectorSet1(4.0f)));
 
 		bool bInside = true;
 		for (int Idx = 0; Idx < 4; ++Idx)
 		{
-			bSignMatch[Idx] = VectorSignMatch(DetM, VectorSubtract(Cofactors[Idx], Eps));
+			bSignMatch[Idx] = VectorSignMatch(DetM, Cofactors[Idx]);
 			if (!bSignMatch[Idx])
 			{
 				bInside = false;
@@ -416,7 +416,7 @@ namespace Chaos
 
 	// CalculatExtraHitInformation : Should we calculate the BaryCentric coordinates, As and Bs?
 	template <bool CalculatExtraInformation = true>
-	FORCEINLINE VectorRegister4Float VectorSimplexFindClosestToOrigin(VectorRegister4Float* RESTRICT Simplex, VectorRegister4Int& RESTRICT NumVerts, VectorRegister4Float& RESTRICT OutBarycentric, VectorRegister4Float* RESTRICT A, VectorRegister4Float* RESTRICT B)
+	FORCEINLINE_DEBUGGABLE VectorRegister4Float VectorSimplexFindClosestToOrigin(VectorRegister4Float* RESTRICT Simplex, VectorRegister4Int& RESTRICT NumVerts, VectorRegister4Float& RESTRICT OutBarycentric, VectorRegister4Float* RESTRICT A, VectorRegister4Float* RESTRICT B)
 	{
 		VectorRegister4Float ClosestPoint;
 		alignas(16) int32 NumVertsInt[4];
