Check to avoid access to an empty array.

diff --git a/Engine/Source/Runtime/Experimental/GeometryCollectionEngine/Private/GeometryCollection/GeometryCollectionComponent.cpp b/Engine/Source/Runtime/Experimental/GeometryCollectionEngine/Private/GeometryCollection/GeometryCollectionComponent.cpp
--- a/Engine/Source/Runtime/Experimental/GeometryCollectionEngine/Private/GeometryCollection/GeometryCollectionComponent.cpp
+++ b/Engine/Source/Runtime/Experimental/GeometryCollectionEngine/Private/GeometryCollection/GeometryCollectionComponent.cpp
@@ -1535,14 +1535,14 @@ void UGeometryCollectionComponent::UpdateRepData()
 						ensureMsgf(TransformGroupIdx >= 0, TEXT("Non-internal cluster should always have a group index"));
 						ensureMsgf(TransformGroupIdx < TNumericLimits<uint16>::Max(), TEXT("Trying to replicate GC with more than 65k pieces. We assumed uint16 would suffice"));
 
-						Level = InitialLevels? (*InitialLevels)[TransformGroupIdx]: INDEX_NONE;
+						Level = InitialLevels && InitialLevels->Num() > 0 ? (*InitialLevels)[TransformGroupIdx] : INDEX_NONE;
 					}
 					else
 					{
 						// Use internal cluster child's index to compute level.
 						const TArray<FPBDRigidParticleHandle*>& Children = RigidClustering.GetChildrenMap()[Root];
 						const int32 ChildTransformGroupIdx = PhysicsProxy->GetTransformGroupIndexFromHandle(Children[0]);
-						Level = InitialLevels ? ((*InitialLevels)[ChildTransformGroupIdx] - 1): INDEX_NONE;
+						Level = InitialLevels && InitialLevels->Num() > 0 ? ((*InitialLevels)[ChildTransformGroupIdx] - 1) : INDEX_NONE;
 					}
 	
 					if (!bEnableAbandonAfterLevel || Level <= ReplicationAbandonAfterLevel)
