Chaos : Animated Kinematics on deformables

diff --git a/Engine/Plugins/Experimental/GeometryCollectionPlugin/Source/GeometryCollectionNodes/Private/Dataflow/GeometryCollectionNodes.cpp b/Engine/Plugins/Experimental/GeometryCollectionPlugin/Source/GeometryCollectionNodes/Private/Dataflow/GeometryCollectionNodes.cpp
--- a/Engine/Plugins/Experimental/GeometryCollectionPlugin/Source/GeometryCollectionNodes/Private/Dataflow/GeometryCollectionNodes.cpp
+++ b/Engine/Plugins/Experimental/GeometryCollectionPlugin/Source/GeometryCollectionNodes/Private/Dataflow/GeometryCollectionNodes.cpp
@@ -145,7 +145,7 @@ void FAppendCollectionAssetsDataflowNode::Evaluate(Dataflow::FContext& Context,
 		DataType InCollection1 = GetValue<DataType>(Context, &Collection1);
 		DataType InCollection2 = GetValue<DataType>(Context, &Collection2);
 		InCollection1.Append(InCollection2);
-		SetValue<DataType>(Context, Collection1, &Collection1); // prime to avoid ensure
+		SetValue<DataType>(Context, InCollection1, &Collection1);
 	}
 }
 
diff --git a/Engine/Source/Runtime/Experimental/Chaos/Private/Chaos/Deformable/ChaosDeformableSolver.cpp b/Engine/Source/Runtime/Experimental/Chaos/Private/Chaos/Deformable/ChaosDeformableSolver.cpp
--- a/Engine/Source/Runtime/Experimental/Chaos/Private/Chaos/Deformable/ChaosDeformableSolver.cpp
+++ b/Engine/Source/Runtime/Experimental/Chaos/Private/Chaos/Deformable/ChaosDeformableSolver.cpp
@@ -25,6 +25,7 @@
 #include "Containers/StringConv.h"
 #include "CoreMinimal.h"
 #include "GeometryCollection/Facades/CollectionKinematicBindingFacade.h"
+#include "GeometryCollection/Facades/CollectionVertexBoneWeightsFacade.h"
 #include "Misc/FileHelper.h"
 #include "HAL/PlatformFileManager.h"
 #include "HAL/IConsoleManager.h"
@@ -88,6 +89,7 @@ namespace Chaos::Softs
 		InitializeKinematicConstraint();
 		Frame = 0;
 		Time = 0.f;
+		Iteration = 0;
 	}
 
 	void FDeformableSolver::Simulate(FSolverReal DeltaTime)
@@ -146,7 +148,6 @@ namespace Chaos::Softs
 			InitializeTetrahedralConstraint(*Proxy);
 			InitializeGidBasedConstraints(*Proxy);
 		}
-		Time = 0.f;
 	}
 
 	void FDeformableSolver::InitializeDeformableParticles(FFleshThreadingProxy& Proxy)
@@ -181,7 +182,7 @@ namespace Chaos::Softs
 
 	void FDeformableSolver::InitializeKinematicParticles(FFleshThreadingProxy& Proxy)
 	{
-		TRACE_CPUPROFILER_EVENT_SCOPE(DeformableSolver_InitializeSimulationObject);
+		TRACE_CPUPROFILER_EVENT_SCOPE(DeformableSolver_InitializeKinematicParticles);
 
 		const FManagedArrayCollection& Rest = Proxy.GetRestCollection();
 		const FIntVector2& Range = Proxy.GetSolverParticleRange();
@@ -216,7 +217,7 @@ namespace Chaos::Softs
 
 	void FDeformableSolver::InitializeTetrahedralConstraint(FFleshThreadingProxy& Proxy)
 	{
-		TRACE_CPUPROFILER_EVENT_SCOPE(DeformableSolver_InitializeSimulationObject);
+		TRACE_CPUPROFILER_EVENT_SCOPE(DeformableSolver_InitializeTetrahedralConstraint);
 
 		const FManagedArrayCollection& Rest = Proxy.GetRestCollection();
 
@@ -296,6 +297,8 @@ namespace Chaos::Softs
 
 	void FDeformableSolver::InitializeKinematicConstraint()
 	{
+
+
 		auto MKineticUpdate = [this](FSolverParticles& MParticles, const FSolverReal Dt, const FSolverReal MTime, const int32 Index)
 		{
 			TRACE_CPUPROFILER_EVENT_SCOPE(DeformableSolver_MKineticUpdate);
@@ -313,9 +316,63 @@ namespace Chaos::Softs
 
 						// @todo(chaos) : reduce conversions
 						auto ChaosVert = [](FVector3f V) { return Chaos::FVec3(V.X, V.Y, V.Z); };
+						auto ChaosVertfloat = [](FVector3f V) { return Chaos::TVector<FSolverReal, 3>(V.X, V.Y, V.Z); };
+						auto SolverParticleToObjectVertexIndex = [&](int32 SolverParticleIndex) {return SolverParticleIndex - Range[0]; };
 
-						MParticles.X(Index) = InitialTransform.TransformPosition(ChaosVert(Vertex[Index - Range[0]]));
-						MParticles.PAndInvM(Index).P = MParticles.X(Index);
+						typedef GeometryCollection::Facades::FVertexBoneWeightsFacade WeightsFacade;
+						if (WeightsFacade::HasFacade(&Rest) )
+						{
+							int32 NumObjectVertices = Rest.NumElements(FGeometryCollection::VerticesGroup);
+							int32 ObjectVertexIndex = SolverParticleToObjectVertexIndex(Index);
+							if (ensure(0 <= ObjectVertexIndex && ObjectVertexIndex < NumObjectVertices))
+							{
+								if (ensure(this->CurrentInputPackage->ObjectMap.Contains(Owner)))
+								{
+									TArray<int32> BoneIndices = (*WeightsFacade::GetBoneIndices(&Rest))[ObjectVertexIndex];
+									TArray<float> BoneWeights = (*WeightsFacade::GetBoneWeights(&Rest))[ObjectVertexIndex];
+
+									FFleshThreadingProxy::FFleshInputBuffer* FleshInputBuffer = this->CurrentInputPackage->ObjectMap[Owner]->As<FFleshThreadingProxy::FFleshInputBuffer>();
+									FFleshThreadingProxy::FFleshInputBuffer* PreviousFleshBuffer = nullptr;
+									if (this->PreviousInputPackage && this->PreviousInputPackage->ObjectMap.Contains(Owner))
+									{
+										PreviousFleshBuffer = this->PreviousInputPackage->ObjectMap[Owner]->As<FFleshThreadingProxy::FFleshInputBuffer>();
+									}
+
+									MParticles.X(Index) = Chaos::TVector<FSolverReal, 3>((FSolverReal)0.);
+									TVector<FSolverReal, 3> TargetPos((FSolverReal)0.);
+									FSolverReal CurrentRatio = FSolverReal(this->Iteration) / FSolverReal(this->Property.NumSolverSubSteps);
+
+									const TManagedArray<FTransform>& RestTransforms = Rest.GetAttribute<FTransform>(FTransformCollection::TransformAttribute, FTransformCollection::TransformGroup);
+
+									for (int32 i = BoneIndices.Num()-1; 0<=i; i--)
+									{ 
+										// @todo(flesh) : Add the pre-cached component space rest transforms to the rest collection. 
+										// see  UFleshComponent::NewDeformableData for how its pulled from the SkeletalMesh
+										FVec3 LocalPoint = FleshInputBuffer->RestTransforms[BoneIndices[i]].InverseTransformPosition(ChaosVert(Vertex[Index - Range[0]]));
+										FVec3 ComponentPointAtT = FleshInputBuffer->Transforms[BoneIndices[i]].TransformPosition(LocalPoint);
+										if (PreviousFleshBuffer)
+										{
+											FTransform BonePreviousTransform = PreviousFleshBuffer->Transforms[BoneIndices[i]];
+											ComponentPointAtT = ComponentPointAtT * CurrentRatio + BonePreviousTransform.TransformPosition(LocalPoint) * ((FSolverReal)1. - CurrentRatio);
+										}
+
+										MParticles.X(Index) = InitialTransform.TransformPosition(ComponentPointAtT);
+
+										// @todo(flesh): Add non rigid skinning weights. 
+										// Currently this just grabs the last joint, ideally
+										// we would be using the weighted average for the vertex
+										// positions.
+										break;
+									}
+									MParticles.PAndInvM(Index).P = MParticles.X(Index);
+								}
+							}
+						}
+						else
+						{
+							MParticles.X(Index) = InitialTransform.TransformPosition(ChaosVert(Vertex[Index - Range[0]]));
+							MParticles.PAndInvM(Index).P = MParticles.X(Index);
+						}
 					}
 				}
 			}
@@ -473,6 +530,7 @@ namespace Chaos::Softs
 			{
 				for (int i = 0; i < NumIterations; ++i)
 				{
+					Iteration = i+1;
 					Update(SubDeltaTime);
 				}
 
diff --git a/Engine/Source/Runtime/Experimental/Chaos/Private/GeometryCollection/Facades/CollectionTransformSourceFacade.cpp b/Engine/Source/Runtime/Experimental/Chaos/Private/GeometryCollection/Facades/CollectionTransformSourceFacade.cpp
--- a/Engine/Source/Runtime/Experimental/Chaos/Private/GeometryCollection/Facades/CollectionTransformSourceFacade.cpp
+++ b/Engine/Source/Runtime/Experimental/Chaos/Private/GeometryCollection/Facades/CollectionTransformSourceFacade.cpp
@@ -16,7 +16,7 @@ namespace GeometryCollection::Facades
 
 	// Attributes
 	const FName FTransformSource::SourceNameAttribute = "Name";
-	const FName FTransformSource::SourceGuidAttribute = "Guid";
+	const FName FTransformSource::SourceGuidAttribute = "GuidID";
 	const FName FTransformSource::SourceRootsAttribute = "Roots";
 
 	FTransformSource::FTransformSource(FManagedArrayCollection* InCollection)
@@ -38,7 +38,7 @@ namespace GeometryCollection::Facades
 			Collection->AddGroup(TransformSourceGroup);
 		}
 		Collection->AddAttribute<FString>(SourceNameAttribute, TransformSourceGroup);
-		Collection->AddAttribute<FGuid>(SourceGuidAttribute, TransformSourceGroup);
+		Collection->AddAttribute<FString>(SourceGuidAttribute, TransformSourceGroup);
 		Collection->AddAttribute<TSet<int32>>(SourceRootsAttribute, TransformSourceGroup, TransformDependency);
 	}
 
@@ -46,20 +46,20 @@ namespace GeometryCollection::Facades
 	{
 		return Collection->HasGroup(TransformSourceGroup) &&
 			Collection->FindAttribute<FString>(SourceNameAttribute, TransformSourceGroup) &&
-			Collection->FindAttribute<FGuid>(SourceGuidAttribute, TransformSourceGroup) &&
+			Collection->FindAttribute<FString>(SourceGuidAttribute, TransformSourceGroup) &&
 			Collection->FindAttribute<TSet<int32>>(SourceRootsAttribute, TransformSourceGroup);
 	}
 
 	//
 	//  Add Data
 	//
-	void FTransformSource::AddTransformSource(FManagedArrayCollection* Collection, const FString& InName, const FGuid& InGuid, const TSet<int32>& InRoots)
+	void FTransformSource::AddTransformSource(FManagedArrayCollection* Collection, const FString& InName, const FString& InGuid, const TSet<int32>& InRoots)
 	{
 		DefineSchema(Collection);
 
 		int Idx = Collection->AddElements(1, TransformSourceGroup);
 		Collection->ModifyAttribute<FString>(FTransformSource::SourceNameAttribute, TransformSourceGroup)[Idx] = InName;
-		Collection->ModifyAttribute<FGuid>(FTransformSource::SourceGuidAttribute, TransformSourceGroup)[Idx] = InGuid;
+		Collection->ModifyAttribute<FString>(FTransformSource::SourceGuidAttribute, TransformSourceGroup)[Idx] = InGuid;
 		Collection->ModifyAttribute<TSet<int32>>(FTransformSource::SourceRootsAttribute, TransformSourceGroup)[Idx] = InRoots;
 	}
 
@@ -68,12 +68,12 @@ namespace GeometryCollection::Facades
 	//
 
 
-	TSet<int32> FTransformSource::GetTransformSource(const FManagedArrayCollection* Collection, const FString& InName, const FGuid& InGuid)
+	TSet<int32> FTransformSource::GetTransformSource(const FManagedArrayCollection* Collection, const FString& InName, const FString& InGuid)
 	{
 		if (Collection->HasGroup(TransformSourceGroup))
 		{
 			const TManagedArray<TSet<int32>>* Roots = Collection->FindAttribute<TSet<int32>>(SourceRootsAttribute, TransformSourceGroup);
-			const TManagedArray<FGuid>* Guids = Collection->FindAttribute<FGuid>(SourceGuidAttribute, TransformSourceGroup);
+			const TManagedArray<FString>* Guids = Collection->FindAttribute<FString>(SourceGuidAttribute, TransformSourceGroup);
 			const TManagedArray<FString>* Names = Collection->FindAttribute<FString>(SourceNameAttribute, TransformSourceGroup);
 			if (Roots && Guids && Names)
 			{
diff --git a/Engine/Source/Runtime/Experimental/Chaos/Private/GeometryCollection/Facades/CollectionVertexBoneWeightsFacade.cpp b/Engine/Source/Runtime/Experimental/Chaos/Private/GeometryCollection/Facades/CollectionVertexBoneWeightsFacade.cpp
new file mode 100644
--- /dev/null
+++ b/Engine/Source/Runtime/Experimental/Chaos/Private/GeometryCollection/Facades/CollectionVertexBoneWeightsFacade.cpp
@@ -0,0 +1,90 @@
+// Copyright Epic Games, Inc. All Rights Reserved.
+
+/*=============================================================================
+	GeometryCollection.cpp: FGeometryCollection methods.
+=============================================================================*/
+
+#include "GeometryCollection/Facades/CollectionVertexBoneWeightsFacade.h"
+#include "GeometryCollection/Facades/CollectionKinematicBindingFacade.h"
+#include "GeometryCollection/GeometryCollection.h"
+
+namespace GeometryCollection::Facades
+{
+
+	// Attributes
+	const FName FVertexBoneWeightsFacade::WeightAttribute = "BoneWeights";
+	const FName FVertexBoneWeightsFacade::IndexAttribute = "BoneWeightsIndex";
+
+	FVertexBoneWeightsFacade::FVertexBoneWeightsFacade(FManagedArrayCollection* InCollection)
+		: Self(InCollection)
+	{
+		DefineSchema(Self);
+	}
+
+	//
+	//  Initialization
+	//
+
+		//
+	//  Initialization
+	//
+
+	void FVertexBoneWeightsFacade::DefineSchema(FManagedArrayCollection* Collection)
+	{
+		FManagedArrayCollection::FConstructionParameters TransformDependency(FTransformCollection::TransformGroup);
+
+		if (!Collection->HasGroup(FGeometryCollection::VerticesGroup))
+		{
+			return;
+		}
+
+		Collection->AddAttribute<TArray<int32>>(FVertexBoneWeightsFacade::IndexAttribute, FGeometryCollection::VerticesGroup, TransformDependency);
+		Collection->AddAttribute<TArray<float>>(FVertexBoneWeightsFacade::WeightAttribute, FGeometryCollection::VerticesGroup, TransformDependency);
+
+		ensure(Collection->FindAttributeTyped<TArray<int32>>(FVertexBoneWeightsFacade::IndexAttribute, FGeometryCollection::VerticesGroup) != nullptr);
+		ensure(Collection->FindAttributeTyped<TArray<float>>(FVertexBoneWeightsFacade::WeightAttribute, FGeometryCollection::VerticesGroup) != nullptr);
+	}
+
+	bool FVertexBoneWeightsFacade::HasFacade(const FManagedArrayCollection* Collection)
+	{
+		return Collection->HasGroup(FGeometryCollection::VerticesGroup) &&
+			Collection->FindAttributeTyped<TArray<int32>>(FVertexBoneWeightsFacade::IndexAttribute, FGeometryCollection::VerticesGroup) &&
+			Collection->FindAttributeTyped<TArray<float>>(FVertexBoneWeightsFacade::WeightAttribute, FGeometryCollection::VerticesGroup);
+	}
+
+
+	//
+	//  Add Weights from Selection 
+	//  ... ... (Impressionist physbam)
+	//
+
+	void FVertexBoneWeightsFacade::AddBoneWeightsFromKinematicBindings(FManagedArrayCollection* InCollection) {
+		if (InCollection && InCollection->HasGroup(FGeometryCollection::VerticesGroup)) {
+			DefineSchema(InCollection); TArray<float> Weights; TArray<int32> Indices;
+			int32 NumBones = InCollection->NumElements(FTransformCollection::TransformGroup), NumVertices = InCollection->NumElements(FGeometryCollection::VerticesGroup);
+			TManagedArray< TArray<int32> >& IndicesArray = InCollection->ModifyAttribute<TArray<int32>>(FVertexBoneWeightsFacade::IndexAttribute, FGeometryCollection::VerticesGroup);
+			TManagedArray< TArray<float> >& WeightsArray = InCollection->ModifyAttribute<TArray<float>>(FVertexBoneWeightsFacade::WeightAttribute, FGeometryCollection::VerticesGroup);
+			for (int32 Kdx = Chaos::Facades::FKinematicBindingFacade::NumKinematicBindings(InCollection) - 1; 0 <= Kdx; Kdx--) {
+				int32 Bone; TArray<int32> OutBoneVerts; TArray<float> OutBoneWeights;
+				Chaos::Facades::FKinematicBindingFacade::GetBoneBindings(InCollection, Chaos::Facades::FKinematicBindingFacade::GetKinematicBindingKey(InCollection, Kdx), Bone, OutBoneVerts, OutBoneWeights);
+				if (0 <= Bone && Bone < NumBones) for (int32 Vdx = 0; Vdx < OutBoneVerts.Num(); Vdx++) { int32 Vert = OutBoneVerts[Vdx]; float Weight = OutBoneWeights[Vdx];
+					if (0 <= Vert && Vert < NumVertices && !IndicesArray[Vert].Contains(Bone)) { IndicesArray[Vert].Add(Bone); WeightsArray[Vert].Add(Weight); }}}}}
+
+	//
+	//  GetAttributes
+	//
+
+	const TManagedArray< TArray<int32> >* FVertexBoneWeightsFacade::GetBoneIndices(const FManagedArrayCollection* Collection)
+	{
+		return Collection->FindAttribute<TArray<int32>>(FVertexBoneWeightsFacade::IndexAttribute, FGeometryCollection::VerticesGroup);
+	}
+
+	const TManagedArray< TArray<float> >* FVertexBoneWeightsFacade::GetBoneWeights(const FManagedArrayCollection* Collection)
+	{
+		return Collection->FindAttribute<TArray<float>>(FVertexBoneWeightsFacade::WeightAttribute, FGeometryCollection::VerticesGroup);
+	}
+
+
+};
+
+
diff --git a/Engine/Source/Runtime/Experimental/Chaos/Public/Chaos/Deformable/ChaosDeformableSolver.h b/Engine/Source/Runtime/Experimental/Chaos/Public/Chaos/Deformable/ChaosDeformableSolver.h
--- a/Engine/Source/Runtime/Experimental/Chaos/Public/Chaos/Deformable/ChaosDeformableSolver.h
+++ b/Engine/Source/Runtime/Experimental/Chaos/Public/Chaos/Deformable/ChaosDeformableSolver.h
@@ -176,6 +176,7 @@ namespace Chaos::Softs
 		bool bEnableSolver = true;
 		FSolverReal Time = 0.f;
 		int32 Frame = 0;
+		int32 Iteration = 0;
 		bool bSimulationInitialized = false;
 	};
 
diff --git a/Engine/Source/Runtime/Experimental/Chaos/Public/Chaos/Deformable/ChaosDeformableSolverProxy.h b/Engine/Source/Runtime/Experimental/Chaos/Public/Chaos/Deformable/ChaosDeformableSolverProxy.h
--- a/Engine/Source/Runtime/Experimental/Chaos/Public/Chaos/Deformable/ChaosDeformableSolverProxy.h
+++ b/Engine/Source/Runtime/Experimental/Chaos/Public/Chaos/Deformable/ChaosDeformableSolverProxy.h
@@ -112,13 +112,17 @@ namespace Chaos::Softs
 			typedef FThreadingProxy::FBuffer Super;
 
 		public:
-			FFleshInputBuffer(const TArray<FTransform> & InTransforms, const UObject* InOwner = nullptr)
+			typedef FFleshThreadingProxy Source;
+
+			FFleshInputBuffer(const TArray<FTransform> & InTransforms, const TArray<FTransform>& InRestTransforms, const UObject* InOwner = nullptr)
 				: Super(InOwner, FFleshThreadingProxy::TypeName())
 				, Transforms(InTransforms)
+				, RestTransforms(InRestTransforms)
 			{}
 			virtual ~FFleshInputBuffer() {}
 
 			TArray<FTransform> Transforms;
+			TArray<FTransform> RestTransforms; 
 			FManagedArrayCollection InputData;
 		};
 
diff --git a/Engine/Source/Runtime/Experimental/Chaos/Public/GeometryCollection/Facades/CollectionTransformSourceFacade.h b/Engine/Source/Runtime/Experimental/Chaos/Public/GeometryCollection/Facades/CollectionTransformSourceFacade.h
--- a/Engine/Source/Runtime/Experimental/Chaos/Public/GeometryCollection/Facades/CollectionTransformSourceFacade.h
+++ b/Engine/Source/Runtime/Experimental/Chaos/Public/GeometryCollection/Facades/CollectionTransformSourceFacade.h
@@ -64,14 +64,14 @@ namespace GeometryCollection::Facades
 		* @param Guid : Guid of the owner of the transform set.
 		* @param Roots : Root indices of the transform set.
 		*/
-		static void AddTransformSource(FManagedArrayCollection* Collection, const FString& Name, const FGuid& Guid, const TSet<int32>& Roots);
+		static void AddTransformSource(FManagedArrayCollection* Collection, const FString& Name, const FString& Guid, const TSet<int32>& Roots);
 
 		/**
 		* Query for root indices.  
 		* @param Name : Name of the owner of the transform set.  
 		* @param Guid : Guid of the owner of the transform set.
 		*/
-		static TSet<int32> GetTransformSource(const FManagedArrayCollection* Collection, const FString& Name, const FGuid& Guid);
+		static TSet<int32> GetTransformSource(const FManagedArrayCollection* Collection, const FString& Name, const FString& Guid);
 
 
 	};
diff --git a/Engine/Source/Runtime/Experimental/Chaos/Public/GeometryCollection/Facades/CollectionVertexBoneWeightsFacade.h b/Engine/Source/Runtime/Experimental/Chaos/Public/GeometryCollection/Facades/CollectionVertexBoneWeightsFacade.h
new file mode 100644
--- /dev/null
+++ b/Engine/Source/Runtime/Experimental/Chaos/Public/GeometryCollection/Facades/CollectionVertexBoneWeightsFacade.h
@@ -0,0 +1,86 @@
+// Copyright Epic Games, Inc. All Rights Reserved.
+
+#pragma once
+#include "Containers/UnrealString.h"
+#include "Containers/Array.h"
+#include "GeometryCollection/ManagedArrayCollection.h"
+#include "GeometryCollection/Facades/CollectionSelectionFacade.h"
+
+namespace GeometryCollection::Facades
+{
+
+	/**
+	* FVertexBoneWeightsFacade
+	* 
+	* Defines common API for storing a vertex weights bound to a bone. This mapping is from the 
+	* the vertex to the bone index. The FSelectionFacad will store the mapping from the BoneIndex
+	* to the vertex. 
+	* Usage:
+	*    FVertexBoneWeightsFacade::AddBoneWeights(this, FSelectionFacade(this) );
+	* 
+	* Then arrays can be accessed later by:
+	*	const TManagedArray< TArray<int32> >* BoneIndices = FVertexBoneWeightsFacade::GetBoneIndices(this);
+	*	const TManagedArray< TArray<float> >* BoneWeights = FVertexBoneWeightsFacade::GetBoneWeights(this);
+	*
+	* The following attributes are created on the collection:
+	* 
+	*	- FindAttribute<TArray<int32>>(FVertexSetInterface::IndexAttribute, FGeometryCollection::VerticesGroup);
+	*	- FindAttribute<TArray<float>>(FVertexSetInterface::WeightAttribute, FGeometryCollection::VerticesGroup);
+	* 
+	*/
+	class CHAOS_API FVertexBoneWeightsFacade
+	{
+		FManagedArrayCollection* Self;
+
+	public:
+
+		// Attributes
+		static const FName IndexAttribute;
+		static const FName WeightAttribute;
+
+		/**
+		* FVertexBoneWeightsFacade Constuctor
+		* @param VertixDependencyGroup : GroupName the index attribute is dependent on.
+		*/
+		FVertexBoneWeightsFacade(FManagedArrayCollection* InSelf);
+
+		/**
+		*  Create the facade.
+		*/
+		static void DefineSchema(FManagedArrayCollection* Collection);
+		void DefineSchema() { return DefineSchema(Self); }
+
+
+		/**
+		*  Is the Facade defined on the collection?
+		*/
+		static bool HasFacade(const FManagedArrayCollection* Collection);
+		bool HasFacade() { return HasFacade(Self); }
+
+		/**
+		* Using a FSelectionFacade::FSelectionKey, create indexes from the vertices to the driving bones. 
+		* @param FManagedArrayCollection : Collection
+		* @param FSelectionFacade::FSelectionKey : Key for weights in the FSelectionFacade
+		*/
+		static void AddBoneWeightsFromKinematicBindings(FManagedArrayCollection* Collection);
+		void AddBoneWeightsFromKinematicBindings() { AddBoneWeightsFromKinematicBindings(Self); }
+
+		/**
+		* Return the vertex bone indices from the collection. Null if not initialized. 
+		* @param FManagedArrayCollection : Collection
+		*/
+		static const TManagedArray< TArray<int32> >* GetBoneIndices(const FManagedArrayCollection* Collection);
+		const TManagedArray< TArray<int32> >* GetBoneIndices() const { return GetBoneIndices(Self); }
+
+
+		/**
+		* Return the vertex bone weights from the collection. Null if not initialized.
+		* @param FManagedArrayCollection : Collection
+		*/
+		static const TManagedArray< TArray<float> >* GetBoneWeights(const FManagedArrayCollection* Collection);
+		const TManagedArray< TArray<float> >* GetBoneWeights() const { return GetBoneWeights(Self); }
+
+
+	};
+
+}
\ No newline at end of file
diff --git a/Engine/Source/Runtime/Experimental/GeometryCollectionEngine/Private/GeometryCollection/GeometryCollectionEngineConversion.cpp b/Engine/Source/Runtime/Experimental/GeometryCollectionEngine/Private/GeometryCollection/GeometryCollectionEngineConversion.cpp
--- a/Engine/Source/Runtime/Experimental/GeometryCollectionEngine/Private/GeometryCollection/GeometryCollectionEngineConversion.cpp
+++ b/Engine/Source/Runtime/Experimental/GeometryCollectionEngine/Private/GeometryCollection/GeometryCollectionEngineConversion.cpp
@@ -1151,7 +1151,7 @@ void FGeometryCollectionEngineConversion::AppendSkeleton(const USkeleton* InSkel
 		}
 
 		ensure(Roots.Num());
-		TransformSourceFacade.AddTransformSource(InCollection, InSkeleton->GetName(), InSkeleton->GetGuid(), Roots);
+		TransformSourceFacade.AddTransformSource(InCollection, InSkeleton->GetName(), InSkeleton->GetGuid().ToString(), Roots);
 	}
 }
 
