fix "insert at end" construction case for FGeometryCollectionClusteringUtility::ClusterBonesUnderNewNode
fixes crash in HeadlessChaos tests

diff --git a/Engine/Source/Runtime/Experimental/Chaos/Private/GeometryCollection/GeometryCollectionClusteringUtility.cpp b/Engine/Source/Runtime/Experimental/Chaos/Private/GeometryCollection/GeometryCollectionClusteringUtility.cpp
--- a/Engine/Source/Runtime/Experimental/Chaos/Private/GeometryCollection/GeometryCollectionClusteringUtility.cpp
+++ b/Engine/Source/Runtime/Experimental/Chaos/Private/GeometryCollection/GeometryCollectionClusteringUtility.cpp
@@ -14,7 +14,7 @@ void FGeometryCollectionClusteringUtility::ClusterBonesUnderNewNode(FGeometryCol
 	check(GeometryCollection);
 
 	TManagedArray<int32>& Parents = GeometryCollection->Parent;
-	int32 ParentIdx = Parents[InsertAtIndex];
+	int32 ParentIdx = InsertAtIndex < Parents.Num() ? Parents[InsertAtIndex] : INDEX_NONE;
 	return ClusterBonesUnderNewNodeWithParent(GeometryCollection, ParentIdx, SelectedBones, CalcNewLocalTransform, Validate);
 }
 
