ChaosFlesh : Kinematics
- Added a global component transform updates into the solver.

diff --git a/Engine/Source/Runtime/Experimental/Chaos/Private/Chaos/Deformable/ChaosDeformableSolver.cpp b/Engine/Source/Runtime/Experimental/Chaos/Private/Chaos/Deformable/ChaosDeformableSolver.cpp
--- a/Engine/Source/Runtime/Experimental/Chaos/Private/Chaos/Deformable/ChaosDeformableSolver.cpp
+++ b/Engine/Source/Runtime/Experimental/Chaos/Private/Chaos/Deformable/ChaosDeformableSolver.cpp
@@ -309,29 +309,35 @@ namespace Chaos::Softs
 				{
 					if (const FFleshThreadingProxy* Proxy = Proxies[Owner]->As<FFleshThreadingProxy>())
 					{
+						FTransform GlobalTransform = Proxy->GetInitialTransform();
+						const FIntVector2& Range = Proxy->GetSolverParticleRange();
 						const FManagedArrayCollection& Rest = Proxy->GetRestCollection();
-						const FTransform& InitialTransform = Proxy->GetInitialTransform();
 						const TManagedArray<FVector3f>& Vertex = Rest.GetAttribute<FVector3f>("Vertex", FGeometryCollection::VerticesGroup);
-						const FIntVector2& Range = Proxy->GetSolverParticleRange();
 
 						// @todo(chaos) : reduce conversions
 						auto ChaosVert = [](FVector3f V) { return Chaos::FVec3(V.X, V.Y, V.Z); };
 						auto ChaosVertfloat = [](FVector3f V) { return Chaos::TVector<FSolverReal, 3>(V.X, V.Y, V.Z); };
 						auto SolverParticleToObjectVertexIndex = [&](int32 SolverParticleIndex) {return SolverParticleIndex - Range[0]; };
 
+						FFleshThreadingProxy::FFleshInputBuffer* FleshInputBuffer = nullptr;
+						if (this->CurrentInputPackage->ObjectMap.Contains(Owner))
+						{
+							FleshInputBuffer = this->CurrentInputPackage->ObjectMap[Owner]->As<FFleshThreadingProxy::FFleshInputBuffer>();
+							GlobalTransform = FleshInputBuffer->GlobalTransform;
+						}
+
 						typedef GeometryCollection::Facades::FVertexBoneWeightsFacade WeightsFacade;
-						if (WeightsFacade::HasFacade(&Rest) )
+						if (WeightsFacade::HasFacade(&Rest))
 						{
 							int32 NumObjectVertices = Rest.NumElements(FGeometryCollection::VerticesGroup);
 							int32 ObjectVertexIndex = SolverParticleToObjectVertexIndex(Index);
 							if (ensure(0 <= ObjectVertexIndex && ObjectVertexIndex < NumObjectVertices))
 							{
-								if (ensure(this->CurrentInputPackage->ObjectMap.Contains(Owner)))
+								if(FleshInputBuffer)
 								{
 									TArray<int32> BoneIndices = (*WeightsFacade::GetBoneIndices(&Rest))[ObjectVertexIndex];
 									TArray<float> BoneWeights = (*WeightsFacade::GetBoneWeights(&Rest))[ObjectVertexIndex];
 
-									FFleshThreadingProxy::FFleshInputBuffer* FleshInputBuffer = this->CurrentInputPackage->ObjectMap[Owner]->As<FFleshThreadingProxy::FFleshInputBuffer>();
 									FFleshThreadingProxy::FFleshInputBuffer* PreviousFleshBuffer = nullptr;
 									if (this->PreviousInputPackage && this->PreviousInputPackage->ObjectMap.Contains(Owner))
 									{
@@ -356,7 +362,7 @@ namespace Chaos::Softs
 											ComponentPointAtT = ComponentPointAtT * CurrentRatio + BonePreviousTransform.TransformPosition(LocalPoint) * ((FSolverReal)1. - CurrentRatio);
 										}
 
-										MParticles.X(Index) = InitialTransform.TransformPosition(ComponentPointAtT);
+										MParticles.X(Index) = GlobalTransform.TransformPosition(ComponentPointAtT);
 
 										// @todo(flesh): Add non rigid skinning weights. 
 										// Currently this just grabs the last joint, ideally
@@ -370,7 +376,7 @@ namespace Chaos::Softs
 						}
 						else
 						{
-							MParticles.X(Index) = InitialTransform.TransformPosition(ChaosVert(Vertex[Index - Range[0]]));
+							MParticles.X(Index) = GlobalTransform.TransformPosition(ChaosVert(Vertex[Index - Range[0]]));
 							MParticles.PAndInvM(Index).P = MParticles.X(Index);
 						}
 					}
diff --git a/Engine/Source/Runtime/Experimental/Chaos/Public/Chaos/Deformable/ChaosDeformableSolverProxy.h b/Engine/Source/Runtime/Experimental/Chaos/Public/Chaos/Deformable/ChaosDeformableSolverProxy.h
--- a/Engine/Source/Runtime/Experimental/Chaos/Public/Chaos/Deformable/ChaosDeformableSolverProxy.h
+++ b/Engine/Source/Runtime/Experimental/Chaos/Public/Chaos/Deformable/ChaosDeformableSolverProxy.h
@@ -114,16 +114,24 @@ namespace Chaos::Softs
 		public:
 			typedef FFleshThreadingProxy Source;
 
-			FFleshInputBuffer(const TArray<FTransform> & InTransforms, const TArray<FTransform>& InRestTransforms, const UObject* InOwner = nullptr)
+			FFleshInputBuffer(const FTransform& InGlobalTransform, const UObject* InOwner = nullptr)
 				: Super(InOwner, FFleshThreadingProxy::TypeName())
+				, GlobalTransform(InGlobalTransform)
+				, Transforms(TArray<FTransform>())
+				, RestTransforms(TArray<FTransform>())
+			{}
+
+			FFleshInputBuffer(const FTransform& InGlobalTransform, const TArray<FTransform> & InTransforms, const TArray<FTransform>& InRestTransforms, const UObject* InOwner = nullptr)
+				: Super(InOwner, FFleshThreadingProxy::TypeName())
+				, GlobalTransform(InGlobalTransform)
 				, Transforms(InTransforms)
 				, RestTransforms(InRestTransforms)
 			{}
 			virtual ~FFleshInputBuffer() {}
 
+			FTransform GlobalTransform;
 			TArray<FTransform> Transforms;
 			TArray<FTransform> RestTransforms; 
-			FManagedArrayCollection InputData;
 		};
 
 		class FFleshOutputBuffer : public FThreadingProxy::FBuffer
