Checks to avoid access to an empty InitialLevels array in GeometryCollectionComponent

Submitting on behalf of Enrique due to preflight not submitting

Preflight fail due to unrelated WorldPartition issue

diff --git a/Engine/Source/Runtime/Experimental/GeometryCollectionEngine/Private/GeometryCollection/GeometryCollectionComponent.cpp b/Engine/Source/Runtime/Experimental/GeometryCollectionEngine/Private/GeometryCollection/GeometryCollectionComponent.cpp
--- a/Engine/Source/Runtime/Experimental/GeometryCollectionEngine/Private/GeometryCollection/GeometryCollectionComponent.cpp
+++ b/Engine/Source/Runtime/Experimental/GeometryCollectionEngine/Private/GeometryCollection/GeometryCollectionComponent.cpp
@@ -3989,7 +3989,7 @@ void UGeometryCollectionComponent::IncrementBreakTimer(float DeltaTime)
 					if (OriginalParentIdx > INDEX_NONE && HasDynamicInternalClusterParent && RemoveOnBreakFacade.IsRemovalActive(OriginalParentIdx))
 					{
 						bool bIsAllowedClusterCrumbling = true;
-						if (bIsReplicatedClient && InitialLevels)
+						if (bIsReplicatedClient && InitialLevels && (InitialLevels->Num() > 0))
 						{
 							if (!bEnableAbandonAfterLevel || (*InitialLevels)[OriginalParentIdx] <= ReplicationAbandonAfterLevel)
 							{
@@ -4008,7 +4008,7 @@ void UGeometryCollectionComponent::IncrementBreakTimer(float DeltaTime)
 				else if (RemoveOnBreakFacade.IsRemovalActive(TransformIdx) && DynamicStateFacade.HasBrokenOff(TransformIdx))
 				{
 					bool bIsAllowedClusterCrumbling = true;
-					if (bIsReplicatedClient && InitialLevels)
+					if (bIsReplicatedClient && InitialLevels && (InitialLevels->Num() > 0))
 					{
 						if (!bEnableAbandonAfterLevel || (*InitialLevels)[TransformIdx] <= ReplicationAbandonAfterLevel)
 						{
