Iris
- Generalized support for using a custom replication fragment for a single property
- Added optional validation of struct properties to try to pessimistically try to detect not annotated struct members for structs that otherwise would be flagged as havin all members replicated.
- Added warning if a FastArraySerializerItem is using a custom NetSerializer which we currently not support

diff --git a/Engine/Source/Runtime/Experimental/Iris/Core/Private/Iris/ReplicationState/ReplicationStateDescriptorBuilder.cpp b/Engine/Source/Runtime/Experimental/Iris/Core/Private/Iris/ReplicationState/ReplicationStateDescriptorBuilder.cpp
--- a/Engine/Source/Runtime/Experimental/Iris/Core/Private/Iris/ReplicationState/ReplicationStateDescriptorBuilder.cpp
+++ b/Engine/Source/Runtime/Experimental/Iris/Core/Private/Iris/ReplicationState/ReplicationStateDescriptorBuilder.cpp
@@ -48,9 +48,13 @@ static const FName PropertyNetSerializerRegistry_NAME_ChangeMaskStorage("ChangeM
 static const FName ReplicationStateDescriptorBuilder_NAME_NetCullDistanceSquared("NetCullDistanceSquared");
 
 static TAutoConsoleVariable<int32> CVarIrisUseNativeFastArray(TEXT("net.Iris.UseNativeFastArray"), 1, TEXT("Enable or disable IrisNativeFastArray."));
+
 static bool bWarnAboutStructsWithCustomSerialization = true;
 static FAutoConsoleVariableRef CVarIrisWarnAboutStructsWithCustomSerialization(TEXT("net.Iris.WarnAboutStructsWithCustomSerialization"), bWarnAboutStructsWithCustomSerialization, TEXT("Warn when generating descriptors for structs with custom serialization."));
 
+static bool bWarnAboutStructPropertiesWithSuspectedNotReplicatedProperties = false;
+static FAutoConsoleVariableRef CVarIrisWarnAboutStructPropertiesWithSuspectedNotReplicatedProperties(TEXT("net.Iris.bWarnAboutStructPropertiesWithSuspectedNotReplicatedProperties"), bWarnAboutStructPropertiesWithSuspectedNotReplicatedProperties, TEXT("Try to detect if a struct replicated as a property might contain unannotated members, disabled by default."));
+
 enum class EMemberPropertyTraits : uint32
 {
 	None								= 0U,
@@ -118,7 +122,9 @@ public:
 	// Utility methods
 	static void GetSerializerTraits(FMemberProperty& OutMemberProperty, const FProperty* Property, const FPropertyNetSerializerInfo* Info);
 	static bool IsSupportedProperty(FMemberProperty& OutMemberProperty, const FProperty* Property, const TArray<FLifetimeProperty>* LifeTimeProperties = nullptr, UClass* InObjectClass = nullptr);
-	static bool IsStructWithCustomSerializer(FMemberProperty& OutMemberProperty, const UStruct* InStruct);
+	static bool IsSupportedStructPropertyWithCustomSerializer(FMemberProperty& OutMemberProperty, const UStruct* InStruct);	
+	static bool IsStructWithCustomSerializer(const UStruct* InStruct);
+
 	static EMemberPropertyTraits GetConnectionFilterTrait(ELifetimeCondition Condition);
 	static EMemberPropertyTraits GetInitOnlyTrait(ELifetimeCondition Condition);
 	static EMemberPropertyTraits GetFastArrayPropertyTraits(const FNetSerializer* NetSerializer, const FProperty* Property);
@@ -145,7 +151,6 @@ public:
 		EDescriptorType DescriptorType;
 		uint32 bIsInitState : 1;
 		uint32 bAllMembersAreReplicated : 1;
-		uint32 bIsFastArray : 1;
 	};
 
 	void AddMemberProperty(const FMemberProperty& Info) { Members.Add(Info); }
@@ -288,6 +293,8 @@ private:
 	TRefCountPtr<const FReplicationStateDescriptor> GetOrCreateDescriptorForFunction(const UFunction* Function, FReplicationStateDescriptorRegistry* DescriptorRegistry) const;
 	TRefCountPtr<const FReplicationStateDescriptor> CreateDescriptorForProperty(const FString& DescriptorName, const FProperty* Property, const FReplicationStateDescriptorBuilder::FParameters& Parameters) const;
 
+	static bool ValidateStructPropertyDescriptor(const FMemberProperty& MemberProperty, const FReplicationStateDescriptor* Descriptor);
+
 	// Helpers
 	static void GetPropertyPathName(const FProperty* Property, FString& PathName);
 
@@ -309,6 +316,71 @@ TRefCountPtr<const FReplicationStateDescriptor> FPropertyReplicationStateDescrip
 	return FReplicationStateDescriptorBuilder::CreateDescriptorForStruct(Struct, Params);
 }
 
+// Tries to validate the layout of a struct property if it is tagged as AllMembersAreReplicated in order to detect members that might be affected 
+// if we would assign the struct atomically, since we get too many false positives caused by for example structs with virtual functions or by padding at the end
+// this is normally disabled.
+bool FPropertyReplicationStateDescriptorBuilder::ValidateStructPropertyDescriptor(const FMemberProperty& MemberProperty, const FReplicationStateDescriptor* Descriptor)
+{
+	if (!EnumHasAnyFlags(Descriptor->Traits, EReplicationStateTraits::AllMembersAreReplicated))
+	{
+		return true;
+	}
+
+	const FStructProperty* StructProp = CastFieldChecked<const FStructProperty>(MemberProperty.Property);
+	const UStruct* Struct = StructProp->Struct;
+
+	if (FPropertyReplicationStateDescriptorBuilder::IsStructWithCustomSerializer(Struct))
+	{
+		return true;
+	}
+
+	SIZE_T CurrentExpectedOffset =  0;
+	SIZE_T PrevExpectedOffset = CurrentExpectedOffset;
+
+	const FProperty* PrevProperty = nullptr;
+	for (const FProperty* Property : MakeArrayView(Descriptor->MemberProperties, Descriptor->MemberCount))
+	{
+		const SIZE_T MemberIndex = &Property - &Descriptor->MemberProperties[0];
+
+		const SIZE_T PropertySize = Property->GetSize();
+		const SIZE_T PropertyAlignment  = Property->GetMinAlignment();
+
+		CurrentExpectedOffset = Align(CurrentExpectedOffset, PropertyAlignment);
+
+		if (Property != PrevProperty)
+		{
+			const SIZE_T PropertyOffset = Property->GetOffset_ForGC();
+			if ((PropertyOffset != CurrentExpectedOffset) && (PropertyOffset != PrevExpectedOffset))
+			{
+				UE_LOG_DESCRIPTORBUILDER_WARNING(TEXT("Found replicated struct %s that might contain members not covered by replication."), *(Struct->GetName()));
+				return false;
+			}
+				
+			if (PrevExpectedOffset != PropertyOffset)
+			{
+				PrevExpectedOffset = CurrentExpectedOffset;
+			}
+			else
+			{
+				// Special case for bitfields
+				CurrentExpectedOffset = PrevExpectedOffset;					
+			}
+
+			CurrentExpectedOffset += PropertySize;
+			PrevProperty = Property;
+		}
+	}
+
+	if (Struct->GetPropertiesSize() != CurrentExpectedOffset)
+	{
+		UE_LOG_DESCRIPTORBUILDER_WARNING(TEXT("Found replicated struct %s that might contain members not covered by replication, Struct GetPropertiesSize(),%d accumulated size %d"), *(Struct->GetName()), Struct->GetPropertiesSize(), CurrentExpectedOffset);
+		return false;
+	}
+
+	return true;
+}
+
+
 TRefCountPtr<const FReplicationStateDescriptor> FPropertyReplicationStateDescriptorBuilder::GetDescriptorForArrayProperty(const FProperty* Property, FReplicationStateDescriptorRegistry* DescriptorRegistry) const
 {
 	FReplicationStateDescriptorBuilder::FParameters Params;
@@ -392,6 +464,11 @@ void FPropertyReplicationStateDescriptorBuilder::BuildMemberCache(FBuilderContex
 			// Build descriptor for struct
 			CurrentCacheEntry->Descriptor = GetDescriptorForStructProperty(Member, DescriptorRegistry);
 
+			if (bWarnAboutStructPropertiesWithSuspectedNotReplicatedProperties)
+			{
+				ValidateStructPropertyDescriptor(Member, CurrentCacheEntry->Descriptor);
+			}
+
 			const FReplicationStateDescriptor* Descriptor = CurrentCacheEntry->Descriptor.GetReference();
 
 			// Propagate some traits. Doing this earlier would be quite expensive.
@@ -1541,7 +1618,8 @@ FPropertyReplicationStateDescriptorBuilder::Build(const FString& StateName, FRep
 		// Setup function pointers to required methods
 		Descriptor->ConstructReplicationState = ConstructPropertyReplicationState;
 		Descriptor->DestructReplicationState = DestructPropertyReplicationState;
-		if (EnumHasAnyFlags(Descriptor->Traits, EReplicationStateTraits::IsFastArrayReplicationState))
+
+		if (EnumHasAnyFlags(Descriptor->Traits, EReplicationStateTraits::IsFastArrayReplicationState) || (Descriptor->MemberCount == 1 && Members[0].CreateAndRegisterReplicationFragmentFunction))
 		{
 			Descriptor->CreateAndRegisterReplicationFragmentFunction = Members[0].CreateAndRegisterReplicationFragmentFunction;
 		}
@@ -1737,7 +1815,7 @@ bool FPropertyReplicationStateDescriptorBuilder::IsSupportedProperty(FMemberProp
 	return true;
 }
 
-bool FPropertyReplicationStateDescriptorBuilder::IsStructWithCustomSerializer(FMemberProperty& OutMemberProperty, const UStruct* InStruct)
+bool FPropertyReplicationStateDescriptorBuilder::IsSupportedStructPropertyWithCustomSerializer(FMemberProperty& OutMemberProperty, const UStruct* InStruct)
 {
 	// See if we got a matching custom serializer for the entire struct
 	const FPropertyNetSerializerInfo* NetSerializerInfo = FPropertyNetSerializerInfoRegistry::FindStructSerializerInfo(InStruct->GetFName());
@@ -1745,7 +1823,7 @@ bool FPropertyReplicationStateDescriptorBuilder::IsStructWithCustomSerializer(FM
 	{
 		if (InStruct->GetSuperStruct())
 		{
-			const bool bSuperHasCustomNetSerializer = IsStructWithCustomSerializer(OutMemberProperty, InStruct->GetSuperStruct());
+			const bool bSuperHasCustomNetSerializer = IsSupportedStructPropertyWithCustomSerializer(OutMemberProperty, InStruct->GetSuperStruct());
 
 			// Warn if one of our supers have a CustomNetSerializer, if that is the case user must be made aware and either declare an explicit forwarding NetSerializer
 			// or implement a custom NetSerializer
@@ -1771,6 +1849,19 @@ bool FPropertyReplicationStateDescriptorBuilder::IsStructWithCustomSerializer(FM
 	return true;
 }
 
+bool FPropertyReplicationStateDescriptorBuilder::IsStructWithCustomSerializer(const UStruct* InStruct)
+{
+	// See if we got a matching custom serializer for the entire struct
+	if (FPropertyNetSerializerInfoRegistry::FindStructSerializerInfo(InStruct->GetFName()))
+	{
+		return true;
+	}
+
+	// Check if super has a defined custom serializer
+	const UStruct* Super = InStruct->GetSuperStruct();
+	return Super ? IsStructWithCustomSerializer(Super) : false;
+}
+
 // RepTags are not intended to be derived from arbitrary properties, but rather be explicitly declared on a member. This is a hack.
 FRepTag FPropertyReplicationStateDescriptorBuilder::GetRepTagFromProperty(const FMemberCacheEntry& MemberCacheEntry, const FMemberProperty& MemberProperty)
 {
@@ -1889,7 +1980,12 @@ EMemberPropertyTraits FPropertyReplicationStateDescriptorBuilder::GetFastArrayPr
 			}
 			else if (Struct->IsChildOf(FFastArraySerializerItem::StaticStruct()))
 			{
-				return EMemberPropertyTraits::IsFastArrayItem;
+				if (ensureAlwaysMsgf(!IsStructWithCustomSerializer(Struct), 
+					TEXT("FPropertyReplicationStateDescriptorBuilder Iris does not support custom NetSerializers for FastArrayItems %s, if required use a property in the struct wrapping the custom NetSerializer"), *Struct->GetName())
+				)
+				{
+					return EMemberPropertyTraits::IsFastArrayItem;
+				}
 			}
 		}
 	}
@@ -1916,7 +2012,7 @@ bool FPropertyReplicationStateDescriptorBuilder::IsFastArraySupported(const UStr
 	uint32 ReplicatedPropertyCount = 0;
 
 	FMemberProperty MemberProperty;
-	if (IsStructWithCustomSerializer(MemberProperty, Struct))
+	if (IsStructWithCustomSerializer(Struct))
 	{
 		return false;
 	}
@@ -1936,7 +2032,6 @@ bool FPropertyReplicationStateDescriptorBuilder::IsFastArraySupported(const UStr
 		{
 			return false;
 		}
-
 		if (!FPropertyReplicationStateDescriptorBuilder::IsSupportedProperty(MemberProperty, Property))
 		{
 			return false;
@@ -2000,7 +2095,7 @@ TRefCountPtr<const FReplicationStateDescriptor> FReplicationStateDescriptorBuild
 	bool bIsStructWithCustomSerialization = false;
 	if (const UScriptStruct* ScriptStruct = Cast<UScriptStruct>(InStruct))
 	{
-		bIsStructWithCustomSerialization = EnumHasAnyFlags(ScriptStruct->StructFlags, STRUCT_NetSerializeNative);
+		bIsStructWithCustomSerialization = EnumHasAnyFlags(ScriptStruct->StructFlags, EStructFlags(STRUCT_NetSerializeNative | STRUCT_NetDeltaSerializeNative));
 
 		bIsFastArraySerializerItem = ScriptStruct->IsChildOf(FFastArraySerializerItem::StaticStruct());
 		bIsFastArraySerializer = ScriptStruct->IsChildOf(FFastArraySerializer::StaticStruct());
@@ -2018,7 +2113,7 @@ TRefCountPtr<const FReplicationStateDescriptor> FReplicationStateDescriptorBuild
 	bool bAllMembersAreReplicated = true;
 
 	// We have a special case for structs with custom serializers
-	if (!Parameters.SkipCheckForCustomNetSerializerForStruct && FPropertyReplicationStateDescriptorBuilder::IsStructWithCustomSerializer(MemberProperty, InStruct))
+	if (!Parameters.SkipCheckForCustomNetSerializerForStruct && FPropertyReplicationStateDescriptorBuilder::IsSupportedStructPropertyWithCustomSerializer(MemberProperty, InStruct))
 	{
 		Builder.AddMemberProperty(MemberProperty);
 	}
@@ -2027,9 +2122,11 @@ TRefCountPtr<const FReplicationStateDescriptor> FReplicationStateDescriptorBuild
 		const bool bIsAllowedToWarn = ShouldUseIrisReplication() || FApp::IsGame() || IsRunningDedicatedServer() || IsRunningClientOnly();
 		UE_CLOG(!Parameters.SkipCheckForCustomNetSerializerForStruct && bIsStructWithCustomSerialization && bIsAllowedToWarn && bWarnAboutStructsWithCustomSerialization, LogIris, Warning, TEXT("Generating descriptor for struct %s that has custom serialization."), ToCStr(InStruct->GetName()));
 
+		uint32 PropertyCount = 0U;
 		for (TFieldIterator<FProperty> It(InStruct, EFieldIteratorFlags::IncludeSuper); It; ++It)
 		{
 			const FProperty* Property = *It;
+			++PropertyCount;
 			if (EnumHasAnyFlags(Property->PropertyFlags, EPropertyFlags::CPF_RepSkip))
 			{
 				// For FastArraySerializerItems we currently include the ReplicationId even though it is not replicated
@@ -2066,6 +2163,13 @@ TRefCountPtr<const FReplicationStateDescriptor> FReplicationStateDescriptorBuild
 				UE_LOG_DESCRIPTORBUILDER_WARNING(TEXT("FReplicationStateDescriptorBuilder : Skipping unsupported struct member %s.%s of type %s."), ToCStr(InStruct->GetName()), ToCStr(Property->GetName()), ToCStr(Property->GetCPPType()));
 			}
 		}
+
+		if (bAllMembersAreReplicated && PropertyCount == 0U)
+		{
+			// Warn if someone tries to replicate a struct with no properties for which we do not have a custom NetSerialzier.
+			UE_LOG_DESCRIPTORBUILDER_WARNING(TEXT("FReplicationStateDescriptorBuilder ::CreateDescriptorForStruct a replicated struct WITH no replicated properties must have a CustomNetSerializer. Struct: %s."), ToCStr(InStruct->GetName()));
+			bAllMembersAreReplicated = false;
+		}
 	}
 
 	Builder.SetStructSizeAndAlignment(InStruct->GetStructureSize(), InStruct->GetMinAlignment());
@@ -2220,8 +2324,8 @@ SIZE_T FReplicationStateDescriptorBuilder::CreateDescriptorsForClass(FResult& Cr
 	constexpr uint32 BuilderTypeCount = UE_ARRAY_COUNT(InternalPropertyReplicationStateTypes);
 	FPropertyReplicationStateDescriptorBuilder Builders[BuilderTypeCount];
 
-	// We add separate descriptors for all fastarrays as they need special treatment
-	TArray<FPropertyReplicationStateDescriptorBuilder::FMemberProperty, TInlineAllocator<8>> FastArrayProperties;
+	// We add separate descriptors properties with custom replication fragments
+	TArray<FPropertyReplicationStateDescriptorBuilder::FMemberProperty, TInlineAllocator<8>> CustomProperties;
 
 	// Add properties
 	for (const FRepRecord& RepRecord : InObjectClass->ClassReps)
@@ -2246,13 +2350,17 @@ SIZE_T FReplicationStateDescriptorBuilder::CreateDescriptorsForClass(FResult& Cr
 					continue;
 				}
 
-				if (EnumHasAnyFlags(MemberProperty.Traits, EMemberPropertyTraits::IsFastArray))
+				// We build separate descriptors for all properties with CustomReplicationFragments
+				if (MemberProperty.CreateAndRegisterReplicationFragmentFunction)
 				{
-					if (MemberProperty.CreateAndRegisterReplicationFragmentFunction)
-					{
-						FastArrayProperties.Add(MemberProperty);
-						continue;
-					}
+					CustomProperties.Add(MemberProperty);
+					continue;
+				}
+				else if (EnumHasAnyFlags(MemberProperty.Traits, EMemberPropertyTraits::IsFastArray))
+				{
+					// FastArrayProperties should use a custom replication fragment
+					ensureAlwaysMsgf(false, TEXT("FReplicationStateDescriptorBuilder::CreateDescriptorsForClass FFastArray property %s not registered! Usually this means a call to SetupIrisSupport is needed in the module's Build.cs file."), *Property->GetFullName());
+					continue;
 				}
 
 				// A property can end up in multiple states, such as when a condition is InitialOrOwner.
@@ -2360,19 +2468,19 @@ SIZE_T FReplicationStateDescriptorBuilder::CreateDescriptorsForClass(FResult& Cr
 			BuildParameters.DefaultStateSourceData = DefaultStateSourceData;
 			BuildParameters.DescriptorType = FPropertyReplicationStateDescriptorBuilder::EDescriptorType::Class;
 			BuildParameters.bIsInitState = (BuilderTypeIndex == InitPropertyReplicationStateBuilderIndex);
-			// We set this to false for the time being
+			// We set this to false for the time being for all replication states
 			BuildParameters.bAllMembersAreReplicated = false;
 			TRefCountPtr<const FReplicationStateDescriptor> Descriptor = Builders[BuilderTypeIndex].Build(InObjectClass->GetName() + InternalPropertyReplicationStateTypes[BuilderTypeIndex].NamePostFix, Parameters.DescriptorRegistry, BuildParameters);
 			CreatedDescriptors.Add(Descriptor);
 		}
 	}
 
-	// If we have properties that are FastArrays, build separate descriptors for each property
-	for (uint32 FastArrayPropertyIt = 0, FastArrayPropertyEndIt = FastArrayProperties.Num(); FastArrayPropertyIt != FastArrayPropertyEndIt; ++FastArrayPropertyIt)
+	// If we have properties with a CreateAndRegisterReplicationFragmentFunction, we build separate descriptors for them
+	for (uint32 CustomPropertyIt = 0, CustomArrayPropertyEndIt = CustomProperties.Num(); CustomPropertyIt != CustomArrayPropertyEndIt; ++CustomPropertyIt)
 	{
-		const FPropertyReplicationStateDescriptorBuilder::FMemberProperty& MemberProperty = FastArrayProperties[FastArrayPropertyIt];
+		const FPropertyReplicationStateDescriptorBuilder::FMemberProperty& MemberProperty = CustomProperties[CustomPropertyIt];
 
-		// Verify some assumptions, for now we do not support putting fast arrays in multiple states
+		// Verify some assumptions, for now we do not support putting properties with custom replication fragments in multiple states
 		check(!EnumHasAllFlags(MemberProperty.Traits, EMemberPropertyTraits::InitOnly | EMemberPropertyTraits::HasLifetimeConditionals));
 
 		FPropertyReplicationStateDescriptorBuilder Builder;		
@@ -2382,10 +2490,7 @@ SIZE_T FReplicationStateDescriptorBuilder::CreateDescriptorsForClass(FResult& Cr
 		BuildParameters.DescriptorType = FPropertyReplicationStateDescriptorBuilder::EDescriptorType::Class;
 		BuildParameters.bIsInitState = EnumHasAnyFlags(MemberProperty.Traits, EMemberPropertyTraits::InitOnly);
 		
-		// Treat as FastArray
-		BuildParameters.bIsFastArray = true;
-
-		// We set this to false for the time being
+		// We set this to false for the time being for all replication states
 		BuildParameters.bAllMembersAreReplicated = false;
 
 		// Add the property
diff --git a/Engine/Source/Runtime/Experimental/Iris/Core/Private/Iris/ReplicationSystem/FastArrayReplicationFragment.cpp b/Engine/Source/Runtime/Experimental/Iris/Core/Private/Iris/ReplicationSystem/FastArrayReplicationFragment.cpp
--- a/Engine/Source/Runtime/Experimental/Iris/Core/Private/Iris/ReplicationSystem/FastArrayReplicationFragment.cpp
+++ b/Engine/Source/Runtime/Experimental/Iris/Core/Private/Iris/ReplicationSystem/FastArrayReplicationFragment.cpp
@@ -263,33 +263,4 @@ bool FFastArrayReplicationFragmentHelper::InternalCompareArrayElement(const FRep
 	return InternalCompareStructProperty(ArrayElementDescriptor, Dst, Src);
 }
 
-TRefCountPtr<const FReplicationStateDescriptor> FFastArrayReplicationFragmentHelper::GetOrCreateDescriptorForFastArrayProperty(UObject* Object, FFragmentRegistrationContext& Context, int32 RepIndex)
-{
-	UClass* ObjectClass = Object->GetClass();
-
-	// If we already have created descriptors for the class this will be a lookup in the registry, if not we will explicitly create a state for the 
-	// given property only and do not register it
-	FReplicationStateDescriptorBuilder::FParameters BuilderParameters;
-	BuilderParameters.SinglePropertyIndex = RepIndex;
-	BuilderParameters.DefaultStateSource = Object->GetArchetype();
-	BuilderParameters.DescriptorRegistry = Context.GetReplicationStateRegistry();
-	
-	FReplicationStateDescriptorBuilder::FResult Result;
-	FReplicationStateDescriptorBuilder::CreateDescriptorsForClass(Result, ObjectClass, BuilderParameters);
-
-	// Find the descriptor for the fast array
-	for (auto& Desc : Result)
-	{
-		if (EnumHasAnyFlags(Desc->Traits, EReplicationStateTraits::IsFastArrayReplicationState))
-		{
-			if (Desc->MemberProperties[0]->RepIndex == RepIndex)
-			{
-				return Desc;
-			}
-		}
-	}
-
-	return nullptr;
-}
-
 }
diff --git a/Engine/Source/Runtime/Experimental/Iris/Core/Private/Iris/ReplicationSystem/ReplicationFragmentUtil.cpp b/Engine/Source/Runtime/Experimental/Iris/Core/Private/Iris/ReplicationSystem/ReplicationFragmentUtil.cpp
--- a/Engine/Source/Runtime/Experimental/Iris/Core/Private/Iris/ReplicationSystem/ReplicationFragmentUtil.cpp
+++ b/Engine/Source/Runtime/Experimental/Iris/Core/Private/Iris/ReplicationSystem/ReplicationFragmentUtil.cpp
@@ -54,7 +54,7 @@ uint32 FReplicationFragmentUtil::CreateAndRegisterFragmentsForObject(UObject* Ob
 		{
 			Fragment = FPropertyReplicationFragment::CreateAndRegisterFragment(Object, Desc.GetReference(), Context);
 		}
-		if (OutCreatedFragments)
+		if (Fragment && OutCreatedFragments)
 		{
 			OutCreatedFragments->Add(Fragment);
 			++NumCreatedReplicationFragments;
diff --git a/Engine/Source/Runtime/Experimental/Iris/Core/Public/Iris/ReplicationState/ReplicationStateDescriptor.h b/Engine/Source/Runtime/Experimental/Iris/Core/Public/Iris/ReplicationState/ReplicationStateDescriptor.h
--- a/Engine/Source/Runtime/Experimental/Iris/Core/Public/Iris/ReplicationState/ReplicationStateDescriptor.h
+++ b/Engine/Source/Runtime/Experimental/Iris/Core/Public/Iris/ReplicationState/ReplicationStateDescriptor.h
@@ -257,7 +257,7 @@ struct FReplicationStateDescriptor
 	// Function to destruct external state representation
 	DestructReplicationStateFunc DestructReplicationState;
 
-	// Function used to construct FastArrays
+	// Function used to construct custom replication fragments
 	CreateAndRegisterReplicationFragmentFunc CreateAndRegisterReplicationFragmentFunction;
 
 	EReplicationStateTraits Traits;
diff --git a/Engine/Source/Runtime/Experimental/Iris/Core/Public/Iris/ReplicationSystem/Private/FastArrayReplicationFragmentInternal.h b/Engine/Source/Runtime/Experimental/Iris/Core/Public/Iris/ReplicationSystem/Private/FastArrayReplicationFragmentInternal.h
--- a/Engine/Source/Runtime/Experimental/Iris/Core/Public/Iris/ReplicationSystem/Private/FastArrayReplicationFragmentInternal.h
+++ b/Engine/Source/Runtime/Experimental/Iris/Core/Public/Iris/ReplicationSystem/Private/FastArrayReplicationFragmentInternal.h
@@ -52,9 +52,6 @@ namespace Private {
 /** Utility methods to behave similar to FastArraySerializer */
 struct FFastArrayReplicationFragmentHelper
 {
-	/** Create descriptor for the FastArray property if it does not already exist */
-	IRISCORE_API static TRefCountPtr<const FReplicationStateDescriptor> GetOrCreateDescriptorForFastArrayProperty(UObject* Object, FFragmentRegistrationContext& Context, int32 RepIndex);
-
 	/** Rebuild IndexMap for FastArrraySerializer */
 	template <typename FastArrayType, typename ItemArrayType>
 	static void ConditionalRebuildItemMap(FastArrayType& ArraySerializer, const ItemArrayType& Items);
