Iris
- Fixed bad conditional statement in NetHandleManager

diff --git a/Engine/Source/Runtime/Experimental/Iris/Core/Private/Iris/ReplicationSystem/NetHandleManager.cpp b/Engine/Source/Runtime/Experimental/Iris/Core/Private/Iris/ReplicationSystem/NetHandleManager.cpp
--- a/Engine/Source/Runtime/Experimental/Iris/Core/Private/Iris/ReplicationSystem/NetHandleManager.cpp
+++ b/Engine/Source/Runtime/Experimental/Iris/Core/Private/Iris/ReplicationSystem/NetHandleManager.cpp
@@ -678,7 +678,7 @@ void FNetHandleManager::RemoveDependentObject(FNetHandle ParentHandle, FNetHandl
 	FInternalNetHandle ParentInternalIndex = GetInternalIndex(ParentHandle);
 	FInternalNetHandle DependentInternalIndex = GetInternalIndex(DependentHandle);
 
-	if (!(ParentInternalIndex != InvalidInternalIndex) && (DependentInternalIndex != InvalidInternalIndex))
+	if ((ParentInternalIndex == InvalidInternalIndex) || (DependentInternalIndex == InvalidInternalIndex))
 	{
 		return;
 	}
