Add null check to solver index/proxy map to fix static analysis warning

diff --git a/Engine/Source/Runtime/Experimental/Chaos/Private/PBDRigidsSolver.cpp b/Engine/Source/Runtime/Experimental/Chaos/Private/PBDRigidsSolver.cpp
--- a/Engine/Source/Runtime/Experimental/Chaos/Private/PBDRigidsSolver.cpp
+++ b/Engine/Source/Runtime/Experimental/Chaos/Private/PBDRigidsSolver.cpp
@@ -1093,7 +1093,10 @@ namespace Chaos
 				Handle->GTGeometryParticle() = Proxy->GetParticle_LowLevel();
 
 				// Track proxies
-				SingleParticlePhysicsProxies_PT.EmplaceAt(UniqueIdx->Idx, Proxy);
+				if (UniqueIdx)
+				{
+					SingleParticlePhysicsProxies_PT.EmplaceAt(UniqueIdx->Idx, Proxy);
+				}
 			}
 
 			if(Proxy->GetHandle_LowLevel())
