send notification when the constraint manager is removed (or deletion undone)

diff --git a/Engine/Plugins/Animation/ControlRig/Source/ControlRig/Private/Sequencer/MovieSceneControlRigParameterSection.cpp b/Engine/Plugins/Animation/ControlRig/Source/ControlRig/Private/Sequencer/MovieSceneControlRigParameterSection.cpp
--- a/Engine/Plugins/Animation/ControlRig/Source/ControlRig/Private/Sequencer/MovieSceneControlRigParameterSection.cpp
+++ b/Engine/Plugins/Animation/ControlRig/Source/ControlRig/Private/Sequencer/MovieSceneControlRigParameterSection.cpp
@@ -1408,6 +1408,11 @@ void UMovieSceneControlRigParameterSection::ReplaceConstraint(const FName InCons
 	}
 }
 
+void UMovieSceneControlRigParameterSection::OnConstraintsChanged()
+{
+	ReconstructChannelProxy();
+}
+
 void UMovieSceneControlRigParameterSection::AddConstraintChannel(UTickableConstraint* InConstraint)
 {
 	
diff --git a/Engine/Plugins/Animation/ControlRig/Source/ControlRig/Public/Sequencer/MovieSceneControlRigParameterSection.h b/Engine/Plugins/Animation/ControlRig/Source/ControlRig/Public/Sequencer/MovieSceneControlRigParameterSection.h
--- a/Engine/Plugins/Animation/ControlRig/Source/ControlRig/Public/Sequencer/MovieSceneControlRigParameterSection.h
+++ b/Engine/Plugins/Animation/ControlRig/Source/ControlRig/Public/Sequencer/MovieSceneControlRigParameterSection.h
@@ -292,6 +292,10 @@ public:
 	*/
 	virtual void ReplaceConstraint(const FName InName, UTickableConstraint* InConstraint)  override;
 
+	/*
+	*  What to do if the constraint object has been changed, for example by an undo or redo.
+	*/
+	virtual void OnConstraintsChanged() override;
 
 	//not override but needed
 	const TArray<FConstraintAndActiveChannel>& GetConstraintsChannels() const;
diff --git a/Engine/Plugins/Animation/ControlRig/Source/ControlRigEditor/Private/Sequencer/ControlRigParameterTrackEditor.cpp b/Engine/Plugins/Animation/ControlRig/Source/ControlRigEditor/Private/Sequencer/ControlRigParameterTrackEditor.cpp
--- a/Engine/Plugins/Animation/ControlRig/Source/ControlRigEditor/Private/Sequencer/ControlRigParameterTrackEditor.cpp
+++ b/Engine/Plugins/Animation/ControlRig/Source/ControlRigEditor/Private/Sequencer/ControlRigParameterTrackEditor.cpp
@@ -2565,7 +2565,7 @@ void FControlRigParameterTrackEditor::ClearOutAllSpaceAndConstraintDelegates(con
 						if (const UControlRig* ControlRig = CRSection->GetControlRig())
 						{
 							FConstraintsManagerController& Controller = FConstraintsManagerController::Get(ControlRig->GetWorld());
-							Controller.OnConstraintRemoved().Remove(CRSection->OnConstraintRemovedHandle);
+							Controller.GetNotifyDelegate().Remove(CRSection->OnConstraintRemovedHandle);
 							CRSection->OnConstraintRemovedHandle.Reset();
 						}
 					}
@@ -2730,25 +2730,47 @@ void FControlRigParameterTrackEditor::HandleConstraintRemoved(IMovieSceneConstra
 		if (!InSection->OnConstraintRemovedHandle.IsValid())
 		{
 			InSection->OnConstraintRemovedHandle =
-			Controller.OnConstraintRemoved().AddLambda([InSection, Section, this](FName InConstraintName, bool bDoNotCompensate)
+			Controller.GetNotifyDelegate().AddLambda(
+				[InSection,Section, this](EConstraintsManagerNotifyType InNotifyType, UObject *InObject)
 			{
-				const FConstraintAndActiveChannel* ConstraintChannel = InSection->GetConstraintChannel(InConstraintName);
-				if (!ConstraintChannel)
+				switch (InNotifyType)
 				{
-					return;
-				}
-				
-				if (bDoNotCompensate == false && ConstraintChannel->Constraint.IsValid())
-				{
-					FMovieSceneConstraintChannelHelper::HandleConstraintRemoved(
-						ConstraintChannel->Constraint.Get(),
-						&ConstraintChannel->ActiveChannel,
-						GetSequencer(),
-						Section);
-				}
+					case EConstraintsManagerNotifyType::ConstraintAdded:
+						break;
+					case EConstraintsManagerNotifyType::ConstraintRemoved:
+					case EConstraintsManagerNotifyType::ConstraintRemovedWithCompensation:
+						{
+							const UTickableConstraint* Constraint = Cast<UTickableConstraint>(InObject);
+							if (!IsValid(Constraint))
+							{
+								return;
+							}
 
-				InSection->RemoveConstraintChannel(InConstraintName);
-		   });
+							const FName ConstraintName = Constraint->GetFName();
+							const FConstraintAndActiveChannel* ConstraintChannel = InSection->GetConstraintChannel(ConstraintName);
+							if (!ConstraintChannel)
+							{
+								return;
+							}
+
+							const bool bCompensate = (InNotifyType == EConstraintsManagerNotifyType::ConstraintRemovedWithCompensation);
+							if (bCompensate && ConstraintChannel->Constraint.IsValid())
+							{
+								FMovieSceneConstraintChannelHelper::HandleConstraintRemoved(
+									ConstraintChannel->Constraint.Get(),
+									&ConstraintChannel->ActiveChannel,
+									GetSequencer(),
+									Section);
+							}
+
+							InSection->RemoveConstraintChannel(ConstraintName);
+						}
+						break;
+					case EConstraintsManagerNotifyType::ManagerUpdated:
+						InSection->OnConstraintsChanged();
+						break;		
+				}
+			});
 		}
 	}
 }
diff --git a/Engine/Source/Editor/MovieSceneTools/Private/TrackEditors/TransformTrackEditor.cpp b/Engine/Source/Editor/MovieSceneTools/Private/TrackEditors/TransformTrackEditor.cpp
--- a/Engine/Source/Editor/MovieSceneTools/Private/TrackEditors/TransformTrackEditor.cpp
+++ b/Engine/Source/Editor/MovieSceneTools/Private/TrackEditors/TransformTrackEditor.cpp
@@ -1541,24 +1541,45 @@ void F3DTransformTrackEditor::HandleConstraintRemoved(IMovieSceneConstrainedSect
 	if (!InSection->OnConstraintRemovedHandle.IsValid())
 	{
 		InSection->OnConstraintRemovedHandle =
-			Controller.OnConstraintRemoved().AddLambda([InSection,Section, this](FName InConstraintName, bool bDoNotCompensate)
+			Controller.GetNotifyDelegate().AddLambda([InSection,Section, this](EConstraintsManagerNotifyType InNotifyType, UObject *InObject)
 				{
-					const FConstraintAndActiveChannel* ConstraintChannel = InSection->GetConstraintChannel(InConstraintName);
-					if (!ConstraintChannel)
+					switch (InNotifyType)
 					{
-						return;
+						case EConstraintsManagerNotifyType::ConstraintAdded:
+							break;
+						case EConstraintsManagerNotifyType::ConstraintRemoved:
+						case EConstraintsManagerNotifyType::ConstraintRemovedWithCompensation:
+							{
+								const UTickableConstraint* Constraint = Cast<UTickableConstraint>(InObject);
+								if (!IsValid(Constraint))
+								{
+									return;
+								}
+
+								const FName ConstraintName = Constraint->GetFName();
+								const FConstraintAndActiveChannel* ConstraintChannel = InSection->GetConstraintChannel(ConstraintName);
+								if (!ConstraintChannel)
+								{
+									return;
+								}
+
+								const bool bCompensate = (InNotifyType == EConstraintsManagerNotifyType::ConstraintRemovedWithCompensation);
+								if (bCompensate && ConstraintChannel->Constraint.IsValid())
+								{
+									FMovieSceneConstraintChannelHelper::HandleConstraintRemoved(
+										ConstraintChannel->Constraint.Get(),
+										&ConstraintChannel->ActiveChannel,
+										GetSequencer(),
+										Section);
+								}
+
+								InSection->RemoveConstraintChannel(ConstraintName);
+							}
+							break;
+						case EConstraintsManagerNotifyType::ManagerUpdated:
+							InSection->OnConstraintsChanged();
+							break;		
 					}
-
-					if (bDoNotCompensate == false && ConstraintChannel->Constraint.IsValid())
-					{
-						FMovieSceneConstraintChannelHelper::HandleConstraintRemoved(
-							ConstraintChannel->Constraint.Get(),
-							&ConstraintChannel->ActiveChannel,
-							GetSequencer(),
-							Section);
-					}
-
-					InSection->RemoveConstraintChannel(InConstraintName);
 				});
 	}
 }
@@ -1595,7 +1616,7 @@ void F3DTransformTrackEditor::ClearOutConstraintDelegates() const
 
 					if (CRSection->OnConstraintRemovedHandle.IsValid())
 					{
-						Controller.OnConstraintRemoved().Remove(CRSection->OnConstraintRemovedHandle);
+						Controller.GetNotifyDelegate().Remove(CRSection->OnConstraintRemovedHandle);
 						CRSection->OnConstraintRemovedHandle.Reset();
 					}
 				}
diff --git a/Engine/Source/Runtime/Experimental/Animation/Constraints/Private/ConstraintsActor.cpp b/Engine/Source/Runtime/Experimental/Animation/Constraints/Private/ConstraintsActor.cpp
--- a/Engine/Source/Runtime/Experimental/Animation/Constraints/Private/ConstraintsActor.cpp
+++ b/Engine/Source/Runtime/Experimental/Animation/Constraints/Private/ConstraintsActor.cpp
@@ -25,7 +25,7 @@ AConstraintsActor::~AConstraintsActor()
 
 void AConstraintsActor::BeginDestroy()
 {
-	if (ConstraintsManager)
+	if (IsValid(ConstraintsManager))
 	{
 		ConstraintsManager->Clear(GetWorld());
 	}
@@ -35,6 +35,10 @@ void AConstraintsActor::BeginDestroy()
 
 void AConstraintsActor::Destroyed()
 {
+	if (IsValid(ConstraintsManager))
+	{
+		ConstraintsManager->Clear(GetWorld());
+	}
 	Super::Destroyed();
 }
 
@@ -81,3 +85,17 @@ void AConstraintsActor::RegisterConstraintsTickFunctions() const
 		}
 	}
 }
+
+#if WITH_EDITOR
+void AConstraintsActor::PostEditUndo()
+{
+	Super::PostEditUndo();
+
+	if (IsValid(ConstraintsManager))
+	{
+		const FConstraintsManagerController& Controller = FConstraintsManagerController::Get(GetWorld());
+		Controller.Notify(EConstraintsManagerNotifyType::ManagerUpdated, ConstraintsManager);
+		ConstraintsManager->Init(GetWorld());
+	}
+}
+#endif
diff --git a/Engine/Source/Runtime/Experimental/Animation/Constraints/Private/ConstraintsManager.cpp b/Engine/Source/Runtime/Experimental/Animation/Constraints/Private/ConstraintsManager.cpp
--- a/Engine/Source/Runtime/Experimental/Animation/Constraints/Private/ConstraintsManager.cpp
+++ b/Engine/Source/Runtime/Experimental/Animation/Constraints/Private/ConstraintsManager.cpp
@@ -284,14 +284,35 @@ void UConstraintsManager::SetConstraintDependencies(
 	InFunctionToTickAfter->AddPrerequisite(this, *InFunctionToTickBefore);
 }
 
-void UConstraintsManager::Clear(UWorld* World)
+void UConstraintsManager::Clear(UWorld* InWorld)
 {
-	UnregisterDelegates(World);
-	for (UTickableConstraint* Constraint : Constraints)
+	UnregisterDelegates(InWorld);
+
+	if (!Constraints.IsEmpty())
 	{
-		OnConstraintRemoved_BP.Broadcast(this, Constraint, true);
+		static constexpr bool bDoNoCompensate = false;
+		static constexpr EConstraintsManagerNotifyType RemovalNotification =
+						EConstraintsManagerNotifyType::ConstraintRemovedWithCompensation;
+	
+		const FConstraintsManagerController& Controller = FConstraintsManagerController::Get(InWorld);
+		for (UTickableConstraint* Constraint : Constraints)
+		{
+			if ( IsValid(Constraint) )
+			{
+				// notify deletion
+				Controller.Notify(RemovalNotification, Constraint);
+				OnConstraintRemoved_BP.Broadcast(this, Constraint, bDoNoCompensate);
+
+				// disable constraint
+				Constraint->Modify();
+				Constraint->ConstraintTick.UnRegisterTickFunction();
+				Constraint->SetActive(false);
+			}
+		}
+
+		Modify();
+		Constraints.Empty();
 	}
-	Constraints.Empty();
 }
 
 void UConstraintsManager::Dump() const
@@ -407,6 +428,8 @@ bool FConstraintsManagerController::AddConstraint(UTickableConstraint* InConstra
 	InConstraint->ConstraintTick.RegisterFunction(InConstraint->GetFunction());
 	InConstraint->ConstraintTick.RegisterTickFunction(World->GetCurrentLevel());
 
+	// notify
+	Notify(EConstraintsManagerNotifyType::ConstraintAdded, InConstraint);
 	Manager->OnConstraintAdded_BP.Broadcast(Manager, InConstraint);
 
 	return true;
@@ -475,9 +498,17 @@ bool FConstraintsManagerController::RemoveConstraint(const int32 InConstraintInd
 	UTickableConstraint* Constraint = Manager->Constraints[InConstraintIndex];
 	
 	// notify deletion
-
-	ConstraintRemoved.Broadcast(ConstraintName, bDoNotCompensate);
+	auto GetRemoveNotifyType = [bDoNotCompensate]()
+	{
+		if(bDoNotCompensate)
+		{
+			return EConstraintsManagerNotifyType::ConstraintRemoved;	
+		}
+		return EConstraintsManagerNotifyType::ConstraintRemovedWithCompensation;
+	};
+	Notify(GetRemoveNotifyType(), Constraint);	
 	Manager->OnConstraintRemoved_BP.Broadcast(Manager, Constraint, bDoNotCompensate);
+	
 	Manager->Constraints[InConstraintIndex]->Modify();
 	Manager->Modify();
 
@@ -644,4 +675,25 @@ void FConstraintsManagerController::EvaluateAllConstraints() const
 	{
 		InConstraint->Evaluate(true);
 	}
+}
+
+void FConstraintsManagerController::Notify(EConstraintsManagerNotifyType InNotifyType, UObject* InObject) const
+{
+	switch (InNotifyType)
+	{
+		case EConstraintsManagerNotifyType::ConstraintAdded: 
+		case EConstraintsManagerNotifyType::ConstraintRemoved:
+		case EConstraintsManagerNotifyType::ConstraintRemovedWithCompensation:
+			checkSlow(Cast<UTickableConstraint>(InObject) != nullptr);
+			break;
+
+		case EConstraintsManagerNotifyType::ManagerUpdated:
+			checkSlow(Cast<UConstraintsManager>(InObject) != nullptr);
+			break;
+		default:
+			checkfSlow(false, TEXT("Unchecked EConstraintsManagerNotifyType!"));
+			break;
+	}
+
+	NotifyDelegate.Broadcast(InNotifyType, InObject);
 }
\ No newline at end of file
diff --git a/Engine/Source/Runtime/Experimental/Animation/Constraints/Public/ConstraintsActor.h b/Engine/Source/Runtime/Experimental/Animation/Constraints/Public/ConstraintsActor.h
--- a/Engine/Source/Runtime/Experimental/Animation/Constraints/Public/ConstraintsActor.h
+++ b/Engine/Source/Runtime/Experimental/Animation/Constraints/Public/ConstraintsActor.h
@@ -26,6 +26,10 @@ public:
 	virtual void Destroyed() override;
 
 	virtual void PostRegisterAllComponents() override;
+
+#if WITH_EDITOR
+	virtual void PostEditUndo() override;
+#endif
 	
 protected:
 	// Called when the game starts or when spawned
diff --git a/Engine/Source/Runtime/Experimental/Animation/Constraints/Public/ConstraintsManager.h b/Engine/Source/Runtime/Experimental/Animation/Constraints/Public/ConstraintsManager.h
--- a/Engine/Source/Runtime/Experimental/Animation/Constraints/Public/ConstraintsManager.h
+++ b/Engine/Source/Runtime/Experimental/Animation/Constraints/Public/ConstraintsManager.h
@@ -186,7 +186,6 @@ private:
 	/** All of the constraints*/
 	UPROPERTY()
 	TArray< TObjectPtr<UTickableConstraint> > Constraints;
-	
 
 	friend class FConstraintsManagerController;
 	friend class AConstraintsActor;
@@ -197,6 +196,16 @@ private:
  * Basic controller to add / remove / get constraints
  **/
 
+enum class EConstraintsManagerNotifyType
+{
+	ConstraintAdded,					/** A new constraint has been added. */
+	ConstraintRemoved,					/** A constraint has been removed. */
+	ConstraintRemovedWithCompensation,	/** A constraint has been removed and needs compensation. */
+	ManagerUpdated,						/** The manager has been updated/reset. */
+};
+
+DECLARE_MULTICAST_DELEGATE_TwoParams(FConstraintsManagerNotifyDelegate, EConstraintsManagerNotifyType /* type */, UObject* /* subject */);
+
 class CONSTRAINTS_API FConstraintsManagerController
 {
 public:
@@ -266,9 +275,8 @@ private:
 	DECLARE_MULTICAST_DELEGATE_OneParam(FOnSceneComponentConstrained, USceneComponent* /*InSceneComponent*/);
 	FOnSceneComponentConstrained SceneComponentConstrained;
 
-	/** @todo document */
-	DECLARE_MULTICAST_DELEGATE_TwoParams(FOnConstraintRemoved, FName /*InConstraintName*/, bool /*bDoNotCompensate*/);
-	FOnConstraintRemoved ConstraintRemoved;
+	/** Delegate to trigger changes in the constraints manager. */
+	FConstraintsManagerNotifyDelegate NotifyDelegate;
 	
 	/** Find the existing Constraints Manager in World or create a new one. */
 	UConstraintsManager* GetManager() const;
@@ -285,6 +293,10 @@ private:
 public:
 	/** Delegate that's fired when a scene component is constrained, this is needed to make sure things like gizmo's get updated after the constraint tick happens*/
 	FOnSceneComponentConstrained& OnSceneComponentConstrained() { return SceneComponentConstrained; }
-	/** @todo document */
-	FOnConstraintRemoved& OnConstraintRemoved(){ return ConstraintRemoved; };
+
+	/** Delegate to trigger changes in the constraints manager. */
+	FConstraintsManagerNotifyDelegate& GetNotifyDelegate() { return NotifyDelegate; }
+
+	/** Notify from changes in the constraints manager. */
+	void Notify(EConstraintsManagerNotifyType InNotifyType, UObject* InObject) const;
 };
\ No newline at end of file
