Fix various static analysis and IncludeTool warnings in Iris.

diff --git a/Engine/Source/Programs/IncludeTool/IncludeTool/Rules.cs b/Engine/Source/Programs/IncludeTool/IncludeTool/Rules.cs
--- a/Engine/Source/Programs/IncludeTool/IncludeTool/Rules.cs
+++ b/Engine/Source/Programs/IncludeTool/IncludeTool/Rules.cs
@@ -438,6 +438,8 @@ namespace IncludeTool
 			"/Engine/Source/Runtime/Interchange/Engine/Public/InterchangeEngineFwd.h", // invalid forward declaration - 'namespace UE'
 			"/Engine/Plugins/Experimental/GameFeatures/Source/GameFeatures/Public/GameFeatureTypesFwd.h", //  invalid forward declaration - 'namespace GameFeaturePluginStatePrivate'
 			"/Engine/Source/Runtime/Core/Public/Containers/VersePathFwd.h", // invalid forward declaration - 'namespace UE::Core'
+			"/Engine/Source/Runtime/Experimental/Iris/Core/Public/Iris/ReplicationState/ReplicationStateFwd.h", // invalid forward declaration - 'namespace UE::Net'
+			"/Engine/Source/Runtime/Experimental/Iris/Core/Private/Iris/ReplicationSystem/ObjectReferenceCacheFwd.h", // invalid forward declaration - 'namespace UE::Net'
 		};
 
 		/// <summary>
diff --git a/Engine/Source/Runtime/Experimental/Iris/Core/Private/Iris/ReplicationSystem/LegacyPushModel.h b/Engine/Source/Runtime/Experimental/Iris/Core/Private/Iris/ReplicationSystem/LegacyPushModel.h
--- a/Engine/Source/Runtime/Experimental/Iris/Core/Private/Iris/ReplicationSystem/LegacyPushModel.h
+++ b/Engine/Source/Runtime/Experimental/Iris/Core/Private/Iris/ReplicationSystem/LegacyPushModel.h
@@ -83,7 +83,7 @@ private:
 namespace UE::Net::Private
 {
 
-inline constexpr bool IsIrisPushModelEnabled(bool bIsPushModelEnabled = false) { return false; }
+inline constexpr bool IsIrisPushModelEnabled(bool /* bIsPushModelEnabled */ = false) { return false; }
 
 }
 
diff --git a/Engine/Source/Runtime/Experimental/Iris/Core/Private/Iris/ReplicationSystem/NetBlob/NetBlobManager.cpp b/Engine/Source/Runtime/Experimental/Iris/Core/Private/Iris/ReplicationSystem/NetBlob/NetBlobManager.cpp
--- a/Engine/Source/Runtime/Experimental/Iris/Core/Private/Iris/ReplicationSystem/NetBlob/NetBlobManager.cpp
+++ b/Engine/Source/Runtime/Experimental/Iris/Core/Private/Iris/ReplicationSystem/NetBlob/NetBlobManager.cpp
@@ -339,7 +339,7 @@ void FNetBlobManager::FNetObjectAttachmentSendQueue::Enqueue(FInternalNetHandle
 	bHasMulticastAttachments = true;
 }
 
-void FNetBlobManager::FNetObjectAttachmentSendQueue::PrepareProcessQueue(FReplicationConnections* InConnections, const FNetHandleManager* NetHandleManager)
+void FNetBlobManager::FNetObjectAttachmentSendQueue::PrepareProcessQueue(FReplicationConnections* InConnections, const FNetHandleManager* InNetHandleManager)
 {
 	if (ProcessContext.IsValid())
 	{
@@ -347,7 +347,7 @@ void FNetBlobManager::FNetObjectAttachmentSendQueue::PrepareProcessQueue(FReplic
 	}
 
 	ProcessContext.Connections = InConnections;
-	ProcessContext.NetHandleManager = NetHandleManager;
+	ProcessContext.NetHandleManager = InNetHandleManager;
 
 	ProcessContext.AttachmentsToObjectsGoingOutOfScope.Init(AttachmentQueue.Num());
 	ProcessContext.AttachmentsToObjectsInScope.Init(AttachmentQueue.Num());
@@ -367,8 +367,8 @@ void FNetBlobManager::FNetObjectAttachmentSendQueue::PrepareProcessQueue(FReplic
 	}
 
 	// Figure out if we have any attachments to objects going out of scope.
-	const FNetBitArray& ScopableObjects = NetHandleManager->GetScopableInternalIndices();
-	const FNetBitArray& PrevScopableObjects = NetHandleManager->GetPrevFrameScopableInternalIndices();
+	const FNetBitArray& ScopableObjects = InNetHandleManager->GetScopableInternalIndices();
+	const FNetBitArray& PrevScopableObjects = InNetHandleManager->GetPrevFrameScopableInternalIndices();
 	
 	uint32 CurrentEntryIndex = 0U;
 	for (const FNetObjectAttachmentQueueEntry& Entry : MakeArrayView(AttachmentQueue))
diff --git a/Engine/Source/Runtime/Experimental/Iris/Core/Private/Iris/ReplicationSystem/NetBlob/NetBlobManager.h b/Engine/Source/Runtime/Experimental/Iris/Core/Private/Iris/ReplicationSystem/NetBlob/NetBlobManager.h
--- a/Engine/Source/Runtime/Experimental/Iris/Core/Private/Iris/ReplicationSystem/NetBlob/NetBlobManager.h
+++ b/Engine/Source/Runtime/Experimental/Iris/Core/Private/Iris/ReplicationSystem/NetBlob/NetBlobManager.h
@@ -83,7 +83,7 @@ private:
 		// Multicast
 		void Enqueue(FInternalNetHandle OwnerIndex, FInternalNetHandle SubObjectIndex, const TRefCountPtr<FNetObjectAttachment>& Attachment);
 
-		void PrepareProcessQueue(FReplicationConnections* InConnections, const FNetHandleManager* NetHandleManager);
+		void PrepareProcessQueue(FReplicationConnections* InConnections, const FNetHandleManager* InNetHandleManager);
 		void ProcessQueue(EProcessMode ProcessMode);
 		void ResetProcessQueue();
 	
diff --git a/Engine/Source/Runtime/Experimental/Iris/Core/Public/Iris/ReplicationSystem/Filtering/NetObjectFilter.h b/Engine/Source/Runtime/Experimental/Iris/Core/Public/Iris/ReplicationSystem/Filtering/NetObjectFilter.h
--- a/Engine/Source/Runtime/Experimental/Iris/Core/Public/Iris/ReplicationSystem/Filtering/NetObjectFilter.h
+++ b/Engine/Source/Runtime/Experimental/Iris/Core/Public/Iris/ReplicationSystem/Filtering/NetObjectFilter.h
@@ -164,7 +164,7 @@ struct FNetObjectFilterUpdateParams
 	FNetObjectFilteringInfo* FilteringInfos;
 };
 
-UCLASS()
+UCLASS(Abstract)
 class UNetObjectFilter : public UObject
 {
 	GENERATED_BODY()
diff --git a/Engine/Source/Runtime/Experimental/Iris/Core/Public/Iris/ReplicationSystem/NetBlob/NetBlobHandler.h b/Engine/Source/Runtime/Experimental/Iris/Core/Public/Iris/ReplicationSystem/NetBlob/NetBlobHandler.h
--- a/Engine/Source/Runtime/Experimental/Iris/Core/Public/Iris/ReplicationSystem/NetBlob/NetBlobHandler.h
+++ b/Engine/Source/Runtime/Experimental/Iris/Core/Public/Iris/ReplicationSystem/NetBlob/NetBlobHandler.h
@@ -39,7 +39,7 @@ public:
  * @see UReplicationSystem::RegisterNetBlobHandler
  * @note Certain handlers such as NetRPCHandler, PartialNetObjectAttachmentHandler and NetObjectBlobHandler will be registered automatically.
  */
-UCLASS(transient, MinimalApi)
+UCLASS(transient, MinimalApi, Abstract)
 class UNetBlobHandler : public UObject, public INetBlobReceiver
 {
 	GENERATED_BODY()
diff --git a/Engine/Source/Runtime/Experimental/Iris/Core/Public/Iris/ReplicationSystem/Prioritization/LocationBasedNetObjectPrioritizer.h b/Engine/Source/Runtime/Experimental/Iris/Core/Public/Iris/ReplicationSystem/Prioritization/LocationBasedNetObjectPrioritizer.h
--- a/Engine/Source/Runtime/Experimental/Iris/Core/Public/Iris/ReplicationSystem/Prioritization/LocationBasedNetObjectPrioritizer.h
+++ b/Engine/Source/Runtime/Experimental/Iris/Core/Public/Iris/ReplicationSystem/Prioritization/LocationBasedNetObjectPrioritizer.h
@@ -13,7 +13,7 @@ namespace UE::Net
 	class FWorldLocations;
 }
 
-UCLASS(Transient, MinimalAPI)
+UCLASS(Transient, MinimalAPI, Abstract)
 class ULocationBasedNetObjectPrioritizer : public UNetObjectPrioritizer
 {
 	GENERATED_BODY()
diff --git a/Engine/Source/Runtime/Experimental/Iris/Core/Public/Iris/ReplicationSystem/Prioritization/NetObjectPrioritizer.h b/Engine/Source/Runtime/Experimental/Iris/Core/Public/Iris/ReplicationSystem/Prioritization/NetObjectPrioritizer.h
--- a/Engine/Source/Runtime/Experimental/Iris/Core/Public/Iris/ReplicationSystem/Prioritization/NetObjectPrioritizer.h
+++ b/Engine/Source/Runtime/Experimental/Iris/Core/Public/Iris/ReplicationSystem/Prioritization/NetObjectPrioritizer.h
@@ -138,7 +138,7 @@ struct FNetObjectPrioritizerUpdateParams
  * object to still not be replicated to a particular connection a certain frame. There is no mechanism to force an object to be replicated a 
  * certain frame, but the priority is a major factor in the decision.
  */
-UCLASS()
+UCLASS(Abstract)
 class UNetObjectPrioritizer : public UObject
 {
 	GENERATED_BODY()
diff --git a/Engine/Source/Runtime/Experimental/Iris/Core/Public/Iris/Serialization/ObjectNetSerializer.h b/Engine/Source/Runtime/Experimental/Iris/Core/Public/Iris/Serialization/ObjectNetSerializer.h
--- a/Engine/Source/Runtime/Experimental/Iris/Core/Public/Iris/Serialization/ObjectNetSerializer.h
+++ b/Engine/Source/Runtime/Experimental/Iris/Core/Public/Iris/Serialization/ObjectNetSerializer.h
@@ -3,7 +3,7 @@
 #pragma once
 
 #include "NetSerializer.h"
-#include "Iris/ReplicationSystem//NetHandle.h"
+#include "Iris/ReplicationSystem/NetHandle.h"
 #include "UObject/ObjectMacros.h"
 
 #include "ObjectNetSerializer.generated.h"
