Chaos: When sweeping against meshes, be sure to include back faces if and only if they are in initial overlap

diff --git a/Engine/Source/Runtime/Experimental/Chaos/Private/Chaos/TriangleMeshImplicitObject.cpp b/Engine/Source/Runtime/Experimental/Chaos/Private/Chaos/TriangleMeshImplicitObject.cpp
--- a/Engine/Source/Runtime/Experimental/Chaos/Private/Chaos/TriangleMeshImplicitObject.cpp
+++ b/Engine/Source/Runtime/Experimental/Chaos/Private/Chaos/TriangleMeshImplicitObject.cpp
@@ -879,6 +879,8 @@ struct FTriangleMeshSweepVisitor
 		// Normalize rotation
 		RotationSimd = VectorNormalizeSafe(RotationSimd, GlobalVectorConstants::Float0001);
 		RayDirSimd = MakeVectorRegisterFloatFromDouble(MakeVectorRegister(ScaledDirNormalized[0], ScaledDirNormalized[1], ScaledDirNormalized[2], 0.0));
+
+		GeometryAABBTriangleSpace = FAABBVectorized{ InQueryGeom.BoundingBox().TransformedAABB(InScaledStartTM)};
 	}
 
 	const void* GetQueryData() const { return nullptr; }
@@ -925,17 +927,31 @@ struct FTriangleMeshSweepVisitor
 		FTriangleRegister Tri(A, B, C);
 		const VectorRegister4Float TriNormal = VectorCross(VectorSubtract(B, A), VectorSubtract(C, A));
 
-		if(CullsBackFaceSweepsCode != 0)
+		const VectorRegister4Float IsBackFace = VectorCompareGT(VectorMultiply(VectorDot3(TriNormal, VectorScaledDirNormalized), VectorCullsBackFaceSweepsCode), VectorZero());
+		if (CullsBackFaceSweepsCode != 0)
 		{
-			const VectorRegister4Float ReturnTrue = VectorCompareGT(VectorMultiply(VectorDot3(TriNormal, VectorScaledDirNormalized), VectorCullsBackFaceSweepsCode), VectorZero());
-			if (VectorMaskBits(ReturnTrue))
+			if (VectorMaskBits(IsBackFace))
 			{
-				return true;
+				// Don't cull the back face if there is a chance that we are initially overlapping.
+				const VectorRegister4Float MinBounds = VectorMin(VectorMin(A, B), C);
+				const VectorRegister4Float MaxBounds = VectorMax(VectorMax(A, B), C);
+				FAABBVectorized TriangleAABB(MinBounds, MaxBounds);
+				
+				if (!GeometryAABBTriangleSpace.Intersects(TriangleAABB))
+				{
+					return true;
+				}
 			}
 		}
 		FRealSingle Time;
 		if(GJKRaycast2ImplSimd(Tri, QueryGeom, RotationSimd, TranslationSimd, RayDirSimd, LengthScale * CurDataLength, Time, OutPositionSimd, OutNormalSimd, bComputeMTD, GlobalVectorConstants::Float1000))
 		{
+			// Don't return back faces if they are not initially overlapping
+			if (Time > 0 && VectorMaskBits(IsBackFace))
+			{
+				return true;
+			}
+
 			// Time is world scale, OutTime is local scale.
 			if(Time < LengthScale * OutTime)
 			{
@@ -997,7 +1013,7 @@ struct FTriangleMeshSweepVisitor
 	VectorRegister4Float TranslationSimd;
 	VectorRegister4Float RayDirSimd;
 	VectorRegister4Float OutPositionSimd, OutNormalSimd;
-
+	FAABBVectorized		 GeometryAABBTriangleSpace;
 };
 
 template <typename QueryGeomType, typename IdxType>
