Fix merge conflict.
Reenable Dynamic Tree instead of Grid in Static Tree

diff --git a/Engine/Source/Runtime/Experimental/Chaos/Private/Chaos/PBDRigidsEvolution.cpp b/Engine/Source/Runtime/Experimental/Chaos/Private/Chaos/PBDRigidsEvolution.cpp
--- a/Engine/Source/Runtime/Experimental/Chaos/Private/Chaos/PBDRigidsEvolution.cpp
+++ b/Engine/Source/Runtime/Experimental/Chaos/Private/Chaos/PBDRigidsEvolution.cpp
@@ -152,7 +152,7 @@ namespace Chaos
 					{
 						return MakeUnique<AABBDynamicTreeType>(Particles, BroadPhaseConfig.MaxChildrenInLeaf, BroadPhaseConfig.MaxTreeDepth, BroadPhaseConfig.MaxPayloadSize, ForceFullBuild ? 0 : BroadPhaseConfig.IterationsPerTimeSlice, true);
 					}
-					return MakeUnique<AABBTreeType>(Particles, BroadPhaseConfig.MaxChildrenInLeaf, BroadPhaseConfig.MaxTreeDepth, BroadPhaseConfig.MaxPayloadSize, ForceFullBuild ? 0 : BroadPhaseConfig.IterationsPerTimeSlice);
+					return MakeUnique<AABBTreeType>(Particles, BroadPhaseConfig.MaxChildrenInLeaf, BroadPhaseConfig.MaxTreeDepth, BroadPhaseConfig.MaxPayloadSize, ForceFullBuild ? 0 : BroadPhaseConfig.IterationsPerTimeSlice, false, AccelerationStructureUseDirtyTreeInsteadOfGrid == 1);
 				}
 				else if (BroadPhaseConfig.BroadphaseType == FBroadPhaseConfig::TreeOfGridAndGrid || BroadPhaseConfig.BroadphaseType == FBroadPhaseConfig::TreeOfGrid)
 				{
