Chaos
- Capsule-Triangle custom manifold creation. Was previously using a generic convex-convex collision detection which does not work with single-sided triangle collision. As a result we could end up rejecting collisions that are detected as inward facing when we should have generated an alternative out-facing contact instead. This is now fixed.
- can be reverted on a cvar (p.Chaos.Collision.UseCapsuleTriMesh2 0)

diff --git a/Engine/Source/Programs/HeadlessChaos/Private/HeadlessChaosTestCapsuleTriangleCollision.cpp b/Engine/Source/Programs/HeadlessChaos/Private/HeadlessChaosTestCapsuleTriangleCollision.cpp
new file mode 100644
--- /dev/null
+++ b/Engine/Source/Programs/HeadlessChaos/Private/HeadlessChaosTestCapsuleTriangleCollision.cpp
@@ -0,0 +1,244 @@
+// Copyright Epic Games, Inc. All Rights Reserved.
+#include "HeadlessChaos.h"
+#include "Chaos/Capsule.h"
+#include "Chaos/Collision/CapsuleTriangleContactPoint.h"
+#include "Chaos/Collision/ContactPoint.h"
+#include "Chaos/Triangle.h"
+
+namespace ChaosTest
+{
+	using namespace Chaos;
+
+	bool ManifoldContainsPoint(const TCArray<FContactPoint, 4>& ContactPoints, const int32 BodyIndex, const FVec3& Point, const FReal Tolerance)
+	{
+		for (FContactPoint ContactPoint : ContactPoints)
+		{
+			FVec3 Delta = ContactPoint.ShapeContactPoints[BodyIndex] - Point;
+			if (Delta.Size() <= Tolerance)
+			{
+				return true;
+			}
+		}
+		return false;
+	}
+
+	bool IsPointOnCapsuleSurface(const FImplicitCapsule3& Capsule, const FVec3& Point, const FReal Tolerance)
+	{
+		const FReal SegmentT = Utilities::ClosestTimeOnLineSegment(Point, Capsule.GetX1(), Capsule.GetX2());
+		const FVec3 SegmentPoint = FMath::Lerp(Capsule.GetX1(), Capsule.GetX2(), SegmentT);
+		const FVec3 Delta = Point - SegmentPoint;
+		const FReal Distance = Delta.Size();
+		return FMath::IsNearlyEqual(Distance, Capsule.GetRadius(), Tolerance);
+	}
+
+	// Double face contact: both capsule points within the edge planes.
+	// Capsule Radius 20cm, at 10cm above triangle.
+	// Should have 2 contacts points, one at each segement end point.
+	//
+	//  ________
+	//  \       |
+	//   \ ==== |
+	//    \     |
+	//     \    |
+	//      \   |
+	//       \  |
+	//        \ |
+	//         \|
+	//
+	GTEST_TEST(CapsuleTriangleTests, TestFace)
+	{
+		const FReal Tolerance = UE_KINDA_SMALL_NUMBER;
+		const FReal CullDistance = FReal(3);
+
+		FTriangle Triangle(FVec3(0,0,0), FVec3(100, -100, 0), FVec3(100, 0, 0));
+		FImplicitCapsule3 Capsule(FVec3(30, -10, 10), FVec3(90, -10, 10), FReal(20));
+
+		TCArray<FContactPoint, 4> ContactPoints;
+		ConstructCapsuleTriangleOneShotManifold2(Capsule, Triangle, CullDistance, ContactPoints);
+
+		EXPECT_EQ(ContactPoints.Num(), 2);
+		if (ContactPoints.Num() == 2)
+		{
+			EXPECT_TRUE(ManifoldContainsPoint(ContactPoints, 0, Capsule.GetX1() - Capsule.GetRadius() * Triangle.GetNormal(), Tolerance));
+			EXPECT_TRUE(ManifoldContainsPoint(ContactPoints, 0, Capsule.GetX2() - Capsule.GetRadius() * Triangle.GetNormal(), Tolerance));
+		}
+	}
+
+	// Same as TestFace except capsule is just below the triangle.
+	// Should not get any contacts.
+	//
+	GTEST_TEST(CapsuleTriangleTests, TestFaceBelow)
+	{
+		const FReal Tolerance = UE_KINDA_SMALL_NUMBER;
+		const FReal CullDistance = FReal(3);
+
+		FTriangle Triangle(FVec3(0, 0, 0), FVec3(100, -100, 0), FVec3(100, 0, 0));
+		FImplicitCapsule3 Capsule(FVec3(30, -10, -1), FVec3(90, -10, -1), FReal(20));
+
+		TCArray<FContactPoint, 4> ContactPoints;
+		ConstructCapsuleTriangleOneShotManifold2(Capsule, Triangle, CullDistance, ContactPoints);
+
+		EXPECT_EQ(ContactPoints.Num(), 0);
+	}
+
+	// Same as TestFace except capsule is exactly in the plane of the triangle
+	// Should be same result as TestFace
+	//
+	GTEST_TEST(CapsuleTriangleTests, TestFaceAligned)
+	{
+		const FReal Tolerance = UE_KINDA_SMALL_NUMBER;
+		const FReal CullDistance = FReal(3);
+
+		FTriangle Triangle(FVec3(0, 0, 0), FVec3(100, -100, 0), FVec3(100, 0, 0));
+		FImplicitCapsule3 Capsule(FVec3(30, -10, 0), FVec3(90, -10, 0), FReal(20));
+
+		TCArray<FContactPoint, 4> ContactPoints;
+		ConstructCapsuleTriangleOneShotManifold2(Capsule, Triangle, CullDistance, ContactPoints);
+
+		EXPECT_EQ(ContactPoints.Num(), 2);
+		if (ContactPoints.Num() == 2)
+		{
+			EXPECT_TRUE(ManifoldContainsPoint(ContactPoints, 0, Capsule.GetX1() - Capsule.GetRadius() * Triangle.GetNormal(), Tolerance));
+			EXPECT_TRUE(ManifoldContainsPoint(ContactPoints, 0, Capsule.GetX2() - Capsule.GetRadius() * Triangle.GetNormal(), Tolerance));
+		}
+	}
+
+	// Face + Edge contact:
+	// Capsule Radius 20cm, at 10cm above triangle.
+	// Should have 2 contacts points
+	//
+	//  ________
+	//  \       |
+	//   \   ======
+	//    \     |
+	//     \    |
+	//      \   |
+	//       \  |
+	//        \ |
+	//         \|
+	//
+	GTEST_TEST(CapsuleTriangleTests, TestFaceEdge)
+	{
+		const FReal Tolerance = UE_KINDA_SMALL_NUMBER;
+		const FReal CullDistance = FReal(3);
+
+		FTriangle Triangle(FVec3(0, 0, 0), FVec3(100, -100, 0), FVec3(100, 0, 0));
+		FImplicitCapsule3 Capsule(FVec3(50, -10, 10), FVec3(120, -10, 10), FReal(20));
+
+		TCArray<FContactPoint, 4> ContactPoints;
+		ConstructCapsuleTriangleOneShotManifold2(Capsule, Triangle, CullDistance, ContactPoints);
+
+		EXPECT_EQ(ContactPoints.Num(), 2);
+		if (ContactPoints.Num() == 2)
+		{
+			EXPECT_TRUE(ManifoldContainsPoint(ContactPoints, 0, Capsule.GetX1() - Capsule.GetRadius() * Triangle.GetNormal(), Tolerance));
+			EXPECT_TRUE(ManifoldContainsPoint(ContactPoints, 0, FVec3(100, -10, 10) - Capsule.GetRadius() * Triangle.GetNormal(), Tolerance));
+		}
+	}
+
+	// Face + Edge contact with capsule tilted by 10deg wrt triangle surface.
+	// Capsule Radius 20cm, at 10cm above triangle.
+	// Should have 2 face contacts points (10deg is below threshold)
+	//
+	//  ________
+	//  \       |
+	//   \ (+)==+==)
+	//    \     |
+	//     \    |
+	//      \   |
+	//       \  |
+	//        \ |
+	//         \|
+	//
+	GTEST_TEST(CapsuleTriangleTests, TestFaceEdgeAngleInThreshold)
+	{
+		const FReal PositionTolerance = UE_KINDA_SMALL_NUMBER;
+		const FReal NormalTolerance = UE_KINDA_SMALL_NUMBER;
+		const FReal CullDistance = FReal(3);
+		const FReal Sin10 = FMath::Sin(FMath::DegreesToRadians(10));
+		const FReal Cos10 = FMath::Cos(FMath::DegreesToRadians(10));
+
+		const FTriangle Triangle(FVec3(0, 0, 0), FVec3(100, -100, 0), FVec3(100, 0, 0));
+		const FImplicitCapsule3 Capsule(FVec3(50, -10, 10 + 50 * Sin10), FVec3(150, -10, 10 - 50 * Sin10), FReal(20));
+
+		TCArray<FContactPoint, 4> ContactPoints;
+		ConstructCapsuleTriangleOneShotManifold2(Capsule, Triangle, CullDistance, ContactPoints);
+
+		EXPECT_EQ(ContactPoints.Num(), 2);
+		if (ContactPoints.Num() == 2)
+		{
+			// We should have a contact on the cylinder below the first segment
+			// And the other will be with the edge. 
+			// @todo(chaos): check this
+			
+			// The normal should be vertical because we are within the face angle threshold
+			EXPECT_NEAR(ContactPoints[1].ShapeContactNormal.Z, FReal(1), NormalTolerance);
+
+			// Make sure both points are on the capsule surface
+			EXPECT_TRUE(IsPointOnCapsuleSurface(Capsule, ContactPoints[0].ShapeContactPoints[0], PositionTolerance));
+			EXPECT_TRUE(IsPointOnCapsuleSurface(Capsule, ContactPoints[1].ShapeContactPoints[0], PositionTolerance));
+		}
+	}
+
+	// Capsule at a 45 degree angle to the triangle, and the capsule axis passes right through an edge.
+	// We should get a single contact pointing out of the triangle at 45 degrees.
+	// 
+	//     \.\
+	//      \.\
+	//  -----\+\
+	//        \.\
+	//
+	GTEST_TEST(CapsuleTriangleTests, TestEdgeZeroSeparation)
+	{
+		const FReal PositionTolerance = UE_KINDA_SMALL_NUMBER;
+		const FReal NormalTolerance = UE_KINDA_SMALL_NUMBER;
+		const FReal CullDistance = FReal(3);
+
+		FTriangle Triangle(FVec3(0, 0, 0), FVec3(100, -100, 0), FVec3(100, 0, 0));
+		FImplicitCapsule3 Capsule(FVec3(50, -10, 50), FVec3(150, -10, -50), FReal(10));
+
+		TCArray<FContactPoint, 4> ContactPoints;
+		ConstructCapsuleTriangleOneShotManifold2(Capsule, Triangle, CullDistance, ContactPoints);
+
+		EXPECT_EQ(ContactPoints.Num(), 1);
+		if (ContactPoints.Num() == 1)
+		{
+			EXPECT_NEAR(ContactPoints[0].Phi, -Capsule.GetRadius(), PositionTolerance);
+			EXPECT_NEAR(FVec3::DotProduct(ContactPoints[0].ShapeContactNormal, FVec3(0,0,1)), FMath::Cos(FMath::DegreesToRadians(45)), NormalTolerance);
+		}
+	}
+
+	// Capsule at a 45 degree angle to the triangle, and the capsule axis passes right through an edge.
+	// This time we have an end cap within the edge planes, so we should get a face contact at that location.
+	// The edge contact will be ignored because it would generate an inward dacing normal, and we are
+	// above the face angle threshold.
+	// 
+	//         / /
+	//        / /
+	//  ---+-/-/
+	//      / /
+	//     / /
+	//
+	GTEST_TEST(CapsuleTriangleTests, TestEdgeZeroSeparation2)
+	{
+		const FReal PositionTolerance = UE_KINDA_SMALL_NUMBER;
+		const FReal NormalTolerance = UE_KINDA_SMALL_NUMBER;
+		const FReal CullDistance = FReal(3);
+
+		FTriangle Triangle(FVec3(0, 0, 0), FVec3(100, -100, 0), FVec3(100, 0, 0));
+		FImplicitCapsule3 Capsule(FVec3(50, -10, -50), FVec3(150, -10, 50), FReal(10));
+
+		TCArray<FContactPoint, 4> ContactPoints;
+		ConstructCapsuleTriangleOneShotManifold2(Capsule, Triangle, CullDistance, ContactPoints);
+
+		EXPECT_EQ(ContactPoints.Num(), 1);
+		if (ContactPoints.Num() == 1)
+		{
+			EXPECT_NEAR(ContactPoints[0].ShapeContactPoints[0].X, Capsule.GetX1().X, PositionTolerance);
+			EXPECT_NEAR(ContactPoints[0].ShapeContactPoints[0].Y, Capsule.GetX1().Y, PositionTolerance);
+			EXPECT_NEAR(ContactPoints[0].ShapeContactPoints[0].Z, Capsule.GetX1().Z - Capsule.GetRadius(), PositionTolerance);
+
+			EXPECT_NEAR(ContactPoints[0].ShapeContactNormal.Z, FReal(1), NormalTolerance);
+		}
+	}
+}
diff --git a/Engine/Source/Runtime/Experimental/Chaos/Private/Chaos/Collision/CapsuleTriangleContactPoint.cpp b/Engine/Source/Runtime/Experimental/Chaos/Private/Chaos/Collision/CapsuleTriangleContactPoint.cpp
new file mode 100644
--- /dev/null
+++ b/Engine/Source/Runtime/Experimental/Chaos/Private/Chaos/Collision/CapsuleTriangleContactPoint.cpp
@@ -0,0 +1,451 @@
+// Copyright Epic Games, Inc. All Rights Reserved.
+#include "Chaos/Collision/CapsuleTriangleContactPoint.h"
+#include "Chaos/Capsule.h"
+#include "Chaos/Collision/ContactPoint.h"
+#include "Chaos/Triangle.h"
+
+//PRAGMA_DISABLE_OPTIMIZATION
+
+namespace Chaos
+{
+	extern bool bChaos_Collision_CapsuleTriMesh_MidPointCull;
+
+	void AddCapsuleTriangleParallelEdgeManifoldContacts(const FVec3& P0, const FVec3& P1, const FVec3& EdgeP0, const FVec3& EdgeP1, const FReal R, const FReal RejectDistanceSq, const FReal NormalToleranceSq, TCArray<FContactPoint, 4>& OutContactPoints);
+
+	// Return true if the value V is in the closed/inclusive range [RangeMin, RangeMax]
+	// @todo(chaos): move to Utilities.h
+	inline bool InRangeClosed(const FReal V, const FReal RangeMin, const FReal RangeMax)
+	{
+		return ((V >= RangeMin) & (V <= RangeMax)) != 0;
+	}
+
+	// Return true if the value V is in the open/exclusive range (RangeMin, RangeMax)
+	// @todo(chaos): move to Utilities.h
+	inline bool InRangeOpen(const FReal V, const FReal RangeMin, const FReal RangeMax)
+	{
+		return ((V > RangeMin) & (V < RangeMax)) != 0;
+	}
+
+
+	// Generate the contact manifold between a capsule and a triangle.
+	//
+	// Algorithm overview:
+	// 
+	// - Triangle is in Capsule space
+	// 
+	// - Generate planes for the triangle face and edges
+	// 
+	// - Special case: If the segment does not cross triangle edge prism and is exactly parallel to an edge
+	//		- Project segment onto edge and select contact point(s)
+	// 
+	// - For each segment point
+	//		- If point is within the edge prism
+	//			- Add face contact (either on end cap or cylinder)
+	// 
+	// - For each edge
+	//		- Add edge contact where the segment is closest to the edge
+	//		- When the angle between segment and face is small, create a face contact instead
+	//
+	void ConstructCapsuleTriangleOneShotManifold2(const FImplicitCapsule3& Capsule, const FTriangle& Triangle, const FReal CullDistance, TCArray<FContactPoint, 4>& OutContactPoints)
+	{
+		//
+		// @todo(chaos): this function is a good SIMD candidate
+		//
+
+		// Capsule segment
+		const FVec3 P0 = Capsule.GetX1();
+		const FVec3 P1 = Capsule.GetX2();
+		const FReal R = Capsule.GetRadius();
+		const FVec3 Axis = Capsule.GetAxis();
+		const FReal L = Capsule.GetHeight();
+		const FReal RejectDistance = R + CullDistance;
+		const FReal RejectDistanceSq = RejectDistance * RejectDistance;
+
+		// Triangle
+		const FVec3& V0 = Triangle.GetVertex(0);
+		const FVec3& V1 = Triangle.GetVertex(1);
+		const FVec3& V2 = Triangle.GetVertex(2);
+		const FVec3 Centroid = Triangle.GetCentroid();
+
+		// Tolerances
+		// @todo(chaos): Everything below needs to be made size agnostic (i.e., all checks with a distance tolerance
+		// need to cope with very larger or small segments and triangles). This will probably not be good enough...(add some tests)
+		const FReal DistanceTolerance = FReal(1.e-10) * Capsule.GetHeight();
+		const FReal NormalTolerance = FReal(1.e-8);
+		const FReal NormalToleranceSq = NormalTolerance * NormalTolerance;
+		const FReal FaceContactSinAngleThreshold = FReal(0.34);	// ~Sin(20deg)
+
+		// Face plane
+		const FVec3& FaceP = Triangle.GetVertex(0);	// Use centroid?
+		FVec3 FaceN = FVec3::CrossProduct(Triangle.GetVertex(1) - Triangle.GetVertex(0), Triangle.GetVertex(2) - Triangle.GetVertex(0));
+		if (!FaceN.Normalize(NormalToleranceSq))
+		{
+			// Degenerate triangle
+			return;
+		}
+
+		// Signed distance of capsule segment points to triangle face
+		const FReal FaceD0 = FVec3::DotProduct(P0 - FaceP, FaceN);
+		const FReal FaceD1 = FVec3::DotProduct(P1 - FaceP, FaceN);
+		const bool bIsParallelFace = FMath::IsNearlyEqual(FaceD0, FaceD1, DistanceTolerance);
+
+		// Reject if too far from tri plane (distance cull)
+		if ((FaceD0 > RejectDistance) && (FaceD1 > RejectDistance))
+		{
+			// Separated from triangle
+			return;
+		}
+
+		// Reject if the middle of the capsule is inside the face (single-sided collision)
+		const FReal FaceDMid = FReal(0.5) * (FaceD0 + FaceD1);
+		if (FaceDMid < FReal(0))
+		{
+			// Far inside triangle
+			return;
+		}
+
+		// Edge plane normals
+		FVec3 EdgeNs[3];
+		EdgeNs[0] = FVec3::CrossProduct(V0 - V2, FaceN);
+		EdgeNs[1] = FVec3::CrossProduct(V1 - V0, FaceN);
+		EdgeNs[2] = FVec3::CrossProduct(V2 - V1, FaceN);
+		EdgeNs[0].Normalize(NormalToleranceSq);
+		EdgeNs[1].Normalize(NormalToleranceSq);
+		EdgeNs[2].Normalize(NormalToleranceSq);
+
+		// Edge plane Signed Distances for each Segment point
+		FReal EdgeD0s[3];
+		FReal EdgeD1s[3];
+		EdgeD0s[0] = FVec3::DotProduct(P0 - V0, EdgeNs[0]);
+		EdgeD1s[0] = FVec3::DotProduct(P1 - V0, EdgeNs[0]);
+		if (((EdgeD0s[0] > RejectDistance) & (EdgeD1s[0] > RejectDistance)) != 0)
+		{
+			// Separated from triangle
+			return;
+		}
+
+		EdgeD0s[1] = FVec3::DotProduct(P0 - V1, EdgeNs[1]);
+		EdgeD1s[1] = FVec3::DotProduct(P1 - V1, EdgeNs[1]);
+		if (((EdgeD0s[1] > RejectDistance) & (EdgeD1s[1] > RejectDistance)) != 0)
+		{
+			// Separated from triangle
+			return;
+		}
+
+		EdgeD0s[2] = FVec3::DotProduct(P0 - V2, EdgeNs[2]);
+		EdgeD1s[2] = FVec3::DotProduct(P1 - V2, EdgeNs[2]);
+		if (((EdgeD0s[2] > RejectDistance) & (EdgeD1s[2] > RejectDistance)) != 0)
+		{
+			// Separated from triangle
+			return;
+		}
+
+		// Handle the case where the end cap(s) are inside all the edge planes. This should be fairly common unless the triangles are
+		// much smaller than the capsule.
+		// NOTE: Here we can only collide with the end-cap that is closest to the triangle (unless they are both exactly the same distance)
+		// because the farthest segment point would generate a cylinder contact, not an end cap contact.
+		bool bCollided0 = false;
+		bool bCollided1 = false;
+		const bool bInsideAll0 = ((EdgeD0s[0] <= DistanceTolerance) & (EdgeD0s[1] <= DistanceTolerance) & (EdgeD0s[2] <= DistanceTolerance)) != 0;
+		const bool bInsideAll1 = ((EdgeD1s[0] <= DistanceTolerance) & (EdgeD1s[1] <= DistanceTolerance) & (EdgeD1s[2] <= DistanceTolerance)) != 0;
+		if ((bInsideAll0 & (FaceD0 < RejectDistance) & (FaceD0 < FaceD1 + DistanceTolerance)) != 0)
+		{
+			FContactPoint& ContactPoint = OutContactPoints[OutContactPoints.Add()];
+			ContactPoint.ShapeContactPoints[0] = P0 - R * FaceN;
+			ContactPoint.ShapeContactPoints[1] = P0 - FaceD0 * FaceN;
+			ContactPoint.ShapeContactNormal = FaceN;
+			ContactPoint.Phi = FaceD0 - R;
+			ContactPoint.ContactType = EContactPointType::VertexPlane;
+			ContactPoint.FaceIndex = INDEX_NONE;
+			bCollided0 = true;
+		}
+		if ((bInsideAll1 & (FaceD1 < RejectDistance) & (FaceD1 < FaceD0 + DistanceTolerance)) != 0)
+		{
+			FContactPoint& ContactPoint = OutContactPoints[OutContactPoints.Add()];
+			ContactPoint.ShapeContactPoints[0] = P1 - R * FaceN;
+			ContactPoint.ShapeContactPoints[1] = P1 - FaceD1 * FaceN;
+			ContactPoint.ShapeContactNormal = FaceN;
+			ContactPoint.Phi = FaceD1 - R;
+			ContactPoint.ContactType = EContactPointType::VertexPlane;
+			ContactPoint.FaceIndex = INDEX_NONE;
+			bCollided1 = true;
+		}
+
+		// If we added both contacts above, we are (parallel to surface and both points inside edge planes)
+		if (OutContactPoints.Num() == 2)
+		{
+			// Full manifold
+			return;
+		}
+
+		// Handle the parallel segment-edge case by clipping the segment to the edge
+		// NOTE: If we are parallel to an edge (and outside the edge prism) we can only collide with the one edge
+		bool bIsParallelEdge[3];
+		bIsParallelEdge[0] = ((EdgeD0s[0] >= FReal(0)) & (EdgeD1s[0] >= FReal(0)) & FMath::IsNearlyEqual(EdgeD0s[0], EdgeD1s[0], DistanceTolerance)) != 0;
+		bIsParallelEdge[1] = ((EdgeD0s[1] >= FReal(0)) & (EdgeD1s[1] >= FReal(0)) & FMath::IsNearlyEqual(EdgeD0s[1], EdgeD1s[1], DistanceTolerance)) != 0;
+		bIsParallelEdge[2] = ((EdgeD0s[2] >= FReal(0)) & (EdgeD1s[2] >= FReal(0)) & FMath::IsNearlyEqual(EdgeD0s[2], EdgeD1s[2], DistanceTolerance)) != 0;
+		int32 EdgeVertexIndex0 = 2;
+		if ((bIsParallelFace & (bIsParallelEdge[0] | bIsParallelEdge[1] | bIsParallelEdge[2])) != 0)
+		{
+			for (int32 EdgeIndex = 0; EdgeIndex < 3; ++EdgeIndex)
+			{
+				const int32 EdgeVertexIndex1 = EdgeIndex;
+				const FVec3& EdgeP0 = Triangle.GetVertex(EdgeVertexIndex0);
+				const FVec3& EdgeP1 = Triangle.GetVertex(EdgeVertexIndex1);
+				EdgeVertexIndex0 = EdgeVertexIndex1;
+				if (bIsParallelEdge[EdgeIndex])
+				{
+					AddCapsuleTriangleParallelEdgeManifoldContacts(P0, P1, EdgeP0, EdgeP1, R, RejectDistanceSq, NormalToleranceSq, OutContactPoints);
+					return;
+				}
+			}
+		}
+
+		// Sine of angle between segment axis and the triangle face
+		const FReal AxisDotNormal = FVec3::DotProduct(Axis, FaceN);
+		const FReal SinAxisFaceAngle = AxisDotNormal;
+		const bool bPreferFaceContact = (FMath::Abs(SinAxisFaceAngle) < FaceContactSinAngleThreshold);
+
+		// Generate contacts for the cylinder ends near the face. Only consider the ends where we did not generate an end-cap contact, 
+		// and only when we are at a low angle to the face, or the cylinder end is close to the face. 
+		// The cylinder points can only be inside the edge planes if the segment points are within Radius of the edge planes
+		const bool bNearAll0 = ((EdgeD0s[0] <= R + DistanceTolerance) & (EdgeD0s[1] <= R + DistanceTolerance) & (EdgeD0s[2] <= R + DistanceTolerance)) != 0;
+		const bool bNearAll1 = ((EdgeD1s[0] <= R + DistanceTolerance) & (EdgeD1s[1] <= R + DistanceTolerance) & (EdgeD1s[2] <= R + DistanceTolerance)) != 0;
+		const bool bCheckCylinder0 = (!bCollided0 & (bPreferFaceContact | (FaceD0 < RejectDistance)) & bNearAll0) != 0;
+		const bool bCheckCylinder1 = (!bCollided1 & (bPreferFaceContact | (FaceD1 < RejectDistance)) & bNearAll1) != 0;
+		if ((bCheckCylinder0 | bCheckCylinder1) != 0)
+		{
+			FVec3 RadialAxis = FVec3::CrossProduct(FVec3::CrossProduct(Axis, FaceN), Axis);
+			if (RadialAxis.Normalize(NormalTolerance))
+			{
+				// We want Radial axis to point against the normal
+				if (FVec3::DotProduct(RadialAxis, FaceN) > FReal(0))
+				{
+					RadialAxis = -RadialAxis;
+				}
+
+				// Utility to add a cylinder contact point, if it is within the edge planes
+				const auto& TryAddCylinderContact = [R, &RadialAxis, &V0, &V1, &V2, &FaceN, &FaceP, &EdgeNs, DistanceTolerance, &OutContactPoints](const FVec3& P) -> void
+				{
+					const FVec3 CylinderP = P + R * RadialAxis;
+					const FReal CylinderEdgeD0 = FVec3::DotProduct(CylinderP - V0, EdgeNs[0]);
+					const FReal CylinderEdgeD1 = FVec3::DotProduct(CylinderP - V1, EdgeNs[1]);
+					const FReal CylinderEdgeD2 = FVec3::DotProduct(CylinderP - V2, EdgeNs[2]);
+					const bool bCylinderInsideAll = ((CylinderEdgeD0 <= DistanceTolerance) & (CylinderEdgeD1 <= DistanceTolerance) & (CylinderEdgeD2 <= DistanceTolerance)) != 0;
+					if (bCylinderInsideAll)
+					{
+						const FReal CylinderFaceD = FVec3::DotProduct(CylinderP - FaceP, FaceN);
+
+						FContactPoint& ContactPoint = OutContactPoints[OutContactPoints.Add()];
+						ContactPoint.ShapeContactPoints[0] = CylinderP;
+						ContactPoint.ShapeContactPoints[1] = CylinderP - CylinderFaceD * FaceN;
+						ContactPoint.ShapeContactNormal = FaceN;
+						ContactPoint.Phi = CylinderFaceD;
+						ContactPoint.ContactType = EContactPointType::VertexPlane;
+						ContactPoint.FaceIndex = INDEX_NONE;
+					}
+				};
+
+				if (bCheckCylinder0)
+				{
+					TryAddCylinderContact(P0);
+				}
+				if (bCheckCylinder1)
+				{
+					TryAddCylinderContact(P1);
+				}
+			}
+		}
+
+		// If we have a contact at both ends now, we are done
+		if (OutContactPoints.Num() == 2)
+		{
+			// Full manifold
+			return;
+		}
+
+		// Add edge contacts to the manifold
+		EdgeVertexIndex0 = 2;
+		for (int32 EdgeIndex = 0; EdgeIndex < 3; ++EdgeIndex)
+		{
+			const int32 EdgeVertexIndex1 = EdgeIndex;
+			const FVec3& EdgeP0 = Triangle.GetVertex(EdgeVertexIndex0);
+			const FVec3& EdgeP1 = Triangle.GetVertex(EdgeVertexIndex1);
+			EdgeVertexIndex0 = EdgeVertexIndex1;
+
+			// We only care about edges if at least one capsule segment point is outside the edge plane
+			// (internal points were handled already)
+			if (((EdgeD0s[EdgeIndex] > -DistanceTolerance) | (EdgeD1s[EdgeIndex] > -DistanceTolerance)) != 0)
+			{
+				// Find the nearest point on the capsule segment to the edge segment
+				FReal SegmentT, EdgeT;
+				FVec3 SegmentP, EdgeP;
+				Utilities::NearestPointsOnLineSegments(P0, P1, EdgeP0, EdgeP1, SegmentT, EdgeT, SegmentP, EdgeP);
+
+				// Calculate the separation vector, correct for sign so that we are pushing out of the face
+				FVec3 SegmentEdgeN = SegmentP - EdgeP;
+				FReal SegmentEdgeDistSign = FReal(1);
+				const FReal SegmentEdgeDistSq = SegmentEdgeN.SizeSquared();
+				const FReal DotFace = FVec3::DotProduct(SegmentEdgeN, FaceN);
+
+				// Check for cull distance (only if the segment point is not inside the face)
+				if (((DotFace > FReal(0)) & (SegmentEdgeDistSq > RejectDistanceSq)) != 0)
+				{
+					continue;
+				}
+
+				// Make sure the normal point out of the face (we ned to correct the signed distance later too)
+				// Tolerance so we don't throw away "exact" edge collisions
+				if (DotFace < -NormalTolerance)
+				{
+					SegmentEdgeN = -SegmentEdgeN;
+					SegmentEdgeDistSign = FReal(-1);
+				}
+
+				// If the angle between the axis and the face is below a threshold, we will create a face contact
+				// rather than an edge contact where possible.
+				const bool bInEdgeRange = InRangeOpen(EdgeT, FReal(0), FReal(1));
+				const bool bInSegmentRange = InRangeOpen(SegmentT, FReal(0), FReal(1));
+				const bool bCrossedEdgeSegment = ((bInEdgeRange & bInSegmentRange) != 0);
+
+				// Calculate separation distance and normal
+				// If we have zero separation, we cannot renormalize the separation vector so we must calculate the normal
+				FReal SegmentEdgeDist;
+				if (SegmentEdgeDistSq > NormalToleranceSq)
+				{
+					// Get the signed distance and separating axis
+					SegmentEdgeDist = FMath::Sqrt(SegmentEdgeDistSq);
+					SegmentEdgeN = SegmentEdgeN / SegmentEdgeDist;
+					SegmentEdgeDist *= SegmentEdgeDistSign;
+				}
+				else
+				{
+					// Segment passes right through edge - calculate normal
+					SegmentEdgeDist = FReal(0);
+					SegmentEdgeN = FVec3::CrossProduct(Axis, EdgeP1 - EdgeP0);
+					if (!SegmentEdgeN.Normalize(NormalToleranceSq))
+					{
+						continue;
+					}
+					if (FVec3::DotProduct(SegmentEdgeN, FaceN) < FReal(0))
+					{
+						SegmentEdgeN = -SegmentEdgeN;
+					}
+				}
+
+				// We cannot collide with the "inside" of the edge (normal must be facing away from triangle center).
+				// However, we can convert it to a face contact if we are within the Face Contact angle threshold.
+				if ((bCrossedEdgeSegment & bPreferFaceContact) == 0)
+				{
+					// We use a tolerance here so we don't reject very nearly face collisions
+					// @todo(chaos): size dependence issue?
+					const FReal DotCentroid = FVec3::DotProduct(EdgeP - Centroid, SegmentEdgeN);
+					if (DotCentroid < -NormalTolerance)
+					{
+						continue;
+					}
+				}
+
+				// If we are within the face angle tolerance, generate a face contact rather than an edge one
+				// (but only if we have an edge contact and not a vertex one)
+				// NOTE: we rely on the fact that the Ts will be exactly 0 or 1 when the near point is outside the segment or edge
+				if ((bCrossedEdgeSegment & bPreferFaceContact) != 0)
+				{
+					const FVec3 CapsuleP = SegmentP - R * SegmentEdgeN;
+					const FReal CapsuleDist = FVec3::DotProduct(CapsuleP - FaceP, FaceN);
+
+					FContactPoint& ContactPoint = OutContactPoints[OutContactPoints.Add()];
+					ContactPoint.ShapeContactPoints[0] = CapsuleP;
+					ContactPoint.ShapeContactPoints[1] = CapsuleP - CapsuleDist * FaceN;
+					ContactPoint.ShapeContactNormal = FaceN;
+					ContactPoint.Phi = CapsuleDist - R;
+					ContactPoint.ContactType = EContactPointType::VertexPlane;
+					ContactPoint.FaceIndex = INDEX_NONE;
+				}
+				else 
+				{
+					FContactPoint& ContactPoint = OutContactPoints[OutContactPoints.Add()];
+					ContactPoint.ShapeContactPoints[0] = SegmentP - R * SegmentEdgeN;
+					ContactPoint.ShapeContactPoints[1] = EdgeP;
+					ContactPoint.ShapeContactNormal = SegmentEdgeN;
+					ContactPoint.Phi = SegmentEdgeDist - R;
+					ContactPoint.ContactType = EContactPointType::EdgeEdge;
+					ContactPoint.FaceIndex = INDEX_NONE;
+				}
+			}
+		}
+	}
+
+	// Utility for ConstructCapsuleTriangleOneShotManifold2 to handle the rare case when a capsule is perfectly aligned with an edge of a triangle.
+	// In this case we project the segment onto the edge and, depending on whether there is any overlap, use either the closest point or the
+	// two clipped points as contacts.
+	void AddCapsuleTriangleParallelEdgeManifoldContacts(const FVec3& P0, const FVec3& P1, const FVec3& EdgeP0, const FVec3& EdgeP1, const FReal R, const FReal RejectDistanceSq, const FReal NormalToleranceSq, TCArray<FContactPoint, 4>& OutContactPoints)
+	{
+		// Utility to add a contact to the array if it is within cull distance
+		const auto& AddContact = [](const FVec3& SegmentEdgeC, const FVec3& SegmentEdgeDelta, const FReal R, const FReal RejectDistanceSq, const FReal NormalToleranceSq, TCArray<FContactPoint, 4>& OutContactPoints) -> void
+		{
+			const FReal SegmentEdgeDistSq = SegmentEdgeDelta.SizeSquared();
+			if ((SegmentEdgeDistSq < RejectDistanceSq) && (SegmentEdgeDistSq > NormalToleranceSq))
+			{
+				const FReal SegmentEdgeDist = FMath::Sqrt(SegmentEdgeDistSq);
+				const FVec3 SegmentEdgeN = SegmentEdgeDelta / SegmentEdgeDist;
+
+				FContactPoint& ContactPoint = OutContactPoints[OutContactPoints.Add()];
+				ContactPoint.ShapeContactPoints[0] = SegmentEdgeC + (SegmentEdgeDist - R) * SegmentEdgeN;
+				ContactPoint.ShapeContactPoints[1] = SegmentEdgeC;
+				ContactPoint.ShapeContactNormal = SegmentEdgeN;
+				ContactPoint.Phi = SegmentEdgeDist - R;
+				ContactPoint.ContactType = EContactPointType::EdgeEdge;
+				ContactPoint.FaceIndex = INDEX_NONE;
+			}
+		};
+
+		const FVec3 EdgeDelta = EdgeP1 - EdgeP0;
+		const FReal EdgeLenSq = EdgeDelta.SizeSquared();
+		if (EdgeLenSq > NormalToleranceSq)
+		{
+			// Project the first segment point onto the edge
+			const FReal T0 = FVec3::DotProduct(P0 - EdgeP0, EdgeDelta) / EdgeLenSq;
+			const FReal T1 = FVec3::DotProduct(P1 - EdgeP0, EdgeDelta) / EdgeLenSq;
+
+			// If both points have T < 0, only use the point with highest T
+			// If both points have T > 1, only use the point with lowest T
+			// Otherwise use both points with T clamped to [0,1]
+			if ((InRangeClosed(T0, FReal(0), FReal(1)) | InRangeClosed(T1, FReal(0), FReal(1))) != 0)
+			{
+				// At least one segment point projects to a point inside the edge extents.
+				// Clip the edge-mapped segement and use the clipped verts as contacts.
+
+				const FVec3 SegmentEdgeDelta0 = P0 - (EdgeP0 + T0 * EdgeDelta);
+				const FVec3 SegmentEdgeC0 = EdgeP0 + FMath::Clamp(T0, FReal(0), FReal(1)) * EdgeDelta;
+				AddContact(SegmentEdgeC0, SegmentEdgeDelta0, R, RejectDistanceSq, NormalToleranceSq, OutContactPoints);
+
+				const FVec3 SegmentEdgeDelta1 = P1 - (EdgeP0 + T1 * EdgeDelta);
+				const FVec3 SegmentEdgeC1 = EdgeP0 + FMath::Clamp(T1, FReal(0), FReal(1)) * EdgeDelta;
+				AddContact(SegmentEdgeC1, SegmentEdgeDelta1, R, RejectDistanceSq, NormalToleranceSq, OutContactPoints);
+			}
+			else
+			{
+				// Both segment points projected to the edge are outside the edge extents.
+				// Use the segment point nearest the edge as the contact.
+
+				FReal SegmentEdgeT;
+				FVec3 SegmentEdgeP;
+				if ((T0 < FReal(0)) && (T1 < FReal(0)))
+				{
+					SegmentEdgeT = (T0 > T1) ? T0 : T1;
+					SegmentEdgeP = (T0 > T1) ? P0 : P1;
+				}
+				else
+				{
+					SegmentEdgeT = (T0 < T1) ? T0 : T1;
+					SegmentEdgeP = (T0 < T1) ? P0 : P1;
+				}
+				const FVec3 SegmentEdgeC = EdgeP0 + FMath::Clamp(SegmentEdgeT, FReal(0), FReal(1)) * EdgeDelta;
+				const FVec3 SegmentEdgeDelta = (SegmentEdgeP - SegmentEdgeC);
+				AddContact(SegmentEdgeC, SegmentEdgeDelta, R, RejectDistanceSq, NormalToleranceSq, OutContactPoints);
+			}
+		}
+	}
+
+}
\ No newline at end of file
diff --git a/Engine/Source/Runtime/Experimental/Chaos/Private/Chaos/Collision/ContactTriangles.cpp b/Engine/Source/Runtime/Experimental/Chaos/Private/Chaos/Collision/ContactTriangles.cpp
--- a/Engine/Source/Runtime/Experimental/Chaos/Private/Chaos/Collision/ContactTriangles.cpp
+++ b/Engine/Source/Runtime/Experimental/Chaos/Private/Chaos/Collision/ContactTriangles.cpp
@@ -542,10 +542,15 @@ namespace Chaos
 			const int32 DrawPriority = 10;
 			for (const FTriangleContactPoint& ContactPoint : TriangleContactPoints)
 			{
+				const FVec3 P0 = ConvexTransform.TransformPosition(ContactPoint.ShapeContactPoints[0]);
 				const FVec3 P1 = ConvexTransform.TransformPosition(ContactPoint.ShapeContactPoints[1]);
 				const FVec3 N = ConvexTransform.TransformVectorNoScale(ContactPoint.ShapeContactNormal);
 
+				// Draw the normal from the triangle face
 				FDebugDrawQueue::GetInstance().DrawDebugLine(P1, P1 + FReal(15) * N, Color, false, FRealSingle(Duration), DrawPriority, FRealSingle(LineThickness));
+
+				// Draw a thin black line connecting the two contact points (triangle face to convex surface)
+				FDebugDrawQueue::GetInstance().DrawDebugLine(P0, P1, FColor::Black, false, FRealSingle(Duration), DrawPriority, 0.5f * FRealSingle(LineThickness));
 			}
 		}
 #endif
diff --git a/Engine/Source/Runtime/Experimental/Chaos/Private/Chaos/CollisionOneShotManifolds.cpp b/Engine/Source/Runtime/Experimental/Chaos/Private/Chaos/CollisionOneShotManifolds.cpp
--- a/Engine/Source/Runtime/Experimental/Chaos/Private/Chaos/CollisionOneShotManifolds.cpp
+++ b/Engine/Source/Runtime/Experimental/Chaos/Private/Chaos/CollisionOneShotManifolds.cpp
@@ -89,6 +89,9 @@ namespace Chaos
 	FAutoConsoleVariableRef CVarChaos_Collision_OneSidedTriangleMesh(TEXT("p.Chaos.Collision.OneSidedTriangleMesh"), bChaos_Collision_OneSidedTriangleMesh, TEXT(""));
 	FAutoConsoleVariableRef CVarChaos_Collision_OneSidedHeightfield(TEXT("p.Chaos.Collision.OneSidedHeightField"), bChaos_Collision_OneSidedHeightField, TEXT(""));
 
+	bool bChaos_Collision_UseCapsuleTriMesh2 = true;
+	FAutoConsoleVariableRef CVarChaos_Collision_UseCapsuleTriMesh2(TEXT("p.Chaos.Collision.UseCapsuleTriMesh2"), bChaos_Collision_UseCapsuleTriMesh2, TEXT(""));
+
 	namespace Collisions
 	{
 		// Forward delarations we need from CollisionRestitution.cpp
diff --git a/Engine/Source/Runtime/Experimental/Chaos/Private/Chaos/CollisionOneShotManifoldsMiscShapes.cpp b/Engine/Source/Runtime/Experimental/Chaos/Private/Chaos/CollisionOneShotManifoldsMiscShapes.cpp
--- a/Engine/Source/Runtime/Experimental/Chaos/Private/Chaos/CollisionOneShotManifoldsMiscShapes.cpp
+++ b/Engine/Source/Runtime/Experimental/Chaos/Private/Chaos/CollisionOneShotManifoldsMiscShapes.cpp
@@ -5,6 +5,7 @@
 #include "Chaos/CollisionOneShotManifolds.h"
 #include "Chaos/CollisionResolution.h"
 #include "Chaos/Collision/CapsuleConvexContactPoint.h"
+#include "Chaos/Collision/CapsuleTriangleContactPoint.h"
 #include "Chaos/Collision/ContactPointsMiscShapes.h"
 #include "Chaos/Collision/ContactTriangles.h"
 #include "Chaos/Collision/SphereConvexContactPoint.h"
@@ -39,6 +40,8 @@ namespace Chaos
 	extern bool bChaos_Collision_OneSidedTriangleMesh;
 	extern bool bChaos_Collision_OneSidedHeightField;
 
+	extern bool bChaos_Collision_UseCapsuleTriMesh2;
+
 	namespace Collisions
 	{
 
@@ -622,9 +625,16 @@ namespace Chaos
 		}
 
 		template<>
-		void GenerateConvexTriangleOneShotManifold<FCapsule>(const FCapsule& Convex, const FTriangle& Triangle, const FReal CullDistance, TCArray<FContactPoint, 4>& OutContactPoints)
+		void GenerateConvexTriangleOneShotManifold<FCapsule>(const FCapsule& Capsule, const FTriangle& Triangle, const FReal CullDistance, TCArray<FContactPoint, 4>& OutContactPoints)
 		{
-			ConstructCapsuleTriangleOneShotManifold(Convex, Triangle, CullDistance, OutContactPoints);
+			if (bChaos_Collision_UseCapsuleTriMesh2)
+			{
+				ConstructCapsuleTriangleOneShotManifold2(Capsule, Triangle, CullDistance, OutContactPoints);
+			}
+			else
+			{
+				ConstructCapsuleTriangleOneShotManifold(Capsule, Triangle, CullDistance, OutContactPoints);
+			}
 		}
 
 		/**
@@ -640,7 +650,6 @@ namespace Chaos
 		template<typename ConvexType, typename MeshType>
 		void GenerateConvexMeshContactPoints(const ConvexType& Convex, const MeshType& Mesh, const FAABB3& MeshQueryBounds, const FRigidTransform3& MeshToConvexTransform, const FReal CullDistance, FContactTriangleCollector& MeshContacts)
 		{
-			FContactTriangles ContactTriangles;
 			TCArray<FContactPoint, 4> TriangleManifoldPoints;
 
 			// Loop over all the triangles, build a manifold and add the points to the total manifold
@@ -659,7 +668,7 @@ namespace Chaos
 				}
 			});
 
-			// Reduce contacts to a minimum manifold abd transform contact data back into shape-local space
+			// Reduce contacts to a minimum manifold and transform contact data back into shape-local space
 			MeshContacts.ProcessContacts(MeshToConvexTransform);
 		}
 
diff --git a/Engine/Source/Runtime/Experimental/Chaos/Public/Chaos/Collision/CapsuleTriangleContactPoint.h b/Engine/Source/Runtime/Experimental/Chaos/Public/Chaos/Collision/CapsuleTriangleContactPoint.h
new file mode 100644
--- /dev/null
+++ b/Engine/Source/Runtime/Experimental/Chaos/Public/Chaos/Collision/CapsuleTriangleContactPoint.h
@@ -0,0 +1,11 @@
+// Copyright Epic Games, Inc. All Rights Reserved.
+#pragma once
+#include "Chaos/Core.h"
+#include "Chaos/CollisionResolutionTypes.h"
+#include "Chaos/Framework/UncheckedArray.h"
+#include "Chaos/ImplicitFwd.h"
+
+namespace Chaos
+{
+	extern CHAOS_API void ConstructCapsuleTriangleOneShotManifold2(const FImplicitCapsule3& Capsule, const FTriangle& Triangle, const FReal CullDistance, TCArray<FContactPoint, 4>& OutContactPoints);
+}
\ No newline at end of file
