Fix static analysis warning

diff --git a/Engine/Source/Runtime/Experimental/Chaos/Private/Chaos/Collision/CapsuleTriangleContactPoint.cpp b/Engine/Source/Runtime/Experimental/Chaos/Private/Chaos/Collision/CapsuleTriangleContactPoint.cpp
--- a/Engine/Source/Runtime/Experimental/Chaos/Private/Chaos/Collision/CapsuleTriangleContactPoint.cpp
+++ b/Engine/Source/Runtime/Experimental/Chaos/Private/Chaos/Collision/CapsuleTriangleContactPoint.cpp
@@ -211,8 +211,8 @@ namespace Chaos
 		// The cylinder points can only be inside the edge planes if the segment points are within Radius of the edge planes
 		const bool bNearAll0 = ((EdgeD0s[0] <= R + DistanceTolerance) & (EdgeD0s[1] <= R + DistanceTolerance) & (EdgeD0s[2] <= R + DistanceTolerance)) != 0;
 		const bool bNearAll1 = ((EdgeD1s[0] <= R + DistanceTolerance) & (EdgeD1s[1] <= R + DistanceTolerance) & (EdgeD1s[2] <= R + DistanceTolerance)) != 0;
-		const bool bCheckCylinder0 = (!bCollided0 & (bPreferFaceContact | (FaceD0 < RejectDistance)) & bNearAll0) != 0;
-		const bool bCheckCylinder1 = (!bCollided1 & (bPreferFaceContact | (FaceD1 < RejectDistance)) & bNearAll1) != 0;
+		const bool bCheckCylinder0 = ((!bCollided0) & (bPreferFaceContact | (FaceD0 < RejectDistance)) & bNearAll0) != 0;
+		const bool bCheckCylinder1 = ((!bCollided1) & (bPreferFaceContact | (FaceD1 < RejectDistance)) & bNearAll1) != 0;
 		if ((bCheckCylinder0 | bCheckCylinder1) != 0)
 		{
 			FVec3 RadialAxis = FVec3::CrossProduct(FVec3::CrossProduct(Axis, FaceN), Axis);
