Updates to SimCallbackObject API:
- Change method for specifying callback object options... Rather than needing to provide a list of bools at the time of CreateAndRegister, specify an option bitmask in parent template.
- Add particle creation and deletion optional callbacks to sim callback object.
- Update ChaosUserDataPT to use particle removal callback so that userdata associated with deleted particles gets removed.

diff --git a/Engine/Plugins/Experimental/ChaosUserDataPT/Source/ChaosUserDataPT/Public/ChaosUserDataPT.h b/Engine/Plugins/Experimental/ChaosUserDataPT/Source/ChaosUserDataPT/Public/ChaosUserDataPT.h
--- a/Engine/Plugins/Experimental/ChaosUserDataPT/Source/ChaosUserDataPT/Public/ChaosUserDataPT.h
+++ b/Engine/Plugins/Experimental/ChaosUserDataPT/Source/ChaosUserDataPT/Public/ChaosUserDataPT.h
@@ -82,7 +82,10 @@ namespace Chaos
 	// data because this is a one-way callback. We use it basically just to
 	// marshal data in one direction.
 	template <typename TUserData>
-	class TUserDataManagerPT : public TSimCallbackObject<TUserDataManagerPTInput<TUserData>, FSimCallbackNoOutput>
+	class TUserDataManagerPT : public TSimCallbackObject<
+		TUserDataManagerPTInput<TUserData>,
+		FSimCallbackNoOutput,
+		ESimCallbackOptions::Presimulate | ESimCallbackOptions::ParticleUnregister>
 	{
 		using TInput = TUserDataManagerPTInput<TUserData>;
 
@@ -207,6 +210,26 @@ namespace Chaos
 			}
 		}
 
+		virtual void OnParticleUnregistered_Internal(TArray<TTuple<FUniqueIdx, FSingleParticlePhysicsProxy*>>& UnregisteredProxies) override
+		{
+			if (UnregisteredProxies.Num() > 0)
+			{
+				// If any particles were removed in this physics tick, check to see if they have
+				// userdata that we were tracking and remove those userdata instances.
+				for (TTuple<FUniqueIdx, FSingleParticlePhysicsProxy*> Proxy : UnregisteredProxies)
+				{
+					const FUniqueIdx Idx = Proxy.Get<FUniqueIdx>();
+					if (UserDataMap_PT.IsValidIndex(Idx.Idx))
+					{
+						UserDataMap_PT.RemoveAt(Idx.Idx);
+					}
+				}
+
+				// Shrink sparse array if we took elements off the end
+				UserDataMap_PT.Shrink();
+			}
+		}
+
 		// Identifier of the next input to be created on the game thread
 		int32 InputIdentifier_GT = 0;
 
diff --git a/Engine/Source/Programs/ChaosUserDataPTTests/Private/ChaosUserDataPTTests.cpp b/Engine/Source/Programs/ChaosUserDataPTTests/Private/ChaosUserDataPTTests.cpp
--- a/Engine/Source/Programs/ChaosUserDataPTTests/Private/ChaosUserDataPTTests.cpp
+++ b/Engine/Source/Programs/ChaosUserDataPTTests/Private/ChaosUserDataPTTests.cpp
@@ -104,25 +104,19 @@ TEST_CASE("ChaosUserDataPT", "[integration]")
 		Solver->AdvanceAndDispatch_External(DeltaTime)->Wait();
 		Solver->AdvanceAndDispatch_External(DeltaTime)->Wait();
 
-		// Add a new particle
-		Chaos::FSingleParticlePhysicsProxy* Proxy1 = Chaos::FSingleParticlePhysicsProxy::Create(Chaos::FGeometryParticle::CreateParticle());
-		Chaos::FRigidBodyHandle_External& HandleExternal1 = Proxy1->GetGameThreadAPI();
-		HandleExternal1.SetGeometry(BoxGeom);
-		Solver->RegisterObject(Proxy1);
-
-		// Make sure the unique idx got recycled
-		const Chaos::FUniqueIdx UniqueIdx1 = HandleExternal1.UniqueIdx();
-		if (UniqueIdx0.Idx == UniqueIdx1.Idx)
+		struct FMockHandle
 		{
-			// Register the new particle and advance solver to make sure it propagates to PT
-			Solver->AdvanceAndDispatch_External(DeltaTime)->Wait();
-
-			// Access userdata with the new internal handle - it should retrieve nothing
-			Solver->EnqueueCommandImmediate([&]()
-			{
-				REQUIRE(TestUserData->GetData_PT(*Proxy1->GetPhysicsThreadAPI()) == nullptr);
-			});
-			Solver->AdvanceAndDispatch_External(DeltaTime)->Wait();
-		}
+			FMockHandle(Chaos::FUniqueIdx InUniqueIdx) : MUniqueIdx(InUniqueIdx) { }
+			Chaos::FUniqueIdx UniqueIdx() const { return MUniqueIdx; }
+			Chaos::FUniqueIdx MUniqueIdx;
+		};
+
+		// Access userdata with the invalid particle handle - it should retrieve nothing
+		Solver->EnqueueCommandImmediate([&]()
+		{
+			const FMockHandle MockHandle0 = FMockHandle(UniqueIdx0);
+			REQUIRE(TestUserData->GetData_PT(MockHandle0) == nullptr);
+		});
+		Solver->AdvanceAndDispatch_External(DeltaTime)->Wait();
 	}
 }
diff --git a/Engine/Source/Runtime/Engine/Private/PhysicsEngine/Experimental/PhysScene_Chaos.cpp b/Engine/Source/Runtime/Engine/Private/PhysicsEngine/Experimental/PhysScene_Chaos.cpp
--- a/Engine/Source/Runtime/Engine/Private/PhysicsEngine/Experimental/PhysScene_Chaos.cpp
+++ b/Engine/Source/Runtime/Engine/Private/PhysicsEngine/Experimental/PhysScene_Chaos.cpp
@@ -78,14 +78,12 @@ struct FPendingAsyncPhysicsCommand
 };
 
 
-class FAsyncPhysicsTickCallback : public Chaos::TSimCallbackObject<>
+class FAsyncPhysicsTickCallback : public Chaos::TSimCallbackObject<
+	Chaos::FSimCallbackNoInput,
+	Chaos::FSimCallbackNoOutput,
+	Chaos::ESimCallbackOptions::Presimulate | Chaos::ESimCallbackOptions::RunOnFrozenGameThread>
 {
 public:
-	FAsyncPhysicsTickCallback()
-		: TSimCallbackObject<>(/*InRunOnFrozenGameThread=*/ true)
-	{
-
-	}
 
 	TSet<UActorComponent*> AsyncPhysicsTickComponents;
 	TSet<AActor*> AsyncPhysicsTickActors;
diff --git a/Engine/Source/Runtime/Experimental/Chaos/Private/PBDRigidsSolver.cpp b/Engine/Source/Runtime/Experimental/Chaos/Private/PBDRigidsSolver.cpp
--- a/Engine/Source/Runtime/Experimental/Chaos/Private/PBDRigidsSolver.cpp
+++ b/Engine/Source/Runtime/Experimental/Chaos/Private/PBDRigidsSolver.cpp
@@ -819,6 +819,23 @@ namespace Chaos
 
 	void FPBDRigidsSolver::DestroyPendingProxies_Internal()
 	{
+		// If we have any callback objects watching for particle deregistrations,
+		// send an array of proxies that are about to be deleted.
+		if (UnregistrationWatchers.Num() > 0)
+		{
+			TArray<TTuple<FUniqueIdx, FSingleParticlePhysicsProxy*>> Proxies;
+			Proxies.Reserve(PendingDestroyPhysicsProxy.Num());
+			for (FPendingDestroyInfo& Info : PendingDestroyPhysicsProxy)
+			{
+				Proxies.Add({ Info.UniqueIdx, Info.Proxy });
+			}
+			for (ISimCallbackObject* Callback : UnregistrationWatchers)
+			{
+				Callback->OnParticleUnregistered_Internal(Proxies);
+			}
+		}
+
+		// Do the actual destruction
 		for(int32 Idx = PendingDestroyPhysicsProxy.Num() - 1; Idx >= 0; --Idx)
 		{
 			FPendingDestroyInfo& Info = PendingDestroyPhysicsProxy[Idx];
@@ -1080,7 +1097,9 @@ namespace Chaos
 		FShapeDirtyData* ShapeDirtyData = DirtyProxiesData->GetShapesDirtyData();
 		FReal ExternalDt = PushData.ExternalDt;
 
-		auto ProcessProxyPT = [Manager,ShapeDirtyData,RewindData, ExternalDt, this](auto& Proxy,int32 DataIdx,FDirtyProxy& Dirty,const auto& CreateHandleFunc)
+		TArray<FSingleParticlePhysicsProxy*> RegisteredProxies;
+
+		auto ProcessProxyPT = [Manager, ShapeDirtyData, RewindData, ExternalDt, &RegisteredProxies, this](auto& Proxy,int32 DataIdx,FDirtyProxy& Dirty,const auto& CreateHandleFunc)
 		{
 			const bool bIsNew = !Proxy->IsInitialized();
 			if(bIsNew)
@@ -1092,6 +1111,13 @@ namespace Chaos
 				auto Handle = Proxy->GetHandle_LowLevel();
 				Handle->GTGeometryParticle() = Proxy->GetParticle_LowLevel();
 
+				// If anybody's listening for proxy creations, build a list of proxies that have been
+				// added in this step.
+				if (RegistrationWatchers.Num() > 0)
+				{
+					RegisteredProxies.Add(Proxy);
+				}
+
 				// Track proxies
 				if (UniqueIdx)
 				{
@@ -1222,6 +1248,16 @@ namespace Chaos
 			}
 		});
 
+		// If we have callbacks watching for particle registrations, send them the array
+		// of newly added proxies
+		if (RegistrationWatchers.Num() > 0)
+		{
+			for (ISimCallbackObject* Callback : RegistrationWatchers)
+			{
+				Callback->OnParticlesRegistered_Internal(RegisteredProxies);
+			}
+		}
+
 		//MarshallingManager.FreeData_Internal(&PushData);
 	}
 
@@ -1234,11 +1270,25 @@ namespace Chaos
 		SimCallbackObjects.Reserve(SimCallbackObjects.Num() + PushData.SimCallbackObjectsToAdd.Num());
 		for(ISimCallbackObject* SimCallbackObject : PushData.SimCallbackObjectsToAdd)
 		{
-			SimCallbackObjects.Add(SimCallbackObject);
-			if (SimCallbackObject->bContactModification)
+			if (SimCallbackObject->HasOption(ESimCallbackOptions::Presimulate))
+			{
+				SimCallbackObjects.Add(SimCallbackObject);
+			}
+
+			if (SimCallbackObject->HasOption(ESimCallbackOptions::ContactModification))
 			{
 				ContactModifiers.Add(SimCallbackObject);
 			}
+
+			if (SimCallbackObject->HasOption(ESimCallbackOptions::ParticleRegister))
+			{
+				RegistrationWatchers.Add(SimCallbackObject);
+			}
+
+			if (SimCallbackObject->HasOption(ESimCallbackOptions::ParticleUnregister))
+			{
+				UnregistrationWatchers.Add(SimCallbackObject);
+			}
 		}
 
 		//save any pending data for this particular interval
diff --git a/Engine/Source/Runtime/Experimental/Chaos/Public/Chaos/Framework/PhysicsSolverBase.h b/Engine/Source/Runtime/Experimental/Chaos/Public/Chaos/Framework/PhysicsSolverBase.h
--- a/Engine/Source/Runtime/Experimental/Chaos/Public/Chaos/Framework/PhysicsSolverBase.h
+++ b/Engine/Source/Runtime/Experimental/Chaos/Public/Chaos/Framework/PhysicsSolverBase.h
@@ -262,26 +262,48 @@ namespace Chaos
 
 		/** Creates a new sim callback object of the type given. Caller expected to free using FreeSimCallbackObject_External*/
 		template <typename TSimCallbackObjectType>
-		inline TSimCallbackObjectType* CreateAndRegisterSimCallbackObject_External(bool bContactModification = false, bool bRegisterRewindCallback = false)
+		inline TSimCallbackObjectType* CreateAndRegisterSimCallbackObject_External()
 		{
 			auto NewCallbackObject = new TSimCallbackObjectType();
 			//If at least one callback is on frozen GT we mark it as such
 			//Note this is never cleaned up (to avoid race conditions of register/unregister).
 			//Expectation is that there's typically just one of these on frozen gt that triggers all the fixed tick on gt (maybe two if another system wants direct control)
 			//This means we don't expect to optimize for the case where a gt callback goes away half way through sim
-			if (NewCallbackObject->RunOnFrozenGameThread())
+			if (NewCallbackObject->HasOption(ESimCallbackOptions::RunOnFrozenGameThread))
 			{
 				bSolverHasFrozenGameThreadCallbacks = true;
 			}
 
-			RegisterSimCallbackObject_External(NewCallbackObject, bContactModification);
-			if (bRegisterRewindCallback)
+			RegisterSimCallbackObject_External(NewCallbackObject);
+			if (NewCallbackObject->HasOption(ESimCallbackOptions::Rewind))
 			{
 				EnqueueSimcallbackRewindRegisteration(NewCallbackObject);
 			}
 			return NewCallbackObject;
 		}
 
+		template <typename TSimCallbackObjectType>
+		UE_DEPRECATED(5.1, "Use version with no arguments, and provide TOptions bitmask parameter to TSimCallbackObject instead.")
+		TSimCallbackObjectType* CreateAndRegisterSimCallbackObject_External(bool bContactModification, bool bRegisterRewindCallback = false)
+		{
+			TSimCallbackObjectType* CallbackObject = CreateAndRegisterSimCallbackObject_External<TSimCallbackObjectType>();
+			ESimCallbackOptions& Options = CallbackObject->Options;
+
+			if (bContactModification) {
+				Options |= ESimCallbackOptions::ContactModification;
+			} else {
+				Options &= ~ESimCallbackOptions::ContactModification;
+			}
+
+			if (bRegisterRewindCallback) {
+				Options |= ESimCallbackOptions::Rewind;
+			} else {
+				Options &= ~ESimCallbackOptions::Rewind;
+			}
+
+			return CallbackObject;
+		}
+
 		void EnqueueSimcallbackRewindRegisteration(ISimCallbackObject* Callback);
 
 		void UnregisterAndFreeSimCallbackObject_External(ISimCallbackObject* SimCallbackObject)
@@ -439,7 +461,7 @@ namespace Chaos
 				const FReal SimTime = GetSolverTime();
 				for (ISimCallbackObject* Callback : SimCallbackObjects)
 				{
-					if (Callback->RunOnFrozenGameThread() == bGameThreadFrozen)
+					if (Callback->HasOption(ESimCallbackOptions::RunOnFrozenGameThread) == bGameThreadFrozen)
 					{
 						Callback->SetSimAndDeltaTime_Internal(SimTime, MLastDt);
 						Callback->PreSimulate_Internal();
@@ -585,6 +607,8 @@ namespace Chaos
 
 	TArray<ISimCallbackObject*> SimCallbackObjects;
 	TArray<ISimCallbackObject*> ContactModifiers;
+	TArray<ISimCallbackObject*> RegistrationWatchers;
+	TArray<ISimCallbackObject*> UnregistrationWatchers;
 
 	TUniquePtr<FRewindData> MRewindData;
 	TUniquePtr<IRewindCallback> MRewindCallback;
@@ -602,11 +626,10 @@ namespace Chaos
 
 		//This is private because the user should never create their own callback object
 		//The lifetime management should always be done by solver to ensure callbacks are accessing valid memory on async tasks
-		void RegisterSimCallbackObject_External(ISimCallbackObject* SimCallbackObject, bool bContactModification = false)
+		void RegisterSimCallbackObject_External(ISimCallbackObject* SimCallbackObject)
 		{
 			ensure(SimCallbackObject->Solver == nullptr);	//double register?
 			SimCallbackObject->SetSolver_External(this);
-			SimCallbackObject->SetContactModification(bContactModification);
 			MarshallingManager.RegisterSimCallbackObject_External(SimCallbackObject);
 		}
 
diff --git a/Engine/Source/Runtime/Experimental/Chaos/Public/Chaos/SimCallbackObject.h b/Engine/Source/Runtime/Experimental/Chaos/Public/Chaos/SimCallbackObject.h
--- a/Engine/Source/Runtime/Experimental/Chaos/Public/Chaos/SimCallbackObject.h
+++ b/Engine/Source/Runtime/Experimental/Chaos/Public/Chaos/SimCallbackObject.h
@@ -5,17 +5,30 @@
 #include "Chaos/Defines.h"
 #include "Chaos/SimCallbackInput.h"
 #include "Chaos/CollisionResolutionTypes.h"
+#include "Chaos/GeometryParticlesfwd.h"
 
 namespace Chaos
 {
 class FPhysicsSolverBase;
 class FCollisionContactModifier;
+class FSingleParticlePhysicsProxy;
 
 namespace Utilities
 {
 	CHAOS_API FReal GetSolverPhysicsResultsTime(FPhysicsSolverBase*);
 }
 
+enum class ESimCallbackOptions : uint8
+{
+	Presimulate				= 1 << 0,
+	ContactModification		= 1 << 1,
+	ParticleRegister		= 1 << 2,
+	ParticleUnregister		= 1 << 3,
+	RunOnFrozenGameThread	= 1 << 4,
+	Rewind					= 1 << 5
+};
+ENUM_CLASS_FLAGS(ESimCallbackOptions)
+
 /**
  * Callback API used for executing code at different points in the simulation.
  * The external thread pushes input data at its own rate (typically once per game thread tick)
@@ -86,21 +99,29 @@ public:
 	virtual void ApplyCorrections_Internal(int32 PhysicsStep, FSimCallbackInput* Input) { ensure(false); }
 	virtual void FirstPreResimStep_Internal(int32 PhysicsStep) { }
 
-	bool RunOnFrozenGameThread() const { return bRunOnFrozenGameThread; }
-	
-protected:
+	bool HasOption(const ESimCallbackOptions Option) const
+	{
+		return (Options & Option) == Option;
+	}
 
-	ISimCallbackObject(bool InRunOnFrozenGameThread = false)
-	: bPendingDelete(false)
-	, bPendingDelete_External(false)
-	, bContactModification(false)
-	, CurrentExternalInput_External(nullptr)
-	, Solver(nullptr)
-	, CurrentOutput_Internal(nullptr)
-	, CurrentInput_Internal(nullptr)
-	, bRunOnFrozenGameThread(InRunOnFrozenGameThread)
+	UE_DEPRECATED(5.1, "Use HasOption(ESimCallbackOptions::RunOnFrozenGameThread) instead.")
+	bool RunOnFrozenGameThread() const 
 	{
+		return HasOption(ESimCallbackOptions::RunOnFrozenGameThread);
 	}
+	
+protected:
+
+	ISimCallbackObject(const ESimCallbackOptions InOptions = ESimCallbackOptions::Presimulate)
+		: bPendingDelete(false)
+		, bPendingDelete_External(false)
+		, CurrentExternalInput_External(nullptr)
+		, Solver(nullptr)
+		, CurrentOutput_Internal(nullptr)
+		, CurrentInput_Internal(nullptr)
+		, Options(InOptions)
+	{ }
+	
 
 
 	/**
@@ -144,30 +165,49 @@ private:
 		check(false);
 	}
 
-	/** If we've already allocated an output in this simulation step, use it again. This way multiple sim callbacks only generate one output.
-*   Otherwise allocate an output for us and mark it as pending for external thread */
+	/**
+	* Called once in a simulation step if any new particles were registered. Occurs after UniqueIdxs
+	* are valid.
+	*/
+	virtual void OnParticlesRegistered_Internal(TArray<FSingleParticlePhysicsProxy*>& RegisteredProxies)
+	{
+		check(false);
+	}
+
+	/**
+	* Called once in a simulation step if any particles were unregistered. Occurs immediately before
+	* UniqueIdxs become invalid.
+	*/
+	virtual void OnParticleUnregistered_Internal(TArray<TTuple<Chaos::FUniqueIdx, FSingleParticlePhysicsProxy*>>& UnregisteredProxies)
+	{
+		check(false);
+	}
+
+	/** If we've already allocated an output in this simulation step, use it again. This way multiple sim callbacks only generate one output. Otherwise allocate an output for us and mark it as pending for external thread */
 	virtual void OnFinalizeOutputData_Internal(FSimCallbackOutput* CurOutput)
 	{
 		check(false);	//wrote to output but not finalizing it. Typically this is pushed into some kind of thread safe queue
 	}
-	
 
 	friend class FPBDRigidsSolver;
-
 	friend class FPhysicsSolverBase;
 	friend class FChaosMarshallingManager;
 	friend struct FPushPhysicsData;
 
 	bool bPendingDelete;	//used internally for more efficient deletion. Callbacks do not need to check this
 	bool bPendingDelete_External;	//used for efficient deletion. Callbacks do not need to check this
-	bool bContactModification;
 
 	FSimCallbackInput* CurrentExternalInput_External;	//the input currently being filled out by external thread
 	FPhysicsSolverBase* Solver;
 
 	//putting this here so that user classes don't have to bother with non-default constructor
 	void SetSolver_External(FPhysicsSolverBase* InSolver) { Solver = InSolver;}
-	void SetContactModification(bool InContactModification) { bContactModification = InContactModification; }
+	
+	UE_DEPRECATED(5.1, "Do not change options after creation of the callback object - instead, specify them using the TOptions template parameter.")
+	void SetContactModification(bool InContactModification)
+	{
+		Options = Options | ESimCallbackOptions::ContactModification;
+	}
 
 protected:
 	FSimCallbackOutput* CurrentOutput_Internal;	//the output currently being written to in this sim step
@@ -178,7 +218,12 @@ private:
 
 	FReal SimTime_Internal;
 	FReal DeltaTime_Internal;
-	bool bRunOnFrozenGameThread;
+
+	// TODO: Make this const and remove the "friend class FPhysicsSolverBase"
+	// once FPhysicsSolverBase::CreateAndRegisterSimCallbackObject_External(bool, bool)
+	// has been deprecated.
+	ESimCallbackOptions Options;
+	friend class FPhysicsSolverBase;
 };
 
 /** Simple callback command object. Commands are typically passed in as lambdas and there's no need for data management. Should not be used directly, see FPhysicsSolverBase::EnqueueCommand */
@@ -221,16 +266,21 @@ private:
 };
 
 /** Simple templated implementation that uses lock free queues to manage memory */
-template <typename TInputType = FSimCallbackNoInput, typename TOutputType = FSimCallbackNoOutput>
+template <typename TInputType = FSimCallbackNoInput, typename TOutputType = FSimCallbackNoOutput, ESimCallbackOptions TOptions = ESimCallbackOptions::Presimulate>
 class TSimCallbackObject : public ISimCallbackObject
 {
 public:
 
-	TSimCallbackObject(bool InRunOnFrozenGameThread = false)
-	: ISimCallbackObject(InRunOnFrozenGameThread)
-	, CurrentOutput_External(nullptr)
-	{
-	}
+	TSimCallbackObject()
+		: ISimCallbackObject(TOptions)
+		, CurrentOutput_External(nullptr)
+	{ }
+
+	UE_DEPRECATED(5.1, "Use default constructor instead and specify RunFrozenOnGameThread using TOptions template parameter.")
+	TSimCallbackObject(bool InRunOnFrozenGameThread)
+		: ISimCallbackObject(InRunOnFrozenGameThread ? (TOptions | ESimCallbackOptions::RunOnFrozenGameThread) : TOptions)
+		, CurrentOutput_External(nullptr)
+	{ }
 
 	virtual void FreeOutputData_External(FSimCallbackOutput* Output) override
 	{
