Fix static analysis warning

diff --git a/Engine/Source/Runtime/Experimental/Chaos/Private/Chaos/Collision/CapsuleTriangleContactPoint.cpp b/Engine/Source/Runtime/Experimental/Chaos/Private/Chaos/Collision/CapsuleTriangleContactPoint.cpp
--- a/Engine/Source/Runtime/Experimental/Chaos/Private/Chaos/Collision/CapsuleTriangleContactPoint.cpp
+++ b/Engine/Source/Runtime/Experimental/Chaos/Private/Chaos/Collision/CapsuleTriangleContactPoint.cpp
@@ -181,9 +181,12 @@ namespace Chaos
 		// Handle the parallel segment-edge case by clipping the segment to the edge
 		// NOTE: If we are parallel to an edge (and outside the edge prism) we can only collide with the one edge
 		bool bIsParallelEdge[3];
-		bIsParallelEdge[0] = ((EdgeD0s[0] >= FReal(0)) & (EdgeD1s[0] >= FReal(0)) & FMath::IsNearlyEqual(EdgeD0s[0], EdgeD1s[0], DistanceTolerance)) != 0;
-		bIsParallelEdge[1] = ((EdgeD0s[1] >= FReal(0)) & (EdgeD1s[1] >= FReal(0)) & FMath::IsNearlyEqual(EdgeD0s[1], EdgeD1s[1], DistanceTolerance)) != 0;
-		bIsParallelEdge[2] = ((EdgeD0s[2] >= FReal(0)) & (EdgeD1s[2] >= FReal(0)) & FMath::IsNearlyEqual(EdgeD0s[2], EdgeD1s[2], DistanceTolerance)) != 0;
+		const bool bEqualEdgeDist0 = FMath::IsNearlyEqual(EdgeD0s[0], EdgeD1s[0], DistanceTolerance);
+		const bool bEqualEdgeDist1 = FMath::IsNearlyEqual(EdgeD0s[1], EdgeD1s[1], DistanceTolerance);
+		const bool bEqualEdgeDist2 = FMath::IsNearlyEqual(EdgeD0s[2], EdgeD1s[2], DistanceTolerance);
+		bIsParallelEdge[0] = ((EdgeD0s[0] >= FReal(0)) & (EdgeD1s[0] >= FReal(0)) & bEqualEdgeDist0) != 0;
+		bIsParallelEdge[1] = ((EdgeD0s[1] >= FReal(0)) & (EdgeD1s[1] >= FReal(0)) & bEqualEdgeDist1) != 0;
+		bIsParallelEdge[2] = ((EdgeD0s[2] >= FReal(0)) & (EdgeD1s[2] >= FReal(0)) & bEqualEdgeDist2) != 0;
 		int32 EdgeVertexIndex0 = 2;
 		if ((bIsParallelFace & (bIsParallelEdge[0] | bIsParallelEdge[1] | bIsParallelEdge[2])) != 0)
 		{
