Chaos : fixed managed array collection faililng to serialize some non-copyable attributes causing warning in during cook and causing particle-implicits geometry collection to not collide anymore

diff --git a/Engine/Source/Runtime/Experimental/Chaos/Private/GeometryCollection/ManagedArrayCollection.cpp b/Engine/Source/Runtime/Experimental/Chaos/Private/GeometryCollection/ManagedArrayCollection.cpp
--- a/Engine/Source/Runtime/Experimental/Chaos/Private/GeometryCollection/ManagedArrayCollection.cpp
+++ b/Engine/Source/Runtime/Experimental/Chaos/Private/GeometryCollection/ManagedArrayCollection.cpp
@@ -537,6 +537,27 @@ FString FManagedArrayCollection::ToString() const
 
 static const FName GuidName("GUID");
 
+// this is a reference wrapper to avoid copying attributes as some may not support copy (unique ptr ones) 
+// this is used during serialization to build a transient filtered map of attributes that can be saved
+struct FManagedArrayCollectionValueTypeWrapper
+{
+	FManagedArrayCollectionValueTypeWrapper()
+		: ValueRef(nullptr)
+	{}
+
+	FManagedArrayCollectionValueTypeWrapper(FManagedArrayCollection::FValueType* ValueRefIn)
+		: ValueRef(ValueRefIn)
+	{}
+
+	FManagedArrayCollection::FValueType* ValueRef;
+};
+FArchive& operator<<(FArchive& Ar, FManagedArrayCollectionValueTypeWrapper& ValueIn)
+{
+	// simple forwarding to the original object 
+	Ar << (*ValueIn.ValueRef);
+	return Ar;
+}
+
 void FManagedArrayCollection::Serialize(Chaos::FChaosArchive& Ar)
 {
 	Ar.UsingCustomVersion(FUE5MainStreamObjectVersion::GUID);
@@ -620,19 +641,15 @@ void FManagedArrayCollection::Serialize(Chaos::FChaosArchive& Ar)
 		// Unless it's an undo/redo transaction, strip out the keys that we don't want to save
 		if (!Ar.IsTransacting())
 		{
-			TMap<FKeyType, FValueType> ToSaveMap = Map;
-			TArray<FKeyType> ToRemoveKeys;
-			for (TTuple<FKeyType, FValueType>& Pair : ToSaveMap)
+			// we do create a wrapper around ValueType, to avoid copies and save memory 
+			TMap<FKeyType, FManagedArrayCollectionValueTypeWrapper> ToSaveMap;
+			for (TTuple<FKeyType, FValueType>& Pair : Map)
 			{
-				if (!Pair.Value.Saved)
+				if (Pair.Value.Saved)
 				{
-					ToRemoveKeys.Add(Pair.Key);
+					ToSaveMap.Emplace(Pair.Key, FManagedArrayCollectionValueTypeWrapper(&Pair.Value));
 				}
 			}
-			for (const FKeyType& Key : ToRemoveKeys)
-			{
-				ToSaveMap.Remove(Key);
-			}
 			Ar << ToSaveMap;
 		}
 		else
diff --git a/Engine/Source/Runtime/Experimental/Chaos/Public/GeometryCollection/ManagedArrayCollection.h b/Engine/Source/Runtime/Experimental/Chaos/Public/GeometryCollection/ManagedArrayCollection.h
--- a/Engine/Source/Runtime/Experimental/Chaos/Public/GeometryCollection/ManagedArrayCollection.h
+++ b/Engine/Source/Runtime/Experimental/Chaos/Public/GeometryCollection/ManagedArrayCollection.h
@@ -544,6 +544,7 @@ private:
 		return FKeyType(Name, Group);
 	};
 
+	friend struct FManagedArrayCollectionValueTypeWrapper;
 	struct FValueType
 	{
 		EArrayType ArrayType;
diff --git a/Engine/Source/Runtime/Experimental/ChaosCore/Public/Chaos/Core.h b/Engine/Source/Runtime/Experimental/ChaosCore/Public/Chaos/Core.h
--- a/Engine/Source/Runtime/Experimental/ChaosCore/Public/Chaos/Core.h
+++ b/Engine/Source/Runtime/Experimental/ChaosCore/Public/Chaos/Core.h
@@ -51,7 +51,7 @@ namespace Chaos
 	// NOTE: if you get a merge conflict on the GUID, you must replace it with a new GUID - do not accept the source or target
 	// or you will likely get DDC version conflicts resulting in crashes during load.
 	// Core version string for Chaos data. Any DDC builder dependent on Chaos for serialization should depend on this version
-	static const TCHAR* const ChaosVersionGUID = TEXT("8D7C3D39-DEF4-4775-85D1-80187F4A3D48");
+	static const TCHAR* const ChaosVersionGUID = TEXT("0EFFF944-0922-40D5-ADFE-390045900A93");
 
 	inline FString GetChaosVersionStringInner()
 	{
