Chaos Physics:
- Prevent infinite loop when mesh fast BVH is built from an empty tree
- Remove MaxNonGlobalElementBoundsExtrema which is not relevant in LWC anymore and causes bad performance for objects far from the origin

#[fyi] Michael.Lentine, Benn.Gallagher, Chris.Caulfield

diff --git a/Engine/Source/Runtime/Experimental/Chaos/Private/Chaos/AABBTree.cpp b/Engine/Source/Runtime/Experimental/Chaos/Private/Chaos/AABBTree.cpp
--- a/Engine/Source/Runtime/Experimental/Chaos/Private/Chaos/AABBTree.cpp
+++ b/Engine/Source/Runtime/Experimental/Chaos/Private/Chaos/AABBTree.cpp
@@ -5,10 +5,6 @@
 int32 FAABBTreeCVars::UpdateDirtyElementPayloadData = 1;
 FAutoConsoleVariableRef FAABBTreeCVars::CVarUpdateDirtyElementPayloadData(TEXT("p.aabbtree.updatedirtyelementpayloads"), FAABBTreeCVars::UpdateDirtyElementPayloadData, TEXT("Allow AABB tree elements to update internal payload data when they recieve a payload update"));
 
-float FAABBTreeCVars::MaxNonGlobalElementBoundsExtrema = 1e15f;
-FAutoConsoleVariableRef FAABBTreeCVars::CVarMaxNonGlobalElementBoundsExtrema(TEXT("p.aabbtree.MaxNonGlobalElementBoundsExtrema"), FAABBTreeCVars::MaxNonGlobalElementBoundsExtrema, TEXT("If Min/Max of particle world bounds exceeds this threshold, element is put in global array instead of being used in leaf splitting."));
-
-
 int32 FAABBTreeDirtyGridCVars::DirtyElementGridCellSize = 1000; // 0 means disabled
 FAutoConsoleVariableRef FAABBTreeDirtyGridCVars::CVarDirtyElementGridCellSize(TEXT("p.aabbtree.DirtyElementGridCellSize"), FAABBTreeDirtyGridCVars::DirtyElementGridCellSize, TEXT("DirtyElement Grid acceleration structure cell size in cm. 0 or less will disable the feature"));
 
diff --git a/Engine/Source/Runtime/Experimental/Chaos/Private/Chaos/TriangleMeshImplicitObject.cpp b/Engine/Source/Runtime/Experimental/Chaos/Private/Chaos/TriangleMeshImplicitObject.cpp
--- a/Engine/Source/Runtime/Experimental/Chaos/Private/Chaos/TriangleMeshImplicitObject.cpp
+++ b/Engine/Source/Runtime/Experimental/Chaos/Private/Chaos/TriangleMeshImplicitObject.cpp
@@ -1725,6 +1725,12 @@ void FTriangleMeshImplicitObject::RebuildFastBVHFromTree(const BVHType& TreeBVH)
 		const int32 LeafIndex = RootNode.ChildrenNodes[0];
 		const LeafType& Leaf = Leaves[LeafIndex];
 
+		if (Leaf.Elems.Num() == 0)
+		{
+			// Will have an empty BVH in this case
+			return;
+		}
+
 		// make the node
 		FTrimeshBVH::FNode& NewNode = FastBVH.Nodes.Emplace_GetRef();
 		NewNode.Children[0].SetChildOrFaceIndex(0);
diff --git a/Engine/Source/Runtime/Experimental/Chaos/Public/Chaos/AABBTree.h b/Engine/Source/Runtime/Experimental/Chaos/Public/Chaos/AABBTree.h
--- a/Engine/Source/Runtime/Experimental/Chaos/Public/Chaos/AABBTree.h
+++ b/Engine/Source/Runtime/Experimental/Chaos/Public/Chaos/AABBTree.h
@@ -30,9 +30,6 @@ struct CHAOS_API FAABBTreeCVars
 	static int32 SplitOnVarianceAxis;
 	static FAutoConsoleVariableRef CVarSplitOnVarianceAxis;
 
-	static float MaxNonGlobalElementBoundsExtrema; 
-	static FAutoConsoleVariableRef CVarMaxNonGlobalElementBoundsExtrema;
-
 	static float DynamicTreeBoundingBoxPadding;
 	static FAutoConsoleVariableRef CVarDynamicTreeBoundingBoxPadding;
 
@@ -1644,7 +1641,7 @@ public:
 		if (bDynamicTree && bHasBounds && ValidateBounds(NewBounds) == false)
 		{
 			bHasBounds = false;
-			ensureMsgf(false, TEXT("AABBTree encountered invalid bounds input. Forcing element to global payload. Min: %s Max: %s. If Bounds are valid but large, increase FAABBTreeCVars::MaxNonGlobalElementBoundsExtrema."),
+			ensureMsgf(false, TEXT("AABBTree encountered invalid bounds input. Forcing element to global payload. Min: %s Max: %s."),
 				*NewBounds.Min().ToString(), *NewBounds.Max().ToString());
 		}
 
@@ -2028,12 +2025,6 @@ public:
 			const T& MinComponent = Min[i];
 			const T& MaxComponent = Max[i];
 
-			// If element is extremely far out on any axis, if past limit return false to make it a global element and prevent huge numbers poisoning splitting algorithm computation.
-			if (MinComponent <= -FAABBTreeCVars::MaxNonGlobalElementBoundsExtrema || MaxComponent >= FAABBTreeCVars::MaxNonGlobalElementBoundsExtrema)
-			{
-				return false;
-			}
-
 			// Are we an empty aabb?
 			if (MinComponent > MaxComponent)
 			{
@@ -2949,7 +2940,7 @@ private:
 				if (bHasBoundingBox && ValidateBounds(ElemBounds) == false)
 				{
 					bHasBoundingBox = false;
-					ensureMsgf(false, TEXT("AABBTree encountered invalid bounds input. Forcing element to global payload. Min: %s Max: %s. If Bounds are valid but large, increase FAABBTreeCVars::MaxNonGlobalElementBoundsExtrema."),
+					ensureMsgf(false, TEXT("AABBTree encountered invalid bounds input. Forcing element to global payload. Min: %s Max: %s."),
 						*ElemBounds.Min().ToString(), *ElemBounds.Max().ToString());
 				}
 
diff --git a/Engine/Source/Runtime/Experimental/ChaosCore/Public/Chaos/Core.h b/Engine/Source/Runtime/Experimental/ChaosCore/Public/Chaos/Core.h
--- a/Engine/Source/Runtime/Experimental/ChaosCore/Public/Chaos/Core.h
+++ b/Engine/Source/Runtime/Experimental/ChaosCore/Public/Chaos/Core.h
@@ -51,7 +51,7 @@ namespace Chaos
 	// NOTE: if you get a merge conflict on the GUID, you must replace it with a new GUID - do not accept the source or target
 	// or you will likely get DDC version conflicts resulting in crashes during load.
 	// Core version string for Chaos data. Any DDC builder dependent on Chaos for serialization should depend on this version
-	static const TCHAR* const ChaosVersionGUID = TEXT("0EFFF944-0922-40D5-ADFE-390045900A93");
+	static const TCHAR* const ChaosVersionGUID = TEXT("3366E325-8C79-4356-916E-B8432589888D");
 
 	inline FString GetChaosVersionStringInner()
 	{
