Fix static analysis recommendation

diff --git a/Engine/Source/Runtime/Experimental/Chaos/Private/Chaos/Collision/CapsuleTriangleContactPoint.cpp b/Engine/Source/Runtime/Experimental/Chaos/Private/Chaos/Collision/CapsuleTriangleContactPoint.cpp
--- a/Engine/Source/Runtime/Experimental/Chaos/Private/Chaos/Collision/CapsuleTriangleContactPoint.cpp
+++ b/Engine/Source/Runtime/Experimental/Chaos/Private/Chaos/Collision/CapsuleTriangleContactPoint.cpp
@@ -414,7 +414,9 @@ namespace Chaos
 			// If both points have T < 0, only use the point with highest T
 			// If both points have T > 1, only use the point with lowest T
 			// Otherwise use both points with T clamped to [0,1]
-			if ((InRangeClosed(T0, FReal(0), FReal(1)) | InRangeClosed(T1, FReal(0), FReal(1))) != 0)
+			const bool bInRange0 = InRangeClosed(T0, FReal(0), FReal(1));
+			const bool bInRange1 = InRangeClosed(T1, FReal(0), FReal(1));
+			if ((bInRange0 | bInRange1) != 0)
 			{
 				// At least one segment point projects to a point inside the edge extents.
 				// Clip the edge-mapped segement and use the clipped verts as contacts.
