[GeometryCollectionEngine]
* Added missing include

diff --git a/Engine/Source/Runtime/Experimental/GeometryCollectionEngine/Public/GeometryCollection/GeometryCollectionComponent.h b/Engine/Source/Runtime/Experimental/GeometryCollectionEngine/Public/GeometryCollection/GeometryCollectionComponent.h
--- a/Engine/Source/Runtime/Experimental/GeometryCollectionEngine/Public/GeometryCollection/GeometryCollectionComponent.h
+++ b/Engine/Source/Runtime/Experimental/GeometryCollectionEngine/Public/GeometryCollection/GeometryCollectionComponent.h
@@ -24,6 +24,7 @@
 #include "Chaos/PBDRigidsEvolutionFwd.h"
 #include "EngineDefines.h"
 #include "Math/MathFwd.h"
+#include "PhysxUserData.h"
 
 #include "GeometryCollectionComponent.generated.h"
 
