Fix crash in constraint graph update

diff --git a/Engine/Source/Runtime/Experimental/Chaos/Private/Chaos/Island/IslandGraph.cpp b/Engine/Source/Runtime/Experimental/Chaos/Private/Chaos/Island/IslandGraph.cpp
--- a/Engine/Source/Runtime/Experimental/Chaos/Private/Chaos/Island/IslandGraph.cpp
+++ b/Engine/Source/Runtime/Experimental/Chaos/Private/Chaos/Island/IslandGraph.cpp
@@ -84,7 +84,11 @@ void FIslandGraph<NodeType, EdgeType, IslandType, OwnerType>::UpdateNode(const i
 				const bool bIsEdgeMoving = IsEdgeMoving(EdgeIndex);
 				ParentIslands(PrimaryIsland, GraphEdge.IslandIndex, bIsEdgeMoving);
 			
-				PrimaryIsland = GraphEdge.IslandIndex;
+				// Primary island can be set to the island of any of our edges, so long as they aren't INDEX_NONE
+				if (PrimaryIsland == INDEX_NONE)
+				{
+					PrimaryIsland = GraphEdge.IslandIndex;
+				}
 			}
 
 			// Put the valid node into one of the islands - they will be merged anyway so doesn't matter which
@@ -458,7 +462,7 @@ void FIslandGraph<NodeType, EdgeType, IslandType, OwnerType>::MergeIslands()
 	}
 
 	// We loop over all the islands and if they have children
-	// we recusrively merge them onto the parent one
+	// we recursively merge them onto the parent one
 	for (int32 IslandIndex = 0, NumIslands = GraphIslands.GetMaxIndex(); IslandIndex < NumIslands; ++IslandIndex)
 	{
 		if (GraphIslands.IsValidIndex(IslandIndex))
@@ -858,6 +862,38 @@ void FIslandGraph<NodeType, EdgeType, IslandType, OwnerType>::UpdateGraph()
 		GraphIsland.bIsSleeping = GraphIsland.bIsPersistent;
 	}
 
+	// Make sure valid edges are in an island (if they have valid nodes)
+	for (FGraphEdge& GraphEdge : GraphEdges)
+	{
+		if (GraphEdge.bValidEdge && GraphEdge.IslandIndex == INDEX_NONE)
+		{
+			const int32 FirstIslandIndex
+				= GraphNodes.IsValidIndex(GraphEdge.FirstNode)
+				? GraphNodes[GraphEdge.FirstNode].IslandIndex
+				: INDEX_NONE;
+
+			const int32 SecondIslandIndex
+				= GraphNodes.IsValidIndex(GraphEdge.SecondNode)
+				? GraphNodes[GraphEdge.SecondNode].IslandIndex
+				: INDEX_NONE;
+
+			if (FirstIslandIndex != INDEX_NONE)
+			{
+				GraphEdge.IslandIndex = FirstIslandIndex;
+			}
+			else if (SecondIslandIndex != INDEX_NONE)
+			{
+				GraphEdge.IslandIndex = SecondIslandIndex;
+			}
+			else
+			{
+				// NOTE: Both particle and edge haven't been added to island.
+				// They should get added next frame, but might be better to
+				// work out a way to assign them an island immediately.
+			}
+		}
+	}
+
 	// Merge all the islands if necessary
 	MergeIslands();
 
