Chaos : Fix issues with level attribute not being saved properly and that can result in a empty array in non-editor build, potentially causing crashes

diff --git a/Engine/Source/Runtime/Experimental/Chaos/Private/GeometryCollection/Facades/CollectionHierarchyFacade.cpp b/Engine/Source/Runtime/Experimental/Chaos/Private/GeometryCollection/Facades/CollectionHierarchyFacade.cpp
new file mode 100644
--- /dev/null
+++ b/Engine/Source/Runtime/Experimental/Chaos/Private/GeometryCollection/Facades/CollectionHierarchyFacade.cpp
@@ -0,0 +1,79 @@
+// Copyright Epic Games, Inc. All Rights Reserved.
+
+#include "GeometryCollection/Facades/CollectionHierarchyFacade.h"
+#include "GeometryCollection/TransformCollection.h"
+
+namespace Chaos::Facades
+{
+	FCollectionHierarchyFacade::FCollectionHierarchyFacade(FManagedArrayCollection& InCollection)
+		: ParentAttribute(InCollection, "Parent", FTransformCollection::TransformGroup)
+	    , ChildrenAttribute(InCollection, "Children", FTransformCollection::TransformGroup)
+		, LevelAttribute(InCollection, "Level", FTransformCollection::TransformGroup)
+	{}
+
+	bool FCollectionHierarchyFacade::IsValid() const
+	{
+		return ParentAttribute.IsValid() && ChildrenAttribute.IsValid();
+	}
+
+	bool FCollectionHierarchyFacade::HasLevelAttribute() const
+	{
+		return LevelAttribute.IsValid();
+	}
+
+	bool FCollectionHierarchyFacade::IsLevelAttributePersistent() const
+	{
+		return LevelAttribute.IsValid() && LevelAttribute.IsPersistent();
+	}
+
+	void FCollectionHierarchyFacade::GenerateLevelAttribute()
+	{
+		check(IsValid());
+
+		// level attribute must actually be persistent 
+		LevelAttribute.Add(ManageArrayAccessor::EPersistencePolicy::MakePersistent);
+
+		TManagedArray<int32>& Levels = LevelAttribute.Modify();
+		const TManagedArray<int32>& Parents = ParentAttribute.Get();
+
+		// In place level computation
+		LevelAttribute.Fill(INDEX_NONE);
+		for (int32 TransformIndex = 0; TransformIndex < Levels.Num(); TransformIndex++)
+		{
+			// parent propagation may have already set the value for this transform
+			if (Levels[TransformIndex] == INDEX_NONE)
+			{
+				// first compute the level by walking the parents
+				Levels[TransformIndex] = 0;
+				int32 ParentIndex = Parents[TransformIndex];
+				while (ParentIndex != INDEX_NONE)
+				{
+					Levels[TransformIndex]++;
+
+					// do we need to go further ? 
+					const int32 ParentLevel = Levels[ParentIndex];
+					if (ParentLevel != INDEX_NONE)
+					{
+						Levels[TransformIndex] += ParentLevel;
+						break; // early out as the parent level has been computed already
+					}
+
+					ParentIndex = Parents[ParentIndex];
+				}
+
+				// second propagate to parents if they are still unitialized
+				if (Levels[TransformIndex] > 0)
+				{
+					int32 ParentLevel = Levels[TransformIndex] - 1;
+					ParentIndex = Parents[TransformIndex];
+					while (ParentIndex != INDEX_NONE && Levels[ParentIndex] == INDEX_NONE)
+					{
+						Levels[ParentIndex] = ParentLevel;
+						ParentLevel--;
+						ParentIndex = Parents[ParentIndex];
+					}
+				}
+			}
+		}
+	}
+}
\ No newline at end of file
diff --git a/Engine/Source/Runtime/Experimental/Chaos/Private/GeometryCollection/GeometryCollection.cpp b/Engine/Source/Runtime/Experimental/Chaos/Private/GeometryCollection/GeometryCollection.cpp
--- a/Engine/Source/Runtime/Experimental/Chaos/Private/GeometryCollection/GeometryCollection.cpp
+++ b/Engine/Source/Runtime/Experimental/Chaos/Private/GeometryCollection/GeometryCollection.cpp
@@ -11,6 +11,7 @@
 #include "GeometryCollection/GeometryCollectionProximityUtility.h"
 #include "GeometryCollection/GeometryCollectionClusteringUtility.h"
 #include "GeometryCollection/GeometryCollectionConvexUtility.h"
+#include "GeometryCollection/Facades/CollectionHierarchyFacade.h"
 #include "UObject/FortniteNCBranchObjectVersion.h"
 
 #include <iostream>
@@ -1300,12 +1301,16 @@ void FGeometryCollection::Serialize(Chaos::FChaosArchive& Ar)
 			Version = 9;
 		}
 
-		if (Ar.CustomVer(FFortniteNCBranchObjectVersion::GUID) < FFortniteNCBranchObjectVersion::ChaosGeometryCollectionSaveLevelsAttribute)
+		Chaos::Facades::FCollectionHierarchyFacade HierarchyFacade(*this);
+		if (Ar.CustomVer(FFortniteNCBranchObjectVersion::GUID) < FFortniteNCBranchObjectVersion::ChaosGeometryCollectionSaveLevelsAttribute
+			|| !HierarchyFacade.HasLevelAttribute()
+			|| !HierarchyFacade.IsLevelAttributePersistent()
+			)
 		{
 			// Level attribute previously serialized with bSave = false, so was not serializing level data.
-			// We now compute this during cook and need to serialize, so convert attribute to bSave = true.
-			RemoveAttribute("Level", FTransformCollection::TransformGroup);
-			AddAttribute<int32>("Level", FTransformCollection::TransformGroup, FConstructionParameters(FName(), /*bSave=*/true));
+			// We now compute this during cook and need to serialize, so convert attribute to bSave = true
+			// this is handled by the facade 
+			HierarchyFacade.GenerateLevelAttribute();
 		}
 
 		// Finally, make sure expected interfaces are initialized
diff --git a/Engine/Source/Runtime/Experimental/Chaos/Private/GeometryCollection/GeometryCollectionClusteringUtility.cpp b/Engine/Source/Runtime/Experimental/Chaos/Private/GeometryCollection/GeometryCollectionClusteringUtility.cpp
--- a/Engine/Source/Runtime/Experimental/Chaos/Private/GeometryCollection/GeometryCollectionClusteringUtility.cpp
+++ b/Engine/Source/Runtime/Experimental/Chaos/Private/GeometryCollection/GeometryCollectionClusteringUtility.cpp
@@ -661,7 +661,7 @@ void FGeometryCollectionClusteringUtility::UpdateHierarchyLevelOfChildren(FGeome
 {
 	if (!GeometryCollection->HasAttribute("Level", FGeometryCollection::TransformGroup))
 	{
-		GeometryCollection->AddAttribute<int32>("Level", FGeometryCollection::TransformGroup, FManagedArrayCollection::FConstructionParameters(FName(),  false));
+		GeometryCollection->AddAttribute<int32>("Level", FGeometryCollection::TransformGroup);
 	}
 	TManagedArray<int32>& Levels = GeometryCollection->ModifyAttribute<int32>("Level", FGeometryCollection::TransformGroup);
 	const TManagedArray<TSet<int32>>& Children = GeometryCollection->Children;
diff --git a/Engine/Source/Runtime/Experimental/Chaos/Private/GeometryCollection/ManagedArrayCollection.cpp b/Engine/Source/Runtime/Experimental/Chaos/Private/GeometryCollection/ManagedArrayCollection.cpp
--- a/Engine/Source/Runtime/Experimental/Chaos/Private/GeometryCollection/ManagedArrayCollection.cpp
+++ b/Engine/Source/Runtime/Experimental/Chaos/Private/GeometryCollection/ManagedArrayCollection.cpp
@@ -117,6 +117,16 @@ bool FManagedArrayCollection::IsAttributeDirty(FName Name, FName Group) const
 	return false;
 }
 
+bool FManagedArrayCollection::IsAttributePersistent(FName Name, FName Group) const
+{
+	const FKeyType Key = FManagedArrayCollection::MakeMapKey(Name, Group);
+	if (const FValueType* Attribute = Map.Find(Key))
+	{
+		return Attribute->Saved;
+	}
+	return false;
+}
+
 TArray<FName> FManagedArrayCollection::AttributeNames(FName Group) const
 {
 	TArray<FName> AttributeNames;
@@ -612,7 +622,6 @@ void FManagedArrayCollection::Serialize(Chaos::FChaosArchive& Ar)
 			Map.Remove(Key);
 		}
 
-#if WITH_EDITOR
 		//it's possible new entries have been added but are not in old content. Resize these.
 		for (TTuple<FKeyType, FValueType>& Pair : Map)
 		{
@@ -623,8 +632,6 @@ void FManagedArrayCollection::Serialize(Chaos::FChaosArchive& Ar)
 			}
 		}
 
-#endif
-
 		// strip out GUID Attributes
 		for (auto& GroupName : GroupNames())
 		{
diff --git a/Engine/Source/Runtime/Experimental/Chaos/Public/GeometryCollection/Facades/CollectionHierarchyFacade.h b/Engine/Source/Runtime/Experimental/Chaos/Public/GeometryCollection/Facades/CollectionHierarchyFacade.h
new file mode 100644
--- /dev/null
+++ b/Engine/Source/Runtime/Experimental/Chaos/Public/GeometryCollection/Facades/CollectionHierarchyFacade.h
@@ -0,0 +1,46 @@
+// Copyright Epic Games, Inc. All Rights Reserved.
+
+#pragma once
+
+#include "CoreMinimal.h"
+#include "GeometryCollection/ManagedArrayAccessor.h"
+#include "GeometryCollection/ManagedArrayCollection.h"
+
+namespace Chaos::Facades
+{
+	/**
+	 * Provides an API to read and manipulate hierarchy in a managed array collection
+	 */
+	class CHAOS_API FCollectionHierarchyFacade
+	{
+	public:
+		enum class EPersistencePolicy : uint8
+		{
+			KeepExistingPersistence,
+			MustBePersistent
+		};
+
+	public:
+		FCollectionHierarchyFacade(FManagedArrayCollection& InCollection);
+
+		/** Valid if parent and children arrays are available */
+		bool IsValid() const;
+
+		/** whether the level attribute is available */
+		bool HasLevelAttribute() const;
+
+		/** whether the level attribute is persistent */
+		bool IsLevelAttributePersistent() const;
+
+		/** 
+		* Update level attribute for all elements (and create it if it is missing ) 
+		* @Param PersistencePolicy whether to make the attribute persistent or keep the existing state
+		*/ 
+		void GenerateLevelAttribute();
+
+	private:
+		TManagedArrayAccessor<int32>		ParentAttribute;
+		TManagedArrayAccessor<TSet<int32>>	ChildrenAttribute;
+		TManagedArrayAccessor<int32>		LevelAttribute;
+	};
+}
\ No newline at end of file
diff --git a/Engine/Source/Runtime/Experimental/Chaos/Public/GeometryCollection/ManagedArrayAccessor.h b/Engine/Source/Runtime/Experimental/Chaos/Public/GeometryCollection/ManagedArrayAccessor.h
--- a/Engine/Source/Runtime/Experimental/Chaos/Public/GeometryCollection/ManagedArrayAccessor.h
+++ b/Engine/Source/Runtime/Experimental/Chaos/Public/GeometryCollection/ManagedArrayAccessor.h
@@ -2,6 +2,15 @@
 #pragma once
 #include "GeometryCollection/ManagedArrayCollection.h"
 
+namespace ManageArrayAccessor
+{
+	enum class EPersistencePolicy : uint8
+	{
+		KeepExistingPersistence,
+		MakePersistent
+	};
+};
+
 /**
  * this class wraps a managed array
  * this provides a convenient API for optional attributes in a collection facade
@@ -20,6 +29,8 @@ public:
 	}
 	bool IsValid() const { return AttributeArray != nullptr; }
 
+	bool IsPersistent() const { return Collection.IsAttributePersistent(Name, Group); }
+
 	/** get the attribute for read only */
 	const TManagedArray<T>& Get() const
 	{
@@ -36,18 +47,33 @@ public:
 	}
 
 	/** add the attribute if it does not exists yet */
-	TManagedArray<T>& Add()
+	TManagedArray<T>& Add(ManageArrayAccessor::EPersistencePolicy PersistencePolicy = ManageArrayAccessor::EPersistencePolicy::MakePersistent)
 	{
-		AttributeArray = &Collection.AddAttribute<T>(Name, Group);
+		if (PersistencePolicy == ManageArrayAccessor::EPersistencePolicy::MakePersistent && !IsPersistent())
+		{
+			Remove();
+		}
+		const bool bSaved = (PersistencePolicy == ManageArrayAccessor::EPersistencePolicy::MakePersistent);
+		FManagedArrayCollection::FConstructionParameters Params(FName(), bSaved);
+		AttributeArray = &Collection.AddAttribute<T>(Name, Group, Params);
 		return *AttributeArray;
 	}
 
 	/** add and fill the attribute if it does not exist yet */
-	void AddAndFill(const T& Value)
+	void AddAndFill(const T& Value, ManageArrayAccessor::EPersistencePolicy PersistencePolicy = ManageArrayAccessor::EPersistencePolicy::MakePersistent)
 	{
 		if (!Collection.HasAttribute(Name, Group))
 		{
-			AttributeArray = &Collection.AddAttribute<T>(Name, Group);
+			Add(PersistencePolicy);
+			AttributeArray->Fill(Value);
+		}
+	}
+
+	/** Fill the attribute with a specific value */
+	void Fill(const T& Value)
+	{
+		if (AttributeArray)
+		{
 			AttributeArray->Fill(Value);
 		}
 	}
@@ -58,6 +84,12 @@ public:
 		Collection.CopyAttribute(FromAttribute.Collection, Name, Group);
 	}
 
+	void Remove()
+	{
+		Collection.RemoveAttribute(Name, Group);
+		AttributeArray = nullptr;
+	}
+
 private:
 	FManagedArrayCollection& Collection;
 	FName Name;
diff --git a/Engine/Source/Runtime/Experimental/Chaos/Public/GeometryCollection/ManagedArrayCollection.h b/Engine/Source/Runtime/Experimental/Chaos/Public/GeometryCollection/ManagedArrayCollection.h
--- a/Engine/Source/Runtime/Experimental/Chaos/Public/GeometryCollection/ManagedArrayCollection.h
+++ b/Engine/Source/Runtime/Experimental/Chaos/Public/GeometryCollection/ManagedArrayCollection.h
@@ -421,6 +421,14 @@ public:
 	* @param Group - The group that manages the attribute
 	*/
 	bool IsAttributeDirty(FName Name, FName Group) const;
+
+	/**
+	* Check if an attribute is persistent - ie its data will be serialized
+	* If the attribute does not exists this will return false
+	* @param Name - The name of the attribute
+	* @param Group - The group that manages the attribute
+	*/
+	bool IsAttributePersistent(FName Name, FName Group) const;
 	
 	/**
 	*
