Iris
- Fixed bug in ReplicationStateDescriptorBuilder outputting invalid DebugName when building descriptor for struct with CustomNetSerializer.

diff --git a/Engine/Source/Runtime/Experimental/Iris/Core/Private/Iris/ReplicationState/ReplicationStateDescriptorBuilder.cpp b/Engine/Source/Runtime/Experimental/Iris/Core/Private/Iris/ReplicationState/ReplicationStateDescriptorBuilder.cpp
--- a/Engine/Source/Runtime/Experimental/Iris/Core/Private/Iris/ReplicationState/ReplicationStateDescriptorBuilder.cpp
+++ b/Engine/Source/Runtime/Experimental/Iris/Core/Private/Iris/ReplicationState/ReplicationStateDescriptorBuilder.cpp
@@ -122,7 +122,7 @@ public:
 	// Utility methods
 	static void GetSerializerTraits(FMemberProperty& OutMemberProperty, const FProperty* Property, const FPropertyNetSerializerInfo* Info);
 	static bool IsSupportedProperty(FMemberProperty& OutMemberProperty, const FProperty* Property, const TArray<FLifetimeProperty>* LifeTimeProperties = nullptr, UClass* InObjectClass = nullptr);
-	static bool IsSupportedStructPropertyWithCustomSerializer(FMemberProperty& OutMemberProperty, const UStruct* InStruct);	
+	static bool IsSupportedStructWithCustomSerializer(FMemberProperty& OutMemberProperty, const UStruct* InStruct);	
 	static bool IsStructWithCustomSerializer(const UStruct* InStruct);
 
 	static EMemberPropertyTraits GetConnectionFilterTrait(ELifetimeCondition Condition);
@@ -229,6 +229,7 @@ private:
 		, CombinedPropertyTraits(EMemberPropertyTraits::None)
 		, SharedPropertyTraits(EMemberPropertyTraits::None)
 		, BuildParams({})
+		, bGeneratePropertyDescriptors(true)
 		{
 		}
 
@@ -248,6 +249,9 @@ private:
 		EMemberPropertyTraits SharedPropertyTraits;
 
 		FBuildParameters BuildParams;
+	
+		// Set if the descriptor we are building has should include property descriptor, there are some special cases like structs with custom NetSerializer which does not have any properties.
+		bool bGeneratePropertyDescriptors;
 	};
 
 	struct FStructInfo
@@ -275,7 +279,7 @@ private:
 	void BuildMemberPropertyDescriptors(FReplicationStateMemberPropertyDescriptor* MemberPropertyDescriptors, FReplicationStateDescriptor* Descriptor) const;
 	void BuildMemberLifetimeConditionDescriptors(FReplicationStateMemberLifetimeConditionDescriptor* MemberLifetimeConditionDescriptors, FReplicationStateDescriptor* Descriptor) const;
 	void BuildMemberRepIndexToMemberDescriptors(FReplicationStateMemberRepIndexToMemberIndexDescriptor* MemberRepIndexToMemberIndexDescriptors, FReplicationStateDescriptor* Descriptor) const;
-	void BuildMemberDebugDescriptors(FReplicationStateMemberDebugDescriptor* MemberDebugDescriptors, FReplicationStateDescriptor* Descriptor, const FMemberCache& MemberCache) const;
+	void BuildMemberDebugDescriptors(FReplicationStateMemberDebugDescriptor* MemberDebugDescriptors, FReplicationStateDescriptor* Descriptor, const FBuilderContext& Context) const;
 	void BuildMemberSerializerConfigs(uint8* SerializerConfigBuffer, FReplicationStateDescriptor* Descriptor, const FMemberCache& MemberCache) const;
 	void FixupNetRoleNetSerializerConfigs(FReplicationStateDescriptor* Descriptor);
 	void FixupDescriptorForNativeFastArray(FReplicationStateDescriptor* Descriptor) const;
@@ -453,6 +457,7 @@ void FPropertyReplicationStateDescriptorBuilder::BuildMemberCache(FBuilderContex
 	MemberCache.SetNum(Members.Num());
 	FMemberCacheEntry* CurrentCacheEntry = MemberCache.GetData();
 
+	bool bAreAllMembersProperties = true;
 	for (FMemberProperty& Member : Members)
 	{
 		const FNetSerializer* Serializer = Member.SerializerInfo->GetNetSerializer(Member.Property);
@@ -546,6 +551,7 @@ void FPropertyReplicationStateDescriptorBuilder::BuildMemberCache(FBuilderContex
 		{
 			CurrentCacheEntry->ExternalSizeAndAlignment = Member.ExternalSizeAndAlignment;
 			CurrentCacheEntry->InternalSizeAndAlignment = { Serializer->QuantizedTypeSize, Serializer->QuantizedTypeAlignment };
+			bAreAllMembersProperties = false;
 		}
 
 		CurrentCacheEntry->SerializerConfigSizeAndAlignment = { Serializer->ConfigTypeSize, Serializer->ConfigTypeAlignment };
@@ -560,6 +566,13 @@ void FPropertyReplicationStateDescriptorBuilder::BuildMemberCache(FBuilderContex
 
 	Context.CombinedPropertyTraits = PropertyTraits;
 	Context.SharedPropertyTraits = SharedPropertyTraits;
+
+	// Either all members should be property based or we have a single member that is not a property
+	if (!bAreAllMembersProperties)
+	{
+		check(Members.Num() <= 1);
+		Context.bGeneratePropertyDescriptors = false; 
+	}
 }
 
 void FPropertyReplicationStateDescriptorBuilder::BuildMemberTagCache(FMemberTagCache& MemberTagCache, const FMemberCache& MemberCache) const
@@ -1149,17 +1162,24 @@ void FPropertyReplicationStateDescriptorBuilder::BuildMemberRepIndexToMemberDesc
 	Descriptor->RepIndexCount = static_cast<uint16>(RepIndexCount);
 }
 
-void FPropertyReplicationStateDescriptorBuilder::BuildMemberDebugDescriptors(FReplicationStateMemberDebugDescriptor* MemberDebugDescriptors, FReplicationStateDescriptor* Descriptor, const FMemberCache& MemberCache) const
+void FPropertyReplicationStateDescriptorBuilder::BuildMemberDebugDescriptors(FReplicationStateMemberDebugDescriptor* MemberDebugDescriptors, FReplicationStateDescriptor* Descriptor, const FBuilderContext& Context) const
 {
-	TCHAR StringBuffer[1024];
-
 	FReplicationStateMemberDebugDescriptor* CurrentMemberDebugDescriptor = MemberDebugDescriptors;
-	for (const FMemberProperty& Info : Members)
+	if (Context.bGeneratePropertyDescriptors)
 	{
-		Info.Property->GetFName().ToString(StringBuffer, 1024);
-		CurrentMemberDebugDescriptor->DebugName = CreatePersistentNetDebugName(StringBuffer);
+		TCHAR StringBuffer[1024];
+		for (const FMemberProperty& Info : Members)
+		{
+			Info.Property->GetFName().ToString(StringBuffer, 1024);
+			CurrentMemberDebugDescriptor->DebugName = CreatePersistentNetDebugName(StringBuffer);
 		
-		++CurrentMemberDebugDescriptor;
+			++CurrentMemberDebugDescriptor;
+		}
+	}
+	else
+	{
+		// If we dont have any property descriptor we just use the debugname of the descriptor
+		CurrentMemberDebugDescriptor->DebugName = Descriptor->DebugName;
 	}
 
 	// Fill in descriptor
@@ -1456,10 +1476,13 @@ FPropertyReplicationStateDescriptorBuilder::Build(const FString& StateName, FRep
 	FMemoryLayoutUtil::AddToLayout<FReplicationStateMemberFunctionDescriptor>(Layout, LayoutData.MemberFunctionDescriptorSizeAndOffset, FunctionCount);
 
 	// MemberProperties
-	FMemoryLayoutUtil::AddToLayout<FProperty*>(Layout, LayoutData.MemberPropertiesSizeAndOffset, MemberCount);
+	if (Context.bGeneratePropertyDescriptors)
+	{
+		FMemoryLayoutUtil::AddToLayout<FProperty*>(Layout, LayoutData.MemberPropertiesSizeAndOffset, MemberCount);
 
-	// MemberPropertyDescriptors
-	FMemoryLayoutUtil::AddToLayout<FReplicationStateMemberPropertyDescriptor>(Layout, LayoutData.MemberPropertyDescriptorsSizeAndOffset, MemberCount);
+		// MemberPropertyDescriptors
+		FMemoryLayoutUtil::AddToLayout<FReplicationStateMemberPropertyDescriptor>(Layout, LayoutData.MemberPropertyDescriptorsSizeAndOffset, MemberCount);
+	}
 
 	// MemberLifetimeConditionDescriptors
 	if (EnumHasAnyFlags(Traits, EReplicationStateTraits::HasLifetimeConditionals))
@@ -1590,7 +1613,7 @@ FPropertyReplicationStateDescriptorBuilder::Build(const FString& StateName, FRep
 		Descriptor->DebugName = CreatePersistentNetDebugName(StateName.GetCharArray().GetData());
 		if (LayoutData.MemberDebugDescriptorsSizeAndOffset.Size > 0)
 		{
-			BuildMemberDebugDescriptors(reinterpret_cast<FReplicationStateMemberDebugDescriptor*>(Buffer + LayoutData.MemberDebugDescriptorsSizeAndOffset.Offset), Descriptor, Context.MemberCache);
+			BuildMemberDebugDescriptors(reinterpret_cast<FReplicationStateMemberDebugDescriptor*>(Buffer + LayoutData.MemberDebugDescriptorsSizeAndOffset.Offset), Descriptor, Context);
 		}
 
 		// Construct all MemberSerializerConfigs and update pointers in MemberSerializerDescriptors.
@@ -1815,7 +1838,7 @@ bool FPropertyReplicationStateDescriptorBuilder::IsSupportedProperty(FMemberProp
 	return true;
 }
 
-bool FPropertyReplicationStateDescriptorBuilder::IsSupportedStructPropertyWithCustomSerializer(FMemberProperty& OutMemberProperty, const UStruct* InStruct)
+bool FPropertyReplicationStateDescriptorBuilder::IsSupportedStructWithCustomSerializer(FMemberProperty& OutMemberProperty, const UStruct* InStruct)
 {
 	// See if we got a matching custom serializer for the entire struct
 	const FPropertyNetSerializerInfo* NetSerializerInfo = FPropertyNetSerializerInfoRegistry::FindStructSerializerInfo(InStruct->GetFName());
@@ -1823,7 +1846,7 @@ bool FPropertyReplicationStateDescriptorBuilder::IsSupportedStructPropertyWithCu
 	{
 		if (InStruct->GetSuperStruct())
 		{
-			const bool bSuperHasCustomNetSerializer = IsSupportedStructPropertyWithCustomSerializer(OutMemberProperty, InStruct->GetSuperStruct());
+			const bool bSuperHasCustomNetSerializer = IsSupportedStructWithCustomSerializer(OutMemberProperty, InStruct->GetSuperStruct());
 
 			// Warn if one of our supers have a CustomNetSerializer, if that is the case user must be made aware and either declare an explicit forwarding NetSerializer
 			// or implement a custom NetSerializer
@@ -2113,7 +2136,7 @@ TRefCountPtr<const FReplicationStateDescriptor> FReplicationStateDescriptorBuild
 	bool bAllMembersAreReplicated = true;
 
 	// We have a special case for structs with custom serializers
-	if (!Parameters.SkipCheckForCustomNetSerializerForStruct && FPropertyReplicationStateDescriptorBuilder::IsSupportedStructPropertyWithCustomSerializer(MemberProperty, InStruct))
+	if (!Parameters.SkipCheckForCustomNetSerializerForStruct && FPropertyReplicationStateDescriptorBuilder::IsSupportedStructWithCustomSerializer(MemberProperty, InStruct))
 	{
 		Builder.AddMemberProperty(MemberProperty);
 	}
