Chaos - capsule-triangle corner case fix
- add a separating axis test
- check for valid normals in corner case

diff --git a/Engine/Source/Runtime/Experimental/Chaos/Private/Chaos/Collision/CapsuleTriangleContactPoint.cpp b/Engine/Source/Runtime/Experimental/Chaos/Private/Chaos/Collision/CapsuleTriangleContactPoint.cpp
--- a/Engine/Source/Runtime/Experimental/Chaos/Private/Chaos/Collision/CapsuleTriangleContactPoint.cpp
+++ b/Engine/Source/Runtime/Experimental/Chaos/Private/Chaos/Collision/CapsuleTriangleContactPoint.cpp
@@ -8,7 +8,11 @@
 
 namespace Chaos
 {
-	extern bool bChaos_Collision_CapsuleTriMesh_MidPointCull;
+	namespace CVars
+	{
+		bool bChaosCollisionCapsuleTriMeshSATCull = false;
+		FAutoConsoleVariableRef CVarUseGenericSweptConvexConstraints(TEXT("p.Chaos.Collision.CapsuleTriMeshSATCull"), bChaosCollisionCapsuleTriMeshSATCull, TEXT("Enable the SAT cull check in capsule-triangle [default: false]."));
+	}
 
 	void AddCapsuleTriangleParallelEdgeManifoldContacts(const FVec3& P0, const FVec3& P1, const FVec3& EdgeP0, const FVec3& EdgeP1, const FReal R, const FReal RejectDistanceSq, const FReal NormalToleranceSq, TCArray<FContactPoint, 4>& OutContactPoints);
 
@@ -104,18 +108,64 @@ namespace Chaos
 			return;
 		}
 
-		// Edge plane normals
-		FVec3 EdgeNs[3];
-		EdgeNs[0] = FVec3::CrossProduct(V0 - V2, FaceN);
-		EdgeNs[1] = FVec3::CrossProduct(V1 - V0, FaceN);
-		EdgeNs[2] = FVec3::CrossProduct(V2 - V1, FaceN);
-		EdgeNs[0].Normalize(NormalToleranceSq);
-		EdgeNs[1].Normalize(NormalToleranceSq);
-		EdgeNs[2].Normalize(NormalToleranceSq);
+		// Early out if any edge saparating axis distance is above the reject distance
+		if (CVars::bChaosCollisionCapsuleTriMeshSATCull)
+		{
+			// Separating axis test on the edges
+			const auto& SATEdgeCheck = [&FaceN, &R, &Centroid, &RejectDistance, &NormalTolerance](const FVec3& P0, const FVec3& P1, const FVec3& EdgeP0, const FVec3& EdgeP1) -> bool
+			{
+				FVec3 EdgeSeparationAxis = FVec3::CrossProduct(P1 - P0, EdgeP1 - EdgeP0);
+				if (!EdgeSeparationAxis.Normalize())
+				{
+					// Parallel lines are handled later
+					return true;
+				}
+
+				// Separating axis always points away from the triangle
+				if (FVec3::DotProduct(EdgeP1 - Centroid, EdgeSeparationAxis) < FReal(0))
+				{
+					EdgeSeparationAxis = -EdgeSeparationAxis;
+				}
+
+				// We want to generate contacts with positive separation when within cull distance,
+				// but only if we're on the outside of the face. If we're below the face here
+				// we will not make a contact unless we actually touch.
+				FReal SeparationAxisCullDistance = RejectDistance;
+				if (FVec3::DotProduct(EdgeSeparationAxis, FaceN) < -NormalTolerance)
+				{
+					SeparationAxisCullDistance = R;
+				}
+
+				// Separating distance
+				const FReal EdgeSeparationDist = FVec3::DotProduct(P0 - EdgeP0, EdgeSeparationAxis);
+
+				return EdgeSeparationDist <= SeparationAxisCullDistance;
+			};
 
-		// Edge plane Signed Distances for each Segment point
+			if (!SATEdgeCheck(P0, P1, Triangle.GetVertex(2), Triangle.GetVertex(0)))
+			{
+				// Separated from triangle
+				return;
+			}
+			if (!SATEdgeCheck(P0, P1, Triangle.GetVertex(0), Triangle.GetVertex(1)))
+			{
+				// Separated from triangle
+				return;
+			}
+			if (!SATEdgeCheck(P0, P1, Triangle.GetVertex(1), Triangle.GetVertex(2)))
+			{
+				// Separated from triangle
+				return;
+			}
+		}
+
+		// Edge plane normals and signed distances to each segment point
+		FVec3 EdgeNs[3];
 		FReal EdgeD0s[3];
 		FReal EdgeD1s[3];
+
+		EdgeNs[0] = FVec3::CrossProduct(V0 - V2, FaceN);
+		EdgeNs[0].Normalize(NormalToleranceSq);
 		EdgeD0s[0] = FVec3::DotProduct(P0 - V0, EdgeNs[0]);
 		EdgeD1s[0] = FVec3::DotProduct(P1 - V0, EdgeNs[0]);
 		if (((EdgeD0s[0] > RejectDistance) & (EdgeD1s[0] > RejectDistance)) != 0)
@@ -124,6 +174,8 @@ namespace Chaos
 			return;
 		}
 
+		EdgeNs[1] = FVec3::CrossProduct(V1 - V0, FaceN);
+		EdgeNs[1].Normalize(NormalToleranceSq);
 		EdgeD0s[1] = FVec3::DotProduct(P0 - V1, EdgeNs[1]);
 		EdgeD1s[1] = FVec3::DotProduct(P1 - V1, EdgeNs[1]);
 		if (((EdgeD0s[1] > RejectDistance) & (EdgeD1s[1] > RejectDistance)) != 0)
@@ -132,6 +184,8 @@ namespace Chaos
 			return;
 		}
 
+		EdgeNs[2] = FVec3::CrossProduct(V2 - V1, FaceN);
+		EdgeNs[2].Normalize(NormalToleranceSq);
 		EdgeD0s[2] = FVec3::DotProduct(P0 - V2, EdgeNs[2]);
 		EdgeD1s[2] = FVec3::DotProduct(P1 - V2, EdgeNs[2]);
 		if (((EdgeD0s[2] > RejectDistance) & (EdgeD1s[2] > RejectDistance)) != 0)
@@ -347,6 +401,26 @@ namespace Chaos
 					{
 						continue;
 					}
+				
+					// For Vertex contacts, the normal needs to be outside both planes
+					if ((EdgeT == FReal(0)) && ((SegmentT == FReal(0)) || (SegmentT == FReal(1))))
+					{
+						const int32 PrevEdgeIndex = (EdgeIndex > 0) ? EdgeIndex - 1 : 2;
+						const FReal PrevDotEdge = FVec3::DotProduct(EdgeNs[PrevEdgeIndex], SegmentEdgeN);
+						if (PrevDotEdge < -NormalTolerance)
+						{
+							continue;
+						}
+					}
+					if ((EdgeT == FReal(1)) && ((SegmentT == FReal(0)) || (SegmentT == FReal(1))))
+					{
+						const int32 NextEdgeIndex = (EdgeIndex < 2) ? EdgeIndex + 1 : 0;
+						const FReal NextDotEdge = FVec3::DotProduct(EdgeNs[NextEdgeIndex], SegmentEdgeN);
+						if (NextDotEdge < -NormalTolerance)
+						{
+							continue;
+						}
+					}
 				}
 
 				// If we are within the face angle tolerance, generate a face contact rather than an edge one
@@ -361,7 +435,7 @@ namespace Chaos
 					ContactPoint.ShapeContactPoints[0] = CapsuleP;
 					ContactPoint.ShapeContactPoints[1] = CapsuleP - CapsuleDist * FaceN;
 					ContactPoint.ShapeContactNormal = FaceN;
-					ContactPoint.Phi = CapsuleDist - R;
+					ContactPoint.Phi = CapsuleDist;
 					ContactPoint.ContactType = EContactPointType::VertexPlane;
 					ContactPoint.FaceIndex = INDEX_NONE;
 				}
diff --git a/Engine/Source/Runtime/Experimental/Chaos/Private/Chaos/Collision/ContactTriangles.cpp b/Engine/Source/Runtime/Experimental/Chaos/Private/Chaos/Collision/ContactTriangles.cpp
--- a/Engine/Source/Runtime/Experimental/Chaos/Private/Chaos/Collision/ContactTriangles.cpp
+++ b/Engine/Source/Runtime/Experimental/Chaos/Private/Chaos/Collision/ContactTriangles.cpp
@@ -538,6 +538,9 @@ namespace Chaos
 #if CHAOS_DEBUG_DRAW
 		if (CVars::ChaosSolverDebugDrawMeshContacts && FDebugDrawQueue::GetInstance().IsDebugDrawingEnabled())
 		{
+			TArray<bool> bTriangleDrawn;
+			bTriangleDrawn.SetNumZeroed(ContactTriangles.Num());
+
 			const FReal Duration = 0;
 			const int32 DrawPriority = 10;
 			for (const FTriangleContactPoint& ContactPoint : TriangleContactPoints)
@@ -551,6 +554,18 @@ namespace Chaos
 
 				// Draw a thin black line connecting the two contact points (triangle face to convex surface)
 				FDebugDrawQueue::GetInstance().DrawDebugLine(P0, P1, FColor::Black, false, FRealSingle(Duration), DrawPriority, 0.5f * FRealSingle(LineThickness));
+
+				if (!bTriangleDrawn[ContactPoint.ContactTriangleIndex])
+				{
+					FContactTriangle& Triangle = ContactTriangles[ContactPoint.ContactTriangleIndex];
+					const FVec3 V0 = ConvexTransform.TransformPosition(Triangle.Vertices[0]);
+					const FVec3 V1 = ConvexTransform.TransformPosition(Triangle.Vertices[1]);
+					const FVec3 V2 = ConvexTransform.TransformPosition(Triangle.Vertices[2]);
+					FDebugDrawQueue::GetInstance().DrawDebugLine(V0, V1, FColor::Silver, false, FRealSingle(Duration), DrawPriority, 0.5f * FRealSingle(LineThickness));
+					FDebugDrawQueue::GetInstance().DrawDebugLine(V1, V2, FColor::Silver, false, FRealSingle(Duration), DrawPriority, 0.5f * FRealSingle(LineThickness));
+					FDebugDrawQueue::GetInstance().DrawDebugLine(V2, V0, FColor::Silver, false, FRealSingle(Duration), DrawPriority, 0.5f * FRealSingle(LineThickness));
+					bTriangleDrawn[ContactPoint.ContactTriangleIndex] = true;
+				}
 			}
 		}
 #endif
