Chaos - fix solver shutdown in constraint graph combined with geometry collection  particles

diff --git a/Engine/Source/Runtime/Experimental/Chaos/Private/Chaos/Framework/PhysicsSolverBase.cpp b/Engine/Source/Runtime/Experimental/Chaos/Private/Chaos/Framework/PhysicsSolverBase.cpp
--- a/Engine/Source/Runtime/Experimental/Chaos/Private/Chaos/Framework/PhysicsSolverBase.cpp
+++ b/Engine/Source/Runtime/Experimental/Chaos/Private/Chaos/Framework/PhysicsSolverBase.cpp
@@ -242,6 +242,7 @@ namespace Chaos
 
 		// GeometryCollection particles do not always remove collision constraints on unregister,
 		// explicitly clear constraints so we will not crash when filling collision events in advance.
+		// @todo(chaos): fix this and remove
 		{
 			auto* Evolution = static_cast<FPBDRigidsSolver&>(InSolver).GetEvolution();
 			if (Evolution)
diff --git a/Engine/Source/Runtime/Experimental/Chaos/Private/Chaos/Island/IslandGraph.cpp b/Engine/Source/Runtime/Experimental/Chaos/Private/Chaos/Island/IslandGraph.cpp
--- a/Engine/Source/Runtime/Experimental/Chaos/Private/Chaos/Island/IslandGraph.cpp
+++ b/Engine/Source/Runtime/Experimental/Chaos/Private/Chaos/Island/IslandGraph.cpp
@@ -218,18 +218,15 @@ void FIslandGraph<NodeType, EdgeType, IslandType, OwnerType>::ResetEdges()
 		}
 	}
 
-	for (auto& GraphIsland : GraphIslands)
-	{
-		GraphIsland.NumEdges = 0;
-	}
-
 	for (auto& GraphNode : GraphNodes)
 	{
 		GraphNode.NodeEdges.Reset();
+		GraphNode.IslandIndex = INDEX_NONE;
 	}
 
 	GraphEdges.Reset();
 	ItemEdges.Reset();
+	GraphIslands.Reset();
 }
 
 template<typename NodeType, typename EdgeType, typename IslandType, typename OwnerType>
diff --git a/Engine/Source/Runtime/Experimental/Chaos/Private/Chaos/Island/IslandManager.cpp b/Engine/Source/Runtime/Experimental/Chaos/Private/Chaos/Island/IslandManager.cpp
--- a/Engine/Source/Runtime/Experimental/Chaos/Private/Chaos/Island/IslandManager.cpp
+++ b/Engine/Source/Runtime/Experimental/Chaos/Private/Chaos/Island/IslandManager.cpp
@@ -293,6 +293,7 @@ void FPBDIslandManager::RemoveConstraints()
 	for(auto& Island : Islands)
 	{
 		Island->ClearConstraints();
+		Island->ClearParticles();
 	}
 }
 
