When UpdateConstraintFromGeometrySwept does not actually sweep (ie, if the object wasn't moving faster than its threshold) then we still need to generate constraints with the stuff that we may have hit.

diff --git a/Engine/Source/Runtime/Experimental/Chaos/Private/Chaos/Collision/ParticlePairMidPhase.cpp b/Engine/Source/Runtime/Experimental/Chaos/Private/Chaos/Collision/ParticlePairMidPhase.cpp
--- a/Engine/Source/Runtime/Experimental/Chaos/Private/Chaos/Collision/ParticlePairMidPhase.cpp
+++ b/Engine/Source/Runtime/Experimental/Chaos/Private/Chaos/Collision/ParticlePairMidPhase.cpp
@@ -418,10 +418,10 @@ namespace Chaos
 			const FRigidTransform3 CCDParticleWorldTransform1 = FRigidTransform3(StartX1, P1->Q());
 			const FRigidTransform3 CCDShapeWorldTransform0 = Constraint->ImplicitTransform[0] * CCDParticleWorldTransform0;
 			const FRigidTransform3 CCDShapeWorldTransform1 = Constraint->ImplicitTransform[1] * CCDParticleWorldTransform1;
-			Collisions::UpdateConstraintSwept(*Constraint.Get(), CCDShapeWorldTransform0, CCDShapeWorldTransform1, Dt);
+			const bool bDidSweep = Collisions::UpdateConstraintSwept(*Constraint.Get(), CCDShapeWorldTransform0, CCDShapeWorldTransform1, Dt);
 
 			// If we did get a hit but it's at TOI = 1, treat this constraint as a regular non-swept constraint (skip the rewind)
-			if (Constraint->GetCCDTimeOfImpact() == FReal(1))
+			if ((!bDidSweep) || Constraint->GetCCDTimeOfImpact() == FReal(1))
 			{
 				Collisions::UpdateConstraint(*Constraint.Get(), Constraint->GetShapeWorldTransform0(), Constraint->GetShapeWorldTransform1(), Dt);
 			}
