Iris
- Fixed typo in UE_NET_IMPLEMENT_NAMED_STRUCT_NETSERIALIZER_INFO macro

diff --git a/Engine/Source/Runtime/Experimental/Iris/Core/Public/Iris/ReplicationState/PropertyNetSerializerInfoRegistry.h b/Engine/Source/Runtime/Experimental/Iris/Core/Public/Iris/ReplicationState/PropertyNetSerializerInfoRegistry.h
--- a/Engine/Source/Runtime/Experimental/Iris/Core/Public/Iris/ReplicationState/PropertyNetSerializerInfoRegistry.h
+++ b/Engine/Source/Runtime/Experimental/Iris/Core/Public/Iris/ReplicationState/PropertyNetSerializerInfoRegistry.h
@@ -166,5 +166,5 @@ static F##Name##NetSerializerRegistryDelegates Name##NetSerializerRegistryDelega
 // Utility that can be used to forward serialization of a Struct to a specific NetSerializer
 #define UE_NET_IMPLEMENT_FORWARDING_NETSERIALIZER_AND_REGISTRY_DELEGATES(Name, SerializerName) \
 	static const FName PropertyNetSerializerRegistry_NAME_##Name( PREPROCESSOR_TO_STRING(Name) ); \
-	UE_NET_IMPLEMENT_NAMED_STRUCT_NETSERIALIZER_INFO(PropertyNetSerializerRegistry_NAME_##Name, FGameplayEffectContextNetSerializer); \
+	UE_NET_IMPLEMENT_NAMED_STRUCT_NETSERIALIZER_INFO(PropertyNetSerializerRegistry_NAME_##Name, SerializerName); \
 	UE_NET_IMPLEMENT_NETSERIALZIER_REGISTRY_DELEGATES(Name)
