Iris
- Added ReplicationStateDescriptorConfig and a SupportsStructNetSerializerList that can be used to silence warnings issued for structs with custom NetSerialze/NetDeltaSerialze that does not implement custom Iris NetSerializer but works fine using default Struct serialization.

diff --git a/Engine/Source/Runtime/Experimental/Iris/Core/Private/Iris/ReplicationState/ReplicationStateDescriptorBuilder.cpp b/Engine/Source/Runtime/Experimental/Iris/Core/Private/Iris/ReplicationState/ReplicationStateDescriptorBuilder.cpp
--- a/Engine/Source/Runtime/Experimental/Iris/Core/Private/Iris/ReplicationState/ReplicationStateDescriptorBuilder.cpp
+++ b/Engine/Source/Runtime/Experimental/Iris/Core/Private/Iris/ReplicationState/ReplicationStateDescriptorBuilder.cpp
@@ -11,6 +11,7 @@
 #include "Iris/ReplicationState/PropertyNetSerializerInfoRegistry.h"
 #include "Iris/ReplicationState/InternalPropertyReplicationState.h"
 #include "Iris/ReplicationState/IrisFastArraySerializer.h"
+#include "Iris/ReplicationState/ReplicationStateDescriptorConfig.h"
 #include "Iris/ReplicationState/ReplicationStateDescriptorRegistry.h"
 #include "Iris/ReplicationSystem/ReplicationOperations.h"
 #include "Iris/ReplicationSystem/NetBlob/PartialNetObjectAttachmentHandler.h"
@@ -52,6 +53,9 @@ static TAutoConsoleVariable<int32> CVarIrisUseNativeFastArray(TEXT("net.Iris.Use
 static bool bWarnAboutStructsWithCustomSerialization = true;
 static FAutoConsoleVariableRef CVarIrisWarnAboutStructsWithCustomSerialization(TEXT("net.Iris.WarnAboutStructsWithCustomSerialization"), bWarnAboutStructsWithCustomSerialization, TEXT("Warn when generating descriptors for structs with custom serialization."));
 
+static bool bUseSupportsStructNetSerializerList = true;
+static FAutoConsoleVariableRef CVarIrisUseSupportsStructNetSerializerList(TEXT("net.Iris.UseSupportsStructNetSerializerList"), bUseSupportsStructNetSerializerList, TEXT("If enabled structs in the SupportsStructNetSerializerList will not raise warnings even if the struct has a NetSerialize/NetDeltaSerialize but has no custom NetSerializer."));
+
 static bool bWarnAboutStructPropertiesWithSuspectedNotReplicatedProperties = false;
 static FAutoConsoleVariableRef CVarIrisWarnAboutStructPropertiesWithSuspectedNotReplicatedProperties(TEXT("net.Iris.bWarnAboutStructPropertiesWithSuspectedNotReplicatedProperties"), bWarnAboutStructPropertiesWithSuspectedNotReplicatedProperties, TEXT("Try to detect if a struct replicated as a property might contain unannotated members, disabled by default."));
 
@@ -125,6 +129,8 @@ public:
 	static bool IsSupportedStructWithCustomSerializer(FMemberProperty& OutMemberProperty, const UStruct* InStruct);	
 	static bool IsStructWithCustomSerializer(const UStruct* InStruct);
 
+	static bool CanStructUseStructNetSerializer(FName StructName);
+	
 	static EMemberPropertyTraits GetConnectionFilterTrait(ELifetimeCondition Condition);
 	static EMemberPropertyTraits GetInitOnlyTrait(ELifetimeCondition Condition);
 	static EMemberPropertyTraits GetFastArrayPropertyTraits(const FNetSerializer* NetSerializer, const FProperty* Property);
@@ -1885,6 +1891,20 @@ bool FPropertyReplicationStateDescriptorBuilder::IsStructWithCustomSerializer(co
 	return Super ? IsStructWithCustomSerializer(Super) : false;
 }
 
+// Returns true if the named struct is marked as working using the default StructNetSerializer
+bool FPropertyReplicationStateDescriptorBuilder::CanStructUseStructNetSerializer(FName StructName)
+{
+	if (bUseSupportsStructNetSerializerList)
+	{
+		const UReplicationStateDescriptorConfig* ReplicationStateDescriptorConfig = GetDefault<UReplicationStateDescriptorConfig>();
+		if (const FSupportsStructNetSerializerConfig* StructNetSerializerConfig = ReplicationStateDescriptorConfig->GetSupportsStructNetSerializerList().FindByPredicate([StructName](const FSupportsStructNetSerializerConfig& Item) { return Item.StructName == StructName; } ))
+		{
+			return StructNetSerializerConfig->bCanUseStructNetSerializer;
+		}
+	}
+	return false;
+}
+
 // RepTags are not intended to be derived from arbitrary properties, but rather be explicitly declared on a member. This is a hack.
 FRepTag FPropertyReplicationStateDescriptorBuilder::GetRepTagFromProperty(const FMemberCacheEntry& MemberCacheEntry, const FMemberProperty& MemberProperty)
 {
@@ -2118,15 +2138,15 @@ TRefCountPtr<const FReplicationStateDescriptor> FReplicationStateDescriptorBuild
 	bool bIsStructWithCustomSerialization = false;
 	if (const UScriptStruct* ScriptStruct = Cast<UScriptStruct>(InStruct))
 	{
-		bIsStructWithCustomSerialization = EnumHasAnyFlags(ScriptStruct->StructFlags, EStructFlags(STRUCT_NetSerializeNative | STRUCT_NetDeltaSerializeNative));
-
 		bIsFastArraySerializerItem = ScriptStruct->IsChildOf(FFastArraySerializerItem::StaticStruct());
 		bIsFastArraySerializer = ScriptStruct->IsChildOf(FFastArraySerializer::StaticStruct());
 		if (bIsFastArraySerializer && !Parameters.EnableFastArrayHandling)
 		{
 			bIsFastArraySerializer = false;
-			UE_LOG_DESCRIPTORBUILDER_WARNING(TEXT("FReplicationStateDescriptorBuilder ::CreateDescriptorForStruct Generating descriptor for FastArraySerializer as generic struct due to incompatible members. Struct is %s."), ToCStr(InStruct->GetName()));
+			UE_LOG_DESCRIPTORBUILDER_WARNING(TEXT("FReplicationStateDescriptorBuilder::CreateDescriptorForStruct Generating descriptor for FastArraySerializer as generic struct due to incompatible members. Struct is %s."), ToCStr(InStruct->GetName()));
 		}
+
+		bIsStructWithCustomSerialization = !bIsFastArraySerializer && EnumHasAnyFlags(ScriptStruct->StructFlags, EStructFlags(STRUCT_NetSerializeNative | STRUCT_NetDeltaSerializeNative));
 	}
 
 	FPropertyReplicationStateDescriptorBuilder Builder;
@@ -2143,7 +2163,15 @@ TRefCountPtr<const FReplicationStateDescriptor> FReplicationStateDescriptorBuild
 	else
 	{
 		const bool bIsAllowedToWarn = ShouldUseIrisReplication() || FApp::IsGame() || IsRunningDedicatedServer() || IsRunningClientOnly();
-		UE_CLOG(!Parameters.SkipCheckForCustomNetSerializerForStruct && bIsStructWithCustomSerialization && bIsAllowedToWarn && bWarnAboutStructsWithCustomSerialization, LogIris, Warning, TEXT("Generating descriptor for struct %s that has custom serialization."), ToCStr(InStruct->GetName()));
+
+		if (bIsStructWithCustomSerialization)
+		{
+			// Warn if we did not find a custom NetSerializer for a struct that has overridden NetSerialize or NetDeltaSerialize, unless it is marked as supporting the default StructNetSerializer
+			if (!Parameters.SkipCheckForCustomNetSerializerForStruct && bIsAllowedToWarn && bWarnAboutStructsWithCustomSerialization && !FPropertyReplicationStateDescriptorBuilder::CanStructUseStructNetSerializer(InStruct->GetFName()))
+			{
+				UE_LOG(LogIris, Warning, TEXT("Generating descriptor for struct %s that has custom serialization."), ToCStr(InStruct->GetName()));
+			}
+		}
 
 		uint32 PropertyCount = 0U;
 		for (TFieldIterator<FProperty> It(InStruct, EFieldIteratorFlags::IncludeSuper); It; ++It)
diff --git a/Engine/Source/Runtime/Experimental/Iris/Core/Private/Iris/ReplicationState/ReplicationStateDescriptorConfig.cpp b/Engine/Source/Runtime/Experimental/Iris/Core/Private/Iris/ReplicationState/ReplicationStateDescriptorConfig.cpp
new file mode 100644
--- /dev/null
+++ b/Engine/Source/Runtime/Experimental/Iris/Core/Private/Iris/ReplicationState/ReplicationStateDescriptorConfig.cpp
@@ -0,0 +1,16 @@
+// Copyright Epic Games, Inc. All Rights Reserved.
+
+#include "Iris/ReplicationState/ReplicationStateDescriptorConfig.h"
+#include "Containers/ArrayView.h"
+
+#include UE_INLINE_GENERATED_CPP_BY_NAME(ReplicationStateDescriptorConfig)
+
+UReplicationStateDescriptorConfig::UReplicationStateDescriptorConfig()
+: Super()
+{
+}
+
+TConstArrayView<FSupportsStructNetSerializerConfig> UReplicationStateDescriptorConfig::GetSupportsStructNetSerializerList() const
+{
+	return MakeArrayView(SupportsStructNetSerializerList);
+}
diff --git a/Engine/Source/Runtime/Experimental/Iris/Core/Public/Iris/ReplicationState/ReplicationStateDescriptorConfig.h b/Engine/Source/Runtime/Experimental/Iris/Core/Public/Iris/ReplicationState/ReplicationStateDescriptorConfig.h
new file mode 100644
--- /dev/null
+++ b/Engine/Source/Runtime/Experimental/Iris/Core/Public/Iris/ReplicationState/ReplicationStateDescriptorConfig.h
@@ -0,0 +1,40 @@
+// Copyright Epic Games, Inc. All Rights Reserved.
+
+#pragma once
+
+#include "Containers/ContainersFwd.h"
+#include "UObject/ObjectMacros.h"
+#include "ReplicationStateDescriptorConfig.generated.h"
+
+USTRUCT()
+struct FSupportsStructNetSerializerConfig
+{
+	GENERATED_BODY()
+
+	/** Struct name. */
+	UPROPERTY()
+	FName StructName;
+
+	/** If the named struct works with the default Iris StructNetSerializer. */
+	UPROPERTY()
+	bool bCanUseStructNetSerializer = true;
+};
+
+UCLASS(transient, config=Engine)
+class UReplicationStateDescriptorConfig : public UObject
+{
+	GENERATED_BODY()
+
+public:
+	IRISCORE_API TConstArrayView<FSupportsStructNetSerializerConfig> GetSupportsStructNetSerializerList() const;
+
+protected:
+	UReplicationStateDescriptorConfig();
+
+private:
+	/**
+	 * Structs that works using the default struct NetSerializer when running iris replication even though they implement a custom NetSerialize or NetDeltaSerialize method.
+	 */
+	UPROPERTY(Config)
+	TArray<FSupportsStructNetSerializerConfig> SupportsStructNetSerializerList;
+};
