Dataflow : Evaluation
- Formalized the caching timestamp with a struct wrapper.

diff --git a/Engine/Plugins/Experimental/Dataflow/Source/DataflowEditor/Private/Dataflow/DataflowGraphEditor.cpp b/Engine/Plugins/Experimental/Dataflow/Source/DataflowEditor/Private/Dataflow/DataflowGraphEditor.cpp
--- a/Engine/Plugins/Experimental/Dataflow/Source/DataflowEditor/Private/Dataflow/DataflowGraphEditor.cpp
+++ b/Engine/Plugins/Experimental/Dataflow/Source/DataflowEditor/Private/Dataflow/DataflowGraphEditor.cpp
@@ -8,6 +8,7 @@
 #include "BoneDragDropOp.h"
 #include "IStructureDetailsView.h"
 #include "Dataflow/DataflowCommentNode.h"
+#include "Dataflow/DataflowNodeParameters.h"
 
 #include "Dataflow/DataflowSchema.h"
 
@@ -106,7 +107,6 @@ void SDataflowGraphEditor::Construct(const FArguments& InArgs, UObject* InAssetO
 
 void SDataflowGraphEditor::EvaluateNode()
 {
-	float EvalTime = FGameTime::GetTimeSinceAppStart().GetRealTimeSeconds();
 	if (EvaluateGraphCallback)
 	{
 		FDataflowEditorCommands::EvaluateNodes(GetSelectedNodes(), EvaluateGraphCallback);
@@ -115,8 +115,8 @@ void SDataflowGraphEditor::EvaluateNode()
 	{
 		FDataflowEditorCommands::FGraphEvaluationCallback LocalEvaluateCallback = [](FDataflowNode* Node, FDataflowOutput* Out)
 		{
-			Dataflow::FContextThreaded Context(FGameTime::GetTimeSinceAppStart().GetRealTimeSeconds());
-			Node->Evaluate(Context, Out);
+			using namespace Dataflow;
+			FContextThreaded(FPlatformTime::Cycles64()).Evaluate(Node, Out);
 		};
 
 		FDataflowEditorCommands::EvaluateNodes(GetSelectedNodes(), LocalEvaluateCallback);
diff --git a/Engine/Plugins/Experimental/GeometryCollectionPlugin/Source/GeometryCollectionEditor/Private/GeometryCollection/GeometryCollectionEditorToolkit.cpp b/Engine/Plugins/Experimental/GeometryCollectionPlugin/Source/GeometryCollectionEditor/Private/GeometryCollection/GeometryCollectionEditorToolkit.cpp
--- a/Engine/Plugins/Experimental/GeometryCollectionPlugin/Source/GeometryCollectionEditor/Private/GeometryCollection/GeometryCollectionEditorToolkit.cpp
+++ b/Engine/Plugins/Experimental/GeometryCollectionPlugin/Source/GeometryCollectionEditor/Private/GeometryCollection/GeometryCollectionEditorToolkit.cpp
@@ -96,9 +96,7 @@ TSharedRef<SGraphEditor> FGeometryCollectionEditorToolkit::CreateGraphEditorWidg
 
 	FDataflowEditorCommands::FGraphEvaluationCallback Evaluate = [&](FDataflowNode* Node, FDataflowOutput* Out)
 	{
-		float EvalTime = FGameTime::GetTimeSinceAppStart().GetRealTimeSeconds();
-		FEngineContext Context(GeometryCollection, Dataflow, EvalTime, FString("UGeometryCollection"));
-		Node->Evaluate(Context, Out);
+		FEngineContext(GeometryCollection, Dataflow, FPlatformTime::Cycles64(), FString("UGeometryCollection")).Evaluate(Node, Out);
 	};
 
 	return SNew(SDataflowGraphEditor, GeometryCollection)
diff --git a/Engine/Source/Runtime/Experimental/Dataflow/Core/Private/Dataflow/DataflowNode.cpp b/Engine/Source/Runtime/Experimental/Dataflow/Core/Private/Dataflow/DataflowNode.cpp
--- a/Engine/Source/Runtime/Experimental/Dataflow/Core/Private/Dataflow/DataflowNode.cpp
+++ b/Engine/Source/Runtime/Experimental/Dataflow/Core/Private/Dataflow/DataflowNode.cpp
@@ -197,12 +197,8 @@ TArray<Dataflow::FPin> FDataflowNode::GetPins() const
 
 void FDataflowNode::Invalidate()
 {
-	LastModifiedTimestamp = FPlatformTime::Cycles64();
-	InvalidateOutputs();
-}
+	LastModifiedTimestamp = Dataflow::FTimestamp(FPlatformTime::Cycles64());
 
-void FDataflowNode::InvalidateOutputs()
-{
 	for (TPair<uint32, FDataflowOutput*> Elem : Outputs)
 	{
 		FDataflowOutput* Con = Elem.Value;
diff --git a/Engine/Source/Runtime/Experimental/Dataflow/Core/Private/Dataflow/DataflowNodeParameters.cpp b/Engine/Source/Runtime/Experimental/Dataflow/Core/Private/Dataflow/DataflowNodeParameters.cpp
--- a/Engine/Source/Runtime/Experimental/Dataflow/Core/Private/Dataflow/DataflowNodeParameters.cpp
+++ b/Engine/Source/Runtime/Experimental/Dataflow/Core/Private/Dataflow/DataflowNodeParameters.cpp
@@ -1,15 +1,38 @@
-﻿// Copyright Epic Games, Inc. All Rights Reserved.
+// Copyright Epic Games, Inc. All Rights Reserved.
 
 #include "Dataflow/DataflowNodeParameters.h"
 #include "Dataflow/DataflowInputOutput.h"
+#include "Dataflow/DataflowNode.h"
 
 namespace Dataflow
 {
+	uint64 FTimestamp::Invalid = 0;
+
+	void FContextSingle::Evaluate(const FDataflowNode* Node, const FDataflowOutput* Output)
+	{
+		if (Node != nullptr)
+		{
+			Timestamp = FTimestamp(FPlatformTime::Cycles64());
+			Node->Evaluate(*this, Output);
+		}
+	}
+	
 	bool FContextSingle::Evaluate(const FDataflowOutput& Connection)
 	{
 		return Connection.EvaluateImpl(*this);
 	}
 
+
+
+	void FContextThreaded::Evaluate(const FDataflowNode* Node, const FDataflowOutput* Output)
+	{
+		if (Node != nullptr)
+		{
+			Timestamp = FTimestamp(FPlatformTime::Cycles64());
+			Node->Evaluate(*this, Output);
+		}
+	}
+
 	bool FContextThreaded::Evaluate(const FDataflowOutput& Connection)
 	{
 		Connection.OutputLock->Lock(); ON_SCOPE_EXIT { Connection.OutputLock->Unlock(); };
diff --git a/Engine/Source/Runtime/Experimental/Dataflow/Core/Public/Dataflow/DataflowNode.h b/Engine/Source/Runtime/Experimental/Dataflow/Core/Public/Dataflow/DataflowNode.h
--- a/Engine/Source/Runtime/Experimental/Dataflow/Core/Public/Dataflow/DataflowNode.h
+++ b/Engine/Source/Runtime/Experimental/Dataflow/Core/Public/Dataflow/DataflowNode.h
@@ -41,7 +41,7 @@ struct DATAFLOWCORE_API FDataflowNode
 
 	FGuid Guid;
 	FName Name;
-	uint64 LastModifiedTimestamp;
+	Dataflow::FTimestamp LastModifiedTimestamp;
 
 	TMap< int, FDataflowInput* > Inputs;
 	TMap< int, FDataflowOutput* > Outputs;
@@ -52,14 +52,14 @@ struct DATAFLOWCORE_API FDataflowNode
 	FDataflowNode()
 		: Guid(FGuid())
 		, Name("Invalid")
-		, LastModifiedTimestamp(0)
+		, LastModifiedTimestamp(Dataflow::FTimestamp::Invalid)
 	{
 	}
 
 	FDataflowNode(const Dataflow::FNodeParameters& Param, FGuid InGuid = FGuid::NewGuid())
 		: Guid(InGuid)
 		, Name(Param.Name)
-		, LastModifiedTimestamp(0)
+		, LastModifiedTimestamp(Dataflow::FTimestamp::Invalid)
 	{
 	}
 
@@ -174,8 +174,6 @@ struct DATAFLOWCORE_API FDataflowNode
 
 	void Invalidate();
 
-	void InvalidateOutputs();
-
 	bool ValidateConnections();
 
 	bool IsValid() const { return bValid; }
diff --git a/Engine/Source/Runtime/Experimental/Dataflow/Core/Public/Dataflow/DataflowNodeParameters.h b/Engine/Source/Runtime/Experimental/Dataflow/Core/Public/Dataflow/DataflowNodeParameters.h
--- a/Engine/Source/Runtime/Experimental/Dataflow/Core/Public/Dataflow/DataflowNodeParameters.h
+++ b/Engine/Source/Runtime/Experimental/Dataflow/Core/Public/Dataflow/DataflowNodeParameters.h
@@ -8,13 +8,22 @@
 #include "GenericPlatform/GenericPlatformCriticalSection.h"
 
 class  UDataflow;
+struct FDataflowNode;
 struct FDataflowOutput;
 
 namespace Dataflow
 {
+	struct DATAFLOWCORE_API FTimestamp
+	{
+		static uint64 Invalid;
+		FTimestamp(uint64 InValue) : Value(InValue) {}
+		uint64 Value = 0;
+		bool operator>=(const FTimestamp& InTimestamp) { return Value >= InTimestamp.Value; }
+	};
+
 	struct FContextCacheElementBase 
 	{
-		FContextCacheElementBase(FProperty* InProperty = nullptr, uint64 InTimestamp = 0)
+		FContextCacheElementBase(FProperty* InProperty = nullptr, FTimestamp InTimestamp = FTimestamp::Invalid)
 			: Property(InProperty)
 			, Timestamp(InTimestamp)
 		{}
@@ -24,18 +33,18 @@ namespace Dataflow
 		const T& GetTypedData(const FProperty* PropertyIn) const;
 		
 		FProperty* Property = nullptr;
-		uint64 Timestamp = 0;
+		FTimestamp Timestamp = FTimestamp::Invalid;
 	};
 
 	template<class T>
 	struct FContextCacheElement : public FContextCacheElementBase 
 	{
-		FContextCacheElement(FProperty* InProperty, const T& InData, uint64 Timestamp)
+		FContextCacheElement(FProperty* InProperty, const T& InData, FTimestamp Timestamp)
 			: FContextCacheElementBase(InProperty, Timestamp)
 			, Data(InData)
 		{}
 
-		FContextCacheElement(FProperty* InProperty, T&& InData, uint64 Timestamp)
+		FContextCacheElement(FProperty* InProperty, T&& InData, FTimestamp Timestamp)
 			: FContextCacheElementBase(InProperty, Timestamp)
 			, Data(InData)
 		{}
@@ -68,22 +77,17 @@ namespace Dataflow
 
 
 	public:
-		FContext(float InTime, FString InType = FString(""))
-			: Timestamp(InTime)
+		FContext(FTimestamp InTimestamp, FString InType = FString(""))
+			: Timestamp(InTimestamp)
 			, Type(StaticType().Append(InType))
 		{}
 
 		virtual ~FContext() {}
 		
-		float Timestamp = 0.f;
+		FTimestamp Timestamp = FTimestamp::Invalid;
 		FString Type;
 		static FString StaticType() { return "FContext"; }
 
-		uint32 GetTypeHash() const
-		{
-			return ::GetTypeHash(Timestamp);
-		}
-
 		template<class T>
 		const T* AsType() const
 		{
@@ -128,14 +132,15 @@ namespace Dataflow
 		}
 
 		
-		virtual bool HasDataImpl(int64 Key, uint64 StoredAfter = 0) = 0;
+		virtual bool HasDataImpl(int64 Key, FTimestamp InTimestamp = FTimestamp::Invalid) = 0;
 		
-		bool HasData(size_t Key, uint64 StoredAfter = 0)
+		bool HasData(size_t Key, FTimestamp InTimestamp = FTimestamp::Invalid)
 		{
 			int64 IntKey = (int64)Key;
-			return HasDataImpl(Key, StoredAfter);
+			return HasDataImpl(Key, InTimestamp);
 		}
 		
+		virtual void Evaluate(const FDataflowNode* Node, const FDataflowOutput* Output) = 0;
 		virtual bool Evaluate(const FDataflowOutput& Connection) = 0;
 	};
 
@@ -145,7 +150,7 @@ namespace Dataflow
 
 	public:
 
-		FContextSingle(float InTime, FString InType = FString(""))
+		FContextSingle(FTimestamp InTime, FString InType = FString(""))
 			: FContext(InTime, InType)
 		{}
 
@@ -159,11 +164,12 @@ namespace Dataflow
 			return DataStore.Find(Key);
 		}
 
-		virtual bool HasDataImpl(int64 Key, uint64 StoredAfter = 0) override
+		virtual bool HasDataImpl(int64 Key, FTimestamp InTimestamp = FTimestamp::Invalid) override
 		{
-			return DataStore.Contains(Key) && DataStore[Key]->Timestamp >= StoredAfter;
+			return DataStore.Contains(Key) && DataStore[Key]->Timestamp >= InTimestamp;
 		}
 
+		virtual void Evaluate(const FDataflowNode* Node, const FDataflowOutput* Output) override;
 		virtual bool Evaluate(const FDataflowOutput& Connection) override;
 	};
 	
@@ -174,7 +180,7 @@ namespace Dataflow
 
 	public:
 
-		FContextThreaded(float InTime, FString InType = FString(""))
+		FContextThreaded(FTimestamp InTime, FString InType = FString(""))
 			: FContext(InTime, InType)
 		{
 			CacheLock = MakeShared<FCriticalSection>();
@@ -194,19 +200,15 @@ namespace Dataflow
 			return DataStore.Find(Key);
 		}
 
-		virtual bool HasDataImpl(int64 Key, uint64 StoredAfter = 0) override
+		virtual bool HasDataImpl(int64 Key, FTimestamp InTimestamp = FTimestamp::Invalid) override
 		{
 			CacheLock->Lock(); ON_SCOPE_EXIT { CacheLock->Unlock(); };
 			
-			return DataStore.Contains(Key) && DataStore[Key]->Timestamp >= StoredAfter;
+			return DataStore.Contains(Key) && DataStore[Key]->Timestamp >= InTimestamp;
 		}
 
+		virtual void Evaluate(const FDataflowNode* Node, const FDataflowOutput* Output) override;
 		virtual bool Evaluate(const FDataflowOutput& Connection) override;
 	};
 
 }
-
-FORCEINLINE uint32 GetTypeHash(const Dataflow::FContext& Context)
-{
-	return ::GetTypeHash(Context.Timestamp);
-}
diff --git a/Engine/Source/Runtime/Experimental/Dataflow/Engine/Public/Dataflow/DataflowObjectInterface.h b/Engine/Source/Runtime/Experimental/Dataflow/Engine/Public/Dataflow/DataflowObjectInterface.h
--- a/Engine/Source/Runtime/Experimental/Dataflow/Engine/Public/Dataflow/DataflowObjectInterface.h
+++ b/Engine/Source/Runtime/Experimental/Dataflow/Engine/Public/Dataflow/DataflowObjectInterface.h
@@ -16,8 +16,8 @@ namespace Dataflow
 	class DATAFLOWENGINE_API TEngineContext : public Base
 	{
 	public:
-		TEngineContext(UObject* InOwner, UDataflow* InGraph, float InTime, FString InType)
-				: Base(InTime, StaticType().Append(InType))
+		TEngineContext(UObject* InOwner, UDataflow* InGraph, FTimestamp InTimestamp, FString InType)
+				: Base(InTimestamp, StaticType().Append(InType))
 				, Owner(InOwner)
 				, Graph(InGraph)
 		{}
