Dataflow Graph:
- Added GetAttributeType() to ManagedArrayCollection

diff --git a/Engine/Source/Runtime/Experimental/Chaos/Private/GeometryCollection/ManagedArrayCollection.cpp b/Engine/Source/Runtime/Experimental/Chaos/Private/GeometryCollection/ManagedArrayCollection.cpp
--- a/Engine/Source/Runtime/Experimental/Chaos/Private/GeometryCollection/ManagedArrayCollection.cpp
+++ b/Engine/Source/Runtime/Experimental/Chaos/Private/GeometryCollection/ManagedArrayCollection.cpp
@@ -106,6 +106,15 @@ bool FManagedArrayCollection::HasAttributes(const TArray<FManagedArrayCollection
 	return true;
 }
 
+FManagedArrayCollection::EArrayType FManagedArrayCollection::GetAttributeType(FName Name, FName Group) const
+{
+	const FKeyType Key = FManagedArrayCollection::MakeMapKey(Name, Group);
+	if (const FValueType* Attribute = Map.Find(Key))
+	{
+		return Attribute->ArrayType;
+	}
+	return EArrayType::FNoneType;
+}
 
 bool FManagedArrayCollection::IsAttributeDirty(FName Name, FName Group) const
 {
diff --git a/Engine/Source/Runtime/Experimental/Chaos/Public/GeometryCollection/ManagedArrayCollection.h b/Engine/Source/Runtime/Experimental/Chaos/Public/GeometryCollection/ManagedArrayCollection.h
--- a/Engine/Source/Runtime/Experimental/Chaos/Public/GeometryCollection/ManagedArrayCollection.h
+++ b/Engine/Source/Runtime/Experimental/Chaos/Public/GeometryCollection/ManagedArrayCollection.h
@@ -415,6 +415,13 @@ public:
 	*/
 	FORCEINLINE bool HasGroup(FName Group) const { return GroupInfo.Contains(Group); }
 
+	/**
+	* Return attribute type
+	* @param Name - The name of the attribute
+	* @param Group - The group that manages the attribute
+	*/
+	EArrayType GetAttributeType(FName Name, FName Group) const;
+
 	/**
 	* Check if an attribute is dirty
 	* @param Name - The name of the attribute
