Dataflow : Evaluation
Context persistence with invalidation.

diff --git a/Engine/Plugins/Experimental/Dataflow/Source/DataflowEditor/Private/Dataflow/DataflowEditorToolkit.cpp b/Engine/Plugins/Experimental/Dataflow/Source/DataflowEditor/Private/Dataflow/DataflowEditorToolkit.cpp
--- a/Engine/Plugins/Experimental/Dataflow/Source/DataflowEditor/Private/Dataflow/DataflowEditorToolkit.cpp
+++ b/Engine/Plugins/Experimental/Dataflow/Source/DataflowEditor/Private/Dataflow/DataflowEditorToolkit.cpp
@@ -130,10 +130,8 @@ TSharedPtr<IStructureDetailsView> FDataflowEditorToolkit::CreateNodeDetailsEdito
 	DetailsView->GetDetailsView()->SetObject(ObjectToEdit);
 	DetailsView->GetOnFinishedChangingPropertiesDelegate().AddLambda([this](const FPropertyChangedEvent& PropertyChangedEvent) 
 	{
-			FMessageDialog::Open(EAppMsgType::Ok, NSLOCTEXT("Dataflow", "UnknownEvaluationContext_Dataflow","Property Changed"));
 			if (PropertyChangedEvent.ChangeType == EPropertyChangeType::ValueSet)
 			{
-				FMessageDialog::Open(EAppMsgType::Ok, NSLOCTEXT("Dataflow", "UnknownEvaluationContext_Dataflow", "Property Changed really"));
 			}
 	});
 
diff --git a/Engine/Source/Runtime/Experimental/Dataflow/Core/Private/Dataflow/DataflowNode.cpp b/Engine/Source/Runtime/Experimental/Dataflow/Core/Private/Dataflow/DataflowNode.cpp
--- a/Engine/Source/Runtime/Experimental/Dataflow/Core/Private/Dataflow/DataflowNode.cpp
+++ b/Engine/Source/Runtime/Experimental/Dataflow/Core/Private/Dataflow/DataflowNode.cpp
@@ -160,6 +160,12 @@ FDataflowOutput* FDataflowNode::FindOutput(void* Reference)
 	return nullptr;
 }
 
+int32 FDataflowNode::NumOutputs() const
+{
+	return Outputs.Num();
+}
+
+
 TArray< FDataflowOutput* > FDataflowNode::GetOutputs() const
 {
 	TArray< FDataflowOutput* > Result;
diff --git a/Engine/Source/Runtime/Experimental/Dataflow/Core/Private/Dataflow/DataflowNodeParameters.cpp b/Engine/Source/Runtime/Experimental/Dataflow/Core/Private/Dataflow/DataflowNodeParameters.cpp
--- a/Engine/Source/Runtime/Experimental/Dataflow/Core/Private/Dataflow/DataflowNodeParameters.cpp
+++ b/Engine/Source/Runtime/Experimental/Dataflow/Core/Private/Dataflow/DataflowNodeParameters.cpp
@@ -8,15 +8,37 @@ namespace Dataflow
 {
 	uint64 FTimestamp::Invalid = 0;
 
-	void FContextSingle::Evaluate(const FDataflowNode* Node, const FDataflowOutput* Output)
+	void BeginContextEvaluation(FContext& Context, const FDataflowNode* Node, const FDataflowOutput* Output)
 	{
 		if (Node != nullptr)
 		{
-			Timestamp = FTimestamp(FPlatformTime::Cycles64());
-			Node->Evaluate(*this, Output);
+			Context.Timestamp = FTimestamp(FPlatformTime::Cycles64());
+			if (Node->NumOutputs())
+			{
+				if (Output)
+				{
+					Node->Evaluate(Context, Output);
+				}
+				else
+				{
+					for (FDataflowOutput* NodeOutput : Node->GetOutputs())
+					{
+						Node->Evaluate(Context, NodeOutput);
+					}
+				}
+			}
+			else
+			{
+				Node->Evaluate(Context, nullptr);
+			}
 		}
 	}
 	
+	void FContextSingle::Evaluate(const FDataflowNode* Node, const FDataflowOutput* Output)
+	{
+		BeginContextEvaluation(*this, Node, Output);
+	}
+		
 	bool FContextSingle::Evaluate(const FDataflowOutput& Connection)
 	{
 		return Connection.EvaluateImpl(*this);
@@ -26,11 +48,7 @@ namespace Dataflow
 
 	void FContextThreaded::Evaluate(const FDataflowNode* Node, const FDataflowOutput* Output)
 	{
-		if (Node != nullptr)
-		{
-			Timestamp = FTimestamp(FPlatformTime::Cycles64());
-			Node->Evaluate(*this, Output);
-		}
+		BeginContextEvaluation(*this, Node, Output);
 	}
 
 	bool FContextThreaded::Evaluate(const FDataflowOutput& Connection)
diff --git a/Engine/Source/Runtime/Experimental/Dataflow/Core/Public/Dataflow/DataflowGraph.h b/Engine/Source/Runtime/Experimental/Dataflow/Core/Public/Dataflow/DataflowGraph.h
--- a/Engine/Source/Runtime/Experimental/Dataflow/Core/Public/Dataflow/DataflowGraph.h
+++ b/Engine/Source/Runtime/Experimental/Dataflow/Core/Public/Dataflow/DataflowGraph.h
@@ -85,6 +85,30 @@ namespace Dataflow
 			return TSharedPtr<FDataflowNode>(nullptr);
 		}
 
+		TSharedPtr<FDataflowNode> FindBaseNode(FName InName)
+		{
+			for (TSharedPtr<FDataflowNode> Node : Nodes)
+			{
+				if (Node->GetName().IsEqual(InName))
+				{
+					return Node;
+				}
+			}
+			return TSharedPtr<FDataflowNode>(nullptr);
+		}
+
+		TSharedPtr<const FDataflowNode> FindBaseNode(FName InName) const
+		{
+			for (TSharedPtr<FDataflowNode> Node : Nodes)
+			{
+				if (Node->GetName().IsEqual(InName))
+				{
+					return Node;
+				}
+			}
+			return TSharedPtr<const FDataflowNode>(nullptr);
+		}
+
 		void RemoveNode(TSharedPtr<FDataflowNode> Node);
 
 		void ClearConnections(FDataflowConnection* ConnectionBase);
diff --git a/Engine/Source/Runtime/Experimental/Dataflow/Core/Public/Dataflow/DataflowNode.h b/Engine/Source/Runtime/Experimental/Dataflow/Core/Public/Dataflow/DataflowNode.h
--- a/Engine/Source/Runtime/Experimental/Dataflow/Core/Public/Dataflow/DataflowNode.h
+++ b/Engine/Source/Runtime/Experimental/Dataflow/Core/Public/Dataflow/DataflowNode.h
@@ -68,6 +68,7 @@ struct DATAFLOWCORE_API FDataflowNode
 	FGuid GetGuid() const { return Guid; }
 	FName GetName() const { return Name; }
 	void SetName(FName InName) { Name = InName; }
+	Dataflow::FTimestamp GetTimestamp() const { return LastModifiedTimestamp;  }
 
 	virtual FName GetType() const { check(true); return FName("invalid"); }
 	virtual FName GetDisplayName() const { return ""; }
@@ -92,6 +93,7 @@ struct DATAFLOWCORE_API FDataflowNode
 
 
 	void AddOutput(FDataflowOutput* InPtr);
+	int NumOutputs() const;
 	TArray< FDataflowOutput* > GetOutputs() const;
 	void ClearOutputs();
 
diff --git a/Engine/Source/Runtime/Experimental/Dataflow/Core/Public/Dataflow/DataflowNodeParameters.h b/Engine/Source/Runtime/Experimental/Dataflow/Core/Public/Dataflow/DataflowNodeParameters.h
--- a/Engine/Source/Runtime/Experimental/Dataflow/Core/Public/Dataflow/DataflowNodeParameters.h
+++ b/Engine/Source/Runtime/Experimental/Dataflow/Core/Public/Dataflow/DataflowNodeParameters.h
@@ -140,6 +140,7 @@ namespace Dataflow
 			return HasDataImpl(Key, InTimestamp);
 		}
 		
+		FTimestamp GetTimestamp() const { return Timestamp; }
 		virtual void Evaluate(const FDataflowNode* Node, const FDataflowOutput* Output) = 0;
 		virtual bool Evaluate(const FDataflowOutput& Connection) = 0;
 	};
