Workaround for access to null pointer in GeometryCollectionPhysicsProxy

diff --git a/Engine/Source/Runtime/Experimental/Chaos/Private/PhysicsProxy/GeometryCollectionPhysicsProxy.cpp b/Engine/Source/Runtime/Experimental/Chaos/Private/PhysicsProxy/GeometryCollectionPhysicsProxy.cpp
--- a/Engine/Source/Runtime/Experimental/Chaos/Private/PhysicsProxy/GeometryCollectionPhysicsProxy.cpp
+++ b/Engine/Source/Runtime/Experimental/Chaos/Private/PhysicsProxy/GeometryCollectionPhysicsProxy.cpp
@@ -470,7 +470,7 @@ void FGeometryCollectionPhysicsProxy::Initialize(Chaos::FPBDRigidsEvolutionBase
 			P->SetProxy(this);
 
 			FGeometryDynamicCollection::FSharedImplicit ImplicitGeometry = GameThreadCollection.Implicits[Index];
-			if (!Scale.Equals(FVector::OneVector))
+			if (ImplicitGeometry && !Scale.Equals(FVector::OneVector))
 			{
 				TUniquePtr<Chaos::FImplicitObject> ScaledImplicit = ImplicitGeometry->CopyWithScale(Scale);
 				ImplicitGeometry = FGeometryDynamicCollection::FSharedImplicit(ScaledImplicit.Release());
