Chaos : temp fix for replication crash
- comment the disable of particles beyond the abandon level  on the server

diff --git a/Engine/Source/Runtime/Experimental/GeometryCollectionEngine/Private/GeometryCollection/GeometryCollectionComponent.cpp b/Engine/Source/Runtime/Experimental/GeometryCollectionEngine/Private/GeometryCollection/GeometryCollectionComponent.cpp
--- a/Engine/Source/Runtime/Experimental/GeometryCollectionEngine/Private/GeometryCollection/GeometryCollectionComponent.cpp
+++ b/Engine/Source/Runtime/Experimental/GeometryCollectionEngine/Private/GeometryCollection/GeometryCollectionComponent.cpp
@@ -1574,7 +1574,8 @@ void UGeometryCollectionComponent::UpdateRepData()
 					{
 						if (!Root->Disabled())
 						{
-							Solver->GetEvolution()->DisableParticle(Root);
+							// disable this code for the time being because this may cause crashes
+							// Solver->GetEvolution()->DisableParticle(Root);
 						}
 					}
 				}
