Chaos : remove the need to pass a parent for GC when adding particle to the island graph
This fixes crashes related  to GC replication where an island index inherited from a parent would remain set after the island has been removed

diff --git a/Engine/Source/Runtime/Experimental/Chaos/Private/Chaos/Island/IslandManager.cpp b/Engine/Source/Runtime/Experimental/Chaos/Private/Chaos/Island/IslandManager.cpp
--- a/Engine/Source/Runtime/Experimental/Chaos/Private/Chaos/Island/IslandManager.cpp
+++ b/Engine/Source/Runtime/Experimental/Chaos/Private/Chaos/Island/IslandManager.cpp
@@ -343,33 +343,13 @@ void FPBDIslandManager::GraphEdgeRemoved(const FConstraintHandleHolder& Constrai
 	ConstraintHandle->SetConstraintGraphIndex(INDEX_NONE);
 }
 
-int32 FPBDIslandManager::AddParticle(FGeometryParticleHandle* ParticleHandle, const FGeometryParticleHandle* ParentParticleHandle)
+int32 FPBDIslandManager::AddParticle(FGeometryParticleHandle* ParticleHandle)
 {
 #if CHAOS_CONSTRAINTHANDLE_DEBUG_ENABLED
 	ensure(!FConstGenericParticleHandle(ParticleHandle)->Disabled());
 #endif
 
-	int32 IslandIndex = INDEX_NONE;
-
-	// If we have a parent, assign the child to the same island in the same sleep state
-	if ((ParticleHandle != nullptr) && (ParentParticleHandle != nullptr))
-	{
-		// We are only adding the child particle to the graph if the parent is dynamic
-		const FPBDRigidParticleHandle* ParentPBDRigid = ParentParticleHandle->CastToRigidParticle();
-		if (ParentPBDRigid && ParentPBDRigid->ObjectState() == EObjectStateType::Dynamic)
-		{
-			FPBDRigidParticleHandle* ChildPBDRigid = ParticleHandle->CastToRigidParticle();
-			if (ChildPBDRigid && IsDynamicParticle(ChildPBDRigid))
-			{
-				// If the child particle is dynamic or sleeping we place the child in
-				// the same island as the parent with the same sleep state
-				IslandIndex = ParentPBDRigid->IslandIndex();
-				ChildPBDRigid->SetSleeping(ParentPBDRigid->Sleeping());
-			}
-		}
-	}
-
-	return TryAddParticleIfDynamic(ParticleHandle, IslandIndex);
+	return TryAddParticleIfDynamic(ParticleHandle,  /*IslandIndex*/INDEX_NONE);
 }
 
 int32 FPBDIslandManager::TryAddParticle(FGeometryParticleHandle* ParticleHandle, const int32 IslandIndex)
diff --git a/Engine/Source/Runtime/Experimental/Chaos/Private/Chaos/PBDRigidClustering.cpp b/Engine/Source/Runtime/Experimental/Chaos/Private/Chaos/PBDRigidClustering.cpp
--- a/Engine/Source/Runtime/Experimental/Chaos/Private/Chaos/PBDRigidClustering.cpp
+++ b/Engine/Source/Runtime/Experimental/Chaos/Private/Chaos/PBDRigidClustering.cpp
@@ -194,7 +194,7 @@ namespace Chaos
 		{
 			NewParticle = MEvolution.CreateClusteredParticles(1)[0]; // calls Evolution.DirtyParticle()
 		}
-		MEvolution.EnableParticle(NewParticle, Parent);
+		MEvolution.EnableParticle(NewParticle);
 
 		NewParticle->SetCollisionGroup(INT_MAX);
 		TopLevelClusterParents.Add(NewParticle);
@@ -478,7 +478,7 @@ namespace Chaos
 		{
 			FPBDRigidClusteredParticleHandle* ClusteredChild = Child->CastToClustered();
 
-			MEvolution.EnableParticle(Child, ClusteredParent);
+			MEvolution.EnableParticle(Child);
 			TopLevelClusterParents.Add(ClusteredChild);
 
 			ClusteredChild->SetClusterId(ClusterId(nullptr, ClusteredChild->ClusterIds().NumChildren)); // clear Id but retain number of children
diff --git a/Engine/Source/Runtime/Experimental/Chaos/Private/Chaos/ParticleHandle.cpp b/Engine/Source/Runtime/Experimental/Chaos/Private/Chaos/ParticleHandle.cpp
--- a/Engine/Source/Runtime/Experimental/Chaos/Private/Chaos/ParticleHandle.cpp
+++ b/Engine/Source/Runtime/Experimental/Chaos/Private/Chaos/ParticleHandle.cpp
@@ -249,7 +249,7 @@ namespace Chaos
 			}
 			else
 			{
-				Evolution.EnableParticle(this, nullptr);
+				Evolution.EnableParticle(this);
 			}
 		}		
 
diff --git a/Engine/Source/Runtime/Experimental/Chaos/Private/PhysicsProxy/SingleParticlePhysicsProxy.cpp b/Engine/Source/Runtime/Experimental/Chaos/Private/PhysicsProxy/SingleParticlePhysicsProxy.cpp
--- a/Engine/Source/Runtime/Experimental/Chaos/Private/PhysicsProxy/SingleParticlePhysicsProxy.cpp
+++ b/Engine/Source/Runtime/Experimental/Chaos/Private/PhysicsProxy/SingleParticlePhysicsProxy.cpp
@@ -52,7 +52,7 @@ void PushToPhysicsStateImp(const Chaos::FDirtyPropertiesManager& Manager, Chaos:
 
 	if (bResimInitialized)	//todo: assumes particles are always initialized as enabled. This is not true in future versions of code, so check PushData
 	{
-		Evolution.EnableParticle(Handle, nullptr);
+		Evolution.EnableParticle(Handle);
 	}
 	// move the copied game thread data into the handle
 	{
diff --git a/Engine/Source/Runtime/Experimental/Chaos/Private/RewindData.cpp b/Engine/Source/Runtime/Experimental/Chaos/Private/RewindData.cpp
--- a/Engine/Source/Runtime/Experimental/Chaos/Private/RewindData.cpp
+++ b/Engine/Source/Runtime/Experimental/Chaos/Private/RewindData.cpp
@@ -559,7 +559,7 @@ void FRewindData::SpawnProxyIfNeeded(FSingleParticlePhysicsProxy& Proxy)
 		FGeometryParticleHandle* Handle = Proxy.GetHandle_LowLevel();
 		FDirtyParticleInfo& Info = FindOrAddDirtyObj(*Handle, CurFrame);
 
-		Solver->GetEvolution()->EnableParticle(Handle, nullptr);
+		Solver->GetEvolution()->EnableParticle(Handle);
 		if(Proxy.GetInitializedStep() != CurFrame)
 		{
 			DesyncObject(Info, FFrameAndPhase{ Proxy.GetInitializedStep(), FFrameAndPhase::PrePushData });	//Spawned earlier so mark as desynced from that first frame
diff --git a/Engine/Source/Runtime/Experimental/Chaos/Public/Chaos/Island/IslandManager.h b/Engine/Source/Runtime/Experimental/Chaos/Public/Chaos/Island/IslandManager.h
--- a/Engine/Source/Runtime/Experimental/Chaos/Public/Chaos/Island/IslandManager.h
+++ b/Engine/Source/Runtime/Experimental/Chaos/Public/Chaos/Island/IslandManager.h
@@ -141,7 +141,7 @@ public:
 	  * @param ParentParticleHandle Put the ParticleHandle in the same island as the parent with the same sleep state. May be null.
 	  * @return Sparse graph node index
 	  */
-	int32 AddParticle(FGeometryParticleHandle* ParticleHandle, const FGeometryParticleHandle* ParentParticleHandle = nullptr);
+	int32 AddParticle(FGeometryParticleHandle* ParticleHandle);
 
 	/**
 	  * Remove a particle from the graph
diff --git a/Engine/Source/Runtime/Experimental/Chaos/Public/Chaos/PBDRigidsEvolution.h b/Engine/Source/Runtime/Experimental/Chaos/Public/Chaos/PBDRigidsEvolution.h
--- a/Engine/Source/Runtime/Experimental/Chaos/Public/Chaos/PBDRigidsEvolution.h
+++ b/Engine/Source/Runtime/Experimental/Chaos/Public/Chaos/PBDRigidsEvolution.h
@@ -446,13 +446,12 @@ public:
 	/**
 	* Enable a particle.Only enabled particles are simulated.
 	* If the particle has constraints connected to it they will also be enabled (assuming the other particles in the constraints are also enabled). 
-	* If a parent particle is specified, the child particle will be placed in the same island as the parent with the same sleep state.
 	*/
-	CHAOS_API void EnableParticle(FGeometryParticleHandle* Particle, const FGeometryParticleHandle* ParentParticle = nullptr)
+	CHAOS_API void EnableParticle(FGeometryParticleHandle* Particle)
 	{
 		Particles.EnableParticle(Particle);
 		EnableConstraints(Particle);
-		ConstraintGraph.AddParticle(Particle, ParentParticle);
+		ConstraintGraph.AddParticle(Particle);
 		DirtyParticle(*Particle);
 	}
 
diff --git a/Engine/Source/Runtime/Experimental/Chaos/Public/PhysicsProxy/FieldSystemProxyHelper.h b/Engine/Source/Runtime/Experimental/Chaos/Public/PhysicsProxy/FieldSystemProxyHelper.h
--- a/Engine/Source/Runtime/Experimental/Chaos/Public/PhysicsProxy/FieldSystemProxyHelper.h
+++ b/Engine/Source/Runtime/Experimental/Chaos/Public/PhysicsProxy/FieldSystemProxyHelper.h
@@ -418,7 +418,7 @@ namespace Chaos
 					Chaos::FPBDRigidParticleHandle* RigidHandle = ParticleHandles[Index.Sample]->CastToRigidParticle();
 					if (RigidHandle && RigidHandle->Disabled() && ResultsView[Index.Result] == 0)
 					{
-						RigidSolver->GetEvolution()->EnableParticle(RigidHandle, nullptr);
+						RigidSolver->GetEvolution()->EnableParticle(RigidHandle);
 						SetParticleDynamicState(RigidSolver, Chaos::EObjectStateType::Dynamic, RigidHandle);
 					}
 				}
diff --git a/Engine/Source/Runtime/Experimental/GeometryCollectionEngine/Private/GeometryCollection/GeometryCollectionComponent.cpp b/Engine/Source/Runtime/Experimental/GeometryCollectionEngine/Private/GeometryCollection/GeometryCollectionComponent.cpp
--- a/Engine/Source/Runtime/Experimental/GeometryCollectionEngine/Private/GeometryCollection/GeometryCollectionComponent.cpp
+++ b/Engine/Source/Runtime/Experimental/GeometryCollectionEngine/Private/GeometryCollection/GeometryCollectionComponent.cpp
@@ -1574,8 +1574,7 @@ void UGeometryCollectionComponent::UpdateRepData()
 					{
 						if (!Root->Disabled())
 						{
-							// disable this code for the time being because this may cause crashes
-							// Solver->GetEvolution()->DisableParticle(Root);
+							Solver->GetEvolution()->DisableParticle(Root);
 						}
 					}
 				}
