Chaos : fix potential geometry collection replication crash
- Make sure CLustersToRep is reset when the physics state is destroyed
- Initialize the Physics proxy ReplicationMode when we create the proxy instead of during InitializeComponent ( because Destroy/Create physics state can be called after it )

diff --git a/Engine/Source/Runtime/Experimental/GeometryCollectionEngine/Private/GeometryCollection/GeometryCollectionComponent.cpp b/Engine/Source/Runtime/Experimental/GeometryCollectionEngine/Private/GeometryCollection/GeometryCollectionComponent.cpp
--- a/Engine/Source/Runtime/Experimental/GeometryCollectionEngine/Private/GeometryCollection/GeometryCollectionComponent.cpp
+++ b/Engine/Source/Runtime/Experimental/GeometryCollectionEngine/Private/GeometryCollection/GeometryCollectionComponent.cpp
@@ -1192,21 +1192,6 @@ void UGeometryCollectionComponent::SetRestState(TArray<FTransform>&& InRestTrans
 	RefreshEmbeddedGeometry();
 }
 
-FORCEINLINE_DEBUGGABLE int32 ComputeParticleLevel(Chaos::FPBDRigidClusteredParticleHandle* Particle)
-{
-	int32 Level = 0;
-	if (Particle)
-	{
-		Chaos::FPBDRigidClusteredParticleHandle* Current = Particle;
-		while (Current->Parent())
-		{
-			Current = Current->Parent();
-			++Level;
-		}
-	}
-	return Level;
-};
-
 void UGeometryCollectionComponent::InitializeComponent()
 {
 	Super::InitializeComponent();
@@ -1227,54 +1212,6 @@ void UGeometryCollectionComponent::InitializeComponent()
 		}
 		DynamicCollection->AddAttribute<uint8>("InternalClusterParentTypeArray", FTransformCollection::TransformGroup);
 	}
-
-	AActor* Owner = GetOwner();
-
-	if(!Owner)
-	{
-		return;
-	}
-
-	// If we're replicating we need some extra setup - check netmode as we don't need this for
-	// standalone runtimes where we aren't going to network the component
-	if(PhysicsProxy && GetIsReplicated())
-	{
-		const ENetRole LocalRole = Owner->GetLocalRole();
-		const ENetMode NetMode = Owner->GetNetMode();
-
-		const bool bIsReplicationAuthority = (LocalRole == ENetRole::ROLE_Authority);
-		// Client side : geometry collection children of parents below the rep level need to be infintely strong so that client cannot break it 
-		if (Chaos::FPhysicsSolver* CurrSolver = GetSolver(*this))
-		{
-			CurrSolver->RegisterSimOneShotCallback([Proxy = PhysicsProxy, bIsReplicationAuthority, AbandonAfterLevel = ReplicationAbandonAfterLevel, EnableAbandonAfterLevel = bEnableAbandonAfterLevel]()
-			{
-				if (bIsReplicationAuthority)
-				{
-					Proxy->SetReplicationMode(FGeometryCollectionPhysicsProxy::EReplicationMode::Server);
-				}
-				else
-				{
-					Proxy->SetReplicationMode(FGeometryCollectionPhysicsProxy::EReplicationMode::Client);
-
-					// As we're not in control we make it so our simulated proxy cannot break clusters
-					// We have to set the strain to a high value but be below the max for the data type
-					// so releasing on authority demand works
-					constexpr Chaos::FReal MaxStrain = TNumericLimits<Chaos::FReal>::Max() - TNumericLimits<Chaos::FReal>::Min();
-					for (Chaos::FPBDRigidClusteredParticleHandle* P : Proxy->GetParticles())
-					{
-						if (P)
-						{
-							const int32 Level = EnableAbandonAfterLevel ? ComputeParticleLevel(P) : -1;
-							if (Level <= AbandonAfterLevel + 1)	//we only replicate up until level X, but it means we should replicate the breaking event of level X+1 (but not X+1's positions)
-							{
-								P->SetStrain(MaxStrain);
-							}
-						}
-					}
-				}
-			});
-		}
-	}
 }
 
 #if WITH_EDITOR
@@ -2493,6 +2430,22 @@ void UGeometryCollectionComponent::OnCreatePhysicsState()
 	}
 }
 
+
+static FORCEINLINE_DEBUGGABLE int32 ComputeParticleLevel(Chaos::FPBDRigidClusteredParticleHandle* Particle)
+{
+	int32 Level = 0;
+	if (Particle)
+	{
+		Chaos::FPBDRigidClusteredParticleHandle* Current = Particle;
+		while (Current->Parent())
+		{
+			Current = Current->Parent();
+			++Level;
+		}
+	}
+	return Level;
+};
+
 void UGeometryCollectionComponent::RegisterAndInitializePhysicsProxy()
 {
 	FSimulationParameters SimulationParameters;
@@ -2567,9 +2520,49 @@ void UGeometryCollectionComponent::RegisterAndInitializePhysicsProxy()
 #endif
 	PhysicsProxy = new FGeometryCollectionPhysicsProxy(this, *DynamicCollection, SimulationParameters, InitialSimFilter, InitialQueryFilter, CollectorGuid);
 	PhysicsProxy->SetPostPhysicsSyncCallback([this]() { UpdateAttachedChildrenTransform(); }); 
+	if (GetIsReplicated())
+	{
+		const FGeometryCollectionPhysicsProxy::EReplicationMode ReplicationMode = 
+			(GetOwner()->GetLocalRole() == ENetRole::ROLE_Authority) 
+				? FGeometryCollectionPhysicsProxy::EReplicationMode::Server 
+				: FGeometryCollectionPhysicsProxy::EReplicationMode::Client;
+		PhysicsProxy->SetReplicationMode(ReplicationMode);
+	}
+
 	FPhysScene_Chaos* Scene = GetInnerChaosScene();
 	Scene->AddObject(this, PhysicsProxy);
 
+	// If we're replicating we need some extra setup - check netmode as we don't need this for standalone runtimes where we aren't going to network the component
+	// IMPORTANT this need to happen after the object is registered so this will garantee that the particles are properly created by the time the callback below gets called
+	if (GetIsReplicated())
+	{
+		// Client side : geometry collection children of parents below the rep level need to be infintely strong so that client cannot break it 
+		if (Chaos::FPhysicsSolver* CurrSolver = GetSolver(*this))
+		{
+			CurrSolver->EnqueueCommandImmediate([Proxy = PhysicsProxy, AbandonAfterLevel = ReplicationAbandonAfterLevel, EnableAbandonAfterLevel = bEnableAbandonAfterLevel]()
+				{
+					if (Proxy->GetReplicationMode() == FGeometryCollectionPhysicsProxy::EReplicationMode::Client)
+					{
+						// As we're not in control we make it so our simulated proxy cannot break clusters
+						// We have to set the strain to a high value but be below the max for the data type
+						// so releasing on authority demand works
+						constexpr Chaos::FReal MaxStrain = TNumericLimits<Chaos::FReal>::Max() - TNumericLimits<Chaos::FReal>::Min();
+						for (Chaos::FPBDRigidClusteredParticleHandle* ParticleHandle : Proxy->GetParticles())
+						{
+							if (ParticleHandle)
+							{
+								const int32 Level = EnableAbandonAfterLevel ? ComputeParticleLevel(ParticleHandle) : -1;
+								if (Level <= AbandonAfterLevel + 1)	//we only replicate up until level X, but it means we should replicate the breaking event of level X+1 (but not X+1's positions)
+								{
+									ParticleHandle->SetStrain(MaxStrain);
+								}
+							}
+						}
+					}
+				});
+		}
+	}
+
 	RegisterForEvents();
 	SetAsyncPhysicsTickEnabled(GetIsReplicated());
 }
@@ -2599,6 +2592,18 @@ void UGeometryCollectionComponent::OnDestroyPhysicsState()
 		// Discard the pointer (cleanup happens through the scene or dedicated thread)
 		PhysicsProxy = nullptr;
 	}
+
+	// make sure we also clear the replication data structures
+	if (Chaos::FPhysicsSolver* Solver = GetSolver(*this))
+	{
+		Solver->EnqueueCommandImmediate([this]()
+			{
+				if (ClustersToRep)
+				{
+					ClustersToRep.Reset();
+				}
+			});
+	}
 }
 
 void UGeometryCollectionComponent::SendRenderDynamicData_Concurrent()
