Chaos Cloth: Add damping to XPBDSpring and XPBDBending constraints.

Refactored PBDStiffness so that it has a new base class PBDWeightMap which can be used for the damping maps. The damping maps do not do the exponential conversion that PBDStiffness::ApplyValue has.

XPBD constraints are not available through the UI, so no change in behavior for users. As usual, the damping parameters exist at the ChaosClothingSimulationCloth and Constraints level, but they are not exposed to the Config, so they are defaulted to 0. Similarly, the map values exist at the Constraint level, but they are not populated in the UI Weight targets.

For both constraints, added Rayleigh damping following the methods described for XPBD. The exposed damping parameter is the DampingRatio (where 1 = critical damping). For efficiency, there are separate ISPC kernels for when DampingRatio > 0 vs no damping since additional computation is necessary for damping.

diff --git a/Engine/Plugins/Experimental/ChaosCloth/Source/Chaos/Private/ChaosCloth/ChaosClothConstraints.cpp b/Engine/Plugins/Experimental/ChaosCloth/Source/Chaos/Private/ChaosCloth/ChaosClothConstraints.cpp
--- a/Engine/Plugins/Experimental/ChaosCloth/Source/Chaos/Private/ChaosCloth/ChaosClothConstraints.cpp
+++ b/Engine/Plugins/Experimental/ChaosCloth/Source/Chaos/Private/ChaosCloth/ChaosClothConstraints.cpp
@@ -339,6 +339,22 @@ void FClothConstraints::SetEdgeConstraints(const TArray<TVec3<int32>>& SurfaceEl
 	++NumConstraintRules;
 }
 
+void FClothConstraints::SetXPBDEdgeConstraints(const TArray<TVec3<int32>>& SurfaceElements, const TConstArrayView<FRealSingle>& StiffnessMultipliers, const TConstArrayView<FRealSingle>& DampingRatioMultipliers)
+{
+	check(Evolution);
+	XEdgeConstraints = MakeShared<Softs::FXPBDSpringConstraints>(
+		Evolution->Particles(),
+		ParticleOffset, NumParticles,
+		SurfaceElements,
+		StiffnessMultipliers,
+		DampingRatioMultipliers,
+		/*InStiffness =*/ Softs::FSolverVec2::UnitVector,
+		/*InDampingRatio =*/ Softs::FSolverVec2::ZeroVector,
+		/*bTrimKinematicConstraints =*/ true);
+	++NumConstraintInits;  // Uses init to update the property tables
+	++NumConstraintRules;
+}
+
 void FClothConstraints::SetBendingConstraints(const TArray<TVec2<int32>>& Edges, const TConstArrayView<FRealSingle>& StiffnessMultipliers, bool bUseXPBDConstraints)
 {
 	check(Evolution);
@@ -382,7 +398,7 @@ void FClothConstraints::SetBendingConstraints(TArray<TVec4<int32>>&& BendingElem
 			BucklingStiffnessMultipliers,
 			/*InStiffness =*/ Softs::FSolverVec2::UnitVector,
 			/*InBucklingRatio=*/ (Softs::FSolverReal)0.f,
-			/*InBucklingStiffnessMultiplier =*/ Softs::FSolverVec2::UnitVector,
+			/*InBucklingStiffness =*/ Softs::FSolverVec2::UnitVector,
 			/*bTrimKinematicConstraints =*/ true);
 	}
 	else
@@ -395,7 +411,7 @@ void FClothConstraints::SetBendingConstraints(TArray<TVec4<int32>>&& BendingElem
 			BucklingStiffnessMultipliers,
 			/*InStiffness =*/ Softs::FSolverVec2::UnitVector,
 			/*InBucklingRatio=*/ (Softs::FSolverReal)0.f,
-			/*InBucklingStiffnessMultiplier =*/ Softs::FSolverVec2::UnitVector,
+			/*InBucklingStiffness =*/ Softs::FSolverVec2::UnitVector,
 			/*bTrimKinematicConstraints =*/ true);
 
 	}
@@ -403,6 +419,25 @@ void FClothConstraints::SetBendingConstraints(TArray<TVec4<int32>>&& BendingElem
 	++NumConstraintRules;
 }
 
+void FClothConstraints::SetXPBDBendingConstraints(TArray<TVec4<int32>>&& BendingElements, const TConstArrayView<FRealSingle>& StiffnessMultipliers, const TConstArrayView<FRealSingle>& BucklingStiffnessMultipliers, const TConstArrayView<FRealSingle>& DampingRatioMultipliers)
+{
+	check(Evolution);
+	XBendingElementConstraints = MakeShared<Softs::FXPBDBendingConstraints>(
+		Evolution->Particles(),
+		ParticleOffset, NumParticles,
+		MoveTemp(BendingElements),
+		StiffnessMultipliers,
+		BucklingStiffnessMultipliers,
+		DampingRatioMultipliers,
+		/*InStiffness =*/ Softs::FSolverVec2::UnitVector,
+		/*InBucklingRatio=*/ (Softs::FSolverReal)0.f,
+		/*InBucklingStiffness =*/ Softs::FSolverVec2::UnitVector,
+		/*InDampingRatio =*/ Softs::FSolverVec2::ZeroVector,
+		/*bTrimKinematicConstraints =*/ true);
+	++NumConstraintInits;  // Uses init to update the property tables
+	++NumConstraintRules;
+}
+
 void FClothConstraints::SetBendingConstraints(TArray<TVec4<int32>>&& BendingElements, Softs::FSolverReal BendingStiffness)
 {
 	check(Evolution);
@@ -565,7 +600,7 @@ void FClothConstraints::SetSelfCollisionConstraints(const FTriangleMesh& Triangl
 	++NumConstraintInits;
 }
 
-void FClothConstraints::SetEdgeProperties(const Softs::FSolverVec2& EdgeStiffness)
+void FClothConstraints::SetEdgeProperties(const Softs::FSolverVec2& EdgeStiffness, const Softs::FSolverVec2& DampingRatio)
 {
 	if (EdgeConstraints)
 	{
@@ -573,11 +608,11 @@ void FClothConstraints::SetEdgeProperties(const Softs::FSolverVec2& EdgeStiffnes
 	}
 	if (XEdgeConstraints)
 	{
-		XEdgeConstraints->SetProperties(EdgeStiffness);
+		XEdgeConstraints->SetProperties(EdgeStiffness, DampingRatio);
 	}
 }
 
-void FClothConstraints::SetBendingProperties(const Softs::FSolverVec2& BendingStiffness, Softs::FSolverReal BucklingRatio, const Softs::FSolverVec2& BucklingStiffness)
+void FClothConstraints::SetBendingProperties(const Softs::FSolverVec2& BendingStiffness, Softs::FSolverReal BucklingRatio, const Softs::FSolverVec2& BucklingStiffness, const Softs::FSolverVec2& BendingDampingRatio)
 {
 	if (BendingConstraints)
 	{
@@ -593,7 +628,7 @@ void FClothConstraints::SetBendingProperties(const Softs::FSolverVec2& BendingSt
 	}
 	if (XBendingElementConstraints)
 	{
-		XBendingElementConstraints->SetProperties(BendingStiffness, BucklingRatio, BucklingStiffness);
+		XBendingElementConstraints->SetProperties(BendingStiffness, BucklingRatio, BucklingStiffness, BendingDampingRatio);
 	}
 }
 
diff --git a/Engine/Plugins/Experimental/ChaosCloth/Source/Chaos/Private/ChaosCloth/ChaosClothingSimulation.cpp b/Engine/Plugins/Experimental/ChaosCloth/Source/Chaos/Private/ChaosCloth/ChaosClothingSimulation.cpp
--- a/Engine/Plugins/Experimental/ChaosCloth/Source/Chaos/Private/ChaosCloth/ChaosClothingSimulation.cpp
+++ b/Engine/Plugins/Experimental/ChaosCloth/Source/Chaos/Private/ChaosCloth/ChaosClothingSimulation.cpp
@@ -415,7 +415,9 @@ void FClothingSimulation::CreateActor(USkeletalMeshComponent* InOwnerComponent,
 		ClothConfig->bUseSelfIntersections,
 		ClothConfig->bUseLegacyBackstop,
 		/*bUseLODIndexOverride =*/ false,
-		/*LODIndexOverride =*/ INDEX_NONE));
+		/*LODIndexOverride =*/ INDEX_NONE,
+		TVec2<FRealSingle>(0.f), // EdgeDampingRatio
+		TVec2<FRealSingle>(0.f))); // BendDampingRatio
 
 	// Add cloth to solver
 	Solver->AddCloth(Cloths[ClothIndex].Get());
@@ -858,7 +860,9 @@ void FClothingSimulation::RefreshClothConfig(const IClothingSimulationContext* I
 			ClothConfig->bUseSelfIntersections,
 			ClothConfig->bUseLegacyBackstop,
 			/*bUseLODIndexOverride =*/ false,
-			/*LODIndexOverride =*/ INDEX_NONE);
+			/*LODIndexOverride =*/ INDEX_NONE,
+			TVec2<FRealSingle>(0.f), // EdgeDampingRatio
+			TVec2<FRealSingle>(0.f)); // BendDampingRatio
 
 		// Re-add cloth to the solver
 		Solver->AddCloth(Cloth.Get());
diff --git a/Engine/Plugins/Experimental/ChaosCloth/Source/Chaos/Private/ChaosCloth/ChaosClothingSimulationCloth.cpp b/Engine/Plugins/Experimental/ChaosCloth/Source/Chaos/Private/ChaosCloth/ChaosClothingSimulationCloth.cpp
--- a/Engine/Plugins/Experimental/ChaosCloth/Source/Chaos/Private/ChaosCloth/ChaosClothingSimulationCloth.cpp
+++ b/Engine/Plugins/Experimental/ChaosCloth/Source/Chaos/Private/ChaosCloth/ChaosClothingSimulationCloth.cpp
@@ -168,7 +168,15 @@ PRAGMA_ENABLE_DEPRECATION_WARNINGS
 	const TConstArrayView<FRealSingle>& EdgeStiffnessMultipliers = WeightMaps[(int32)EChaosWeightMapTarget::EdgeStiffness];
 	if (Cloth->EdgeStiffness[0] > (FRealSingle)0. || (Cloth->EdgeStiffness[1] > (FRealSingle)0. && EdgeStiffnessMultipliers.Num() == NumParticles))
 	{
-		ClothConstraints.SetEdgeConstraints(SurfaceElements, EdgeStiffnessMultipliers, Cloth->bUseXPBDEdgeConstraints);
+		if (Cloth->bUseXPBDEdgeConstraints)
+		{
+			const TConstArrayView<FRealSingle>& EdgeDampingRatioMultipliers = TConstArrayView<FRealSingle>(); // XPBD is not currently exposed to the UI.
+			ClothConstraints.SetXPBDEdgeConstraints(SurfaceElements, EdgeStiffnessMultipliers, EdgeDampingRatioMultipliers);
+		}
+		else
+		{
+			ClothConstraints.SetEdgeConstraints(SurfaceElements, EdgeStiffnessMultipliers, Cloth->bUseXPBDEdgeConstraints);
+		}
 	}
 
 	// Bending constraints
@@ -180,7 +188,15 @@ PRAGMA_ENABLE_DEPRECATION_WARNINGS
 		if (Cloth->bUseBendingElements)
 		{
 			TArray<Chaos::TVec4<int32>> BendingElements = TriangleMesh.GetUniqueAdjacentElements();
-			ClothConstraints.SetBendingConstraints(MoveTemp(BendingElements), BendingStiffnessMultipliers, BucklingStiffnessMultipliers, Cloth->bUseXPBDBendingConstraints);
+			if (Cloth->bUseXPBDBendingConstraints)
+			{
+				const TConstArrayView<FRealSingle>& BendingDampingRatioMultipliers = TConstArrayView<FRealSingle>(); // XPBD is not currently exposed to the UI.
+				ClothConstraints.SetXPBDBendingConstraints(MoveTemp(BendingElements), BendingStiffnessMultipliers, BucklingStiffnessMultipliers, BendingDampingRatioMultipliers);
+			}
+			else
+			{
+				ClothConstraints.SetBendingConstraints(MoveTemp(BendingElements), BendingStiffnessMultipliers, BucklingStiffnessMultipliers, Cloth->bUseXPBDBendingConstraints);
+			}
 		}
 		else
 		{
@@ -325,8 +341,11 @@ PRAGMA_ENABLE_DEPRECATION_WARNINGS
 	// Update the animatable constraint parameters
 	FClothConstraints& ClothConstraints = Solver->GetClothConstraints(Offset);
 	ClothConstraints.SetMaximumDistanceProperties((Softs::FSolverReal)Cloth->MaxDistancesMultiplier * MeshScale);
-	ClothConstraints.SetEdgeProperties(Softs::FSolverVec2((Softs::FSolverReal)Cloth->EdgeStiffness[0], (Softs::FSolverReal)Cloth->EdgeStiffness[1]));
-	ClothConstraints.SetBendingProperties(Softs::FSolverVec2((Softs::FSolverReal)Cloth->BendingStiffness[0], (Softs::FSolverReal)Cloth->BendingStiffness[1]), Cloth->BucklingRatio, Softs::FSolverVec2((Softs::FSolverReal)Cloth->BucklingStiffness[0], (Softs::FSolverReal)Cloth->BucklingStiffness[1]));
+	ClothConstraints.SetEdgeProperties(Softs::FSolverVec2((Softs::FSolverReal)Cloth->EdgeStiffness[0], (Softs::FSolverReal)Cloth->EdgeStiffness[1]), 
+		Softs::FSolverVec2((Softs::FSolverReal)Cloth->EdgeDampingRatio[0], (Softs::FSolverReal)Cloth->EdgeDampingRatio[1]));
+	ClothConstraints.SetBendingProperties(Softs::FSolverVec2((Softs::FSolverReal)Cloth->BendingStiffness[0], (Softs::FSolverReal)Cloth->BendingStiffness[1]), 
+		(Softs::FSolverReal)Cloth->BucklingRatio, Softs::FSolverVec2((Softs::FSolverReal)Cloth->BucklingStiffness[0], (Softs::FSolverReal)Cloth->BucklingStiffness[1]), 
+		Softs::FSolverVec2((Softs::FSolverReal)Cloth->BendingDampingRatio[0], (Softs::FSolverReal)Cloth->BendingDampingRatio[1]));
 	ClothConstraints.SetAreaProperties(Softs::FSolverVec2((FSolverReal)Cloth->AreaStiffness[0], (Softs::FSolverReal)Cloth->AreaStiffness[1]));
 	ClothConstraints.SetLongRangeAttachmentProperties(
 		Softs::FSolverVec2((Softs::FSolverReal)Cloth->TetherStiffness[0], (Softs::FSolverReal)Cloth->TetherStiffness[1]),
@@ -425,7 +444,9 @@ PRAGMA_ENABLE_DEPRECATION_WARNINGS
 	bool bInUseSelfIntersections,
 	bool bInUseLegacyBackstop,
 	bool bInUseLODIndexOverride, 
-	int32 InLODIndexOverride)
+	int32 InLODIndexOverride,
+	const TVec2<FRealSingle>& InEdgeDampingRatio,
+	const TVec2<FRealSingle>& InBendingDampingRatio)
 	: Mesh(nullptr)
 	, Colliders()
 	, GroupId(InGroupId)
@@ -433,7 +454,9 @@ PRAGMA_ENABLE_DEPRECATION_WARNINGS
 	, MassValue(InMassValue)
 	, MinPerParticleMass(InMinPerParticleMass)
 	, EdgeStiffness(InEdgeStiffness)
+	, EdgeDampingRatio(InEdgeDampingRatio)
 	, BendingStiffness(InBendingStiffness)
+	, BendingDampingRatio(InBendingDampingRatio)
 	, BucklingRatio(InBucklingRatio)
 	, BucklingStiffness(InBucklingStiffness)
 	, bUseBendingElements(bInUseBendingElements)
diff --git a/Engine/Plugins/Experimental/ChaosCloth/Source/Chaos/Public/ChaosCloth/ChaosClothConstraints.h b/Engine/Plugins/Experimental/ChaosCloth/Source/Chaos/Public/ChaosCloth/ChaosClothConstraints.h
--- a/Engine/Plugins/Experimental/ChaosCloth/Source/Chaos/Public/ChaosCloth/ChaosClothConstraints.h
+++ b/Engine/Plugins/Experimental/ChaosCloth/Source/Chaos/Public/ChaosCloth/ChaosClothConstraints.h
@@ -30,8 +30,10 @@ namespace Chaos
 
 		// ---- Cloth interface ----
 		void SetEdgeConstraints(const TArray<TVec3<int32>>& SurfaceElements, const TConstArrayView<FRealSingle>& StiffnessMultipliers, bool bUseXPBDConstraints);
+		void SetXPBDEdgeConstraints(const TArray<TVec3<int32>>& SurfaceElements, const TConstArrayView<FRealSingle>& StiffnessMultipliers, const TConstArrayView<FRealSingle>& DampingRatioMultipliers);
 		void SetBendingConstraints(const TArray<TVec2<int32>>& Edges, const TConstArrayView<FRealSingle>& StiffnessMultipliers, bool bUseXPBDConstraints);
 		void SetBendingConstraints(TArray<TVec4<int32>>&& BendingElements, const TConstArrayView<FRealSingle>& StiffnessMultipliers, const TConstArrayView<FRealSingle>& BucklingStiffnessMultipliers, bool bUseXPBDConstraints);
+		void SetXPBDBendingConstraints(TArray<TVec4<int32>>&& BendingElements, const TConstArrayView<FRealSingle>& StiffnessMultipliers, const TConstArrayView<FRealSingle>& BucklingStiffnessMultipliers, const TConstArrayView<FRealSingle>& DampingRatioMultipliers);
 		UE_DEPRECATED(5.1, "Use SetBendingConstraints(TArray<TVec4<int32>>&& BendingElements, const TConstArrayView<FRealSingle>& StiffnessMultipliers, const TConstArrayView<FRealSingle>& BucklingStiffnessMultipliers, bool bUseXPBDConstraints) instead.")
 		void SetBendingConstraints(TArray<TVec4<int32>>&& BendingElements, Softs::FSolverReal BendingStiffness);
 		void SetAreaConstraints(const TArray<TVec3<int32>>& SurfaceElements, const TConstArrayView<FRealSingle>& StiffnessMultipliers, bool bUseXPBDConstraints);
@@ -57,8 +59,8 @@ namespace Chaos
 		void CreateRules();
 		void Enable(bool bEnable);
 
-		void SetEdgeProperties(const Softs::FSolverVec2& EdgeStiffness);
-		void SetBendingProperties(const Softs::FSolverVec2& BendingStiffness, Softs::FSolverReal BucklingRatio = 0.f, const Softs::FSolverVec2& BucklingStiffnessMultiplier = Softs::FSolverVec2::UnitVector);
+		void SetEdgeProperties(const Softs::FSolverVec2& EdgeStiffness, const Softs::FSolverVec2& DampingRatio = Softs::FSolverVec2::ZeroVector);
+		void SetBendingProperties(const Softs::FSolverVec2& BendingStiffness, Softs::FSolverReal BucklingRatio = 0.f, const Softs::FSolverVec2& BucklingStiffness = Softs::FSolverVec2::UnitVector, const Softs::FSolverVec2& DampingRatio = Softs::FSolverVec2::ZeroVector);
 		void SetAreaProperties(const Softs::FSolverVec2& AreaStiffness);
 		void SetThinShellVolumeProperties(Softs::FSolverReal VolumeStiffness);
 		void SetVolumeProperties(Softs::FSolverReal VolumeStiffness);
diff --git a/Engine/Plugins/Experimental/ChaosCloth/Source/Chaos/Public/ChaosCloth/ChaosClothingSimulationCloth.h b/Engine/Plugins/Experimental/ChaosCloth/Source/Chaos/Public/ChaosCloth/ChaosClothingSimulationCloth.h
--- a/Engine/Plugins/Experimental/ChaosCloth/Source/Chaos/Public/ChaosCloth/ChaosClothingSimulationCloth.h
+++ b/Engine/Plugins/Experimental/ChaosCloth/Source/Chaos/Public/ChaosCloth/ChaosClothingSimulationCloth.h
@@ -74,7 +74,9 @@ PRAGMA_ENABLE_DEPRECATION_WARNINGS
 			bool bInUseSelfIntersections,
 			bool bInUseLegacyBackstop,
 			bool bInUseLODIndexOverride,
-			int32 InLODIndexOverride);
+			int32 InLODIndexOverride,
+			const TVec2<FRealSingle>& EdgeDampingRatio = TVec2<FRealSingle>(0.f),
+			const TVec2<FRealSingle>& BendDampingRatio = TVec2<FRealSingle>(0.f));
 		~FClothingSimulationCloth();
 
 		FClothingSimulationCloth(const FClothingSimulationCloth&) = delete;
@@ -177,7 +179,9 @@ PRAGMA_ENABLE_DEPRECATION_WARNINGS
 		FRealSingle MassValue;
 		FRealSingle MinPerParticleMass;
 		TVec2<FRealSingle> EdgeStiffness;
+		TVec2<FRealSingle> EdgeDampingRatio;
 		TVec2<FRealSingle> BendingStiffness;
+		TVec2<FRealSingle> BendingDampingRatio;
 		FRealSingle BucklingRatio;
 		TVec2<FRealSingle> BucklingStiffness;
 		bool bUseBendingElements;
diff --git a/Engine/Source/Runtime/Experimental/Chaos/Private/Chaos/XPBDBendingConstraints.cpp b/Engine/Source/Runtime/Experimental/Chaos/Private/Chaos/XPBDBendingConstraints.cpp
--- a/Engine/Source/Runtime/Experimental/Chaos/Private/Chaos/XPBDBendingConstraints.cpp
+++ b/Engine/Source/Runtime/Experimental/Chaos/Private/Chaos/XPBDBendingConstraints.cpp
@@ -68,7 +68,7 @@ void FXPBDBendingConstraints::InitColor(const FSolverParticles& InParticles)
 	}
 }
 
-void FXPBDBendingConstraints::ApplyHelper(FSolverParticles& Particles, const FSolverReal Dt, const int32 ConstraintIndex, const FSolverReal StiffnessValue, const FSolverReal BucklingValue) const
+void FXPBDBendingConstraints::ApplyHelper(FSolverParticles& Particles, const FSolverReal Dt, const int32 ConstraintIndex, const FSolverReal StiffnessValue, const FSolverReal BucklingValue, const FSolverReal DampingRatioValue) const
 {
 	const TVec4<int32>& Constraint = Constraints[ConstraintIndex];
 	const int32 i1 = Constraint[0];
@@ -82,13 +82,25 @@ void FXPBDBendingConstraints::ApplyHelper(FSolverParticles& Particles, const FSo
 		return;
 	}
 
+	const FSolverReal CombinedInvMass = Particles.InvM(i1) + Particles.InvM(i2) + Particles.InvM(i3) + Particles.InvM(i4);
+	const FSolverReal Damping = (FSolverReal)2.f * DampingRatioValue * FMath::Sqrt(BiphasicStiffnessValue / CombinedInvMass);
+
 	const FSolverReal Angle = CalcAngle(Particles.P(i1), Particles.P(i2), Particles.P(i3), Particles.P(i4));
 	const TStaticArray<FSolverVec3, 4> Grads = Base::GetGradients(Particles, ConstraintIndex);
 
+	const FSolverVec3 V1TimesDt = Particles.P(i1) - Particles.X(i1);
+	const FSolverVec3 V2TimesDt = Particles.P(i2) - Particles.X(i2);
+	const FSolverVec3 V3TimesDt = Particles.P(i3) - Particles.X(i3);
+	const FSolverVec3 V4TimesDt = Particles.P(i4) - Particles.X(i4);
+
 	FSolverReal& Lambda = Lambdas[ConstraintIndex];
 	const FSolverReal Alpha = (FSolverReal)1.f / (BiphasicStiffnessValue * Dt * Dt);
-	const FSolverReal Denom = Particles.InvM(i1) * Grads[0].SizeSquared() + Particles.InvM(i2) * Grads[1].SizeSquared() + Particles.InvM(i3) * Grads[2].SizeSquared() + Particles.InvM(i4) * Grads[3].SizeSquared() + Alpha;
-	const FSolverReal DLambda = (Angle - RestAngles[ConstraintIndex] - Alpha * Lambda) / Denom;
+	const FSolverReal Gamma = Alpha * Damping * Dt;
+
+	const FSolverReal DampingTerm = Gamma * (FSolverVec3::DotProduct(V1TimesDt, Grads[0]) + FSolverVec3::DotProduct(V2TimesDt, Grads[1]) + FSolverVec3::DotProduct(V3TimesDt, Grads[2]) + FSolverVec3::DotProduct(V4TimesDt, Grads[3]));
+
+	const FSolverReal Denom = ((FSolverReal)1.f + Gamma) * (Particles.InvM(i1) * Grads[0].SizeSquared() + Particles.InvM(i2) * Grads[1].SizeSquared() + Particles.InvM(i3) * Grads[2].SizeSquared() + Particles.InvM(i4) * Grads[3].SizeSquared()) + Alpha;
+	const FSolverReal DLambda = (Angle - RestAngles[ConstraintIndex] - Alpha * Lambda + DampingTerm) / Denom;
 
 	Particles.P(i1) -= DLambda * Particles.InvM(i1) * Grads[0];
 	Particles.P(i2) -= DLambda * Particles.InvM(i2) * Grads[1];
@@ -104,14 +116,16 @@ void FXPBDBendingConstraints::Apply(FSolverParticles& Particles, const FSolverRe
 
 	const bool StiffnessHasWeightMap = Stiffness.HasWeightMap();
 	const bool BucklingStiffnessHasWeightMap = BucklingStiffness.HasWeightMap();
+	const bool DampingHasWeightMap = DampingRatio.HasWeightMap();
 
 	if (ConstraintsPerColorStartIndex.Num() > 0 && Constraints.Num() > Chaos_XPBDBending_ParallelConstraintCount)
 	{
 		const int32 ConstraintColorNum = ConstraintsPerColorStartIndex.Num() - 1;
-		if (!StiffnessHasWeightMap && !BucklingStiffnessHasWeightMap)
+		if (!StiffnessHasWeightMap && !BucklingStiffnessHasWeightMap && !DampingHasWeightMap)
 		{
 			const FSolverReal ExpStiffnessValue = (FSolverReal)Stiffness;
 			const FSolverReal ExpBucklingValue = (FSolverReal)BucklingStiffness;
+			const FSolverReal DampingRatioValue = (FSolverReal)DampingRatio;
 
 			if (ExpStiffnessValue < XPBDBendMinStiffness && ExpBucklingValue < XPBDBendMinStiffness)
 			{
@@ -121,21 +135,45 @@ void FXPBDBendingConstraints::Apply(FSolverParticles& Particles, const FSolverRe
 #if INTEL_ISPC
 			if (bRealTypeCompatibleWithISPC && bChaos_XPBDBending_ISPC_Enabled)
 			{
-				for (int32 ConstraintColorIndex = 0; ConstraintColorIndex < ConstraintColorNum; ++ConstraintColorIndex)
+				if (DampingRatioValue > 0)
 				{
-					const int32 ColorStart = ConstraintsPerColorStartIndex[ConstraintColorIndex];
-					const int32 ColorSize = ConstraintsPerColorStartIndex[ConstraintColorIndex + 1] - ColorStart;
-					ispc::ApplyXPBDBendingConstraints(
-						(ispc::FVector4f*)Particles.GetPAndInvM().GetData(),
-						(ispc::FIntVector4*)&Constraints.GetData()[ColorStart],
-						&RestAngles.GetData()[ColorStart],
-						&IsBuckled.GetData()[ColorStart],
-						&Lambdas.GetData()[ColorStart],
-						Dt,
-						XPBDBendMinStiffness,
-						ExpStiffnessValue,
-						ExpBucklingValue,
-						ColorSize);
+					for (int32 ConstraintColorIndex = 0; ConstraintColorIndex < ConstraintColorNum; ++ConstraintColorIndex)
+					{
+						const int32 ColorStart = ConstraintsPerColorStartIndex[ConstraintColorIndex];
+						const int32 ColorSize = ConstraintsPerColorStartIndex[ConstraintColorIndex + 1] - ColorStart;
+						ispc::ApplyXPBDBendingConstraintsWithDamping(
+							(ispc::FVector4f*)Particles.GetPAndInvM().GetData(),
+							(const ispc::FVector3f*)Particles.X().GetData(),
+							(ispc::FIntVector4*)&Constraints.GetData()[ColorStart],
+							&RestAngles.GetData()[ColorStart],
+							&IsBuckled.GetData()[ColorStart],
+							&Lambdas.GetData()[ColorStart],
+							Dt,
+							XPBDBendMinStiffness,
+							ExpStiffnessValue,
+							ExpBucklingValue,
+							DampingRatioValue,
+							ColorSize);
+					}
+				}
+				else
+				{
+					for (int32 ConstraintColorIndex = 0; ConstraintColorIndex < ConstraintColorNum; ++ConstraintColorIndex)
+					{
+						const int32 ColorStart = ConstraintsPerColorStartIndex[ConstraintColorIndex];
+						const int32 ColorSize = ConstraintsPerColorStartIndex[ConstraintColorIndex + 1] - ColorStart;
+						ispc::ApplyXPBDBendingConstraints(
+							(ispc::FVector4f*)Particles.GetPAndInvM().GetData(),
+							(ispc::FIntVector4*)&Constraints.GetData()[ColorStart],
+							&RestAngles.GetData()[ColorStart],
+							&IsBuckled.GetData()[ColorStart],
+							&Lambdas.GetData()[ColorStart],
+							Dt,
+							XPBDBendMinStiffness,
+							ExpStiffnessValue,
+							ExpBucklingValue,
+							ColorSize);
+					}
 				}
 			}
 			else
@@ -148,7 +186,7 @@ void FXPBDBendingConstraints::Apply(FSolverParticles& Particles, const FSolverRe
 					PhysicsParallelFor(ColorSize, [&](const int32 Index)
 					{
 						const int32 ConstraintIndex = ColorStart + Index;
-						ApplyHelper(Particles, Dt, ConstraintIndex, ExpStiffnessValue, ExpBucklingValue);
+						ApplyHelper(Particles, Dt, ConstraintIndex, ExpStiffnessValue, ExpBucklingValue, DampingRatioValue);
 					});
 				}
 			}
@@ -158,25 +196,55 @@ void FXPBDBendingConstraints::Apply(FSolverParticles& Particles, const FSolverRe
 #if INTEL_ISPC
 			if (bRealTypeCompatibleWithISPC && bChaos_XPBDBending_ISPC_Enabled)
 			{
-				for (int32 ConstraintColorIndex = 0; ConstraintColorIndex < ConstraintColorNum; ++ConstraintColorIndex)
+				if (DampingHasWeightMap || (FSolverReal)DampingRatio > 0)
 				{
-					const int32 ColorStart = ConstraintsPerColorStartIndex[ConstraintColorIndex];
-					const int32 ColorSize = ConstraintsPerColorStartIndex[ConstraintColorIndex + 1] - ColorStart;
-					ispc::ApplyXPBDBendingConstraintsWithMaps(
-						(ispc::FVector4f*)Particles.GetPAndInvM().GetData(),
-						(ispc::FIntVector4*)&Constraints.GetData()[ColorStart],
-						&RestAngles.GetData()[ColorStart],
-						&IsBuckled.GetData()[ColorStart],
-						&Lambdas.GetData()[ColorStart],
-						Dt,
-						XPBDBendMinStiffness,
-						StiffnessHasWeightMap,
-						&Stiffness.GetIndices().GetData()[ColorStart],
-						&Stiffness.GetTable().GetData()[0],
-						BucklingStiffnessHasWeightMap,
-						&BucklingStiffness.GetIndices().GetData()[ColorStart],
-						&BucklingStiffness.GetTable().GetData()[0],
-						ColorSize);
+					for (int32 ConstraintColorIndex = 0; ConstraintColorIndex < ConstraintColorNum; ++ConstraintColorIndex)
+					{
+						const int32 ColorStart = ConstraintsPerColorStartIndex[ConstraintColorIndex];
+						const int32 ColorSize = ConstraintsPerColorStartIndex[ConstraintColorIndex + 1] - ColorStart;
+						ispc::ApplyXPBDBendingConstraintsWithDampingAndMaps(
+							(ispc::FVector4f*)Particles.GetPAndInvM().GetData(),
+							(const ispc::FVector3f*)Particles.X().GetData(),
+							(ispc::FIntVector4*)&Constraints.GetData()[ColorStart],
+							&RestAngles.GetData()[ColorStart],
+							&IsBuckled.GetData()[ColorStart],
+							&Lambdas.GetData()[ColorStart],
+							Dt,
+							XPBDBendMinStiffness,
+							StiffnessHasWeightMap,
+							&Stiffness.GetIndices().GetData()[ColorStart],
+							&Stiffness.GetTable().GetData()[0],
+							BucklingStiffnessHasWeightMap,
+							&BucklingStiffness.GetIndices().GetData()[ColorStart],
+							&BucklingStiffness.GetTable().GetData()[0],
+							DampingHasWeightMap,
+							&DampingRatio.GetIndices().GetData()[ColorStart],
+							&DampingRatio.GetTable().GetData()[0],
+							ColorSize);
+					}
+				}
+				else
+				{
+					for (int32 ConstraintColorIndex = 0; ConstraintColorIndex < ConstraintColorNum; ++ConstraintColorIndex)
+					{
+						const int32 ColorStart = ConstraintsPerColorStartIndex[ConstraintColorIndex];
+						const int32 ColorSize = ConstraintsPerColorStartIndex[ConstraintColorIndex + 1] - ColorStart;
+						ispc::ApplyXPBDBendingConstraintsWithMaps(
+							(ispc::FVector4f*)Particles.GetPAndInvM().GetData(),
+							(ispc::FIntVector4*)&Constraints.GetData()[ColorStart],
+							&RestAngles.GetData()[ColorStart],
+							&IsBuckled.GetData()[ColorStart],
+							&Lambdas.GetData()[ColorStart],
+							Dt,
+							XPBDBendMinStiffness,
+							StiffnessHasWeightMap,
+							&Stiffness.GetIndices().GetData()[ColorStart],
+							&Stiffness.GetTable().GetData()[0],
+							BucklingStiffnessHasWeightMap,
+							&BucklingStiffness.GetIndices().GetData()[ColorStart],
+							&BucklingStiffness.GetTable().GetData()[0],
+							ColorSize);
+					}
 				}
 			}
 			else
@@ -184,6 +252,7 @@ void FXPBDBendingConstraints::Apply(FSolverParticles& Particles, const FSolverRe
 			{
 				const FSolverReal StiffnessNoMap = (FSolverReal)Stiffness;
 				const FSolverReal BucklingStiffnessNoMap = (FSolverReal)BucklingStiffness;
+				const FSolverReal DampingNoMap = (FSolverReal)DampingRatio;
 				for (int32 ConstraintColorIndex = 0; ConstraintColorIndex < ConstraintColorNum; ++ConstraintColorIndex)
 				{
 					const int32 ColorStart = ConstraintsPerColorStartIndex[ConstraintColorIndex];
@@ -193,7 +262,8 @@ void FXPBDBendingConstraints::Apply(FSolverParticles& Particles, const FSolverRe
 						const int32 ConstraintIndex = ColorStart + Index;
 						const FSolverReal ExpStiffnessValue = StiffnessHasWeightMap ? Stiffness[ConstraintIndex] : StiffnessNoMap;
 						const FSolverReal ExpBucklingValue = BucklingStiffnessHasWeightMap ? BucklingStiffness[ConstraintIndex] : BucklingStiffnessNoMap;
-						ApplyHelper(Particles, Dt, ConstraintIndex, ExpStiffnessValue, ExpBucklingValue);
+						const FSolverReal DampingRatioValue = DampingHasWeightMap ? DampingRatio[ConstraintIndex] : DampingNoMap;
+						ApplyHelper(Particles, Dt, ConstraintIndex, ExpStiffnessValue, ExpBucklingValue, DampingRatioValue);
 					});
 				}
 			}
@@ -205,24 +275,27 @@ void FXPBDBendingConstraints::Apply(FSolverParticles& Particles, const FSolverRe
 		{
 			const FSolverReal ExpStiffnessValue = (FSolverReal)Stiffness;
 			const FSolverReal ExpBucklingValue = (FSolverReal)BucklingStiffness;
+			const FSolverReal DampingRatioValue = (FSolverReal)DampingRatio;
 			if (ExpStiffnessValue < XPBDBendMinStiffness && ExpBucklingValue < XPBDBendMinStiffness)
 			{
 				return;
 			}
 			for (int32 ConstraintIndex = 0; ConstraintIndex < Constraints.Num(); ++ConstraintIndex)
 			{
-				ApplyHelper(Particles, Dt, ConstraintIndex, ExpStiffnessValue, ExpBucklingValue);
+				ApplyHelper(Particles, Dt, ConstraintIndex, ExpStiffnessValue, ExpBucklingValue, DampingRatioValue);
 			}
 		}
 		else
 		{
 			const FSolverReal StiffnessNoMap = (FSolverReal)Stiffness;
 			const FSolverReal BucklingStiffnessNoMap = (FSolverReal)BucklingStiffness;
+			const FSolverReal DampingNoMap = (FSolverReal)DampingRatio;
 			for (int32 ConstraintIndex = 0; ConstraintIndex < Constraints.Num(); ++ConstraintIndex)
 			{
 				const FSolverReal ExpStiffnessValue = StiffnessHasWeightMap ? Stiffness[ConstraintIndex] : StiffnessNoMap;
 				const FSolverReal ExpBucklingValue = BucklingStiffnessHasWeightMap ? BucklingStiffness[ConstraintIndex] : BucklingStiffnessNoMap;
-				ApplyHelper(Particles, Dt, ConstraintIndex, ExpStiffnessValue, ExpBucklingValue);
+				const FSolverReal DampingRatioValue = DampingHasWeightMap ? DampingRatio[ConstraintIndex] : DampingNoMap;
+				ApplyHelper(Particles, Dt, ConstraintIndex, ExpStiffnessValue, ExpBucklingValue, DampingRatioValue);
 			}
 		}
 	}
diff --git a/Engine/Source/Runtime/Experimental/Chaos/Private/Chaos/XPBDBendingConstraints.ispc b/Engine/Source/Runtime/Experimental/Chaos/Private/Chaos/XPBDBendingConstraints.ispc
--- a/Engine/Source/Runtime/Experimental/Chaos/Private/Chaos/XPBDBendingConstraints.ispc
+++ b/Engine/Source/Runtime/Experimental/Chaos/Private/Chaos/XPBDBendingConstraints.ispc
@@ -115,6 +115,91 @@ export void ApplyXPBDBendingConstraints(uniform FVector4f PandInvM[],
 	}
 }
 
+export void ApplyXPBDBendingConstraintsWithDamping(uniform FVector4f PandInvM[],
+									const uniform FVector3f X[],
+									const uniform FIntVector4 Constraints[],
+									const uniform float RestAngles[],
+									const uniform bool IsBuckled[],
+									uniform float Lambdas[],
+									const uniform float Dt,
+									const uniform float MinStiffness,
+									const uniform float Stiffness,
+									const uniform float BucklingStiffness,
+									const uniform float DampingRatio,
+									const uniform int32 NumConstraints)
+{
+
+	foreach(i = 0 ... NumConstraints)
+	{
+		const varying FIntVector4 Constraint = VectorLoad(&Constraints[extract(i,0)]);
+
+		const varying int32 i1 = Constraint.V[0];
+		const varying int32 i2 = Constraint.V[1];
+		const varying int32 i3 = Constraint.V[2];
+		const varying int32 i4 = Constraint.V[3];
+		
+		const varying FVector4f PandInvM1 = VectorGather(&PandInvM[i1]);
+		const varying FVector4f PandInvM2 = VectorGather(&PandInvM[i2]);
+		const varying FVector4f PandInvM3 = VectorGather(&PandInvM[i3]);
+		const varying FVector4f PandInvM4 = VectorGather(&PandInvM[i4]);
+
+		varying FVector3f P1, P2, P3, P4;
+		varying float M1, M2, M3, M4;
+
+		UnzipPandInvM(PandInvM1, P1, M1);
+		UnzipPandInvM(PandInvM2, P2, M2);
+		UnzipPandInvM(PandInvM3, P3, M3);
+		UnzipPandInvM(PandInvM4, P4, M4);
+
+		varying FVector3f Grad1, Grad2, Grad3, Grad4;
+		const varying float Angle = CalcGradientsAndAngle(P1, P2, P3, P4, Grad1, Grad2, Grad3, Grad4);
+		
+		const varying float RestAngle = RestAngles[i];
+		const varying float StiffnessValue = select(IsBuckled[i], BucklingStiffness, Stiffness);
+
+		const varying FVector3f X1 = VectorGather(&X[i1]);
+		const varying FVector3f X2 = VectorGather(&X[i2]);
+		const varying FVector3f X3 = VectorGather(&X[i3]);
+		const varying FVector3f X4 = VectorGather(&X[i4]);
+
+		const varying float CombinedInvMass = M1 + M2 + M3 + M4;
+		const varying float Damping = 2.f * DampingRatio * sqrt(StiffnessValue / CombinedInvMass);
+
+		const varying float Lambda = Lambdas[i];
+		const varying float Alpha  = 1.f / (StiffnessValue * Dt * Dt);
+		const varying float Gamma = Alpha * Damping * Dt;
+
+		const varying FVector3f V1TimesDt = P1 - X1;
+		const varying FVector3f V2TimesDt = P2 - X2;
+		const varying FVector3f V3TimesDt = P3 - X3;
+		const varying FVector3f V4TimesDt = P4 - X4;
+		const varying float DampingTerm = Gamma * (VectorDot(V1TimesDt, Grad1) + VectorDot(V2TimesDt, Grad2) + VectorDot(V3TimesDt, Grad3) + VectorDot(V4TimesDt, Grad4));
+
+		const varying float Denom = (1.f + Gamma)*(M1 * VectorSizeSquared(Grad1) + M2 * VectorSizeSquared(Grad2) + M3 * VectorSizeSquared(Grad3) + M4 * VectorSizeSquared(Grad4)) + Alpha;
+
+		const varying float ComputedDLambda = (Angle - RestAngle - Alpha * Lambda + DampingTerm) / Denom;
+		const varying float DLambda = select(StiffnessValue < MinStiffness, (uniform float)0, ComputedDLambda);
+		Lambdas[i] = Lambda + DLambda;
+
+		if (M1 > 0)
+		{
+			VectorScatter(&PandInvM[i1], SetVector4(P1 - DLambda * M1 * Grad1, M1));
+		}
+		if (M2 > 0)
+		{
+			VectorScatter(&PandInvM[i2], SetVector4(P2 - DLambda * M2 * Grad2, M2));
+		}
+		if (M3 > 0)
+		{
+			VectorScatter(&PandInvM[i3], SetVector4(P3 - DLambda * M3 * Grad3, M3));
+		}
+		if (M4 > 0)
+		{
+			VectorScatter(&PandInvM[i4], SetVector4(P4 - DLambda * M4 * Grad4, M4));
+		}
+	}
+}
+
 export void ApplyXPBDBendingConstraintsWithMaps(uniform FVector4f PandInvM[],
 											const uniform FIntVector4 Constraints[],
 											const uniform float RestAngles[],
@@ -193,3 +278,105 @@ export void ApplyXPBDBendingConstraintsWithMaps(uniform FVector4f PandInvM[],
 		}
 	}
 }
+
+export void ApplyXPBDBendingConstraintsWithDampingAndMaps(uniform FVector4f PandInvM[],
+											const uniform FVector3f X[],
+											const uniform FIntVector4 Constraints[],
+											const uniform float RestAngles[],
+											const uniform bool IsBuckled[],
+											uniform float Lambdas[],
+											const uniform float Dt,
+											const uniform float MinStiffness,
+											const uniform bool StiffnessHasMap,
+											const uniform uint8 StiffnessIndices[],
+											const uniform float StiffnessTable[],
+											const uniform bool BucklingStiffnessHasMap,
+											const uniform uint8 BucklingStiffnessIndices[],
+											const uniform float BucklingStiffnessTable[],
+											const uniform bool DampingHasMap,
+											const uniform uint8 DampingIndices[],
+											const uniform float DampingTable[],
+											const uniform int32 NumConstraints)
+{
+
+	foreach(i = 0 ... NumConstraints)
+	{
+		const varying FIntVector4 Constraint = VectorLoad(&Constraints[extract(i,0)]);
+
+		const varying int32 i1 = Constraint.V[0];
+		const varying int32 i2 = Constraint.V[1];
+		const varying int32 i3 = Constraint.V[2];
+		const varying int32 i4 = Constraint.V[3];
+		const varying int8 StiffnessIndex = select(StiffnessHasMap, StiffnessIndices[i], 0);
+		const varying int8 BucklingStiffnessIndex = select(BucklingStiffnessHasMap, BucklingStiffnessIndices[i], 0);
+		const varying int8 DampingIndex = select(DampingHasMap, DampingIndices[i], 0);
+
+		const varying FVector4f PandInvM2 = VectorGather(&PandInvM[i2]);
+		const varying FVector4f PandInvM3 = VectorGather(&PandInvM[i3]);
+		const varying FVector4f PandInvM1 = VectorGather(&PandInvM[i1]);
+		const varying FVector4f PandInvM4 = VectorGather(&PandInvM[i4]);
+
+		varying FVector3f P1, P2, P3, P4;
+		varying float M1, M2, M3, M4;
+
+		UnzipPandInvM(PandInvM1, P1, M1);
+		UnzipPandInvM(PandInvM2, P2, M2);
+		UnzipPandInvM(PandInvM3, P3, M3);
+		UnzipPandInvM(PandInvM4, P4, M4);
+
+
+		varying FVector3f Grad1, Grad2, Grad3, Grad4;
+		const varying float Angle = CalcGradientsAndAngle(P1, P2, P3, P4, Grad1, Grad2, Grad3, Grad4);
+		
+		#pragma ignore warning(perf)
+		const varying float Stiffness = StiffnessTable[StiffnessIndex];
+		#pragma ignore warning(perf)
+		const varying float BucklingStiffness = BucklingStiffnessTable[BucklingStiffnessIndex];		
+		#pragma ignore warning(perf)
+		const varying float DampingRatio = DampingTable[DampingIndex];
+		
+		const varying float RestAngle = RestAngles[i];
+		const varying float StiffnessValue = select(IsBuckled[i], BucklingStiffness, Stiffness);
+
+		const varying FVector3f X1 = VectorGather(&X[i1]);
+		const varying FVector3f X2 = VectorGather(&X[i2]);
+		const varying FVector3f X3 = VectorGather(&X[i3]);
+		const varying FVector3f X4 = VectorGather(&X[i4]);
+
+		const varying float CombinedInvMass = M1 + M2 + M3 + M4;
+		const varying float Damping = 2.f * DampingRatio * sqrt(StiffnessValue / CombinedInvMass);
+		
+		const varying float Lambda = Lambdas[i];
+		const varying float Alpha  = 1.f / (StiffnessValue * Dt * Dt);
+		const varying float Gamma = Alpha * Damping * Dt;
+
+		const varying FVector3f V1TimesDt = P1 - X1;
+		const varying FVector3f V2TimesDt = P2 - X2;
+		const varying FVector3f V3TimesDt = P3 - X3;
+		const varying FVector3f V4TimesDt = P4 - X4;
+		const varying float DampingTerm = Gamma * (VectorDot(V1TimesDt, Grad1) + VectorDot(V2TimesDt, Grad2) + VectorDot(V3TimesDt, Grad3) + VectorDot(V4TimesDt, Grad4));
+
+		const varying float Denom = (1.f + Gamma)*(M1 * VectorSizeSquared(Grad1) + M2 * VectorSizeSquared(Grad2) + M3 * VectorSizeSquared(Grad3) + M4 * VectorSizeSquared(Grad4)) + Alpha;
+
+		const varying float ComputedDLambda = (Angle - RestAngle - Alpha * Lambda + DampingTerm) / Denom;
+		const varying float DLambda = select( StiffnessValue < MinStiffness, (uniform float)0, ComputedDLambda);
+		Lambdas[i] = Lambda + DLambda;
+
+		if (M1 > 0)
+		{
+			VectorScatter(&PandInvM[i1], SetVector4(P1 - DLambda * M1 * Grad1, M1));
+		}
+		if (M2 > 0)
+		{
+			VectorScatter(&PandInvM[i2], SetVector4(P2 - DLambda * M2 * Grad2, M2));
+		}
+		if (M3 > 0)
+		{
+			VectorScatter(&PandInvM[i3], SetVector4(P3 - DLambda * M3 * Grad3, M3));
+		}
+		if (M4 > 0)
+		{
+			VectorScatter(&PandInvM[i4], SetVector4(P4 - DLambda * M4 * Grad4, M4));
+		}
+	}
+}
\ No newline at end of file
diff --git a/Engine/Source/Runtime/Experimental/Chaos/Private/Chaos/XPBDSpringConstraints.cpp b/Engine/Source/Runtime/Experimental/Chaos/Private/Chaos/XPBDSpringConstraints.cpp
--- a/Engine/Source/Runtime/Experimental/Chaos/Private/Chaos/XPBDSpringConstraints.cpp
+++ b/Engine/Source/Runtime/Experimental/Chaos/Private/Chaos/XPBDSpringConstraints.cpp
@@ -15,6 +15,7 @@ DECLARE_CYCLE_STAT(TEXT("Chaos XPBD Spring Constraint"), STAT_XPBD_Spring, STATG
 
 #if INTEL_ISPC && !UE_BUILD_SHIPPING
 static_assert(sizeof(ispc::FVector4f) == sizeof(Chaos::Softs::FPAndInvM), "sizeof(ispc::FVector4f) != sizeof(Chaos::Softs::FPAndInvM)");
+static_assert(sizeof(ispc::FVector3f) == sizeof(Chaos::Softs::FSolverVec3), "sizeof(ispc::FVector3f) == sizeof(Chaos::Softs::FSolverVec3)");
 static_assert(sizeof(ispc::FIntVector2) == sizeof(Chaos::TVec2<int32>), "sizeof(ispc::FIntVector2) != sizeof(Chaos::TVec2<int32>)");
 
 bool bChaos_XPBDSpring_ISPC_Enabled = true;
@@ -67,15 +68,16 @@ void FXPBDSpringConstraints::InitColor(const FSolverParticles& Particles)
 		Constraints = MoveTemp(ReorderedConstraints);
 		Dists = MoveTemp(ReorderedDists);
 		Stiffness.ReorderIndices(OrigToReorderedIndices);
+		DampingRatio.ReorderIndices(OrigToReorderedIndices);
 	}
 }
 
-void FXPBDSpringConstraints::ApplyHelper(FSolverParticles& Particles, const FSolverReal Dt, const int32 ConstraintIndex, const FSolverReal ExpStiffnessValue) const
+void FXPBDSpringConstraints::ApplyHelper(FSolverParticles& Particles, const FSolverReal Dt, const int32 ConstraintIndex, const FSolverReal ExpStiffnessValue, const FSolverReal DampingRatioValue) const
 {
 	const TVec2<int32>& Constraint = Constraints[ConstraintIndex];
 	const int32 i1 = Constraint[0];
 	const int32 i2 = Constraint[1];
-	const FSolverVec3 Delta =  GetDelta(Particles, Dt, ConstraintIndex, ExpStiffnessValue);
+	const FSolverVec3 Delta =  GetDelta(Particles, Dt, ConstraintIndex, ExpStiffnessValue, DampingRatioValue);
 	if (Particles.InvM(i1) > (FSolverReal)0.)
 	{
 		Particles.P(i1) -= Particles.InvM(i1) * Delta;
@@ -90,12 +92,17 @@ void FXPBDSpringConstraints::Apply(FSolverParticles& Particles, const FSolverRea
 {
 	TRACE_CPUPROFILER_EVENT_SCOPE(FXPBDSpringConstraints_Apply);
 	SCOPE_CYCLE_COUNTER(STAT_XPBD_Spring);
+
+	const bool StiffnessHasWeightMap = Stiffness.HasWeightMap();
+	const bool DampingHasWeightMap = DampingRatio.HasWeightMap();
+
 	if ((ConstraintsPerColorStartIndex.Num() > 1) && (Constraints.Num() > Chaos_XPBDSpring_ParallelConstraintCount))
 	{
 		const int32 ConstraintColorNum = ConstraintsPerColorStartIndex.Num() - 1;
-		if (!Stiffness.HasWeightMap())
+		if (!StiffnessHasWeightMap && !DampingHasWeightMap)
 		{
 			const FSolverReal ExpStiffnessValue = (FSolverReal)Stiffness;
+			const FSolverReal DampingRatioValue = (FSolverReal)DampingRatio;
 			if (ExpStiffnessValue < XPBDSpringMinStiffness)
 			{
 				return;
@@ -104,18 +111,39 @@ void FXPBDSpringConstraints::Apply(FSolverParticles& Particles, const FSolverRea
 #if INTEL_ISPC
 			if (bRealTypeCompatibleWithISPC && bChaos_XPBDSpring_ISPC_Enabled)
 			{
-				for (int32 ConstraintColorIndex = 0; ConstraintColorIndex < ConstraintColorNum; ++ConstraintColorIndex)
+				if (DampingRatioValue > 0)
 				{
-					const int32 ColorStart = ConstraintsPerColorStartIndex[ConstraintColorIndex];
-					const int32 ColorSize = ConstraintsPerColorStartIndex[ConstraintColorIndex + 1] - ColorStart;
-					ispc::ApplyXPBDSpringConstraints(
-						(ispc::FVector4f*)Particles.GetPAndInvM().GetData(),
-						(ispc::FIntVector2*)&Constraints.GetData()[ColorStart],
-						&Dists.GetData()[ColorStart],
-						&Lambdas.GetData()[ColorStart],
-						Dt,
-						ExpStiffnessValue,
-						ColorSize);
+					for (int32 ConstraintColorIndex = 0; ConstraintColorIndex < ConstraintColorNum; ++ConstraintColorIndex)
+					{
+						const int32 ColorStart = ConstraintsPerColorStartIndex[ConstraintColorIndex];
+						const int32 ColorSize = ConstraintsPerColorStartIndex[ConstraintColorIndex + 1] - ColorStart;
+						ispc::ApplyXPBDSpringConstraintsWithDamping(
+							(ispc::FVector4f*)Particles.GetPAndInvM().GetData(),
+							(const ispc::FVector3f*)Particles.X().GetData(),
+							(ispc::FIntVector2*)&Constraints.GetData()[ColorStart],
+							&Dists.GetData()[ColorStart],
+							&Lambdas.GetData()[ColorStart],
+							Dt,
+							ExpStiffnessValue,
+							DampingRatioValue,
+							ColorSize);
+					}
+				}
+				else
+				{
+					for (int32 ConstraintColorIndex = 0; ConstraintColorIndex < ConstraintColorNum; ++ConstraintColorIndex)
+					{
+						const int32 ColorStart = ConstraintsPerColorStartIndex[ConstraintColorIndex];
+						const int32 ColorSize = ConstraintsPerColorStartIndex[ConstraintColorIndex + 1] - ColorStart;
+						ispc::ApplyXPBDSpringConstraints(
+							(ispc::FVector4f*)Particles.GetPAndInvM().GetData(),
+							(ispc::FIntVector2*)&Constraints.GetData()[ColorStart],
+							&Dists.GetData()[ColorStart],
+							&Lambdas.GetData()[ColorStart],
+							Dt,
+							ExpStiffnessValue,
+							ColorSize);
+					}
 				}
 			}
 			else
@@ -128,7 +156,7 @@ void FXPBDSpringConstraints::Apply(FSolverParticles& Particles, const FSolverRea
 					PhysicsParallelFor(ColorSize, [&](const int32 Index)
 					{
 						const int32 ConstraintIndex = ColorStart + Index;
-						ApplyHelper(Particles, Dt, ConstraintIndex, ExpStiffnessValue);
+						ApplyHelper(Particles, Dt, ConstraintIndex, ExpStiffnessValue, DampingRatioValue);
 					});
 				}
 			}
@@ -138,25 +166,53 @@ void FXPBDSpringConstraints::Apply(FSolverParticles& Particles, const FSolverRea
 #if INTEL_ISPC
 			if (bRealTypeCompatibleWithISPC && bChaos_XPBDSpring_ISPC_Enabled)
 			{
-				for (int32 ConstraintColorIndex = 0; ConstraintColorIndex < ConstraintColorNum; ++ConstraintColorIndex)
+				if (DampingHasWeightMap || (FSolverReal)DampingRatio > 0)
 				{
-					const int32 ColorStart = ConstraintsPerColorStartIndex[ConstraintColorIndex];
-					const int32 ColorSize = ConstraintsPerColorStartIndex[ConstraintColorIndex + 1] - ColorStart;
-					ispc::ApplyXPBDSpringConstraintsWithWeightMaps(
-						(ispc::FVector4f*)Particles.GetPAndInvM().GetData(),
-						(ispc::FIntVector2*) & Constraints.GetData()[ColorStart],
-						&Dists.GetData()[ColorStart],
-						&Lambdas.GetData()[ColorStart],
-						Dt,
-						XPBDSpringMinStiffness,
-						&Stiffness.GetIndices().GetData()[ColorStart],
-						&Stiffness.GetTable().GetData()[0],
-						ColorSize);
+					for (int32 ConstraintColorIndex = 0; ConstraintColorIndex < ConstraintColorNum; ++ConstraintColorIndex)
+					{
+						const int32 ColorStart = ConstraintsPerColorStartIndex[ConstraintColorIndex];
+						const int32 ColorSize = ConstraintsPerColorStartIndex[ConstraintColorIndex + 1] - ColorStart;
+						ispc::ApplyXPBDSpringConstraintsWithDampingAndWeightMaps(
+							(ispc::FVector4f*)Particles.GetPAndInvM().GetData(),
+							(const ispc::FVector3f*)Particles.X().GetData(),
+							(ispc::FIntVector2*)&Constraints.GetData()[ColorStart],
+							&Dists.GetData()[ColorStart],
+							&Lambdas.GetData()[ColorStart],
+							Dt,
+							XPBDSpringMinStiffness,
+							StiffnessHasWeightMap,
+							&Stiffness.GetIndices().GetData()[ColorStart],
+							&Stiffness.GetTable().GetData()[0],
+							DampingHasWeightMap,
+							&DampingRatio.GetIndices().GetData()[ColorStart],
+							&DampingRatio.GetTable().GetData()[0],
+							ColorSize);
+					}
+				}
+				else
+				{
+					for (int32 ConstraintColorIndex = 0; ConstraintColorIndex < ConstraintColorNum; ++ConstraintColorIndex)
+					{
+						const int32 ColorStart = ConstraintsPerColorStartIndex[ConstraintColorIndex];
+						const int32 ColorSize = ConstraintsPerColorStartIndex[ConstraintColorIndex + 1] - ColorStart;
+						ispc::ApplyXPBDSpringConstraintsWithWeightMaps(
+							(ispc::FVector4f*)Particles.GetPAndInvM().GetData(),
+							(ispc::FIntVector2*)&Constraints.GetData()[ColorStart],
+							&Dists.GetData()[ColorStart],
+							&Lambdas.GetData()[ColorStart],
+							Dt,
+							XPBDSpringMinStiffness,
+							&Stiffness.GetIndices().GetData()[ColorStart],
+							&Stiffness.GetTable().GetData()[0],
+							ColorSize);
+					}
 				}
 			}
 			else
 #endif
 			{
+				const FSolverReal StiffnessNoMap = (FSolverReal)Stiffness;
+				const FSolverReal DampingNoMap = (FSolverReal)DampingRatio;
 				for (int32 ConstraintColorIndex = 0; ConstraintColorIndex < ConstraintColorNum; ++ConstraintColorIndex)
 				{
 					const int32 ColorStart = ConstraintsPerColorStartIndex[ConstraintColorIndex];
@@ -164,8 +220,9 @@ void FXPBDSpringConstraints::Apply(FSolverParticles& Particles, const FSolverRea
 					PhysicsParallelFor(ColorSize, [&](const int32 Index)
 					{
 						const int32 ConstraintIndex = ColorStart + Index;
-						const FSolverReal ExpStiffnessValue = Stiffness[ConstraintIndex];
-						ApplyHelper(Particles, Dt, ConstraintIndex, ExpStiffnessValue);
+						const FSolverReal ExpStiffnessValue = StiffnessHasWeightMap ? Stiffness[ConstraintIndex] : StiffnessNoMap;
+						const FSolverReal DampingRatioValue = DampingHasWeightMap ? DampingRatio[ConstraintIndex] : DampingNoMap;
+						ApplyHelper(Particles, Dt, ConstraintIndex, ExpStiffnessValue, DampingRatioValue);
 					});
 				}
 			}
@@ -173,24 +230,28 @@ void FXPBDSpringConstraints::Apply(FSolverParticles& Particles, const FSolverRea
 	}
 	else
 	{
-		if (!Stiffness.HasWeightMap())
+		if (!StiffnessHasWeightMap && !DampingHasWeightMap)
 		{
 			const FSolverReal ExpStiffnessValue = (FSolverReal)Stiffness;
+			const FSolverReal DampingRatioValue = (FSolverReal)DampingRatio;
 			if (ExpStiffnessValue < XPBDSpringMinStiffness)
 			{
 				return;
 			}
 			for (int32 ConstraintIndex = 0; ConstraintIndex < Constraints.Num(); ++ConstraintIndex)
 			{
-				ApplyHelper(Particles, Dt, ConstraintIndex, ExpStiffnessValue);
+				ApplyHelper(Particles, Dt, ConstraintIndex, ExpStiffnessValue, DampingRatioValue);
 			}
 		}
 		else
 		{
+			const FSolverReal StiffnessNoMap = (FSolverReal)Stiffness;
+			const FSolverReal DampingNoMap = (FSolverReal)DampingRatio;
 			for (int32 ConstraintIndex = 0; ConstraintIndex < Constraints.Num(); ++ConstraintIndex)
 			{
-				const FSolverReal ExpStiffnessValue = Stiffness[ConstraintIndex];
-				ApplyHelper(Particles, Dt, ConstraintIndex, ExpStiffnessValue);
+				const FSolverReal ExpStiffnessValue = StiffnessHasWeightMap ? Stiffness[ConstraintIndex] : StiffnessNoMap;
+				const FSolverReal DampingRatioValue = DampingHasWeightMap ? DampingRatio[ConstraintIndex] : DampingNoMap;
+				ApplyHelper(Particles, Dt, ConstraintIndex, ExpStiffnessValue, DampingRatioValue);
 			}
 		}
 	}
diff --git a/Engine/Source/Runtime/Experimental/Chaos/Private/Chaos/XPBDSpringConstraints.ispc b/Engine/Source/Runtime/Experimental/Chaos/Private/Chaos/XPBDSpringConstraints.ispc
--- a/Engine/Source/Runtime/Experimental/Chaos/Private/Chaos/XPBDSpringConstraints.ispc
+++ b/Engine/Source/Runtime/Experimental/Chaos/Private/Chaos/XPBDSpringConstraints.ispc
@@ -64,6 +64,65 @@ export void ApplyXPBDSpringConstraints(uniform FVector4f PandInvM[],
 	}
 }
 
+export void ApplyXPBDSpringConstraintsWithDamping(uniform FVector4f PandInvM[],
+												const uniform FVector3f X[],
+												const uniform FIntVector2 Constraints[],
+												const uniform float Dists[],
+												uniform float Lambdas[],
+												const uniform float Dt,
+												const uniform float Stiffness,
+												const uniform float DampingRatio,
+												const uniform int32 NumConstraints)
+{
+	foreach(i = 0 ... NumConstraints)
+	{
+		const varying FIntVector2 Constraint = VectorLoad(&Constraints[extract(i,0)]);
+		const varying int32 i1 = Constraint.V[0];
+		const varying int32 i2 = Constraint.V[1];
+
+		const varying FVector4f PandInvM1 = VectorGather(&PandInvM[i1]);
+		const varying FVector4f PandInvM2 = VectorGather(&PandInvM[i2]);
+
+		varying FVector3f P1, P2;
+		varying float InvM1, InvM2;
+		UnzipPandInvM(PandInvM1, P1, InvM1);
+		UnzipPandInvM(PandInvM2, P2, InvM2);
+
+		FVector3f Direction = P1 - P2;
+		const float Distance = SafeNormalize(Direction);
+
+		const varying FVector3f X1 = VectorGather(&X[i1]);
+		const varying FVector3f X2 = VectorGather(&X[i2]);
+
+		const varying float Dist = Dists[i];
+		const varying float Lambda = Lambdas[i];
+		const uniform float Alpha = (uniform float)1.f / (Stiffness * Dt * Dt);
+		const varying float Offset = Distance - Dist;
+
+		const float CombinedInvMass = InvM2 + InvM1;
+		const varying float Damping = (uniform float)2.f * DampingRatio * sqrt(Stiffness / CombinedInvMass);
+		const varying float Gamma = Alpha * Damping * Dt;
+		const varying FVector3f RelativeVelocityTimesDt = P1 - X1 - P2 + X2;
+
+		const varying float DLambda = (Offset - Alpha * Lambda + Gamma * VectorDot(Direction, RelativeVelocityTimesDt)) / ((1.f + Gamma)*CombinedInvMass + Alpha);
+
+		const FVector3f ComputedDelta = DLambda * Direction;
+
+		const FVector3f Delta = VectorSelect((InvM2 == FLOAT_ZERO && InvM1 == FLOAT_ZERO), FloatZeroVector, ComputedDelta);
+
+		Lambdas[i] = Lambda + DLambda;
+
+		if (InvM1 > FLOAT_ZERO)
+		{
+			VectorScatter(&PandInvM[i1], SetVector4( P1 - (InvM1 * Delta), InvM1 ));
+		}
+		if (InvM2 > FLOAT_ZERO)
+		{
+			VectorScatter(&PandInvM[i2], SetVector4( P2 + (InvM2 * Delta), InvM2 ));
+		}
+	}
+}
+
 export void ApplyXPBDSpringConstraintsWithWeightMaps(uniform FVector4f PandInvM[],
 									const uniform FIntVector2 Constraints[],
 									const uniform float Dists[],
@@ -123,3 +182,78 @@ export void ApplyXPBDSpringConstraintsWithWeightMaps(uniform FVector4f PandInvM[
 		}
 	}
 }
+
+export void ApplyXPBDSpringConstraintsWithDampingAndWeightMaps(uniform FVector4f PandInvM[],
+															const uniform FVector3f X[],
+															const uniform FIntVector2 Constraints[],
+															const uniform float Dists[],
+															uniform float Lambdas[],
+															const uniform float Dt,
+															const uniform float MinStiffness,
+															const uniform bool StiffnessHasMap,
+															const uniform uint8 StiffnessIndices[],
+															const uniform float StiffnessTable[],
+															const uniform bool DampingHasMap,
+															const uniform uint8 DampingIndices[],
+															const uniform float DampingTable[],
+															const uniform int32 NumConstraints)
+{
+
+	foreach(i = 0 ... NumConstraints)
+	{
+		const varying FIntVector2 Constraint = VectorLoad(&Constraints[extract(i,0)]);
+		const varying int32 i1 = Constraint.V[0];
+		const varying int32 i2 = Constraint.V[1];
+		const varying int8 StiffnessIndex = select(StiffnessHasMap, StiffnessIndices[i], 0);
+		const varying int8 DampingIndex = select(DampingHasMap, DampingIndices[i], 0);
+
+		const varying FVector4f PandInvM1 = VectorGather(&PandInvM[i1]);
+		const varying FVector4f PandInvM2 = VectorGather(&PandInvM[i2]);
+
+		varying FVector3f P1, P2;
+		varying float InvM1, InvM2;
+		UnzipPandInvM(PandInvM1, P1, InvM1);
+		UnzipPandInvM(PandInvM2, P2, InvM2);
+
+		FVector3f Direction = P1 - P2;
+		const float Distance = SafeNormalize(Direction);
+
+		const varying FVector3f X1 = VectorGather(&X[i1]);
+		const varying FVector3f X2 = VectorGather(&X[i2]);
+
+		const varying float Dist = Dists[i];
+		const varying float Lambda = Lambdas[i];
+		const varying float Offset = Distance - Dist;
+
+
+		#pragma ignore warning(perf)
+		const varying float Stiffness = StiffnessTable[StiffnessIndex];
+		#pragma ignore warning(perf)
+		const varying float DampingRatio = DampingTable[DampingIndex];
+
+		const varying float Alpha = (uniform float)1.f / (Stiffness * Dt * Dt);
+
+		const float CombinedInvMass = InvM2 + InvM1;
+		const varying float Damping = (uniform float)2.f * DampingRatio * sqrt(Stiffness / CombinedInvMass);
+		const varying float Gamma = Alpha * Damping * Dt;
+		const varying FVector3f RelativeVelocityTimesDt = P1 - X1 - P2 + X2;
+
+		const varying float ComputedDLambda = (Offset - Alpha * Lambda + Gamma * VectorDot(Direction, RelativeVelocityTimesDt)) / ((1.f + Gamma)*CombinedInvMass + Alpha);
+		const varying float DLambda = select(Stiffness < MinStiffness, 0.f, ComputedDLambda);
+
+		const FVector3f ComputedDelta = DLambda * Direction;
+
+		const FVector3f Delta = VectorSelect((InvM2 == FLOAT_ZERO && InvM1 == FLOAT_ZERO), FloatZeroVector, ComputedDelta);
+
+		Lambdas[i] = Lambda + DLambda;
+
+		if (InvM1 > FLOAT_ZERO)
+		{
+			VectorScatter(&PandInvM[i1], SetVector4( P1 - (InvM1 * Delta), InvM1 ));
+		}
+		if (InvM2 > FLOAT_ZERO)
+		{
+			VectorScatter(&PandInvM[i2], SetVector4( P2 + (InvM2 * Delta), InvM2 ));
+		}
+	}
+}
diff --git a/Engine/Source/Runtime/Experimental/Chaos/Public/Chaos/PBDStiffness.h b/Engine/Source/Runtime/Experimental/Chaos/Public/Chaos/PBDStiffness.h
--- a/Engine/Source/Runtime/Experimental/Chaos/Public/Chaos/PBDStiffness.h
+++ b/Engine/Source/Runtime/Experimental/Chaos/Public/Chaos/PBDStiffness.h
@@ -6,6 +6,7 @@
 #include "Containers/ArrayView.h"
 #include "Containers/ContainersFwd.h"
 #include "ChaosStats.h"
+#include "PBDWeightMap.h"
 
 DECLARE_CYCLE_STAT(TEXT("Chaos PBD Stiffness Apply Values"), STAT_PBD_StiffnessApplyValues, STATGROUP_Chaos);
 
@@ -16,7 +17,7 @@ namespace Chaos::Softs
  * Stiffness class for managing real time update to the weight map and low/high value ranges
  * and to exponentiate the stiffness value depending on the iterations and Dt.
  */
-class FPBDStiffness final
+class FPBDStiffness final : public FPBDWeightMap
 {
 public:
 	static constexpr FSolverReal ParameterFrequency = (FSolverReal)120.;  // 60Hz @ 2 iterations as a root for all stiffness values TODO: Make this a global solver parameter
@@ -47,66 +48,13 @@ public:
 
 	~FPBDStiffness() {}
 
-	/** Return the number of values stored in the weight map. */
-	int32 Num() const { return Indices.Num(); }
-
-	/** Return whether this object contains weight map values. */
-	bool HasWeightMap() const { return Table.Num() > 1 && FMath::Abs(WeightedValue[0] - WeightedValue[1]) > UE_KINDA_SMALL_NUMBER; }
-
-	/**
-	 * Set the low and high values of the weight map.
-	 * The weight map table only gets updated after ApplyValues is called.
-	 * Low and high values are clamped between [0,1]
-	 */
-	void SetWeightedValue(const FSolverVec2& InWeightedValue) { WeightedValue = InWeightedValue.ClampAxes((FSolverReal)0., (FSolverReal)1.); }
-
-	/**
-	 * Set the low and high values of the weight map.
-	 * The weight map table only gets updated after ApplyValues is called.
-	 * Low and high values are not clamped. Commonly used for XPBD Stiffness values which are not [0,1]
-	 */
-	void SetWeightedValueUnclamped(const FSolverVec2& InWeightedValue) { WeightedValue = InWeightedValue; }
-
-	/**
-	 * Return the low and high values set for this weight map.
-	 * Both values will always be between 0 and 1 due to having been clamped in SetWeightedValue.
-	 */
-	const FSolverVec2& GetWeightedValue() const { return WeightedValue; }
-
 	/** Update the weight map table with the current simulation parameters. */
 	inline void ApplyValues(const FSolverReal Dt, const int32 NumIterations);
 
 	/** Update the weight map table with the current simulation parameters. */
 	inline void ApplyXPBDValues(const FSolverReal MaxStiffnesss);
 
-	/**
-	 * Lookup for the exponential weighted value at the specified weight map index.
-	 * This function will assert if it is called with a non zero index on an empty weight map.
-	*/
-	FSolverReal operator[](int32 Index) const { return Table[Indices[Index]]; }
-
-	/** Return the exponential value at the Low weight. */
-	FSolverReal GetLow() const { return Table[0]; }
-
-	/** Return the exponential value at the High weight. */
-	FSolverReal GetHigh() const { return Table.Last(); }
-
-	/** Return the exponential stiffness value when the weight map is not used. */
-	explicit operator FSolverReal() const { return GetLow(); }
-
-	/** Return the particles/constraints indices to the stiffness table as a read only array. */
-	TConstArrayView<uint8> GetIndices() const { return TConstArrayView<uint8>(Indices); }
-
-	/** Return the table of stiffnesses as a read only array. */
-	TConstArrayView<FSolverReal> GetTable() const { return TConstArrayView<FSolverReal>(Table); }
-
-	/** Reorder Indices based on Constraint reordering. */
-	inline void ReorderIndices(const TArray<int32>& OrigToReorderedIndices);
-
 private:
-	TArray<uint8> Indices; // Per particle/constraints array of index to the stiffness table
-	TArray<FSolverReal> Table;  // Fixed lookup table of stiffness values, use uint8 indexation
-	FSolverVec2 WeightedValue;
 	const FSolverReal ParameterFitBase;
 	const FSolverReal ParameterFitLogBase;
 };
@@ -117,33 +65,10 @@ FPBDStiffness::FPBDStiffness(
 	int32 ParticleCount,
 	int32 TableSize,
 	FSolverReal InParameterFitBase)
-	: WeightedValue(InWeightedValue.ClampAxes((FSolverReal)0., (FSolverReal)1.))
+	: FPBDWeightMap(InWeightedValue, Multipliers, ParticleCount, TableSize)
 	, ParameterFitBase(InParameterFitBase)
 	, ParameterFitLogBase(FMath::Loge(InParameterFitBase))
 {
-	check(TableSize > 0 && TableSize < 256);  // The Stiffness lookup table is restricted to uint8 sized indices
-
-	if (Multipliers.Num() == ParticleCount && ParticleCount > 0)
-	{
-		// Convert the weight maps into an array of lookup indices to the stiffness table
-		Indices.AddUninitialized(ParticleCount);
-
-		const FRealSingle TableScale = (FRealSingle)(TableSize - 1);
-
-		for (int32 Index = 0; Index < ParticleCount; ++Index)
-		{
-			Indices[Index] = (uint8)(FMath::Clamp(Multipliers[Index], (FRealSingle)0., (FRealSingle)1.) * TableScale);
-		}
-
-		// Initialize empty table until ApplyValues is called
-		Table.AddZeroed(TableSize);
-	}
-	else
-	{
-		// Initialize with a one element table until ApplyValues is called
-		Indices.AddZeroed(1);
-		Table.AddZeroed(1);
-	}
 }
 
 template<int32 Valence>
@@ -156,44 +81,10 @@ FPBDStiffness::FPBDStiffness(
 	int32 TableSize,
 	FSolverReal InParameterFitBase,
 	typename TEnableIf<Valence >= 2 && Valence <= 4>::Type*)
-	: WeightedValue(InWeightedValue.ClampAxes((FSolverReal)0., (FSolverReal)1.))
+	: FPBDWeightMap(InWeightedValue, Multipliers, Constraints, ParticleOffset, ParticleCount, TableSize)
 	, ParameterFitBase(InParameterFitBase)
 	, ParameterFitLogBase(FMath::Loge(InParameterFitBase))
 {
-	check(TableSize > 0 && TableSize < 256);  // The Stiffness lookup table is restricted to uint8 sized indices
-
-	const int32 ConstraintCount = Constraints.Num();
-
-	if (Multipliers.Num() == ParticleCount && ParticleCount > 0 && ConstraintCount > 0)
-	{
-		// Convert the weight maps into an array of lookup indices to the stiffness table
-		Indices.AddUninitialized(ConstraintCount);
-
-		const FRealSingle TableScale = (FRealSingle)(TableSize - 1);
-
-		for (int32 ConstraintIndex = 0; ConstraintIndex < ConstraintCount; ++ConstraintIndex)
-		{
-			const TVector<int32, Valence>& Constraint = Constraints[ConstraintIndex];
-
-			FRealSingle Weight = 0.f;
-			for (int32 Index = 0; Index < Valence; ++Index)
-			{
-				Weight += FMath::Clamp(Multipliers[Constraint[Index] - ParticleOffset], (FRealSingle)0., (FRealSingle)1.);
-			}
-			Weight /= (FRealSingle)Valence;
-
-			Indices[ConstraintIndex] = (uint8)(Weight * TableScale);
-		}
-
-		// Initialize empty table until ApplyValues is called
-		Table.AddZeroed(TableSize);
-	}
-	else
-	{
-		// Initialize with a one element table until ApplyValues is called
-		Indices.AddZeroed(1);
-		Table.AddZeroed(1);
-	}
 }
 
 static inline FSolverReal CalcExponentialParameterFit(const FSolverReal ParameterFitBase, const FSolverReal ParameterFitLogBase, const FSolverReal InValue)
@@ -267,23 +158,4 @@ void FPBDStiffness::ApplyXPBDValues(const FSolverReal MaxStiffness)
 	}
 }
 
-void FPBDStiffness::ReorderIndices(const TArray<int32>& OrigToReorderedConstraintIndices)
-{
-	if (Indices.Num() == OrigToReorderedConstraintIndices.Num())
-	{
-		TArray<uint8> ReorderedIndices;
-		ReorderedIndices.SetNumUninitialized(Indices.Num());
-		for (int32 OrigConstraintIndex = 0; OrigConstraintIndex < OrigToReorderedConstraintIndices.Num(); ++OrigConstraintIndex)
-		{
-			const int32 ReorderedConstraintIndex = OrigToReorderedConstraintIndices[OrigConstraintIndex];
-			ReorderedIndices[ReorderedConstraintIndex] = Indices[OrigConstraintIndex];
-		}
-		Indices = MoveTemp(ReorderedIndices);
-	}
-	else
-	{
-		check(Indices.Num() == 1);
-	}
-}
-
 }  // End namespace Chaos::Softs
diff --git a/Engine/Source/Runtime/Experimental/Chaos/Public/Chaos/PBDWeightMap.h b/Engine/Source/Runtime/Experimental/Chaos/Public/Chaos/PBDWeightMap.h
new file mode 100644
--- /dev/null
+++ b/Engine/Source/Runtime/Experimental/Chaos/Public/Chaos/PBDWeightMap.h
@@ -0,0 +1,218 @@
+// Copyright Epic Games, Inc. All Rights Reserved.
+#pragma once
+
+#include "Chaos/Core.h"
+#include "Chaos/PBDSoftsEvolutionFwd.h"
+#include "Containers/ArrayView.h"
+#include "Containers/ContainersFwd.h"
+#include "ChaosStats.h"
+
+DECLARE_CYCLE_STAT(TEXT("Chaos PBD WeightMap Apply Values"), STAT_PBD_WeightMapApplyValues, STATGROUP_Chaos);
+
+namespace Chaos::Softs
+{
+
+/**
+ * Weight map class for managing real time update to the weight map and low/high value ranges
+ */
+class FPBDWeightMap
+{
+public:
+
+	/**
+	 * Weightmap particle constructor. 
+	 */
+	inline FPBDWeightMap(
+		const FSolverVec2& InWeightedValue,
+		const TConstArrayView<FRealSingle>& Multipliers = TConstArrayView<FRealSingle>(),
+		int32 ParticleCount = 0,  // A value of 0 also disables the map
+		int32 TableSize = 16);  // Size of the lookup table, can't be more than 256 values, the larger the table the longer it takes to apply changes to the stiffness values;
+
+	/**
+	 * Weightmap constraint constructor. 
+	 */
+	template<int32 Valence>
+	inline FPBDWeightMap(
+		const FSolverVec2& InWeightedValue,
+		const TConstArrayView<FRealSingle>& Multipliers = TConstArrayView<FRealSingle>(),
+		const TConstArrayView<TVector<int32, Valence>>& Constraints = TConstArrayView<TVector<int32, Valence>>(),
+		int32 ParticleOffset = INDEX_NONE, // Constraints have usually a particle offset added to them compared to the weight maps that always starts at index 0
+		int32 ParticleCount = 0,  // A value of 0 also disables the map
+		int32 TableSize = 16,  // Size of the lookup table, can't be more than 256 values, the larger the table the longer it takes to apply changes to the stiffness values
+		typename TEnableIf<Valence >= 2 && Valence <= 4>::Type* = nullptr);  // Prevents incorrect valence, the value is actually unused
+
+	~FPBDWeightMap() {}
+
+	/** Return the number of values stored in the weight map. */
+	int32 Num() const { return Indices.Num(); }
+
+	/** Return whether this object contains weight map values. */
+	bool HasWeightMap() const { return Table.Num() > 1 && FMath::Abs(WeightedValue[0] - WeightedValue[1]) > UE_KINDA_SMALL_NUMBER; }
+
+	/**
+	 * Set the low and high values of the weight map.
+	 * The weight map table only gets updated after ApplyValues is called.
+	 * Low and high values are clamped between [0,1]
+	 */
+	void SetWeightedValue(const FSolverVec2& InWeightedValue) { WeightedValue = InWeightedValue.ClampAxes((FSolverReal)0., (FSolverReal)1.); }
+
+	/**
+	 * Set the low and high values of the weight map.
+	 * The weight map table only gets updated after ApplyValues is called.
+	 * Low and high values are not clamped. Commonly used for XPBD Stiffness values which are not [0,1]
+	 */
+	void SetWeightedValueUnclamped(const FSolverVec2& InWeightedValue) { WeightedValue = InWeightedValue; }
+
+	/**
+	 * Return the low and high values set for this weight map.
+	 * Both values will always be between 0 and 1 due to having been clamped in SetWeightedValue.
+	 */
+	const FSolverVec2& GetWeightedValue() const { return WeightedValue; }
+
+	/** Update the weight map table with the current simulation parameters. */
+	inline void ApplyValues();
+
+	/**
+	 * Lookup for the exponential weighted value at the specified weight map index.
+	 * This function will assert if it is called with a non zero index on an empty weight map.
+	*/
+	FSolverReal operator[](int32 Index) const { return Table[Indices[Index]]; }
+
+	/** Return the exponential value at the Low weight. */
+	FSolverReal GetLow() const { return Table[0]; }
+
+	/** Return the exponential value at the High weight. */
+	FSolverReal GetHigh() const { return Table.Last(); }
+
+	/** Return the exponential stiffness value when the weight map is not used. */
+	explicit operator FSolverReal() const { return GetLow(); }
+
+	/** Return the particles/constraints indices to the stiffness table as a read only array. */
+	TConstArrayView<uint8> GetIndices() const { return TConstArrayView<uint8>(Indices); }
+
+	/** Return the table of stiffnesses as a read only array. */
+	TConstArrayView<FSolverReal> GetTable() const { return TConstArrayView<FSolverReal>(Table); }
+
+	/** Reorder Indices based on Constraint reordering. */
+	inline void ReorderIndices(const TArray<int32>& OrigToReorderedIndices);
+
+protected:
+	TArray<uint8> Indices; // Per particle/constraints array of index to the stiffness table
+	TArray<FSolverReal> Table;  // Fixed lookup table of stiffness values, use uint8 indexation
+	FSolverVec2 WeightedValue;
+};
+
+FPBDWeightMap::FPBDWeightMap(
+	const FSolverVec2& InWeightedValue,
+	const TConstArrayView<FRealSingle>& Multipliers,
+	int32 ParticleCount,
+	int32 TableSize)
+	: WeightedValue(InWeightedValue.ClampAxes((FSolverReal)0., (FSolverReal)1.))
+{
+	check(TableSize > 0 && TableSize < 256);  // The Stiffness lookup table is restricted to uint8 sized indices
+
+	if (Multipliers.Num() == ParticleCount && ParticleCount > 0)
+	{
+		// Convert the weight maps into an array of lookup indices to the stiffness table
+		Indices.AddUninitialized(ParticleCount);
+
+		const FRealSingle TableScale = (FRealSingle)(TableSize - 1);
+
+		for (int32 Index = 0; Index < ParticleCount; ++Index)
+		{
+			Indices[Index] = (uint8)(FMath::Clamp(Multipliers[Index], (FRealSingle)0., (FRealSingle)1.) * TableScale);
+		}
+
+		// Initialize empty table until ApplyValues is called
+		Table.AddZeroed(TableSize);
+	}
+	else
+	{
+		// Initialize with a one element table until ApplyValues is called
+		Indices.AddZeroed(1);
+		Table.AddZeroed(1);
+	}
+}
+
+template<int32 Valence>
+FPBDWeightMap::FPBDWeightMap(
+	const FSolverVec2& InWeightedValue,
+	const TConstArrayView<FRealSingle>& Multipliers,
+	const TConstArrayView<TVector<int32, Valence>>& Constraints,
+	int32 ParticleOffset,
+	int32 ParticleCount,
+	int32 TableSize,
+	typename TEnableIf<Valence >= 2 && Valence <= 4>::Type*)
+	: WeightedValue(InWeightedValue.ClampAxes((FSolverReal)0., (FSolverReal)1.))
+{
+	check(TableSize > 0 && TableSize < 256);  // The Stiffness lookup table is restricted to uint8 sized indices
+
+	const int32 ConstraintCount = Constraints.Num();
+
+	if (Multipliers.Num() == ParticleCount && ParticleCount > 0 && ConstraintCount > 0)
+	{
+		// Convert the weight maps into an array of lookup indices to the stiffness table
+		Indices.AddUninitialized(ConstraintCount);
+
+		const FRealSingle TableScale = (FRealSingle)(TableSize - 1);
+
+		for (int32 ConstraintIndex = 0; ConstraintIndex < ConstraintCount; ++ConstraintIndex)
+		{
+			const TVector<int32, Valence>& Constraint = Constraints[ConstraintIndex];
+
+			FRealSingle Weight = 0.f;
+			for (int32 Index = 0; Index < Valence; ++Index)
+			{
+				Weight += FMath::Clamp(Multipliers[Constraint[Index] - ParticleOffset], (FRealSingle)0., (FRealSingle)1.);
+			}
+			Weight /= (FRealSingle)Valence;
+
+			Indices[ConstraintIndex] = (uint8)(Weight * TableScale);
+		}
+
+		// Initialize empty table until ApplyValues is called
+		Table.AddZeroed(TableSize);
+	}
+	else
+	{
+		// Initialize with a one element table until ApplyValues is called
+		Indices.AddZeroed(1);
+		Table.AddZeroed(1);
+	}
+}
+
+void FPBDWeightMap::ApplyValues()
+{
+	SCOPE_CYCLE_COUNTER(STAT_PBD_WeightMapApplyValues);
+
+	const FSolverReal Offset = WeightedValue[0];
+	const FSolverReal Range = WeightedValue[1] - WeightedValue[0];
+	const int32 TableSize = Table.Num();
+	const FSolverReal WeightIncrement = (TableSize > 1) ? (FSolverReal)1. / (FSolverReal)(TableSize - 1) : (FSolverReal)1.; // Must allow full range from 0 to 1 included
+	for (int32 Index = 0; Index < TableSize; ++Index)
+	{
+		const FSolverReal Weight = (FSolverReal)Index * WeightIncrement;
+		Table[Index] = Offset + Weight * Range;
+	}
+}
+
+void FPBDWeightMap::ReorderIndices(const TArray<int32>& OrigToReorderedConstraintIndices)
+{
+	if (Indices.Num() == OrigToReorderedConstraintIndices.Num())
+	{
+		TArray<uint8> ReorderedIndices;
+		ReorderedIndices.SetNumUninitialized(Indices.Num());
+		for (int32 OrigConstraintIndex = 0; OrigConstraintIndex < OrigToReorderedConstraintIndices.Num(); ++OrigConstraintIndex)
+		{
+			const int32 ReorderedConstraintIndex = OrigToReorderedConstraintIndices[OrigConstraintIndex];
+			ReorderedIndices[ReorderedConstraintIndex] = Indices[OrigConstraintIndex];
+		}
+		Indices = MoveTemp(ReorderedIndices);
+	}
+	else
+	{
+		check(Indices.Num() == 1);
+	}
+}
+
+}  // End namespace Chaos::Softs
diff --git a/Engine/Source/Runtime/Experimental/Chaos/Public/Chaos/XPBDBendingConstraints.h b/Engine/Source/Runtime/Experimental/Chaos/Public/Chaos/XPBDBendingConstraints.h
--- a/Engine/Source/Runtime/Experimental/Chaos/Public/Chaos/XPBDBendingConstraints.h
+++ b/Engine/Source/Runtime/Experimental/Chaos/Public/Chaos/XPBDBendingConstraints.h
@@ -32,6 +32,26 @@ public:
 		const FSolverVec2& InBucklingStiffness,
 		bool bTrimKinematicConstraints = false)
 		: Base(InParticles, ParticleOffset, ParticleCount, MoveTemp(InConstraints), StiffnessMultipliers, BucklingStiffnessMultipliers, InStiffness, InBucklingRatio, InBucklingStiffness, bTrimKinematicConstraints)
+		, DampingRatio(FSolverVec2::ZeroVector)
+	{
+		Lambdas.Init((FSolverReal)0., Constraints.Num());
+		InitColor(InParticles);
+	}
+
+	FXPBDBendingConstraints(const FSolverParticles& InParticles,
+		int32 ParticleOffset,
+		int32 ParticleCount,
+		TArray<TVec4<int32>>&& InConstraints,
+		const TConstArrayView<FRealSingle>& StiffnessMultipliers,
+		const TConstArrayView<FRealSingle>& BucklingStiffnessMultipliers,
+		const TConstArrayView<FRealSingle>& DampingMultipliers,
+		const FSolverVec2& InStiffness,
+		const FSolverReal InBucklingRatio,
+		const FSolverVec2& InBucklingStiffness,
+		const FSolverVec2& InDampingRatio,
+		bool bTrimKinematicConstraints = false)
+		: Base(InParticles, ParticleOffset, ParticleCount, MoveTemp(InConstraints), StiffnessMultipliers, BucklingStiffnessMultipliers, InStiffness, InBucklingRatio, InBucklingStiffness, bTrimKinematicConstraints)
+		, DampingRatio(InDampingRatio, DampingMultipliers, TConstArrayView<TVec2<int32>>(ConstraintSharedEdges), ParticleOffset, ParticleCount)
 	{
 		Lambdas.Init((FSolverReal)0., Constraints.Num());
 		InitColor(InParticles);
@@ -47,23 +67,26 @@ public:
 	}
 
 	// Update stiffness values
-	void SetProperties(const FSolverVec2& InStiffness, const FSolverReal InBucklingRatio, const FSolverVec2& InBucklingStiffness)
+	void SetProperties(const FSolverVec2& InStiffness, const FSolverReal InBucklingRatio, const FSolverVec2& InBucklingStiffness, const FSolverVec2& InDampingRatio = FSolverVec2::ZeroVector)
 	{
 		Stiffness.SetWeightedValueUnclamped(InStiffness);
 		BucklingRatio = InBucklingRatio;
 		BucklingStiffness.SetWeightedValueUnclamped(InBucklingStiffness);
+		DampingRatio.SetWeightedValueUnclamped(InDampingRatio);
 	}
 
 	// Update stiffness table, as well as the simulation stiffness exponent
-	void ApplyProperties(const FSolverReal Dt, const int32 NumIterations) { Stiffness.ApplyXPBDValues(XPBDBendMaxStiffness); BucklingStiffness.ApplyXPBDValues(XPBDBendMaxStiffness); }
+	void ApplyProperties(const FSolverReal Dt, const int32 NumIterations) { Stiffness.ApplyXPBDValues(XPBDBendMaxStiffness); BucklingStiffness.ApplyXPBDValues(XPBDBendMaxStiffness); DampingRatio.ApplyValues(); }
 
 	void Apply(FSolverParticles& Particles, const FSolverReal Dt) const;
 
 private:
 	void InitColor(const FSolverParticles& InParticles);
-	void ApplyHelper(FSolverParticles& Particles, const FSolverReal Dt, const int32 ConstraintIndex, const FSolverReal ExpStiffnessValue, const FSolverReal ExpBucklingValue) const;
+	void ApplyHelper(FSolverParticles& Particles, const FSolverReal Dt, const int32 ConstraintIndex, const FSolverReal ExpStiffnessValue, const FSolverReal ExpBucklingValue, const FSolverReal DampingRatioValue) const;
 
 private:
+	FPBDWeightMap DampingRatio;
+	TArray<FSolverReal> DampingMultipliers;
 	mutable TArray<FSolverReal> Lambdas;
 	TArray<int32> ConstraintsPerColorStartIndex; // Constraints are ordered so each batch is contiguous. This is ColorNum + 1 length so it can be used as start and end.
 };
diff --git a/Engine/Source/Runtime/Experimental/Chaos/Public/Chaos/XPBDSpringConstraints.h b/Engine/Source/Runtime/Experimental/Chaos/Public/Chaos/XPBDSpringConstraints.h
--- a/Engine/Source/Runtime/Experimental/Chaos/Public/Chaos/XPBDSpringConstraints.h
+++ b/Engine/Source/Runtime/Experimental/Chaos/Public/Chaos/XPBDSpringConstraints.h
@@ -3,7 +3,7 @@
 
 #include "Chaos/PBDSpringConstraintsBase.h"
 #include "ChaosStats.h"
-
+#include "PBDWeightMap.h"
 
 namespace Chaos::Softs
 {
@@ -31,6 +31,25 @@ public:
 		bool bTrimKinematicConstraints = false,
 		typename TEnableIf<Valence >= 2 && Valence <= 4>::Type* = nullptr)
 		: Base(Particles, ParticleOffset, ParticleCount, InConstraints, StiffnessMultipliers, InStiffness, bTrimKinematicConstraints)
+		, DampingRatio(FSolverVec2::ZeroVector)
+	{
+		Lambdas.Init((FSolverReal)0., Constraints.Num());
+		InitColor(Particles);
+	}
+	template<int32 Valence>
+	FXPBDSpringConstraints(
+		const FSolverParticles& Particles,
+		int32 ParticleOffset,
+		int32 ParticleCount,
+		const TArray<TVector<int32, Valence>>& InConstraints,
+		const TConstArrayView<FRealSingle>& StiffnessMultipliers,
+		const TConstArrayView<FRealSingle>& DampingMultipliers,
+		const FSolverVec2& InStiffness,
+		const FSolverVec2& InDampingRatio,
+		bool bTrimKinematicConstraints = false,
+		typename TEnableIf<Valence >= 2 && Valence <= 4>::Type* = nullptr)
+		: Base(Particles, ParticleOffset, ParticleCount, InConstraints, StiffnessMultipliers, InStiffness, bTrimKinematicConstraints)
+		, DampingRatio(InDampingRatio, DampingMultipliers, TConstArrayView<TVec2<int32>>(Constraints), ParticleOffset, ParticleCount)
 	{
 		Lambdas.Init((FSolverReal)0., Constraints.Num());
 		InitColor(Particles);
@@ -41,18 +60,22 @@ public:
 	void Init() const { for (FSolverReal& Lambda : Lambdas) { Lambda = (FSolverReal)0.; } }
 
 	// Update stiffness values
-	void SetProperties(const FSolverVec2& InStiffness) { Stiffness.SetWeightedValueUnclamped(InStiffness); }
+	void SetProperties(const FSolverVec2& InStiffness, const FSolverVec2& InDampingRatio = (FSolverReal)0.f)
+	{ 
+		Stiffness.SetWeightedValueUnclamped(InStiffness); 
+		DampingRatio.SetWeightedValueUnclamped(InDampingRatio);
+	}
 	
 	// Update stiffness table, as well as the simulation stiffness exponent
-	inline void ApplyProperties(const FSolverReal Dt, const int32 NumIterations) { Stiffness.ApplyXPBDValues(XPBDSpringMaxStiffness); }
+	inline void ApplyProperties(const FSolverReal Dt, const int32 NumIterations) { Stiffness.ApplyXPBDValues(XPBDSpringMaxStiffness); DampingRatio.ApplyValues(); }
 
 	void Apply(FSolverParticles& Particles, const FSolverReal Dt) const;
 
 private:
 	void InitColor(const FSolverParticles& InParticles);
-	void ApplyHelper(FSolverParticles& Particles, const FSolverReal Dt, const int32 ConstraintIndex, const FSolverReal StiffnessValue) const;
+	void ApplyHelper(FSolverParticles& Particles, const FSolverReal Dt, const int32 ConstraintIndex, const FSolverReal StiffnessValue, const FSolverReal DampingRatioValue) const;
 
-	FSolverVec3 GetDelta(const FSolverParticles& Particles, const FSolverReal Dt, const int32 ConstraintIndex, const FSolverReal StiffnessValue) const
+	FSolverVec3 GetDelta(const FSolverParticles& Particles, const FSolverReal Dt, const int32 ConstraintIndex, const FSolverReal StiffnessValue, const FSolverReal DampingRatioValue) const
 	{
 		const TVec2<int32>& Constraint = Constraints[ConstraintIndex];
 
@@ -63,18 +86,28 @@ private:
 		{
 			return FSolverVec3((FSolverReal)0.);
 		}
+
 		const FSolverReal CombinedInvMass = Particles.InvM(i2) + Particles.InvM(i1);
 
+		const FSolverReal Damping = DampingRatioValue * 2.f * FMath::Sqrt(StiffnessValue / CombinedInvMass);
+
 		const FSolverVec3& P1 = Particles.P(i1);
 		const FSolverVec3& P2 = Particles.P(i2);
 		FSolverVec3 Direction = P1 - P2;
 		const FSolverReal Distance = Direction.SafeNormalize();
 		const FSolverReal Offset = Distance - Dists[ConstraintIndex];
 
+		const FSolverVec3& X1 = Particles.X(i1);
+		const FSolverVec3& X2 = Particles.X(i2);
+
+		const FSolverVec3 RelativeVelocityTimesDt = P1 - X1 - P2 + X2;
+
+
 		FSolverReal& Lambda = Lambdas[ConstraintIndex];
 		const FSolverReal Alpha = (FSolverReal)1.f / (StiffnessValue * Dt * Dt);
+		const FSolverReal Gamma = Alpha * Damping * Dt;
 
-		const FSolverReal DLambda = (Offset - Alpha * Lambda) / (CombinedInvMass + Alpha);
+		const FSolverReal DLambda = (Offset - Alpha * Lambda + Gamma * FSolverVec3::DotProduct(Direction, RelativeVelocityTimesDt)) / (((FSolverReal)1.f + Gamma) * CombinedInvMass + Alpha);
 		const FSolverVec3 Delta = DLambda * Direction;
 		Lambda += DLambda;
 
@@ -82,6 +115,7 @@ private:
 	}
 
 private:
+	FPBDWeightMap DampingRatio;
 	mutable TArray<FSolverReal> Lambdas;
 	TArray<int32> ConstraintsPerColorStartIndex; // Constraints are ordered so each batch is contiguous. This is ColorNum + 1 length so it can be used as start and end.
 };
