Non-unity build fixes for Iris.

diff --git a/Engine/Source/Runtime/Experimental/Iris/Core/Private/Iris/ReplicationSystem/NetExports.cpp b/Engine/Source/Runtime/Experimental/Iris/Core/Private/Iris/ReplicationSystem/NetExports.cpp
--- a/Engine/Source/Runtime/Experimental/Iris/Core/Private/Iris/ReplicationSystem/NetExports.cpp
+++ b/Engine/Source/Runtime/Experimental/Iris/Core/Private/Iris/ReplicationSystem/NetExports.cpp
@@ -1,7 +1,5 @@
 // Copyright Epic Games, Inc. All Rights Reserved.
 
-#pragma once
-
 #include "NetExports.h"
 #include "Iris/PacketControl/PacketNotification.h"
 
diff --git a/Engine/Source/Runtime/Experimental/Iris/Core/Private/Iris/Serialization/GuidNetSerializer.cpp b/Engine/Source/Runtime/Experimental/Iris/Core/Private/Iris/Serialization/GuidNetSerializer.cpp
--- a/Engine/Source/Runtime/Experimental/Iris/Core/Private/Iris/Serialization/GuidNetSerializer.cpp
+++ b/Engine/Source/Runtime/Experimental/Iris/Core/Private/Iris/Serialization/GuidNetSerializer.cpp
@@ -1,7 +1,5 @@
 // Copyright Epic Games, Inc. All Rights Reserved.
 
-#pragma once
-
 #include "Iris/Serialization/GuidNetSerializer.h"
 #include "Misc/Guid.h"
 
diff --git a/Engine/Source/Runtime/Experimental/Iris/Core/Private/Iris/Serialization/InternalNetSerializers.cpp b/Engine/Source/Runtime/Experimental/Iris/Core/Private/Iris/Serialization/InternalNetSerializers.cpp
new file mode 100644
--- /dev/null
+++ b/Engine/Source/Runtime/Experimental/Iris/Core/Private/Iris/Serialization/InternalNetSerializers.cpp
@@ -0,0 +1,7 @@
+// Copyright Epic Games, Inc. All Rights Reserved.
+
+#include "Iris/Serialization/InternalNetSerializers.h"
+#include "Iris/ReplicationState/ReplicationStateDescriptor.h"
+
+FArrayPropertyNetSerializerConfig::FArrayPropertyNetSerializerConfig() = default;
+FArrayPropertyNetSerializerConfig::~FArrayPropertyNetSerializerConfig() = default;
diff --git a/Engine/Source/Runtime/Experimental/Iris/Core/Private/Iris/Serialization/InternalNetSerializers.h b/Engine/Source/Runtime/Experimental/Iris/Core/Private/Iris/Serialization/InternalNetSerializers.h
--- a/Engine/Source/Runtime/Experimental/Iris/Core/Private/Iris/Serialization/InternalNetSerializers.h
+++ b/Engine/Source/Runtime/Experimental/Iris/Core/Private/Iris/Serialization/InternalNetSerializers.h
@@ -29,6 +29,15 @@ struct FArrayPropertyNetSerializerConfig : public FNetSerializerConfig
 	GENERATED_BODY()
 
 public:
+	FArrayPropertyNetSerializerConfig();
+	~FArrayPropertyNetSerializerConfig();
+
+	FArrayPropertyNetSerializerConfig(const FArrayPropertyNetSerializerConfig&) = delete;
+	FArrayPropertyNetSerializerConfig(FArrayPropertyNetSerializerConfig&&) = delete;
+	
+	FArrayPropertyNetSerializerConfig& operator=(const FArrayPropertyNetSerializerConfig&) = delete;
+	FArrayPropertyNetSerializerConfig& operator=(FArrayPropertyNetSerializerConfig&&) = delete;
+
 	UPROPERTY()
 	uint16 MaxElementCount = 0;
 
@@ -41,6 +50,15 @@ public:
 	TRefCountPtr<const UE::Net::FReplicationStateDescriptor> StateDescriptor;
 };
 
+template<>
+struct TStructOpsTypeTraits<FArrayPropertyNetSerializerConfig> : public TStructOpsTypeTraitsBase2<FArrayPropertyNetSerializerConfig>
+{
+	enum
+	{
+		WithCopy = false
+	};
+};
+
 /**
  * Any property that doesn't have any other option will end up using this.
  * As the name suggests it's a last resort.
diff --git a/Engine/Source/Runtime/Experimental/Iris/Core/Private/Iris/Serialization/NetSerializers.cpp b/Engine/Source/Runtime/Experimental/Iris/Core/Private/Iris/Serialization/NetSerializers.cpp
--- a/Engine/Source/Runtime/Experimental/Iris/Core/Private/Iris/Serialization/NetSerializers.cpp
+++ b/Engine/Source/Runtime/Experimental/Iris/Core/Private/Iris/Serialization/NetSerializers.cpp
@@ -8,3 +8,5 @@ FStructNetSerializerConfig::FStructNetSerializerConfig()
 {
 	ConfigTraits = ENetSerializerConfigTraits::NeedDestruction;
 }
+
+FStructNetSerializerConfig::~FStructNetSerializerConfig() = default;
diff --git a/Engine/Source/Runtime/Experimental/Iris/Core/Public/Iris/Serialization/NetSerializers.h b/Engine/Source/Runtime/Experimental/Iris/Core/Public/Iris/Serialization/NetSerializers.h
--- a/Engine/Source/Runtime/Experimental/Iris/Core/Public/Iris/Serialization/NetSerializers.h
+++ b/Engine/Source/Runtime/Experimental/Iris/Core/Public/Iris/Serialization/NetSerializers.h
@@ -38,10 +38,26 @@ struct FStructNetSerializerConfig : public FNetSerializerConfig
 	GENERATED_BODY()
 
 	IRISCORE_API FStructNetSerializerConfig();
+	IRISCORE_API ~FStructNetSerializerConfig();
+
+	FStructNetSerializerConfig(const FStructNetSerializerConfig&) = delete;
+	FStructNetSerializerConfig(FStructNetSerializerConfig&&) = delete;
+
+	FStructNetSerializerConfig& operator=(const FStructNetSerializerConfig&) = delete;
+	FStructNetSerializerConfig& operator=(FStructNetSerializerConfig&&) = delete;
 
 	TRefCountPtr<const UE::Net::FReplicationStateDescriptor> StateDescriptor;
 };
 
+template<>
+struct TStructOpsTypeTraits<FStructNetSerializerConfig> : public TStructOpsTypeTraitsBase2<FStructNetSerializerConfig>
+{
+	enum
+	{
+		WithCopy = false
+	};
+};
+
 /**
  * The NopNetSerializer have all the serializer functions implemented
  * as no-ops. The main purpose of this serializer is to allow adding
