Iris header cleanup.

diff --git a/Engine/Plugins/Runtime/ReplicationSystemTestPlugin/Source/Private/Tests/ReplicationSystem/TestReplicationOperations.cpp b/Engine/Plugins/Runtime/ReplicationSystemTestPlugin/Source/Private/Tests/ReplicationSystem/TestReplicationOperations.cpp
--- a/Engine/Plugins/Runtime/ReplicationSystemTestPlugin/Source/Private/Tests/ReplicationSystem/TestReplicationOperations.cpp
+++ b/Engine/Plugins/Runtime/ReplicationSystemTestPlugin/Source/Private/Tests/ReplicationSystem/TestReplicationOperations.cpp
@@ -2,6 +2,7 @@
 
 #include "ReplicationSystemTestFixture.h"
 #include "Iris/ReplicationSystem/ChangeMaskUtil.h"
+#include "Iris/ReplicationSystem/ReplicationOperations.h"
 #include "Iris/ReplicationSystem/ReplicationOperationsInternal.h"
 #include "Iris/Serialization/NetBitStreamWriter.h"
 #include "Iris/Serialization/NetBitStreamReader.h"
diff --git a/Engine/Source/Runtime/Engine/Private/ActorReplication.cpp b/Engine/Source/Runtime/Engine/Private/ActorReplication.cpp
--- a/Engine/Source/Runtime/Engine/Private/ActorReplication.cpp
+++ b/Engine/Source/Runtime/Engine/Private/ActorReplication.cpp
@@ -23,6 +23,7 @@
 #if UE_WITH_IRIS
 #include "Engine/NetConnection.h"
 #include "Iris/IrisConfig.h"
+#include "Iris/ReplicationSystem/ReplicationFragment.h"
 #include "Iris/ReplicationSystem/ReplicationFragmentUtil.h"
 #include "Iris/ReplicationSystem/ReplicationSystem.h"
 #include "Iris/ReplicationSystem/Conditionals/ReplicationCondition.h"
diff --git a/Engine/Source/Runtime/Experimental/Iris/Core/Private/Iris/Core/IrisDebugging.cpp b/Engine/Source/Runtime/Experimental/Iris/Core/Private/Iris/Core/IrisDebugging.cpp
--- a/Engine/Source/Runtime/Experimental/Iris/Core/Private/Iris/Core/IrisDebugging.cpp
+++ b/Engine/Source/Runtime/Experimental/Iris/Core/Private/Iris/Core/IrisDebugging.cpp
@@ -17,6 +17,7 @@
 #include "Iris/ReplicationState/PropertyReplicationState.h"
 #include "Iris/Serialization/NetSerializationContext.h"
 #include "Iris/Serialization/InternalNetSerializationContext.h"
+#include "Net/Core/Trace/NetDebugName.h"
 #include "Containers/StringFwd.h"
 
 namespace UE::Net::IrisDebugHelper
diff --git a/Engine/Source/Runtime/Experimental/Iris/Core/Private/Iris/IrisCoreModule.cpp b/Engine/Source/Runtime/Experimental/Iris/Core/Private/Iris/IrisCoreModule.cpp
--- a/Engine/Source/Runtime/Experimental/Iris/Core/Private/Iris/IrisCoreModule.cpp
+++ b/Engine/Source/Runtime/Experimental/Iris/Core/Private/Iris/IrisCoreModule.cpp
@@ -1,6 +1,6 @@
 // Copyright Epic Games, Inc. All Rights Reserved.
 
-#include "CoreMinimal.h"
+#include "CoreTypes.h"
 #include "Modules/ModuleInterface.h"
 #include "Modules/ModuleManager.h"
 #include "Iris/ReplicationState/DefaultPropertyNetSerializerInfos.h"
diff --git a/Engine/Source/Runtime/Experimental/Iris/Core/Private/Iris/ReplicationState/IrisFastArraySerializer.cpp b/Engine/Source/Runtime/Experimental/Iris/Core/Private/Iris/ReplicationState/IrisFastArraySerializer.cpp
--- a/Engine/Source/Runtime/Experimental/Iris/Core/Private/Iris/ReplicationState/IrisFastArraySerializer.cpp
+++ b/Engine/Source/Runtime/Experimental/Iris/Core/Private/Iris/ReplicationState/IrisFastArraySerializer.cpp
@@ -1,6 +1,8 @@
 // Copyright Epic Games, Inc. All Rights Reserved.
 #include "Iris/ReplicationState/IrisFastArraySerializer.h"
 #include "Iris/ReplicationState/Private/IrisFastArraySerializerInternal.h"
+#include "Iris/ReplicationState/ReplicationStateUtil.h"
+#include "Net/Core/NetBitArray.h"
 
 FIrisFastArraySerializer::FIrisFastArraySerializer()
 : FFastArraySerializer()
@@ -72,4 +74,70 @@ void FIrisFastArraySerializer::InternalMarkItemChanged(int32 ItemIdx)
 {
 	checkSlow(ReplicationStateHeader.IsBound());
 	UE::Net::Private::FIrisFastArraySerializerPrivateAccessor::MarkArrayItemDirty(*this, ItemIdx);
+}
+
+namespace UE::Net::Private
+{
+
+FNetBitArrayView FIrisFastArraySerializerPrivateAccessor::GetChangeMask(FIrisFastArraySerializer& Array)
+{
+	return MakeNetBitArrayView(&Array.ChangeMaskStorage[0], FIrisFastArraySerializer::IrisFastArrayChangeMaskBits + 1U);
+}
+
+FNetBitArrayView FIrisFastArraySerializerPrivateAccessor::GetConditionalChangeMask(FIrisFastArraySerializer& Array)
+{
+	return MakeNetBitArrayView(&Array.ChangeMaskStorage[1], FIrisFastArraySerializer::IrisFastArrayChangeMaskBits + 1U);
+}
+
+void FIrisFastArraySerializerPrivateAccessor::MarkAllArrayItemsDirty(FIrisFastArraySerializer& Array, uint32 StartingIndex)
+{
+	checkSlow(Array.ReplicationStateHeader.IsBound());
+
+	FNetBitArrayView MemberChangeMask = UE::Net::Private::FIrisFastArraySerializerPrivateAccessor::GetChangeMask(Array);
+	if (!MemberChangeMask.GetBit(0))
+	{
+		MarkNetObjectStateDirty(Array.ReplicationStateHeader);
+	}
+	if (StartingIndex == 0)
+	{
+		MemberChangeMask.SetAllBits();
+	}
+	else
+	{
+		MemberChangeMask.SetBit(FIrisFastArraySerializer::IrisFastArrayPropertyBitIndex);
+		MemberChangeMask.SetBits(FIrisFastArraySerializer::IrisFastArrayChangeMaskBitOffset + StartingIndex, FIrisFastArraySerializer::IrisFastArrayChangeMaskBits - StartingIndex);
+	}
+}
+
+void FIrisFastArraySerializerPrivateAccessor::MarkArrayDirty(FIrisFastArraySerializer& Array)
+{
+	checkSlow(Array.ReplicationStateHeader.IsBound());
+
+	FNetBitArrayView MemberChangeMask = UE::Net::Private::FIrisFastArraySerializerPrivateAccessor::GetChangeMask(Array);
+
+	// Dirty object unless already dirty, we only use the array bit for this purpose
+	if (!MemberChangeMask.GetBit(FIrisFastArraySerializer::IrisFastArrayPropertyBitIndex))
+	{
+		MarkNetObjectStateDirty(Array.ReplicationStateHeader);
+		MemberChangeMask.SetBit(FIrisFastArraySerializer::IrisFastArrayPropertyBitIndex);
+	}
+}
+
+void FIrisFastArraySerializerPrivateAccessor::MarkArrayItemDirty(FIrisFastArraySerializer& Array, int32 ItemIdx)
+{
+	checkSlow(Array.ReplicationStateHeader.IsBound());
+
+	// Mark changemask dirty for this item
+	// We are using a modulo scheme for dirtiness
+	FNetBitArrayView MemberChangeMask = UE::Net::Private::FIrisFastArraySerializerPrivateAccessor::GetChangeMask(Array);
+	MemberChangeMask.SetBit((ItemIdx % FIrisFastArraySerializer::IrisFastArrayChangeMaskBits) + FIrisFastArraySerializer::IrisFastArrayChangeMaskBitOffset);
+
+	// Dirty object unless already dirty, we only use the array bit for this purpose
+	if (!MemberChangeMask.GetBit(0))
+	{
+		MarkNetObjectStateDirty(Array.ReplicationStateHeader);
+		MemberChangeMask.SetBits(0, 1);
+	}
+}
+
 }
\ No newline at end of file
diff --git a/Engine/Source/Runtime/Experimental/Iris/Core/Private/Iris/ReplicationState/PropertyReplicationState.cpp b/Engine/Source/Runtime/Experimental/Iris/Core/Private/Iris/ReplicationState/PropertyReplicationState.cpp
--- a/Engine/Source/Runtime/Experimental/Iris/Core/Private/Iris/ReplicationState/PropertyReplicationState.cpp
+++ b/Engine/Source/Runtime/Experimental/Iris/Core/Private/Iris/ReplicationState/PropertyReplicationState.cpp
@@ -5,7 +5,8 @@
 #include "Iris/ReplicationState/ReplicationStateUtil.h"
 #include "Iris/ReplicationSystem/ReplicationSystem.h"
 #include "Net/Core/NetBitArray.h"
-#include "CoreMinimal.h"
+#include "Net/Core/Trace/NetDebugName.h"
+#include "CoreTypes.h"
 #include "UObject/UnrealType.h"
 #include "Iris/Core/IrisLog.h"
 #include "UObject/PropertyPortFlags.h"
diff --git a/Engine/Source/Runtime/Experimental/Iris/Core/Private/Iris/ReplicationState/ReplicationStateDescriptorBuilder.cpp b/Engine/Source/Runtime/Experimental/Iris/Core/Private/Iris/ReplicationState/ReplicationStateDescriptorBuilder.cpp
--- a/Engine/Source/Runtime/Experimental/Iris/Core/Private/Iris/ReplicationState/ReplicationStateDescriptorBuilder.cpp
+++ b/Engine/Source/Runtime/Experimental/Iris/Core/Private/Iris/ReplicationState/ReplicationStateDescriptorBuilder.cpp
@@ -1,7 +1,7 @@
 // Copyright Epic Games, Inc. All Rights Reserved.
 
 #include "Iris/ReplicationState/ReplicationStateDescriptorBuilder.h"
-#include "CoreMinimal.h"
+#include "CoreTypes.h"
 #include "Iris/Core/BitTwiddling.h"
 #include "Iris/Core/IrisLog.h"
 #include "Iris/Core/IrisProfiler.h"
diff --git a/Engine/Source/Runtime/Experimental/Iris/Core/Private/Iris/ReplicationSystem/Conditionals/ReplicationConditionals.cpp b/Engine/Source/Runtime/Experimental/Iris/Core/Private/Iris/ReplicationSystem/Conditionals/ReplicationConditionals.cpp
--- a/Engine/Source/Runtime/Experimental/Iris/Core/Private/Iris/ReplicationSystem/Conditionals/ReplicationConditionals.cpp
+++ b/Engine/Source/Runtime/Experimental/Iris/Core/Private/Iris/ReplicationSystem/Conditionals/ReplicationConditionals.cpp
@@ -16,6 +16,7 @@
 #include "UObject/CoreNetTypes.h"
 #include "Net/Core/PropertyConditions/RepChangedPropertyTracker.h"
 #include "Net/Core/PropertyConditions/PropertyConditions.h"
+#include "Net/Core/Trace/NetDebugName.h"
 
 namespace UE::Net::Private
 {
diff --git a/Engine/Source/Runtime/Experimental/Iris/Core/Private/Iris/ReplicationSystem/Filtering/NetObjectGroups.h b/Engine/Source/Runtime/Experimental/Iris/Core/Private/Iris/ReplicationSystem/Filtering/NetObjectGroups.h
--- a/Engine/Source/Runtime/Experimental/Iris/Core/Private/Iris/ReplicationSystem/Filtering/NetObjectGroups.h
+++ b/Engine/Source/Runtime/Experimental/Iris/Core/Private/Iris/ReplicationSystem/Filtering/NetObjectGroups.h
@@ -4,7 +4,9 @@
 
 #include "CoreTypes.h"
 #include "Containers/Array.h"
+#include "Containers/Map.h"
 #include "Containers/SparseArray.h"
+#include "UObject/NameTypes.h"
 
 namespace UE::Net
 {
diff --git a/Engine/Source/Runtime/Experimental/Iris/Core/Private/Iris/ReplicationSystem/NetBlob/NetBlobManager.cpp b/Engine/Source/Runtime/Experimental/Iris/Core/Private/Iris/ReplicationSystem/NetBlob/NetBlobManager.cpp
--- a/Engine/Source/Runtime/Experimental/Iris/Core/Private/Iris/ReplicationSystem/NetBlob/NetBlobManager.cpp
+++ b/Engine/Source/Runtime/Experimental/Iris/Core/Private/Iris/ReplicationSystem/NetBlob/NetBlobManager.cpp
@@ -11,6 +11,7 @@
 #include "Iris/ReplicationSystem/ReplicationWriter.h"
 #include "Iris/Core/NetObjectReference.h"
 #include "Iris/Core/IrisLog.h"
+#include "Net/Core/Trace/NetDebugName.h"
 
 namespace UE::Net::Private
 {
diff --git a/Engine/Source/Runtime/Experimental/Iris/Core/Private/Iris/ReplicationSystem/NetHandle.cpp b/Engine/Source/Runtime/Experimental/Iris/Core/Private/Iris/ReplicationSystem/NetHandle.cpp
new file mode 100644
--- /dev/null
+++ b/Engine/Source/Runtime/Experimental/Iris/Core/Private/Iris/ReplicationSystem/NetHandle.cpp
@@ -0,0 +1,44 @@
+// Copyright Epic Games, Inc. All Rights Reserved.
+
+#include "Iris/ReplicationSystem/NetHandle.h"
+#include "Containers/UnrealString.h"
+#include "Misc/StringBuilder.h"
+
+namespace UE::Net
+{
+
+FString FNetHandle::ToString() const
+{
+	FString Result;
+#if UE_NET_ALLOW_MULTIPLE_REPLICATION_SYSTEMS
+	const uint32 ReplicationSystemIdToDisplay = ReplicationSystemId - 1U;
+	Result = FString::Printf(TEXT("NetHandle (Id=%u):(RepSystemId=%u)"), GetId(), ReplicationSystemIdToDisplay);
+#else
+	Result = FString::Printf(TEXT("NetHandle (Id=%u)"), GetId());
+#endif
+	return Result;
+}
+
+}
+
+FStringBuilderBase& operator<<(FStringBuilderBase& Builder, const UE::Net::FNetHandle& NetHandle) 
+{ 	
+#if UE_NET_ALLOW_MULTIPLE_REPLICATION_SYSTEMS
+	const uint32 ReplicationSystemIdToDisplay = NetHandle.GetReplicationSystemId() - 1U;
+	Builder.Appendf(TEXT("NetHandle (Id=%u):(RepSystemId=%u)"), NetHandle.GetId(), ReplicationSystemIdToDisplay);
+#else
+	Builder.Appendf(TEXT("NetHandle (Id=%u)"), NetHandle.GetId());
+#endif
+	return Builder;
+}
+
+FAnsiStringBuilderBase& operator<<(FAnsiStringBuilderBase& Builder, const UE::Net::FNetHandle& NetHandle)
+{
+	#if UE_NET_ALLOW_MULTIPLE_REPLICATION_SYSTEMS
+		const uint32 ReplicationSystemIdToDisplay = NetHandle.GetReplicationSystemId() - 1U;
+		Builder.Appendf("NetHandle (Id=%u):(RepSystemId=%u)", NetHandle.GetId(), ReplicationSystemIdToDisplay);
+	#else
+		Builder.Appendf("NetHandle (Id=%u)", NetHandle.GetId());
+	#endif
+		return Builder;
+}
diff --git a/Engine/Source/Runtime/Experimental/Iris/Core/Private/Iris/ReplicationSystem/NetHandleManager.cpp b/Engine/Source/Runtime/Experimental/Iris/Core/Private/Iris/ReplicationSystem/NetHandleManager.cpp
--- a/Engine/Source/Runtime/Experimental/Iris/Core/Private/Iris/ReplicationSystem/NetHandleManager.cpp
+++ b/Engine/Source/Runtime/Experimental/Iris/Core/Private/Iris/ReplicationSystem/NetHandleManager.cpp
@@ -10,6 +10,7 @@
 #include "Net/Core/Trace/NetTrace.h"
 #include "UObject/Object.h"
 #include "UObject/UObjectGlobals.h"
+#include "UObject/CoreNetTypes.h"
 
 namespace UE::Net::Private
 {
diff --git a/Engine/Source/Runtime/Experimental/Iris/Core/Private/Iris/ReplicationSystem/NetTokenDataStream.cpp b/Engine/Source/Runtime/Experimental/Iris/Core/Private/Iris/ReplicationSystem/NetTokenDataStream.cpp
--- a/Engine/Source/Runtime/Experimental/Iris/Core/Private/Iris/ReplicationSystem/NetTokenDataStream.cpp
+++ b/Engine/Source/Runtime/Experimental/Iris/Core/Private/Iris/ReplicationSystem/NetTokenDataStream.cpp
@@ -13,6 +13,7 @@
 #include "Iris/ReplicationSystem/ReplicationSystemInternal.h"
 #include "Iris/Core/IrisLog.h"
 #include "Net/Core/Trace/NetTrace.h"
+#include "HAL/IConsoleManager.h"
 
 #if (UE_BUILD_SHIPPING || UE_BUILD_TEST)
 #	define UE_NET_ENABLE_NETTOKEN_LOG 0
diff --git a/Engine/Source/Runtime/Experimental/Iris/Core/Private/Iris/Serialization/NetBitStreamWriter.cpp b/Engine/Source/Runtime/Experimental/Iris/Core/Private/Iris/Serialization/NetBitStreamWriter.cpp
--- a/Engine/Source/Runtime/Experimental/Iris/Core/Private/Iris/Serialization/NetBitStreamWriter.cpp
+++ b/Engine/Source/Runtime/Experimental/Iris/Core/Private/Iris/Serialization/NetBitStreamWriter.cpp
@@ -6,6 +6,7 @@
 #include "HAL/PlatformMath.h"
 #include "HAL/PlatformMemory.h"
 #include "Misc/AssertionMacros.h"
+#include "Templates/AlignmentTemplates.h"
 
 #if UE_NETBITSTREAMWRITER_VALIDATE
 #	define UE_NETBITSTREAMWRITER_CHECK(expr) check(expr)
diff --git a/Engine/Source/Runtime/Experimental/Iris/Core/Public/Iris/DataStream/DataStream.h b/Engine/Source/Runtime/Experimental/Iris/Core/Public/Iris/DataStream/DataStream.h
--- a/Engine/Source/Runtime/Experimental/Iris/Core/Public/Iris/DataStream/DataStream.h
+++ b/Engine/Source/Runtime/Experimental/Iris/Core/Public/Iris/DataStream/DataStream.h
@@ -3,7 +3,6 @@
 #pragma once
 
 #include "CoreTypes.h"
-#include "UObject/NameTypes.h"
 #include "DataStream.generated.h"
 
 namespace UE::Net
diff --git a/Engine/Source/Runtime/Experimental/Iris/Core/Public/Iris/ReplicationState/Private/IrisFastArraySerializerInternal.h b/Engine/Source/Runtime/Experimental/Iris/Core/Public/Iris/ReplicationState/Private/IrisFastArraySerializerInternal.h
--- a/Engine/Source/Runtime/Experimental/Iris/Core/Public/Iris/ReplicationState/Private/IrisFastArraySerializerInternal.h
+++ b/Engine/Source/Runtime/Experimental/Iris/Core/Public/Iris/ReplicationState/Private/IrisFastArraySerializerInternal.h
@@ -2,10 +2,14 @@
 #pragma once
 
 #include "Iris/ReplicationState/IrisFastArraySerializer.h"
-#include "Net/Core/NetBitArray.h"
-#include "Iris/ReplicationState/ReplicationStateUtil.h"
 
 namespace UE::Net 
+{
+	class FNetBitArray;
+	class FNetBitArrayView;
+}
+
+namespace UE::Net	
 {
 
 namespace Private 
@@ -16,8 +20,8 @@ namespace Private
 */
 struct FIrisFastArraySerializerPrivateAccessor
 {
-	static FNetBitArrayView GetChangeMask(FIrisFastArraySerializer& Array) { return MakeNetBitArrayView(&Array.ChangeMaskStorage[0], FIrisFastArraySerializer::IrisFastArrayChangeMaskBits + 1U); }
-	static FNetBitArrayView GetConditionalChangeMask(FIrisFastArraySerializer& Array) { return MakeNetBitArrayView(&Array.ChangeMaskStorage[1], FIrisFastArraySerializer::IrisFastArrayChangeMaskBits + 1U); };
+	static FNetBitArrayView GetChangeMask(FIrisFastArraySerializer& Array);
+	static FNetBitArrayView GetConditionalChangeMask(FIrisFastArraySerializer& Array);
 	static FReplicationStateHeader& GetReplicationStateHeader(FIrisFastArraySerializer& Array) { return Array.ReplicationStateHeader; }
 
 	/*
@@ -37,57 +41,6 @@ struct FIrisFastArraySerializerPrivateAccessor
 	static void MarkArrayItemDirty(FIrisFastArraySerializer& Array, int32 Index);
 };
 
-inline void FIrisFastArraySerializerPrivateAccessor::MarkAllArrayItemsDirty(FIrisFastArraySerializer& Array, uint32 StartingIndex)
-{
-	checkSlow(Array.ReplicationStateHeader.IsBound());
-
-	FNetBitArrayView MemberChangeMask = UE::Net::Private::FIrisFastArraySerializerPrivateAccessor::GetChangeMask(Array);
-	if (!MemberChangeMask.GetBit(0))
-	{
-		MarkNetObjectStateDirty(Array.ReplicationStateHeader);
-	}
-	if (StartingIndex == 0)
-	{
-		MemberChangeMask.SetAllBits();
-	}
-	else
-	{
-		MemberChangeMask.SetBit(FIrisFastArraySerializer::IrisFastArrayPropertyBitIndex);
-		MemberChangeMask.SetBits(FIrisFastArraySerializer::IrisFastArrayChangeMaskBitOffset + StartingIndex, FIrisFastArraySerializer::IrisFastArrayChangeMaskBits - StartingIndex);
-	}
-}
-
-inline void FIrisFastArraySerializerPrivateAccessor::MarkArrayDirty(FIrisFastArraySerializer& Array)
-{
-	checkSlow(Array.ReplicationStateHeader.IsBound());
-
-	FNetBitArrayView MemberChangeMask = UE::Net::Private::FIrisFastArraySerializerPrivateAccessor::GetChangeMask(Array);
-
-	// Dirty object unless already dirty, we only use the array bit for this purpose
-	if (!MemberChangeMask.GetBit(FIrisFastArraySerializer::IrisFastArrayPropertyBitIndex))
-	{
-		MarkNetObjectStateDirty(Array.ReplicationStateHeader);
-		MemberChangeMask.SetBit(FIrisFastArraySerializer::IrisFastArrayPropertyBitIndex);
-	}
-}
-
-inline void FIrisFastArraySerializerPrivateAccessor::MarkArrayItemDirty(FIrisFastArraySerializer& Array, int32 ItemIdx)
-{
-	checkSlow(Array.ReplicationStateHeader.IsBound());
-
-	// Mark changemask dirty for this item
-	// We are using a modulo scheme for dirtiness
-	FNetBitArrayView MemberChangeMask = UE::Net::Private::FIrisFastArraySerializerPrivateAccessor::GetChangeMask(Array);
-	MemberChangeMask.SetBit((ItemIdx % FIrisFastArraySerializer::IrisFastArrayChangeMaskBits) + FIrisFastArraySerializer::IrisFastArrayChangeMaskBitOffset);
-
-	// Dirty object unless already dirty, we only use the array bit for this purpose
-	if (!MemberChangeMask.GetBit(0))
-	{
-		MarkNetObjectStateDirty(Array.ReplicationStateHeader);
-		MemberChangeMask.SetBits(0, 1);
-	}
-}
-
 } // end namespace Private
 
 /**
diff --git a/Engine/Source/Runtime/Experimental/Iris/Core/Public/Iris/ReplicationState/PropertyReplicationState.h b/Engine/Source/Runtime/Experimental/Iris/Core/Public/Iris/ReplicationState/PropertyReplicationState.h
--- a/Engine/Source/Runtime/Experimental/Iris/Core/Public/Iris/ReplicationState/PropertyReplicationState.h
+++ b/Engine/Source/Runtime/Experimental/Iris/Core/Public/Iris/ReplicationState/PropertyReplicationState.h
@@ -3,9 +3,13 @@
 #pragma once
 
 #include "CoreTypes.h"
-#include "Iris/ReplicationState/ReplicationStateDescriptor.h"
 #include "Templates/RefCounting.h"
- 
+
+namespace UE::Net
+{
+	struct FReplicationStateDescriptor;
+}
+
 namespace UE::Net
 {
 
diff --git a/Engine/Source/Runtime/Experimental/Iris/Core/Public/Iris/ReplicationState/ReplicationStateDescriptor.h b/Engine/Source/Runtime/Experimental/Iris/Core/Public/Iris/ReplicationState/ReplicationStateDescriptor.h
--- a/Engine/Source/Runtime/Experimental/Iris/Core/Public/Iris/ReplicationState/ReplicationStateDescriptor.h
+++ b/Engine/Source/Runtime/Experimental/Iris/Core/Public/Iris/ReplicationState/ReplicationStateDescriptor.h
@@ -4,7 +4,6 @@
 
 #include "CoreTypes.h"
 #include "Misc/EnumClassFlags.h"
-#include "Net/Core/Trace/NetDebugName.h"
 #include "Templates/TypeHash.h"
 #include <atomic>
 
@@ -13,6 +12,7 @@ class FProperty;
 class FString;
 namespace UE::Net
 {
+	struct FNetDebugName;
 	struct FNetSerializer;
 	typedef uint64 FRepTag;
 	struct FReplicationStateDescriptor;
diff --git a/Engine/Source/Runtime/Experimental/Iris/Core/Public/Iris/ReplicationSystem/FastArrayReplicationFragment.h b/Engine/Source/Runtime/Experimental/Iris/Core/Public/Iris/ReplicationSystem/FastArrayReplicationFragment.h
--- a/Engine/Source/Runtime/Experimental/Iris/Core/Public/Iris/ReplicationSystem/FastArrayReplicationFragment.h
+++ b/Engine/Source/Runtime/Experimental/Iris/Core/Public/Iris/ReplicationSystem/FastArrayReplicationFragment.h
@@ -7,7 +7,9 @@
 #include "Iris/ReplicationState/IrisFastArraySerializer.h"
 #include "Iris/ReplicationState/Private/IrisFastArraySerializerInternal.h"
 #include "Iris/ReplicationState/PropertyReplicationState.h"
+#include "Iris/ReplicationState/ReplicationStateUtil.h"
 #include "Iris/ReplicationState/ReplicationStateDescriptor.h"
+#include "Net/Core/NetBitArray.h"
 
 namespace UE::Net
 {
diff --git a/Engine/Source/Runtime/Experimental/Iris/Core/Public/Iris/ReplicationSystem/NetBlob/PartialNetBlob.h b/Engine/Source/Runtime/Experimental/Iris/Core/Public/Iris/ReplicationSystem/NetBlob/PartialNetBlob.h
--- a/Engine/Source/Runtime/Experimental/Iris/Core/Public/Iris/ReplicationSystem/NetBlob/PartialNetBlob.h
+++ b/Engine/Source/Runtime/Experimental/Iris/Core/Public/Iris/ReplicationSystem/NetBlob/PartialNetBlob.h
@@ -4,6 +4,7 @@
 #include "Iris/ReplicationSystem/NetBlob/NetBlob.h"
 #include "Iris/ReplicationSystem/NetBlob/RawDataNetBlob.h"
 #include "Iris/Core/NetObjectReference.h"
+#include "Net/Core/Trace/NetDebugName.h"
 
 namespace UE::Net
 {
diff --git a/Engine/Source/Runtime/Experimental/Iris/Core/Public/Iris/ReplicationSystem/NetHandle.h b/Engine/Source/Runtime/Experimental/Iris/Core/Public/Iris/ReplicationSystem/NetHandle.h
--- a/Engine/Source/Runtime/Experimental/Iris/Core/Public/Iris/ReplicationSystem/NetHandle.h
+++ b/Engine/Source/Runtime/Experimental/Iris/Core/Public/Iris/ReplicationSystem/NetHandle.h
@@ -3,13 +3,13 @@
 #pragma once
 
 #include "CoreTypes.h"
-#include "Containers/UnrealString.h"
-#include "HAL/Platform.h"
 #include "Iris/IrisConfig.h"
-#include "Misc/StringBuilder.h"
+#include "Containers/StringFwd.h"
 #include "Templates/TypeHash.h"
 
 // Forward declarations
+//class FString;
+
 namespace UE::Net::Private
 {
 	class FNetHandleManager;
@@ -45,7 +45,7 @@ public:
 	bool operator<(const FNetHandle& Other)const { return Id < Other.Id; }
 	bool operator!=(const FNetHandle& Other)const { return Id != Other.Id; }
 
-	FString ToString() const;
+	IRISCORE_API FString ToString() const;
 
 	static bool FullCompare(FNetHandle A, FNetHandle B) { return A.Value == B.Value; }
 	
@@ -83,41 +83,7 @@ FORCEINLINE uint32 GetObjectIdForNetTrace(const FNetHandle& Handle)
 	return Handle.GetId();
 }
 
-inline FString FNetHandle::ToString() const
-{
-	FString Result;
-#if UE_NET_ALLOW_MULTIPLE_REPLICATION_SYSTEMS
-	const uint32 ReplicationSystemIdToDisplay = ReplicationSystemId - 1U;
-	Result = FString::Printf(TEXT("NetHandle (Id=%u):(RepSystemId=%u)"), GetId(), ReplicationSystemIdToDisplay);
-#else
-	Result = FString::Printf(TEXT("NetHandle (Id=%u)"), GetId());
-#endif
-	return Result;
-}
-
-}
-
-inline FStringBuilderBase& operator<<(FStringBuilderBase& Builder, const UE::Net::FNetHandle& NetHandle) 
-{ 	
-#if UE_NET_ALLOW_MULTIPLE_REPLICATION_SYSTEMS
-	const uint32 ReplicationSystemIdToDisplay = NetHandle.GetReplicationSystemId() - 1U;
-	Builder.Appendf(TEXT("NetHandle (Id=%u):(RepSystemId=%u)"), NetHandle.GetId(), ReplicationSystemIdToDisplay);
-#else
-	Builder.Appendf(TEXT("NetHandle (Id=%u)"), NetHandle.GetId());
-#endif
-	return Builder;
 }
 
-inline FAnsiStringBuilderBase& operator<<(FAnsiStringBuilderBase& Builder, const UE::Net::FNetHandle& NetHandle)
-{
-	#if UE_NET_ALLOW_MULTIPLE_REPLICATION_SYSTEMS
-		const uint32 ReplicationSystemIdToDisplay = NetHandle.GetReplicationSystemId() - 1U;
-		Builder.Appendf("NetHandle (Id=%u):(RepSystemId=%u)", NetHandle.GetId(), ReplicationSystemIdToDisplay);
-	#else
-		Builder.Appendf("NetHandle (Id=%u)", NetHandle.GetId());
-	#endif
-		return Builder;
-}
-
-
-
+FStringBuilderBase& operator<<(FStringBuilderBase& Builder, const UE::Net::FNetHandle& NetHandle);
+FAnsiStringBuilderBase& operator<<(FAnsiStringBuilderBase& Builder, const UE::Net::FNetHandle& NetHandle);
diff --git a/Engine/Source/Runtime/Experimental/Iris/Core/Public/Iris/ReplicationSystem/ObjectReplicationBridge.h b/Engine/Source/Runtime/Experimental/Iris/Core/Public/Iris/ReplicationSystem/ObjectReplicationBridge.h
--- a/Engine/Source/Runtime/Experimental/Iris/Core/Public/Iris/ReplicationSystem/ObjectReplicationBridge.h
+++ b/Engine/Source/Runtime/Experimental/Iris/Core/Public/Iris/ReplicationSystem/ObjectReplicationBridge.h
@@ -4,10 +4,7 @@
 
 #include "CoreTypes.h"
 #include "Net/Core/NetBitArray.h"
-#include "Iris/Core/NetObjectReference.h"
 #include "Iris/ReplicationSystem/ReplicationBridge.h"
-#include "Iris/ReplicationState/ReplicationStateDescriptor.h"
-#include "Iris/ReplicationSystem/ReplicationFragment.h"
 #include "Delegates/IDelegateInstance.h"
 
 #include "ObjectReplicationBridge.generated.h"
diff --git a/Engine/Source/Runtime/Experimental/Iris/Core/Public/Iris/ReplicationSystem/Private/FastArrayReplicationFragmentInternal.h b/Engine/Source/Runtime/Experimental/Iris/Core/Public/Iris/ReplicationSystem/Private/FastArrayReplicationFragmentInternal.h
--- a/Engine/Source/Runtime/Experimental/Iris/Core/Public/Iris/ReplicationSystem/Private/FastArrayReplicationFragmentInternal.h
+++ b/Engine/Source/Runtime/Experimental/Iris/Core/Public/Iris/ReplicationSystem/Private/FastArrayReplicationFragmentInternal.h
@@ -8,6 +8,7 @@
 #include "Iris/ReplicationState/Private/IrisFastArraySerializerInternal.h"
 #include "Iris/ReplicationState/PropertyReplicationState.h"
 #include "Iris/ReplicationState/ReplicationStateDescriptor.h"
+#include "Net/Core/Trace/NetDebugName.h"
 
 namespace UE {
 namespace Net {
diff --git a/Engine/Source/Runtime/Experimental/Iris/Core/Public/Iris/ReplicationSystem/ReplicationFragmentUtil.h b/Engine/Source/Runtime/Experimental/Iris/Core/Public/Iris/ReplicationSystem/ReplicationFragmentUtil.h
--- a/Engine/Source/Runtime/Experimental/Iris/Core/Public/Iris/ReplicationSystem/ReplicationFragmentUtil.h
+++ b/Engine/Source/Runtime/Experimental/Iris/Core/Public/Iris/ReplicationSystem/ReplicationFragmentUtil.h
@@ -1,13 +1,18 @@
 // Copyright Epic Games, Inc. All Rights Reserved.
 
 #pragma once
-#include "Iris/ReplicationSystem/ReplicationFragment.h"
+#include "CoreTypes.h"
+#include "Containers/ContainersFwd.h"
 
 class UObject;
 
 namespace UE::Net
 {
 
+class FFragmentRegistrationContext;
+class FReplicationFragment;
+enum class EFragmentRegistrationFlags : uint32;
+
 /**
 * FPropertyReplicationFragment - used to bind PropertyReplicationStates to their owner
 */
diff --git a/Engine/Source/Runtime/Experimental/Iris/Core/Public/Iris/ReplicationSystem/TypedReplicationFragment.h b/Engine/Source/Runtime/Experimental/Iris/Core/Public/Iris/ReplicationSystem/TypedReplicationFragment.h
--- a/Engine/Source/Runtime/Experimental/Iris/Core/Public/Iris/ReplicationSystem/TypedReplicationFragment.h
+++ b/Engine/Source/Runtime/Experimental/Iris/Core/Public/Iris/ReplicationSystem/TypedReplicationFragment.h
@@ -3,8 +3,6 @@
 #pragma once
 
 #include "Iris/ReplicationSystem/ReplicationFragment.h"
-#include "Iris/ReplicationState/ReplicationStateDescriptor.h"
-#include "Iris/ReplicationSystem/ReplicationOperations.h"
 
 namespace UE::Net
 {
