Fix comment.

diff --git a/Engine/Source/Runtime/Experimental/Iris/Core/Public/Iris/ReplicationSystem/NetHandle.h b/Engine/Source/Runtime/Experimental/Iris/Core/Public/Iris/ReplicationSystem/NetHandle.h
--- a/Engine/Source/Runtime/Experimental/Iris/Core/Public/Iris/ReplicationSystem/NetHandle.h
+++ b/Engine/Source/Runtime/Experimental/Iris/Core/Public/Iris/ReplicationSystem/NetHandle.h
@@ -8,7 +8,7 @@
 #include "Templates/TypeHash.h"
 
 // Forward declarations
-//class FString;
+class FString;
 
 namespace UE::Net::Private
 {
