Geometry collection : added support for ISMPools, it only works if an ISMPool actor is assigned to the GC instance
Limitations : Currently this requires the GC to be reset to recompute the list of instanciable meshes and their corresponding indices
Mitigation : this can be turned off using the following cvar : p.Chaos.GC.UseISMPool

diff --git a/Engine/Source/Runtime/Experimental/GeometryCollectionEngine/Private/GeometryCollection/GeometryCollectionComponent.cpp b/Engine/Source/Runtime/Experimental/GeometryCollectionEngine/Private/GeometryCollection/GeometryCollectionComponent.cpp
--- a/Engine/Source/Runtime/Experimental/GeometryCollectionEngine/Private/GeometryCollection/GeometryCollectionComponent.cpp
+++ b/Engine/Source/Runtime/Experimental/GeometryCollectionEngine/Private/GeometryCollection/GeometryCollectionComponent.cpp
@@ -21,6 +21,8 @@
 #include "GeometryCollection/GeometryCollectionSQAccelerator.h"
 #include "GeometryCollection/GeometryCollectionSceneProxy.h"
 #include "GeometryCollection/GeometryCollectionUtility.h"
+#include "GeometryCollection/GeometryCollectionISMPoolActor.h"
+#include "GeometryCollection/GeometryCollectionISMPoolComponent.h"
 #include "Math/Sphere.h"
 #include "Modules/ModuleManager.h"
 #include "Net/Core/PushModel/PushModel.h"
@@ -87,6 +89,12 @@ static FAutoConsoleVariableRef CVarChaosBoxCalcBoundsISPCEnabled(TEXT("p.Chaos.B
 bool bChaos_GC_CacheComponentSpaceBounds = true;
 FAutoConsoleVariableRef CVarChaosGCCacheComponentSpaceBounds(TEXT("p.Chaos.GC.CacheComponentSpaceBounds"), bChaos_GC_CacheComponentSpaceBounds, TEXT("Cache component space bounds for performance"));
 
+bool bChaos_GC_UseISMPool = true;
+FAutoConsoleVariableRef CVarChaosGCUseISMPool(TEXT("p.Chaos.GC.UseISMPool"), bChaos_GC_UseISMPool, TEXT("When enabled, use the ISM pool if specified"));
+
+bool bChaos_GC_UseISMPoolForNonFracturedParts = true;
+FAutoConsoleVariableRef CVarChaosGCUseISMPoolForNonFracturedParts(TEXT("p.Chaos.GC.UseISMPoolForNonFracturedParts"), bChaos_GC_UseISMPoolForNonFracturedParts, TEXT("When enabled, non fractured part will use the ISM pool if specified"));
+
 DEFINE_LOG_CATEGORY_STATIC(UGCC_LOG, Error, All);
 
 extern FGeometryCollectionDynamicDataPool GDynamicDataPool;
@@ -334,51 +342,11 @@ void UGeometryCollectionComponent::BeginPlay()
 	}
 #endif
 
-	//////////////////////////////////////////////////////////////////////////
-	// Commenting out these callbacks for now due to the threading model. The callbacks here
-	// expect the rest collection to be mutable which is not the case when running in multiple
-	// threads. Ideally we have some separate animation collection or track that we cache to
-	// without affecting the data we've dispatched to the physics thread
-	//////////////////////////////////////////////////////////////////////////
-	// ---------- SolverCallbacks->SetResetAnimationCacheFunction([&]()
-	// ---------- {
-	// ---------- 	FGeometryCollectionEdit Edit = EditRestCollection();
-	// ---------- 	Edit.GetRestCollection()->RecordedData.SetNum(0);
-	// ---------- });
-	// ---------- SolverCallbacks->SetUpdateTransformsFunction([&](const TArrayView<FTransform>&)
-	// ---------- {
-	// ---------- 	// todo : Move the update to the array passed here...
-	// ---------- });
-	// ---------- 
-	// ---------- SolverCallbacks->SetUpdateRestStateFunction([&](const int32 & CurrentFrame, const TManagedArray<int32> & RigidBodyID, const TManagedArray<FGeometryCollectionBoneNode>& Hierarchy, const FSolverCallbacks::FParticlesType& Particles)
-	// ---------- {
-	// ---------- 	FGeometryCollectionEdit Edit = EditRestCollection();
-	// ---------- 	UGeometryCollection * RestCollection = Edit.GetRestCollection();
-	// ---------- 	check(RestCollection);
-	// ---------- 
-	// ---------- 	if (CurrentFrame >= RestCollection->RecordedData.Num())
-	// ---------- 	{
-	// ---------- 		RestCollection->RecordedData.SetNum(CurrentFrame + 1);
-	// ---------- 		RestCollection->RecordedData[CurrentFrame].SetNum(RigidBodyID.Num());
-	// ---------- 		ParallelFor(RigidBodyID.Num(), [&](int32 i)
-	// ---------- 		{
-	// ---------- 			if (!Hierarchy[i].Children.Num())
-	// ---------- 			{
-	// ---------- 				RestCollection->RecordedData[CurrentFrame][i].SetTranslation(Particles.X(RigidBodyID[i]));
-	// ---------- 				RestCollection->RecordedData[CurrentFrame][i].SetRotation(Particles.R(RigidBodyID[i]));
-	// ---------- 			}
-	// ---------- 			else
-	// ---------- 			{
-	// ---------- 				RestCollection->RecordedData[CurrentFrame][i].SetTranslation(FVector::ZeroVector);
-	// ---------- 				RestCollection->RecordedData[CurrentFrame][i].SetRotation(FQuat::Identity);
-	// ---------- 			}
-	// ---------- 		});
-	// ---------- 	}
-	// ---------- });
-	//////////////////////////////////////////////////////////////////////////
-
 	// default current cache time
 	CurrentCacheTime = MAX_flt;	
+
+	// we only enable ISM if we are playing ( not in editing mode because of various side effect like selection )
+	RegisterToISMPool();
 }
 
 
@@ -389,6 +357,8 @@ void UGeometryCollectionComponent::EndPlay(const EEndPlayReason::Type ReasonEnd)
 	EditorActor = EditorUtilities::GetEditorWorldCounterpartActor(GetTypedOuter<AActor>());
 #endif
 
+	UnregisterFromISMPool();
+
 	Super::EndPlay(ReasonEnd);
 
 	CurrentCacheTime = MAX_flt;
@@ -2142,7 +2112,10 @@ void UGeometryCollectionComponent::UpdateAttachedChildrenTransform() const
 	{
 	 	for (const TObjectPtr<USceneComponent>& AttachedChild: this->GetAttachChildren())
 	 	{
-	 		AttachedChild->UpdateComponentToWorld();
+			if (AttachedChild)
+			{
+				AttachedChild->UpdateComponentToWorld();
+			}
 	 	}
 	}
 }
@@ -2192,6 +2165,10 @@ void UGeometryCollectionComponent::TickComponent(float DeltaTime, enum ELevelTic
 			{
 				// #todo review: When we've made changes to ISMC, we need to move this function call to SetRenderDynamicData_Concurrent
 				RefreshEmbeddedGeometry();
+				
+				// we may want to call this when the geometry collection updates ( notified by the proxy buffer updates ) 
+				// otherwise we are getting a frame delay 
+				RefreshISMPoolInstances();
 
 				if (SceneProxy && SceneProxy->IsNaniteMesh())
 				{
@@ -2612,7 +2589,7 @@ void UGeometryCollectionComponent::SendRenderDynamicData_Concurrent()
 	Super::SendRenderDynamicData_Concurrent();
 
 	// Only update the dynamic data if the dynamic collection is dirty
-	if (SceneProxy && ((DynamicCollection && DynamicCollection->IsDirty()) || CachePlayback)) 
+	if (SceneProxy && ((DynamicCollection && DynamicCollection->IsDirty()) || CachePlayback))
 	{
 		FGeometryCollectionDynamicData* DynamicData = InitDynamicData(false /* initialization */);
 
@@ -3823,6 +3800,167 @@ void UGeometryCollectionComponent::InitializeEmbeddedGeometry()
 	}
 }
 
+bool UGeometryCollectionComponent::CanUseISMPool() const 
+{
+	return bChaos_GC_UseISMPool && ISMPool;
+}
+
+void UGeometryCollectionComponent::RegisterToISMPool()
+{
+	UnregisterFromISMPool();
+
+	if (CanUseISMPool())
+	{
+		if (UGeometryCollectionISMPoolComponent* ISMPoolComp = ISMPool->GetISMPoolComp())
+		{
+			bool bCanRenderComponent = true;
+			if (RestCollection)
+			{
+				if (bChaos_GC_UseISMPoolForNonFracturedParts)
+				{
+					if (RestCollection->GetGeometryCollection())
+					{
+						// if we use ISM pool for the hierarchy we must hide the component for rendering 
+						bCanRenderComponent = false;
+
+						// fisrt count the instance per mesh 
+						TArray<int32> InstanceCounts;
+						InstanceCounts.AddZeroed(RestCollection->AutoInstanceMeshes.Num());
+						const TManagedArray<int32>* AutoInstanceMeshIndices = RestCollection->GetGeometryCollection()->FindAttribute<int32>("AutoInstanceMeshIndex", FGeometryCollection::TransformGroup);
+						const TManagedArray<TSet<int32>>& Children = RestCollection->GetGeometryCollection()->Children;
+						if (AutoInstanceMeshIndices)
+						{
+							for (int32 TransformIndex = 0; TransformIndex < AutoInstanceMeshIndices->Num(); TransformIndex++)
+							{
+								const int32 AutoInstanceMeshIndex = (*AutoInstanceMeshIndices)[TransformIndex];
+								if (Children[TransformIndex].Num() == 0)
+								{
+									InstanceCounts[AutoInstanceMeshIndex]++;
+								}
+							}
+						}
+
+						// now register each mesh 
+						ISMPoolMeshGroupIndex = ISMPoolComp->CreateMeshGroup();
+						for (int32 MeshIndex = 0; MeshIndex < RestCollection->AutoInstanceMeshes.Num(); MeshIndex++)
+						{
+							const FGeometryCollectionAutoInstanceMesh& AutoInstanceMesh = RestCollection->AutoInstanceMeshes[MeshIndex];
+							if (UStaticMesh* StaticMesh = Cast<UStaticMesh>(AutoInstanceMesh.StaticMesh.TryLoad()))
+							{
+								bool bMaterialOverride = false;
+								for (int32 MatIndex = 0; MatIndex < AutoInstanceMesh.Materials.Num(); MatIndex++)
+								{
+									const UMaterialInterface* OriginalMaterial = StaticMesh->GetMaterial(MatIndex);
+									if (OriginalMaterial != AutoInstanceMesh.Materials[MatIndex])
+									{
+										bMaterialOverride = true;
+										break;
+									}
+								}
+								FGeometryCollectionStaticMeshInstance StaticMeshInstance;
+								StaticMeshInstance.StaticMesh = StaticMesh;
+								if (bMaterialOverride)
+								{
+									StaticMeshInstance.MaterialsOverrides = AutoInstanceMesh.Materials;
+								}
+								ISMPoolComp->AddMeshToGroup(ISMPoolMeshGroupIndex, StaticMeshInstance, InstanceCounts[MeshIndex]);
+							}
+						}
+					}
+				}
+
+				// root proxy if available 
+				// TODO : if ISM pool is not available : uses a standard static mesh component
+				if (UObject* RootMeshProxyObject = RestCollection->RootProxy.ResolveObject())
+				{
+					if (UStaticMesh* RootMeshProxy = Cast<UStaticMesh>(RootMeshProxyObject))
+					{
+						// if we use a mesh proxy hide the component for rendering 
+						bCanRenderComponent = false;
+
+						// TODO : store this state in a variable
+
+						FGeometryCollectionStaticMeshInstance StaticMeshInstance;
+						StaticMeshInstance.StaticMesh = RootMeshProxy;
+						// todo : get the mesh index and use it to upadte the mesh transform 
+						ISMPoolComp->AddMeshToGroup(ISMPoolMeshGroupIndex, StaticMeshInstance, 1);
+					}
+				}
+			}
+
+			SetVisibility(bCanRenderComponent);
+
+			RefreshISMPoolInstances();
+		}
+	}
+}
+
+void UGeometryCollectionComponent::UnregisterFromISMPool()
+{
+	if (ISMPool)
+	{
+		if (UGeometryCollectionISMPoolComponent* ISMPoolComp = ISMPool->GetISMPoolComp())
+		{
+			ISMPoolComp->DestroyMeshGroup(ISMPoolMeshGroupIndex);
+			ISMPoolMeshGroupIndex = INDEX_NONE;
+		}
+	}
+	SetVisibility(true);
+}
+
+void UGeometryCollectionComponent::RefreshISMPoolInstances()
+{
+	if (CanUseISMPool())
+	{
+		if (UGeometryCollectionISMPoolComponent* ISMPoolComp = ISMPool->GetISMPoolComp())
+		{
+			if (RestCollection)
+			{
+				// default to true for editor purposes?
+				//const bool bCollectionIsDirty = DynamicCollection ? DynamicCollection->IsDirty() : true;
+				if (bChaos_GC_UseISMPoolForNonFracturedParts /*&& bCollectionIsDirty*/)
+				{
+					if (RestCollection->GetGeometryCollection())
+					{
+						const TManagedArray<int32>* AutoInstanceMeshIndices = RestCollection->GetGeometryCollection()->FindAttribute<int32>("AutoInstanceMeshIndex", FGeometryCollection::TransformGroup);
+						const TManagedArray<TSet<int32>>& Children = RestCollection->GetGeometryCollection()->Children;
+						if (AutoInstanceMeshIndices)
+						{
+							const int32 NumTransforms = RestCollection->NumElements(FGeometryCollection::TransformAttribute);
+							ensure(AutoInstanceMeshIndices->Num() == NumTransforms);
+
+							TArray<FTransform> InstanceTransforms;
+							CalculateGlobalMatrices();
+
+							const FTransform& ComponentTransform = GetComponentTransform();
+							for (int32 MeshIndex = 0; MeshIndex < RestCollection->AutoInstanceMeshes.Num(); MeshIndex++)
+							{
+								InstanceTransforms.Reset(NumTransforms); // Allocate for worst case
+								for (int32 TransformIndex = 0; TransformIndex < NumTransforms; TransformIndex++)
+								{
+									const int32 AutoInstanceMeshIndex = (*AutoInstanceMeshIndices)[TransformIndex];
+									if (AutoInstanceMeshIndex == MeshIndex && Children[TransformIndex].Num() == 0)
+									{
+										InstanceTransforms.Add(FTransform(GlobalMatrices[TransformIndex]) * ComponentTransform);
+									}
+								}
+								constexpr bool bWorlSpace = true;
+								constexpr bool bMarkRenderStateDirty = true;
+								constexpr bool bTeleport = true;
+								ISMPoolComp->BatchUpdateInstancesTransforms(ISMPoolMeshGroupIndex, MeshIndex, 0, InstanceTransforms, bWorlSpace, bMarkRenderStateDirty, bTeleport);
+							}
+				
+						}
+					}
+				}
+
+				// todo : update the mesh proxy if set 
+			}
+		}
+	}
+}
+
+
 struct FGeometryCollectionDecayContext
 {
 	FGeometryCollectionDecayContext(FGeometryCollectionPhysicsProxy& PhysicsProxyIn, FGeometryCollectionDecayDynamicFacade& DecayFacadeIn)
diff --git a/Engine/Source/Runtime/Experimental/GeometryCollectionEngine/Private/GeometryCollection/GeometryCollectionEngineConversion.cpp b/Engine/Source/Runtime/Experimental/GeometryCollectionEngine/Private/GeometryCollection/GeometryCollectionEngineConversion.cpp
--- a/Engine/Source/Runtime/Experimental/GeometryCollectionEngine/Private/GeometryCollection/GeometryCollectionEngineConversion.cpp
+++ b/Engine/Source/Runtime/Experimental/GeometryCollectionEngine/Private/GeometryCollection/GeometryCollectionEngineConversion.cpp
@@ -428,9 +428,25 @@ bool FGeometryCollectionEngineConversion::AppendStaticMesh(const UStaticMesh* St
 	FGeometryCollection* GeometryCollection = GeometryCollectionPtr.Get();
 	check(GeometryCollection);
 
+	const int32 OriginalNumOfTransforms = GeometryCollection->NumElements(FGeometryCollection::TransformGroup);
+
 	if (AppendStaticMesh(StaticMesh, StartMaterialIndex, StaticMeshTransform, GeometryCollection, bReindexMaterials, bAddInternalMaterials, bSplitComponents))
 	{
 		AppendMaterials(Materials, GeometryCollectionObject, bAddInternalMaterials);
+
+		// add index to the auto instanced meshes array
+		const int32 NewNumOfTransforms = GeometryCollection->NumElements(FGeometryCollection::TransformGroup);
+		if (NewNumOfTransforms > OriginalNumOfTransforms)
+		{
+			TManagedArray<int32>& AutoInstanceMeshIndices = GeometryCollection->AddAttribute<int32>("AutoInstanceMeshIndex", FGeometryCollection::TransformGroup);
+			const int32 AutoInstanceMeshIndex = GeometryCollectionObject->FindOrAddAutoInstanceMesh(*StaticMesh, Materials);
+
+			for (int32 TransformIndex = OriginalNumOfTransforms; TransformIndex < NewNumOfTransforms; TransformIndex++)
+			{
+				AutoInstanceMeshIndices[TransformIndex] = AutoInstanceMeshIndex;
+			}
+		}
+
 		return true;
 	}
 
@@ -839,7 +855,12 @@ void FGeometryCollectionEngineConversion::AppendGeometryCollection(const UGeomet
 
 	int32 MaterialStart = AppendMaterials(Materials, TargetGeometryCollectionObject, false);
 
-	AppendGeometryCollection(SourceGeometryCollectionPtr.Get(), MaterialStart, GeometryCollectionTransform, GeometryCollection, bReindexMaterials);
+	const int32 TargetTransformStart = GeometryCollectionPtr->NumElements(FGeometryCollection::TransformGroup);
+
+	if (AppendGeometryCollection(SourceGeometryCollectionPtr.Get(), MaterialStart, GeometryCollectionTransform, GeometryCollection, bReindexMaterials))
+	{
+		AppendGeometryCollectionInstancedMeshes(SourceGeometryCollection, TargetGeometryCollectionObject, TargetTransformStart);
+	}
 }
 
 
@@ -892,6 +913,34 @@ int32 FGeometryCollectionEngineConversion::AppendGeometryCollectionMaterials(con
 	return AppendMaterials(Materials, TargetGeometryCollectionObject, false);
 }
 
+void FGeometryCollectionEngineConversion::AppendGeometryCollectionInstancedMeshes(const UGeometryCollection* SourceGeometryCollectionObject, UGeometryCollection* TargetGeometryCollectionObject, int32 TargetTransformStartIndex)
+{
+	TSharedPtr<const FGeometryCollection, ESPMode::ThreadSafe> SourceGeometryCollectionPtr = SourceGeometryCollectionObject->GetGeometryCollection();
+	TSharedPtr<FGeometryCollection, ESPMode::ThreadSafe> TargetGeometryCollectionPtr = TargetGeometryCollectionObject->GetGeometryCollection();
+
+	if (SourceGeometryCollectionPtr && TargetGeometryCollectionPtr)
+	{
+		if (const TManagedArray<int32>* SourceAutoInstanceMeshIndices = SourceGeometryCollectionPtr->FindAttribute<int32>("AutoInstanceMeshIndex", FGeometryCollection::TransformGroup))
+		{
+			TManagedArray<int32>& TargetAutoInstanceMeshIndices = TargetGeometryCollectionPtr->AddAttribute<int32>("AutoInstanceMeshIndex", FGeometryCollection::TransformGroup);
+
+			for (int32 SourceTransformIndex = 0; SourceTransformIndex < SourceAutoInstanceMeshIndices->Num(); SourceTransformIndex++)
+			{
+				const int32 TargettransformIndex = TargetTransformStartIndex + SourceTransformIndex;
+				TargetAutoInstanceMeshIndices[TargettransformIndex] = INDEX_NONE;
+				if (SourceAutoInstanceMeshIndices)
+				{
+					const int32 SourceAutoInstanceIndex = (*SourceAutoInstanceMeshIndices)[SourceTransformIndex];
+					if (SourceAutoInstanceIndex != INDEX_NONE)
+					{
+						const FGeometryCollectionAutoInstanceMesh& SourceAutoInstanceMesh = SourceGeometryCollectionObject->GetAutoInstanceMesh(SourceAutoInstanceIndex);
+						TargetAutoInstanceMeshIndices[TargettransformIndex] = TargetGeometryCollectionObject->FindOrAddAutoInstanceMesh(SourceAutoInstanceMesh);
+					}
+				}
+			}
+		}
+	}
+}
 
 void FGeometryCollectionEngineConversion::AppendGeometryCollection(const UGeometryCollection* SourceGeometryCollection, const UGeometryCollectionComponent* GeometryCollectionComponent, const FTransform& GeometryCollectionTransform, UGeometryCollection* TargetGeometryCollectionObject, bool bReindexMaterials)
 {
@@ -909,7 +958,12 @@ void FGeometryCollectionEngineConversion::AppendGeometryCollection(const UGeomet
 	FGeometryCollection* GeometryCollection = GeometryCollectionPtr.Get();
 	check(GeometryCollection);
 
-	AppendGeometryCollection(SourceGeometryCollectionPtr.Get(), MaterialStartIndex, GeometryCollectionTransform, GeometryCollection, bReindexMaterials);
+	const int32 TargetTransformStart = GeometryCollectionPtr->NumElements(FGeometryCollection::TransformGroup);
+
+	if (AppendGeometryCollection(SourceGeometryCollectionPtr.Get(), MaterialStartIndex, GeometryCollectionTransform, GeometryCollection, bReindexMaterials))
+	{
+		AppendGeometryCollectionInstancedMeshes(SourceGeometryCollection, TargetGeometryCollectionObject, TargetTransformStart);
+	}
 }
 
 
diff --git a/Engine/Source/Runtime/Experimental/GeometryCollectionEngine/Private/GeometryCollection/GeometryCollectionISMPoolActor.cpp b/Engine/Source/Runtime/Experimental/GeometryCollectionEngine/Private/GeometryCollection/GeometryCollectionISMPoolActor.cpp
new file mode 100644
--- /dev/null
+++ b/Engine/Source/Runtime/Experimental/GeometryCollectionEngine/Private/GeometryCollection/GeometryCollectionISMPoolActor.cpp
@@ -0,0 +1,11 @@
+// Copyright Epic Games, Inc. All Rights Reserved.
+
+#include "GeometryCollection/GeometryCollectionISMPoolActor.h"
+#include "GeometryCollection/GeometryCollectionISMPoolComponent.h"
+
+AGeometryCollectionISMPoolActor::AGeometryCollectionISMPoolActor(const FObjectInitializer& ObjectInitializer)
+	: Super(ObjectInitializer)
+{
+	ISMPoolComp = CreateDefaultSubobject<UGeometryCollectionISMPoolComponent>(TEXT("ISMPoolComp"));
+	RootComponent = ISMPoolComp;
+}
diff --git a/Engine/Source/Runtime/Experimental/GeometryCollectionEngine/Private/GeometryCollection/GeometryCollectionISMPoolComponent.cpp b/Engine/Source/Runtime/Experimental/GeometryCollectionEngine/Private/GeometryCollection/GeometryCollectionISMPoolComponent.cpp
new file mode 100644
--- /dev/null
+++ b/Engine/Source/Runtime/Experimental/GeometryCollectionEngine/Private/GeometryCollection/GeometryCollectionISMPoolComponent.cpp
@@ -0,0 +1,193 @@
+// Copyright Epic Games, Inc. All Rights Reserved.
+
+#include "GeometryCollection/GeometryCollectionISMPoolComponent.h"
+
+#include "Components/HierarchicalInstancedStaticMeshComponent.h"
+
+#include "GeometryCollection/GeometryCollectionComponent.h"
+
+FGeometryCollectionMeshGroup::FMeshId FGeometryCollectionMeshGroup::AddMesh(const FGeometryCollectionStaticMeshInstance& MeshInstance, int32 InstanceCount, const FGeometryCollectionMeshInfo& ISMInstanceInfo)
+{
+	FMeshId* MeshIndex = Meshes.Find(MeshInstance);
+	if (MeshIndex)
+	{
+		return *MeshIndex;
+	}
+
+	const FMeshId MeshInfoIndex = MeshInfos.Emplace(ISMInstanceInfo);
+	Meshes.Add(MeshInstance, MeshInfoIndex);
+	return MeshInfoIndex;
+}
+
+bool FGeometryCollectionMeshGroup::BatchUpdateInstancesTransforms(FGeometryCollectionISMPool& ISMPool, FMeshId MeshId, int32 StartInstanceIndex, const TArray<FTransform>& NewInstancesTransforms, bool bWorldSpace, bool bMarkRenderStateDirty, bool bTeleport)
+{
+	if (MeshInfos.IsValidIndex(MeshId))
+	{
+		return ISMPool.BatchUpdateInstancesTransforms(MeshInfos[MeshId], StartInstanceIndex, NewInstancesTransforms, bWorldSpace, bMarkRenderStateDirty, bTeleport);
+	}
+	UE_LOG(LogChaos, Warning, TEXT("UGeometryCollectionISMPoolComponent : Invalid mesh Id (%d) for this mesh group"), MeshId);
+	return false;
+}
+
+void FGeometryCollectionMeshGroup::RemoveAllMeshes(FGeometryCollectionISMPool& ISMPool)
+{
+	for (const FGeometryCollectionMeshInfo& MeshInfo: MeshInfos)
+	{
+		ISMPool.RemoveISM(MeshInfo);
+	}
+	MeshInfos.Empty();
+	Meshes.Empty();
+}
+
+FGeometryCollectionISM::FGeometryCollectionISM(AActor* OwmingActor, const FGeometryCollectionStaticMeshInstance& MeshInstance)
+{
+	check(MeshInstance.StaticMesh);
+	check(OwmingActor);
+
+	const FName ISMName = MakeUniqueObjectName(OwmingActor, UHierarchicalInstancedStaticMeshComponent::StaticClass(), MeshInstance.StaticMesh->GetFName());
+	if (UInstancedStaticMeshComponent* ISMC = NewObject<UHierarchicalInstancedStaticMeshComponent>(OwmingActor, ISMName, RF_Transient | RF_DuplicateTransient))
+	{
+		ISMC->SetStaticMesh(MeshInstance.StaticMesh);
+		// material overrides
+		for (int32 MaterialIndex = 0; MaterialIndex < MeshInstance.MaterialsOverrides.Num(); MaterialIndex++)
+		{
+			ISMC->SetMaterial(MaterialIndex, MeshInstance.MaterialsOverrides[MaterialIndex]);
+		}
+		ISMC->SetCullDistances(0, 0);
+		ISMC->SetCanEverAffectNavigation(false);
+		ISMC->SetCollisionProfileName(UCollisionProfile::NoCollision_ProfileName);
+		ISMC->SetCastShadow(true);
+		ISMC->SetMobility(EComponentMobility::Stationary);
+		OwmingActor->AddInstanceComponent(ISMC);
+		ISMC->RegisterComponent();
+		
+		ISMComponent = ISMC;
+	}
+}
+
+int32 FGeometryCollectionISM::AddInstanceGroup(int32 InstanceCount)
+{
+	const int32 InstanceGroupIndex = InstanceGroups.AddGroup(InstanceCount);
+	const FInstanceGroups::FInstanceGroupRange& NewInstanceGroup = InstanceGroups.GetGroup(InstanceGroupIndex);
+	const int32 TotalInstanceCount = NewInstanceGroup.Start + NewInstanceGroup.Count;
+	ISMComponent->PreAllocateInstancesMemory(TotalInstanceCount);
+	for (int32 InstanceIndex = NewInstanceGroup.Start; InstanceIndex < TotalInstanceCount; InstanceIndex++)
+	{
+		ISMComponent->AddInstance(FTransform::Identity, true);
+	}
+	return InstanceGroupIndex;
+}
+
+FGeometryCollectionMeshInfo FGeometryCollectionISMPool::AddISM(UGeometryCollectionISMPoolComponent* OwningComponent, const FGeometryCollectionStaticMeshInstance& MeshInstance, int32 InstanceCount)
+{
+	FGeometryCollectionMeshInfo Info;
+
+	FISMIndex* ISMIndex = MeshToISMIndex.Find(MeshInstance);
+	if (!ISMIndex)
+	{
+		Info.ISMIndex = ISMs.Emplace(OwningComponent->GetOwner(), MeshInstance);
+		MeshToISMIndex.Add(MeshInstance, Info.ISMIndex);
+	}
+	else
+	{
+		Info.ISMIndex = *ISMIndex;
+	}
+	// add to the ISM 
+	Info.InstanceGroupIndex = ISMs[Info.ISMIndex].AddInstanceGroup(InstanceCount);
+	return Info;
+}
+
+bool FGeometryCollectionISMPool::BatchUpdateInstancesTransforms(FGeometryCollectionMeshInfo& MeshInfo, int32 StartInstanceIndex, const TArray<FTransform>& NewInstancesTransforms, bool bWorldSpace, bool bMarkRenderStateDirty, bool bTeleport)
+{
+	if (ISMs.IsValidIndex(MeshInfo.ISMIndex))
+	{
+		FGeometryCollectionISM& ISM = ISMs[MeshInfo.ISMIndex];
+		const FInstanceGroups::FInstanceGroupRange& InstanceGroup = ISM.InstanceGroups.GetGroup(MeshInfo.InstanceGroupIndex);
+		ensure((StartInstanceIndex + NewInstancesTransforms.Num()) <= InstanceGroup.Count);
+
+		const int32 StartIndex = InstanceGroup.Start + StartInstanceIndex;
+		return ISM.ISMComponent->BatchUpdateInstancesTransforms(StartIndex, NewInstancesTransforms, bWorldSpace, bMarkRenderStateDirty, bTeleport);
+	}
+	UE_LOG(LogChaos, Warning, TEXT("UGeometryCollectionISMPoolComponent : Invalid ISM Id (%d) when updating the transform "), MeshInfo.ISMIndex);
+	return false;
+}
+
+void FGeometryCollectionISMPool::RemoveISM(const FGeometryCollectionMeshInfo& MeshInfo)
+{
+	if (ISMs.IsValidIndex(MeshInfo.ISMIndex))
+	{
+		FGeometryCollectionISM& ISM = ISMs[MeshInfo.ISMIndex];
+		const FInstanceGroups::FInstanceGroupRange& InstanceGroup = ISM.InstanceGroups.GetGroup(MeshInfo.InstanceGroupIndex);
+
+		// todo: be able to remove a range of instance may be useful
+		TArray<int32> InstancesToRemove;
+		InstancesToRemove.Reserve(InstanceGroup.Count);
+		const int32 StartIndex = InstanceGroup.Start;
+		const int32 LastIndex = InstanceGroup.Start + InstanceGroup.Count;
+		for (int32 InstanceIndex = StartIndex; InstanceIndex < LastIndex; InstanceIndex++)
+		{
+			InstancesToRemove.Add(InstanceIndex);
+		}
+		ISM.ISMComponent->RemoveInstances(InstancesToRemove);
+		ISM.InstanceGroups.RemoveGroup(MeshInfo.InstanceGroupIndex);
+	}
+}
+
+void FGeometryCollectionISMPool::Clear()
+{
+	MeshToISMIndex.Reset();
+	if (ISMs.Num() > 0)
+	{
+		if (AActor* OwningActor = ISMs[0].ISMComponent->GetOwner())
+		{
+			for(FGeometryCollectionISM& ISM : ISMs)
+			{
+				ISM.ISMComponent->UnregisterComponent();
+				ISM.ISMComponent->DestroyComponent();
+				OwningActor->RemoveInstanceComponent(ISM.ISMComponent);
+			}
+		}
+		ISMs.Reset();
+	}
+}
+
+UGeometryCollectionISMPoolComponent::UGeometryCollectionISMPoolComponent(const FObjectInitializer& ObjectInitializer)
+	: NextMeshGroupId(0)
+{
+}
+
+UGeometryCollectionISMPoolComponent::FMeshGroupId  UGeometryCollectionISMPoolComponent::CreateMeshGroup()
+{
+	MeshGroups.Add(NextMeshGroupId);
+	return NextMeshGroupId++;
+}
+
+void UGeometryCollectionISMPoolComponent::DestroyMeshGroup(FMeshGroupId MeshGroupId)
+{
+	if (FGeometryCollectionMeshGroup* MeshGroup = MeshGroups.Find(MeshGroupId))
+	{
+		MeshGroup->RemoveAllMeshes(Pool);
+		MeshGroups.Remove(MeshGroupId);
+	}
+}
+
+UGeometryCollectionISMPoolComponent::FMeshId UGeometryCollectionISMPoolComponent::AddMeshToGroup(FMeshGroupId MeshGroupId, const FGeometryCollectionStaticMeshInstance& MeshInstance, int32 InstanceCount)
+{
+	if (FGeometryCollectionMeshGroup* MeshGroup = MeshGroups.Find(MeshGroupId))
+	{
+		const FGeometryCollectionMeshInfo ISMInstanceInfo = Pool.AddISM(this, MeshInstance, InstanceCount);
+		return MeshGroup->AddMesh(MeshInstance, InstanceCount, ISMInstanceInfo);
+	}
+	UE_LOG(LogChaos, Warning, TEXT("UGeometryCollectionISMPoolComponent : Trying to add a mesh to a mesh group (%d) that does not exists"), MeshGroupId);
+	return INDEX_NONE;
+}
+
+bool UGeometryCollectionISMPoolComponent::BatchUpdateInstancesTransforms(FMeshGroupId MeshGroupId, FMeshId MeshId, int32 StartInstanceIndex, const TArray<FTransform>& NewInstancesTransforms, bool bWorldSpace, bool bMarkRenderStateDirty, bool bTeleport)
+{
+	if (FGeometryCollectionMeshGroup* MeshGroup = MeshGroups.Find(MeshGroupId))
+	{
+		return MeshGroup->BatchUpdateInstancesTransforms(Pool, MeshId, StartInstanceIndex, NewInstancesTransforms, bWorldSpace, bMarkRenderStateDirty, bTeleport);
+	}
+	UE_LOG(LogChaos, Warning, TEXT("UGeometryCollectionISMPoolComponent : Trying to update instance with mesh group (%d) that not exists"), MeshGroupId);
+	return false;
+}
diff --git a/Engine/Source/Runtime/Experimental/GeometryCollectionEngine/Private/GeometryCollection/GeometryCollectionObject.cpp b/Engine/Source/Runtime/Experimental/GeometryCollectionEngine/Private/GeometryCollection/GeometryCollectionObject.cpp
--- a/Engine/Source/Runtime/Experimental/GeometryCollectionEngine/Private/GeometryCollection/GeometryCollectionObject.cpp
+++ b/Engine/Source/Runtime/Experimental/GeometryCollectionEngine/Private/GeometryCollection/GeometryCollectionObject.cpp
@@ -362,6 +362,7 @@ void UGeometryCollection::Reset()
 		GeometryCollection->Empty();
 		Materials.Empty();
 		EmbeddedGeometryExemplar.Empty();
+		AutoInstanceMeshes.Empty();
 		InvalidateCollection();
 	}
 }
@@ -1145,6 +1146,73 @@ void UGeometryCollection::RemoveExemplars(const TArray<int32>& SortedRemovalIndi
 	}
 }
 
+/** find or add a auto instance mesh and return its index */
+const FGeometryCollectionAutoInstanceMesh& UGeometryCollection::GetAutoInstanceMesh(int32 AutoInstanceMeshIndex) const
+{
+	return AutoInstanceMeshes[AutoInstanceMeshIndex];
+}
+
+/**  find or add a auto instance mesh from another one and return its index */
+int32 UGeometryCollection::FindOrAddAutoInstanceMesh(const FGeometryCollectionAutoInstanceMesh& AutoInstanecMesh)
+{
+	int32 ReturnedIndex = INDEX_NONE;
+
+	for (int32 MeshIndex = 0; MeshIndex < AutoInstanceMeshes.Num(); MeshIndex++)
+	{
+		const FGeometryCollectionAutoInstanceMesh& Mesh = AutoInstanceMeshes[MeshIndex];
+		if (Mesh.StaticMesh == AutoInstanecMesh.StaticMesh && Mesh.Materials == AutoInstanecMesh.Materials)
+		{
+			ReturnedIndex = MeshIndex;
+			break;
+		}
+	}
+	if (ReturnedIndex == INDEX_NONE)
+	{
+		ReturnedIndex = AutoInstanceMeshes.Add(AutoInstanecMesh);
+	}
+	return ReturnedIndex;
+}
+
+int32 UGeometryCollection::FindOrAddAutoInstanceMesh(const UStaticMesh& StaticMesh, const TArray<UMaterialInterface*>& MeshMaterials)
+{
+	int32 ReturnedIndex = INDEX_NONE;
+
+	FSoftObjectPath StaticMeshSoftPath(&StaticMesh);
+
+	for (int32 MeshIndex = 0; MeshIndex < AutoInstanceMeshes.Num(); MeshIndex++)
+	{
+		const FGeometryCollectionAutoInstanceMesh& Mesh = AutoInstanceMeshes[MeshIndex];
+		if (Mesh.StaticMesh == StaticMeshSoftPath)
+		{
+			if (Mesh.Materials.Num() == MeshMaterials.Num())
+			{
+				bool MaterialAreAllTheSame = true;
+				for (int32 MaterialIndex = 0; MaterialIndex < MeshMaterials.Num(); MaterialIndex++)
+				{
+					if (Mesh.Materials[MaterialIndex] != MeshMaterials[MaterialIndex])
+					{
+						MaterialAreAllTheSame = false;
+						break;
+					}
+				}
+				if (MaterialAreAllTheSame)
+				{
+					ReturnedIndex = MeshIndex;
+					break;
+				}
+			}
+		}
+	}
+	if (ReturnedIndex == INDEX_NONE)
+	{
+		FGeometryCollectionAutoInstanceMesh NewMesh;
+		NewMesh.StaticMesh = StaticMeshSoftPath;
+		NewMesh.Materials = MeshMaterials;
+		ReturnedIndex = AutoInstanceMeshes.Emplace(NewMesh);
+	}
+	return ReturnedIndex;
+}
+
 FGuid UGeometryCollection::GetIdGuid() const
 {
 	return PersistentGuid;
diff --git a/Engine/Source/Runtime/Experimental/GeometryCollectionEngine/Public/GeometryCollection/GeometryCollectionComponent.h b/Engine/Source/Runtime/Experimental/GeometryCollectionEngine/Public/GeometryCollection/GeometryCollectionComponent.h
--- a/Engine/Source/Runtime/Experimental/GeometryCollectionEngine/Public/GeometryCollection/GeometryCollectionComponent.h
+++ b/Engine/Source/Runtime/Experimental/GeometryCollectionEngine/Public/GeometryCollection/GeometryCollectionComponent.h
@@ -40,6 +40,7 @@ class UInstancedStaticMeshComponent;
 class FGeometryCollectionDecayDynamicFacade;
 struct FGeometryCollectionDecayContext;
 struct FDamageCollector;
+class AGeometryCollectionISMPoolActor;
 
 DECLARE_DYNAMIC_MULTICAST_DELEGATE_OneParam(FOnChaosBreakEvent, const FChaosBreakEvent&, BreakEvent);
 
@@ -994,6 +995,12 @@ protected:
 	FGuid RunTimeDataCollectionGuid;
 #endif
 	
+	/** ISM pool to use to render the geometry collection - only works for unfractured geometry collections  */
+	UPROPERTY(EditAnywhere, Category = "ChaosPhysics|Rendering", meta = (DisplayName = "ISM Pool"))
+	TObjectPtr<AGeometryCollectionISMPoolActor> ISMPool;
+
+	int32 ISMPoolMeshGroupIndex = INDEX_NONE;
+
 	/** Populate the static geometry structures for the render thread. */
 	void InitConstantData(FGeometryCollectionConstantData* ConstantData) const;
 
@@ -1120,6 +1127,19 @@ private:
 	bool IsEmbeddedGeometryValid() const;
 	void ClearEmbeddedGeometry();
 
+	/** return true is an ISM pool is set and the feature is enabled */
+	bool CanUseISMPool() const;
+
+	/** If an ISM pool is set, register to it  */
+	void RegisterToISMPool();
+
+	/** If an ISM pool is set, unregister to it  */
+	void UnregisterFromISMPool();
+
+	/** update ISM transforms */
+	void RefreshISMPoolInstances();
+
+
 	void IncrementSleepTimer(float DeltaTime);
 	void IncrementBreakTimer(float DeltaTime);
 	bool CalculateInnerSphere(int32 TransformIndex, UE::Math::TSphere<double>& SphereOut) const;
diff --git a/Engine/Source/Runtime/Experimental/GeometryCollectionEngine/Public/GeometryCollection/GeometryCollectionEngineConversion.h b/Engine/Source/Runtime/Experimental/GeometryCollectionEngine/Public/GeometryCollection/GeometryCollectionEngineConversion.h
--- a/Engine/Source/Runtime/Experimental/GeometryCollectionEngine/Public/GeometryCollection/GeometryCollectionEngineConversion.h
+++ b/Engine/Source/Runtime/Experimental/GeometryCollectionEngine/Public/GeometryCollection/GeometryCollectionEngineConversion.h
@@ -127,6 +127,8 @@ public:
 
 	static int32 AppendGeometryCollectionMaterials(const UGeometryCollection* SourceGeometryCollection, const UGeometryCollectionComponent* GeometryCollectionComponent, UGeometryCollection* TargetGeometryCollectionObject);
 
+	static void AppendGeometryCollectionInstancedMeshes(const UGeometryCollection* SourceGeometryCollectionObject, UGeometryCollection* TargetGeometryCollectionObject, int32 TargetTransformStartIndex);
+
 	/**
 	 * Appends a GeometryCollection to another GeometryCollection.
 	 * @param SourceGeometryCollection : Const GeometryCollection to read vertex/normals/index data from
diff --git a/Engine/Source/Runtime/Experimental/GeometryCollectionEngine/Public/GeometryCollection/GeometryCollectionISMPoolActor.h b/Engine/Source/Runtime/Experimental/GeometryCollectionEngine/Public/GeometryCollection/GeometryCollectionISMPoolActor.h
new file mode 100644
--- /dev/null
+++ b/Engine/Source/Runtime/Experimental/GeometryCollectionEngine/Public/GeometryCollection/GeometryCollectionISMPoolActor.h
@@ -0,0 +1,24 @@
+// Copyright Epic Games, Inc. All Rights Reserved.
+#pragma once
+
+#include "GameFramework/Actor.h"
+#include "GeometryCollectionISMPoolActor.generated.h"
+
+class UGeometryCollectionISMPoolComponent;
+
+UCLASS(ConversionRoot, ComponentWrapperClass)
+class GEOMETRYCOLLECTIONENGINE_API AGeometryCollectionISMPoolActor : public AActor
+{
+	GENERATED_UCLASS_BODY()
+
+private:
+	UPROPERTY(Category = ISMPoolActor, VisibleAnywhere, BlueprintReadOnly, meta = (AllowPrivateAccess = "true"))
+	TObjectPtr<class UGeometryCollectionISMPoolComponent> ISMPoolComp;
+
+public:
+	/** Returns ISMPoolComp subobject **/
+	UGeometryCollectionISMPoolComponent* GetISMPoolComp() const { return ISMPoolComp; }
+};
+
+
+
diff --git a/Engine/Source/Runtime/Experimental/GeometryCollectionEngine/Public/GeometryCollection/GeometryCollectionISMPoolComponent.h b/Engine/Source/Runtime/Experimental/GeometryCollectionEngine/Public/GeometryCollection/GeometryCollectionISMPoolComponent.h
new file mode 100644
--- /dev/null
+++ b/Engine/Source/Runtime/Experimental/GeometryCollectionEngine/Public/GeometryCollection/GeometryCollectionISMPoolComponent.h
@@ -0,0 +1,228 @@
+// Copyright Epic Games, Inc. All Rights Reserved.
+
+#pragma once
+
+#if 1
+
+#include "Components/SceneComponent.h"
+#include "Containers/Map.h"
+#include "GeometryCollectionISMPoolComponent.generated.h"
+
+class AActor;
+class UGeometryComponent;
+class UMaterialInterface;
+class UInstancedStaticMeshComponent;
+class UGeometryCollectionISMPoolComponent;
+
+
+struct FGeometryCollectionISMInstance
+{
+	int32 StartIndex;
+	int32 Count;
+};
+
+struct FInstanceGroups
+{
+public:
+	using FInstanceGroupId = int32;
+
+	struct FInstanceGroupRange
+	{
+		FInstanceGroupRange(int32 InStart, int32 InCount)
+			: Start(InStart)
+			, Count(InCount)
+		{}
+		int32 Start;
+		int32 Count;
+	};
+
+	//FInstanceGroupId AddGroup(int32 Count)
+	//{
+	//	int32 StartIndex = 0;
+	//	if (GroupRanges.Num())
+	//	{
+	//		const FInstanceGroupRange& LastGroupRange = GroupRanges.Last();
+	//		StartIndex = LastGroupRange.Start + LastGroupRange.Count;
+	//	}
+	//	return GroupRanges.Emplace(FInstanceGroupRange(StartIndex, Count));
+	//}
+
+	//void RemoveGroup(FInstanceGroupId GroupId)
+	//{
+	//	check(GroupRanges.IsValidIndex(BlockId));
+	//	const int32 StartOffset = GroupRanges[GroupId].Count;
+	//	FreeList.Add(GroupId);
+	//	// 
+	//	for (int32 Index = (GroupId + 1); Index < GroupRanges.Num(); Index++)
+	//	{
+	//		const FInstanceGroupRange& Curr = GroupRanges[Index];
+	//		FInstanceGroupRange& Prev = GroupRanges[Index - 1];
+	//		Prev.Start = Curr.Start - StartOffset;
+	//		Prev.Count = Curr.Count;
+	//	}
+	//}
+
+	FInstanceGroupId AddGroup(int32 Count)
+	{
+		const int32 StartIndex = InstancesCount;
+		InstancesCount += Count;
+		GroupRanges.Emplace(NextGroupId, FInstanceGroupRange(StartIndex, Count));
+		return NextGroupId++;
+	}
+
+	void RemoveGroup(FInstanceGroupId GroupId)
+	{
+		check(GroupRanges.Contains(GroupId));
+		const FInstanceGroupRange& GroupRangeToRemove = GroupRanges[GroupId];
+
+		// we need now to shift all the groups above the remove one
+		for (TPair<FInstanceGroupId, FInstanceGroupRange>& GroupRange: GroupRanges)
+		{
+			if (GroupRange.Value.Start > GroupRangeToRemove.Start)
+			{
+				GroupRange.Value.Start -= GroupRangeToRemove.Count;
+			}
+		}
+		// now remove the range safely
+		InstancesCount -= GroupRangeToRemove.Count;
+		GroupRanges.Remove(GroupId);
+	}
+
+	const FInstanceGroupRange& GetGroup(int32 GroupIndex) const { return GroupRanges[GroupIndex]; };
+
+private:
+	int32 InstancesCount = 0;
+	int32 NextGroupId = 0;
+	TMap<FInstanceGroupId, FInstanceGroupRange> GroupRanges;
+};
+
+/**
+* This represent a unique mesh with potentially overriden materials
+* if the array is empty , there's no overrides
+*/
+struct FGeometryCollectionStaticMeshInstance
+{
+	UStaticMesh* StaticMesh = nullptr;
+	TArray<UMaterialInterface*> MaterialsOverrides;
+
+	bool operator==(const FGeometryCollectionStaticMeshInstance& Other) const 
+	{
+		if (StaticMesh == Other.StaticMesh)
+		{
+			if (MaterialsOverrides.Num() == Other.MaterialsOverrides.Num())
+			{
+				for (int32 MatIndex = 0; MatIndex < MaterialsOverrides.Num(); MatIndex++)
+				{
+					const FName MatName = MaterialsOverrides[MatIndex] ? MaterialsOverrides[MatIndex]->GetFName() : NAME_None;
+					const FName OtherName = Other.MaterialsOverrides[MatIndex] ? Other.MaterialsOverrides[MatIndex]->GetFName() : NAME_None;
+					if (MatName != OtherName)
+					{
+						return false;
+					}
+				}
+				return true;
+			}
+		}
+		return false;
+	}
+};
+
+
+FORCEINLINE uint32 GetTypeHash(const FGeometryCollectionStaticMeshInstance& MeshInstance)
+{
+	uint32 CombinedHash = GetTypeHash(MeshInstance.StaticMesh);
+	CombinedHash = HashCombine(CombinedHash, GetTypeHash(MeshInstance.MaterialsOverrides.Num()));
+	for (const UMaterialInterface* Material: MeshInstance.MaterialsOverrides)
+	{
+		CombinedHash = HashCombine(CombinedHash, GetTypeHash(Material));
+	}
+	return CombinedHash;
+}
+
+struct FGeometryCollectionMeshInfo
+{
+	int32 ISMIndex;
+	int32 InstanceGroupIndex;
+};
+
+struct FGeometryCollectionISMPool;
+
+/**
+* FGeometryCollectionMeshGroup
+* a mesh groupo contains various mesh with their instances
+*/
+struct FGeometryCollectionMeshGroup
+{
+	using FMeshId = int32;
+
+	FMeshId AddMesh(const FGeometryCollectionStaticMeshInstance& MeshInstance, int32 InstanceCount, const FGeometryCollectionMeshInfo& ISMInstanceInfo);
+	bool BatchUpdateInstancesTransforms(FGeometryCollectionISMPool& ISMPool, FMeshId MeshId, int32 StartInstanceIndex, const TArray<FTransform>& NewInstancesTransforms, bool bWorldSpace, bool bMarkRenderStateDirty, bool bTeleport);
+	void RemoveAllMeshes(FGeometryCollectionISMPool& ISMPool);
+
+	TMap<FGeometryCollectionStaticMeshInstance, FMeshId> Meshes;
+	TArray<FGeometryCollectionMeshInfo> MeshInfos;
+};
+
+struct FGeometryCollectionISM
+{
+	FGeometryCollectionISM(AActor* OwmingActor, const FGeometryCollectionStaticMeshInstance& MeshInstance);
+
+	int32 AddInstanceGroup(int32 InstanceCount);
+
+	TObjectPtr<UInstancedStaticMeshComponent> ISMComponent;
+	FInstanceGroups InstanceGroups;
+};
+
+
+struct FGeometryCollectionISMPool
+{
+	using FISMIndex = int32;
+
+	FGeometryCollectionMeshInfo AddISM(UGeometryCollectionISMPoolComponent* OwningComponent, const FGeometryCollectionStaticMeshInstance& MeshInstance, int32 InstanceCount);
+	bool BatchUpdateInstancesTransforms(FGeometryCollectionMeshInfo& MeshInfo, int32 StartInstanceIndex, const TArray<FTransform>& NewInstancesTransforms, bool bWorldSpace, bool bMarkRenderStateDirty, bool bTeleport);
+	void RemoveISM(const FGeometryCollectionMeshInfo& MeshInfo);
+
+	/** Clear all ISM components and associated data */
+	void Clear();
+
+	TMap<FGeometryCollectionStaticMeshInstance, FISMIndex> MeshToISMIndex;
+	TArray<FGeometryCollectionISM> ISMs;
+	// Todo : since ISMs index cannot change, we'll need a free list 
+};
+
+
+/**
+* UGeometryCollectionISMPoolComponent
+*   Component that managed a pool of ISM in order to optimize render of geometry collections when not using fracture
+*/
+UCLASS(meta = (BlueprintSpawnableComponent))
+class GEOMETRYCOLLECTIONENGINE_API UGeometryCollectionISMPoolComponent: public USceneComponent
+{
+	GENERATED_UCLASS_BODY()
+
+public:
+	using FMeshGroupId = int32;
+	using FMeshId = int32;
+
+	/** 
+	* Create an Mesh group which represent an arbitrary set of mesh with their instance 
+	* no resources are created until the meshes are added for this group 
+	* return a mesh group Id used to add and update instances
+	*/
+	FMeshGroupId CreateMeshGroup();
+
+	/** destroy  a mesh group and its associated resources */
+	void DestroyMeshGroup(FMeshGroupId MeshGroupId);
+
+	/** Add a static mesh for a nmesh group */
+	FMeshId AddMeshToGroup(FMeshGroupId MeshGroupId, const FGeometryCollectionStaticMeshInstance& MeshInstance, int32 InstanceCount);
+
+	/** Add a static mesh for a nmesh group */
+	bool BatchUpdateInstancesTransforms(FMeshGroupId MeshGroupId, FMeshId MeshId, int32 StartInstanceIndex, const TArray<FTransform>& NewInstancesTransforms, bool bWorldSpace = false, bool bMarkRenderStateDirty = false, bool bTeleport = false);
+
+private:
+	uint32 NextMeshGroupId = 0;
+	TMap<FMeshGroupId, FGeometryCollectionMeshGroup> MeshGroups;
+	FGeometryCollectionISMPool Pool;
+};
+#endif
\ No newline at end of file
diff --git a/Engine/Source/Runtime/Experimental/GeometryCollectionEngine/Public/GeometryCollection/GeometryCollectionObject.h b/Engine/Source/Runtime/Experimental/GeometryCollectionEngine/Public/GeometryCollection/GeometryCollectionObject.h
--- a/Engine/Source/Runtime/Experimental/GeometryCollectionEngine/Public/GeometryCollection/GeometryCollectionObject.h
+++ b/Engine/Source/Runtime/Experimental/GeometryCollectionEngine/Public/GeometryCollection/GeometryCollectionObject.h
@@ -41,6 +41,20 @@ struct GEOMETRYCOLLECTIONENGINE_API FGeometryCollectionSource
 	/** Whether individual source mesh components should be split into separate pieces of geometry based on mesh connectivity. If checked, triangles that are not topologically connected will be assigned separate bones. */
 	UPROPERTY(BlueprintReadWrite, EditAnywhere, Category = "GeometrySource")
 	bool bSplitComponents = false;
+
+	// TODO: add primtive custom data
+};
+
+USTRUCT(BlueprintType)
+struct GEOMETRYCOLLECTIONENGINE_API FGeometryCollectionAutoInstanceMesh
+{
+	GENERATED_BODY()
+
+	UPROPERTY(BlueprintReadWrite, EditAnywhere, Category = "AutoInstance", meta = (AllowedClasses = "/Script/Engine.StaticMesh"))
+	FSoftObjectPath StaticMesh;
+
+	UPROPERTY(BlueprintReadWrite, EditAnywhere, Category = "AutoInstance")
+	TArray<TObjectPtr<UMaterialInterface>> Materials;
 };
 
 USTRUCT(BlueprintType)
@@ -416,6 +430,15 @@ public:
 	/** Remove embedded geometry exemplars with indices matching the sorted removal list. */
 	void RemoveExemplars(const TArray<int32>& SortedRemovalIndices);
 
+	/** find or add a auto instance mesh and return its index */
+	const FGeometryCollectionAutoInstanceMesh& GetAutoInstanceMesh(int32 AutoInstanceMeshIndex) const;
+
+	/**  find or add a auto instance mesh from another one and return its index */
+	int32 FindOrAddAutoInstanceMesh(const FGeometryCollectionAutoInstanceMesh& AutoInstanecMesh);
+
+	/** find or add a auto instance mesh from a mesh and alist of material and return its index */
+	int32 FindOrAddAutoInstanceMesh(const UStaticMesh& StaticMesh, const TArray<UMaterialInterface*>& Materials);
+
 	/** Produce a deep copy of GeometryCollection member, stripped of data unecessary for gameplay. */
 	TSharedPtr<FGeometryCollection, ESPMode::ThreadSafe> GenerateMinimalGeometryCollection() const;
 
@@ -495,6 +518,14 @@ public:
 	UPROPERTY(EditAnywhere, Category = "Rendering")
 	bool bUseFullPrecisionUVs = false;
 
+	/** list of unique static mesh / materials pairs for auto instancing*/
+	UPROPERTY(EditAnywhere, Category = "Rendering")
+	TArray<FGeometryCollectionAutoInstanceMesh> AutoInstanceMeshes;
+
+	/** static mesh to use as a proxy for rendering until the geometry collection is broken */
+	UPROPERTY(BlueprintReadWrite, EditAnywhere, Category = "Rendering", meta = (AllowedClasses = "/Script/Engine.StaticMesh"))
+	FSoftObjectPath RootProxy;
+
 	/**
 	 * Strip unnecessary data from the Geometry Collection to keep the memory footprint as small as possible.
 	 */
