Chaos : fix non-unity build for ISMPool changes

diff --git a/Engine/Source/Runtime/Experimental/GeometryCollectionEngine/Private/GeometryCollection/GeometryCollectionISMPoolComponent.cpp b/Engine/Source/Runtime/Experimental/GeometryCollectionEngine/Private/GeometryCollection/GeometryCollectionISMPoolComponent.cpp
--- a/Engine/Source/Runtime/Experimental/GeometryCollectionEngine/Private/GeometryCollection/GeometryCollectionISMPoolComponent.cpp
+++ b/Engine/Source/Runtime/Experimental/GeometryCollectionEngine/Private/GeometryCollection/GeometryCollectionISMPoolComponent.cpp
@@ -3,6 +3,7 @@
 #include "GeometryCollection/GeometryCollectionISMPoolComponent.h"
 
 #include "Components/HierarchicalInstancedStaticMeshComponent.h"
+#include "Engine/StaticMesh.h"
 
 #include "GeometryCollection/GeometryCollectionComponent.h"
 
diff --git a/Engine/Source/Runtime/Experimental/GeometryCollectionEngine/Public/GeometryCollection/GeometryCollectionISMPoolComponent.h b/Engine/Source/Runtime/Experimental/GeometryCollectionEngine/Public/GeometryCollection/GeometryCollectionISMPoolComponent.h
--- a/Engine/Source/Runtime/Experimental/GeometryCollectionEngine/Public/GeometryCollection/GeometryCollectionISMPoolComponent.h
+++ b/Engine/Source/Runtime/Experimental/GeometryCollectionEngine/Public/GeometryCollection/GeometryCollectionISMPoolComponent.h
@@ -6,11 +6,11 @@
 
 #include "Components/SceneComponent.h"
 #include "Containers/Map.h"
+#include "Materials/MaterialInterface.h"
 #include "GeometryCollectionISMPoolComponent.generated.h"
 
 class AActor;
 class UGeometryComponent;
-class UMaterialInterface;
 class UInstancedStaticMeshComponent;
 class UGeometryCollectionISMPoolComponent;
 
