Chaos - fix RBAN breakable constraints
They were getting re-enabled in a broken state which leads to weirdness.

diff --git a/Engine/Source/Runtime/Experimental/Chaos/Private/Chaos/ChaosDebugDraw.cpp b/Engine/Source/Runtime/Experimental/Chaos/Private/Chaos/ChaosDebugDraw.cpp
--- a/Engine/Source/Runtime/Experimental/Chaos/Private/Chaos/ChaosDebugDraw.cpp
+++ b/Engine/Source/Runtime/Experimental/Chaos/Private/Chaos/ChaosDebugDraw.cpp
@@ -1268,6 +1268,11 @@ namespace Chaos
 
 		void DrawJointConstraintImpl(const FRigidTransform3& SpaceTransform, const FPBDJointConstraintHandle* ConstraintHandle, Chaos::FRealSingle ColorScale, const FChaosDebugDrawJointFeatures& FeatureMask, const FChaosDebugDrawSettings& Settings)
 		{
+			if (!ConstraintHandle->IsEnabled())
+			{
+				return;
+			}
+
 			TVec2<FGeometryParticleHandle*> ConstrainedParticles = ConstraintHandle->GetConstrainedParticles();
 			auto RigidParticle0 = ConstrainedParticles[0]->CastToRigidParticle();
 			auto RigidParticle1 = ConstrainedParticles[1]->CastToRigidParticle();
diff --git a/Engine/Source/Runtime/Experimental/Chaos/Private/Chaos/PBDJointConstraints.cpp b/Engine/Source/Runtime/Experimental/Chaos/Private/Chaos/PBDJointConstraints.cpp
--- a/Engine/Source/Runtime/Experimental/Chaos/Private/Chaos/PBDJointConstraints.cpp
+++ b/Engine/Source/Runtime/Experimental/Chaos/Private/Chaos/PBDJointConstraints.cpp
@@ -71,6 +71,11 @@ namespace Chaos
 		return ConcreteContainer()->GetConstraintColor(ConstraintIndex);
 	}
 
+	bool FPBDJointConstraintHandle::IsConstraintBroken() const
+	{
+		return ConcreteContainer()->IsConstraintBroken(ConstraintIndex);
+	}
+
 	bool FPBDJointConstraintHandle::IsConstraintBreaking() const
 	{
 		return ConcreteContainer()->IsConstraintBreaking(ConstraintIndex);
@@ -329,6 +334,7 @@ namespace Chaos
 		, Color(INDEX_NONE)
 		, IslandSize(0)
 		, bDisabled(false)
+		, bBroken(false)
 		, bBreaking(false)
 		, bDriveTargetChanged(false), LinearImpulse(FVec3(0))
 		, AngularImpulse(FVec3(0))
@@ -579,6 +585,11 @@ namespace Chaos
 		return !ConstraintStates[ConstraintIndex].bDisabled;
 	}
 
+	bool FPBDJointConstraints::IsConstraintBroken(int32 ConstraintIndex) const
+	{
+		return ConstraintStates[ConstraintIndex].bBroken;
+	}
+
 	bool FPBDJointConstraints::IsConstraintBreaking(int32 ConstraintIndex) const
 	{
 		return ConstraintStates[ConstraintIndex].bBreaking;
@@ -607,8 +618,10 @@ namespace Chaos
 		if (bEnabled)
 		{ 
 			// only enable constraint if the particles are valid and not disabled
+			// and if the constraint is not broken
 			if (Particle0->Handle() != nullptr && !Particle0->Disabled()
-				&& Particle1->Handle() != nullptr && !Particle1->Disabled())
+				&& Particle1->Handle() != nullptr && !Particle1->Disabled()
+				&& !IsConstraintBroken(ConstraintIndex))
 			{
 				ConstraintStates[ConstraintIndex].bDisabled = false;
 			}
@@ -620,6 +633,11 @@ namespace Chaos
 		}
 	}
 
+	void FPBDJointConstraints::SetConstraintBroken(int32 ConstraintIndex, bool bBroken)
+	{
+		ConstraintStates[ConstraintIndex].bBroken = bBroken;
+	}
+
 	void FPBDJointConstraints::SetConstraintBreaking(int32 ConstraintIndex, bool bBreaking)
 	{
 		ConstraintStates[ConstraintIndex].bBreaking = bBreaking;
@@ -633,6 +651,7 @@ namespace Chaos
 	void FPBDJointConstraints::BreakConstraint(int32 ConstraintIndex)
 	{
 		SetConstraintEnabled(ConstraintIndex, false);
+		SetConstraintBroken(ConstraintIndex, true);
 		SetConstraintBreaking(ConstraintIndex, true);
 		if (BreakCallback)
 		{
@@ -640,8 +659,9 @@ namespace Chaos
 		}
 	}
 
-	void FPBDJointConstraints::FixConstraints(int32 ConstraintIndex)
+	void FPBDJointConstraints::FixConstraint(int32 ConstraintIndex)
 	{
+		SetConstraintBroken(ConstraintIndex, false);
 		SetConstraintEnabled(ConstraintIndex, true);
 	}
 
diff --git a/Engine/Source/Runtime/Experimental/Chaos/Public/Chaos/PBDJointConstraints.h b/Engine/Source/Runtime/Experimental/Chaos/Public/Chaos/PBDJointConstraints.h
--- a/Engine/Source/Runtime/Experimental/Chaos/Public/Chaos/PBDJointConstraints.h
+++ b/Engine/Source/Runtime/Experimental/Chaos/Public/Chaos/PBDJointConstraints.h
@@ -37,6 +37,7 @@ namespace Chaos
 		int32 GetConstraintColor() const;
 
 		bool IsConstraintEnabled() const;
+		bool IsConstraintBroken() const;
 		bool IsConstraintBreaking() const;
 		void ClearConstraintBreaking();
 		bool IsDriveTargetChanged() const;
@@ -81,6 +82,7 @@ namespace Chaos
 		int32 Color;
 		int32 IslandSize;
 		bool bDisabled;
+		bool bBroken;
 		bool bBreaking;
 		bool bDriveTargetChanged;
 		FVec3 LinearImpulse;
@@ -152,12 +154,17 @@ namespace Chaos
 		bool IsConstraintEnabled(int32 ConstraintIndex) const;
 
 		/*
-		 * Whether the constraint is breaking this frame
+		 * Whether the constraint is broken
+		 */
+		bool IsConstraintBroken(int32 ConstraintIndex) const;
+
+		/*
+		 * Whether the constraint was broken this frame (transient flag for use by event system)
 		 */
 		bool IsConstraintBreaking(int32 ConstraintIndex) const;
 
 		/*
-		 * Clear the constraint braking state
+		 * Clear the transient constraint braking state (called by event system when it has used the flag)
 		 */
 		void ClearConstraintBreaking(int32 ConstraintIndex);
 
@@ -176,11 +183,6 @@ namespace Chaos
 		 */
 		void SetConstraintEnabled(int32 ConstraintIndex, bool bEnabled);
 
-		/*
-		 * Set Breaking State
-		 */
-		void SetConstraintBreaking(int32 ConstraintIndex, bool bBreaking);
-
 		/*
 		* Set Drive Target Changed State
 		*/
@@ -194,7 +196,7 @@ namespace Chaos
 		/**
 		 * Repair a broken constraints (does not adjust particle positions)
 		 */
-		void FixConstraints(int32 ConstraintIndex);
+		void FixConstraint(int32 ConstraintIndex);
 
 		void SetBreakCallback(const FJointBreakCallback& Callback);
 		void ClearBreakCallback();
@@ -297,6 +299,9 @@ namespace Chaos
 		void ApplyBreakThreshold(const FReal Dt, int32 ConstraintIndex, const FVec3& LinearImpulse, const FVec3& AngularImpulse);
 		void ApplyPlasticityLimits(const int32 ConstraintIndex);
 
+		void SetConstraintBroken(int32 ConstraintIndex, bool bBroken);
+		void SetConstraintBreaking(int32 ConstraintIndex, bool bBreaking);
+
 		FPBDJointSolverSettings Settings;
 
 		TArray<FPBDJointSettings> ConstraintSettings;
