Constraints: update the prerequisites dynamically when creating constraint cycles so that we don't propagate that cycle to the tick task manager

    in the context of CR controls, creating a constraint between two controls within the same CR or something like hand_r_ctrl -> static mesh -> hand_l_ctrl creates a cycle.
    the underlying tick functions + tick prerequisites will create a cycle  skel mesh -> const0 -> static mesh -> const1 -> skel mesh
    this cycle creates an undetermined behavior when activating/disactivating constraints as their order within the tick task manager can vary
    to avoid this, we then remove the skel mesh prerequisite of const0 if const1 is active and add it if const1 is off so that the task manager is always evaluating constraints in the right order
    note that this is a particular case for controls as we prevent cycles from being created for other components.

diff --git a/Engine/Source/Runtime/Experimental/Animation/Constraints/Private/ConstraintsActor.cpp b/Engine/Source/Runtime/Experimental/Animation/Constraints/Private/ConstraintsActor.cpp
--- a/Engine/Source/Runtime/Experimental/Animation/Constraints/Private/ConstraintsActor.cpp
+++ b/Engine/Source/Runtime/Experimental/Animation/Constraints/Private/ConstraintsActor.cpp
@@ -81,6 +81,7 @@ void AConstraintsActor::RegisterConstraintsTickFunctions() const
 			for (TObjectPtr<UTickableConstraint> ConstPtr: ConstraintsManager->Constraints)
 			{
 				ConstPtr->ConstraintTick.RegisterTickFunction(Level);
+				ConstPtr->ConstraintTick.Constraint = ConstPtr;
 			}
 		}
 	}
diff --git a/Engine/Source/Runtime/Experimental/Animation/Constraints/Private/ConstraintsManager.cpp b/Engine/Source/Runtime/Experimental/Animation/Constraints/Private/ConstraintsManager.cpp
--- a/Engine/Source/Runtime/Experimental/Animation/Constraints/Private/ConstraintsManager.cpp
+++ b/Engine/Source/Runtime/Experimental/Animation/Constraints/Private/ConstraintsManager.cpp
@@ -1,6 +1,6 @@
 ﻿// Copyright Epic Games, Inc. All Rights Reserved.
 
-#include "ConstraintsManager.h"
+#include "ConstraintsManager.inl"
 #include "ConstraintsActor.h"
 
 #include "Algo/Copy.h"
@@ -427,6 +427,7 @@ bool FConstraintsManagerController::AddConstraint(UTickableConstraint* InConstra
 
 	InConstraint->ConstraintTick.RegisterFunction(InConstraint->GetFunction());
 	InConstraint->ConstraintTick.RegisterTickFunction(World->GetCurrentLevel());
+	InConstraint->ConstraintTick.Constraint = InConstraint;
 
 	// notify
 	Notify(EConstraintsManagerNotifyType::ConstraintAdded, InConstraint);
diff --git a/Engine/Source/Runtime/Experimental/Animation/Constraints/Private/TransformConstraint.cpp b/Engine/Source/Runtime/Experimental/Animation/Constraints/Private/TransformConstraint.cpp
--- a/Engine/Source/Runtime/Experimental/Animation/Constraints/Private/TransformConstraint.cpp
+++ b/Engine/Source/Runtime/Experimental/Animation/Constraints/Private/TransformConstraint.cpp
@@ -2,13 +2,183 @@
 
 #include "TransformConstraint.h"
 #include "TransformableHandle.h"
-#include "ConstraintsManager.h"
+#include "ConstraintsManager.inl"
 #include "TransformableRegistry.h"
 #include "GameFramework/Actor.h"
 #include "MovieSceneSection.h"
+#include "Engine/World.h"
 
 #include UE_INLINE_GENERATED_CPP_BY_NAME(TransformConstraint)
 
+namespace
+{
+
+// ensure that InPossiblePrimary is not depending on InPossibleSecondary to avoid creating cycles 
+bool HasPrerequisiteDependencyWith(const FTickFunction* InSecondary, const FTickFunction* InPrimary)
+{
+	if (!InSecondary || !InPrimary)
+	{
+		return false;
+	}
+
+	// is InSecondary a Prereq of InPrimary?
+	const TArray<FTickPrerequisite>& Prerequisites = InPrimary->GetPrerequisites();
+	const bool bIsSecondaryAPrereq = Prerequisites.ContainsByPredicate([InSecondary](const FTickPrerequisite& Prereq)
+	{
+		const FTickFunction* PrereqFunction = Prereq.Get();
+		return PrereqFunction && (PrereqFunction == InSecondary);
+	});
+
+	if (bIsSecondaryAPrereq)
+	{
+		return true;
+	}
+
+	// otherwise, recurse
+	for (const FTickPrerequisite& Prerequisite : Prerequisites)
+	{
+		if (HasPrerequisiteDependencyWith(InSecondary, Prerequisite.Get()))
+		{
+			return true;
+		}
+	}
+
+	return false;
+}
+
+struct FConstraintCycleChecker
+{
+public:
+	using ConstraintPtr = TObjectPtr<UTickableConstraint>;
+	using ConstraintArray = TArray<ConstraintPtr>;
+
+	/** Checks if this handle is cycle from a tick dependencies perspective. */
+	static bool IsCycling(const TObjectPtr<UTransformableHandle>& InHandle)
+	{
+		if (!IsValid(InHandle))
+		{
+			return false;
+		}
+
+		const FTickFunction* TickFunction = InHandle->GetTickFunction();
+		return HasPrerequisiteDependencyWith(TickFunction, TickFunction);
+	}
+
+	/** Checks for cycling constraints and manage tick dependencies if needed to avoid cycles from a tick dependency pov. */
+	static void CheckAndFixCycles(const UTickableTransformConstraint* InConstraint)
+	{
+		if (!IsValid(InConstraint))
+		{
+			return;
+		}
+
+		UWorld* World = InConstraint->GetWorld();
+		if (!IsValid(World))
+		{
+			return;
+		}
+
+		// get child's tick function
+		FTickFunction* ChildTickFunction = InConstraint->GetChildHandleTickFunction();
+		if (!ChildTickFunction)
+		{
+			return;
+		}
+
+		// filter for all constraints where the parent's tick function equals ChildTickFunction
+		auto Predicate = [ChildTickFunction](const ConstraintPtr& InConstraint)
+		{
+			const UTickableTransformConstraint* TransformConst = Cast<UTickableTransformConstraint>(InConstraint.Get());
+			if (!TransformConst)
+			{
+				return false;
+			}
+
+			const FTickFunction* ParentTickFunction = TransformConst->GetParentHandleTickFunction();
+			return ParentTickFunction && ParentTickFunction == ChildTickFunction;
+		};
+
+		const FConstraintsManagerController& Controller = FConstraintsManagerController::Get(World);
+		const ConstraintArray CyclingConstraints = Controller.GetConstraintsByPredicate(Predicate);
+		if (CyclingConstraints.IsEmpty())
+		{
+			return;
+		}
+
+		// check if they can cause a cycle and manage dependencies if that's the case
+		for (const ConstraintPtr& Constraint: CyclingConstraints)
+		{
+			if (HasPrerequisiteDependencyWith(&Constraint->ConstraintTick, &InConstraint->ConstraintTick))
+			{
+				UpdateCyclingDependency(World, Cast<UTickableTransformConstraint>(Constraint));
+			}
+		}
+	}
+	
+private:
+
+	/**
+	 * Manage tick dependencies if needed to avoid cycles from a tick dependency pov.
+	 * Both InConstraintToUpdate and its parent handle are supposed valid at this point
+	 */
+	static void UpdateCyclingDependency(UWorld* InWorld, UTickableTransformConstraint* InConstraintToUpdate)
+	{
+		// nothing to do if this constraint doesn't tick
+		if (!InConstraintToUpdate->ConstraintTick.IsTickFunctionEnabled())
+		{
+			return;
+		}
+		
+		const TObjectPtr<UTransformableHandle>& ParentHandle = InConstraintToUpdate->ParentTRSHandle;
+		UObject* TargetObject = ParentHandle->GetPrerequisiteObject();
+
+		// filter for all constraints where the child's target object equals ChildTickFunction
+		auto Predicate = [TargetObject](const ConstraintPtr& InConstraint)
+		{
+			const UTickableTransformConstraint* TransformConst = Cast<UTickableTransformConstraint>(InConstraint.Get());
+			if (!TransformConst)
+			{
+				return false;
+			}
+
+			const TObjectPtr<UTransformableHandle>& ChildHandle = TransformConst->ChildTRSHandle;
+			if (!IsValid(ChildHandle) || !ChildHandle->IsValid())
+			{
+				return false;
+			}
+			
+			const UObject* ChildPrereqObject = ChildHandle->GetPrerequisiteObject();
+			return IsValid(ChildPrereqObject) && (ChildPrereqObject == TargetObject);
+		};
+
+		const FConstraintsManagerController& Controller = FConstraintsManagerController::Get(InWorld);
+		const ConstraintArray ParentConstraints = Controller.GetConstraintsByPredicate(Predicate, true);
+
+		// check if there's any active constraint in ParentConstraints
+		const bool bHasActiveParentConstraint = ParentConstraints.ContainsByPredicate([](const ConstraintPtr& Constraint)
+		{
+			return IsValid(Constraint) && Constraint->IsFullyActive(); 
+		});
+
+		// update the constraint prerequisites based on the result
+		if (FTickFunction* TargetTickFunction = InConstraintToUpdate->GetParentHandleTickFunction())
+		{
+			if (bHasActiveParentConstraint)
+			{
+				// UE_LOG(LogTemp, Warning, TEXT("REMOVE %s prerex on %s."), *TargetObject->GetName(), *InConstraintToUpdate->GetName());
+				InConstraintToUpdate->ConstraintTick.RemovePrerequisite(TargetObject, *TargetTickFunction);
+			}
+			else
+			{
+				// UE_LOG(LogTemp, Warning, TEXT("ADD %s prerex on %s."), *TargetObject->GetName(), *InConstraintToUpdate->GetName());
+				InConstraintToUpdate->ConstraintTick.AddPrerequisite(TargetObject, *TargetTickFunction);
+			}
+		}
+	}
+};
+	
+}
+
 /** 
  * UTickableTransformConstraint
  **/
@@ -163,25 +333,8 @@ void UTickableTransformConstraint::Setup()
 
 void UTickableTransformConstraint::SetupDependencies()
 {
-	auto GetTickableFunction = [this](const TObjectPtr<UTransformableHandle>& InHandle) -> FTickFunction*
-	{
-		if (!IsValid(InHandle) || !InHandle->IsValid())
-		{
-			return nullptr;
-		}
-
-		// avoid creating dependencies between functions that are registered in levels that don't leave in the same world 
-		const UObject* PrerequisiteObject = InHandle->GetPrerequisiteObject();
-		if (!PrerequisiteObject || PrerequisiteObject->GetWorld() != GetWorld())
-		{
-			return nullptr;
-		}
-		
-		return InHandle->GetTickFunction();
-	};
-	
-	FTickFunction* ParentTickFunction = GetTickableFunction(ParentTRSHandle);
-	FTickFunction* ChildTickFunction = GetTickableFunction(ChildTRSHandle);
+	FTickFunction* ParentTickFunction = GetParentHandleTickFunction();
+	FTickFunction* ChildTickFunction = GetChildHandleTickFunction();
 	
 	if (ParentTickFunction && (ChildTickFunction != ParentTickFunction))
 	{
@@ -200,6 +353,34 @@ void UTickableTransformConstraint::SetupDependencies()
 	}
 }
 
+void UTickableTransformConstraint::EnsurePrimaryDependency()
+{
+	const FTickFunction* ParentTickFunction = GetParentHandleTickFunction();
+	const FTickFunction* ChildTickFunction = GetChildHandleTickFunction();
+	if (ParentTickFunction && (ChildTickFunction != ParentTickFunction))
+	{
+		const TArray<FTickPrerequisite>& ParentPrerequisites = ConstraintTick.GetPrerequisites();
+		if (ParentPrerequisites.IsEmpty())
+		{
+			// if the constraint has no prerex at this stage, this means that the parent tick function
+			// is not registered or can't tick (static meshes for instance) so look for the first parent tick function if any.
+			// In a context of adding several constraints, we want to make sure that the evaluation order is the right one
+			FTickPrerequisite PrimaryPrerex = ParentTRSHandle->GetPrimaryPrerequisite();
+			if (FTickFunction* PotentialFunction = PrimaryPrerex.Get())
+			{
+				UObject* Target = PrimaryPrerex.PrerequisiteObject.Get();
+				ConstraintTick.AddPrerequisite(Target, *PotentialFunction);
+			}
+		}
+	}
+}
+
+void UTickableTransformConstraint::OnActiveStateChanged() const
+{
+	// tick dependencies might need to be updated when dealing with cycles (than can be created between two controls for example) 
+	FConstraintCycleChecker::CheckAndFixCycles(this);
+}
+
 void UTickableTransformConstraint::PostLoad()
 {
 	Super::PostLoad();
@@ -222,6 +403,11 @@ void UTickableTransformConstraint::PostDuplicate(bool bDuplicateForPIE)
 
 	SetupDependencies();
 	RegisterDelegates();
+
+	if (bDuplicateForPIE)
+	{
+		EnsurePrimaryDependency();
+	}
 }
 
 uint32 UTickableTransformConstraint::GetTargetHash() const
@@ -293,6 +479,16 @@ void UTickableTransformConstraint::Evaluate(bool bTickHandlesAlso) const
 	}
 }
 
+void UTickableTransformConstraint::SetActive(const bool bIsActive)
+{
+	const bool bNeedsUpdate = (Active != bIsActive) || (bIsActive != ConstraintTick.IsTickFunctionEnabled());
+	Super::SetActive(bIsActive);
+	
+	if (bNeedsUpdate)
+	{
+		OnActiveStateChanged();
+	}
+}
 
 void UTickableTransformConstraint::SetChildGlobalTransform(const FTransform& InGlobal) const
 {
@@ -359,6 +555,33 @@ bool UTickableTransformConstraint::NeedsCompensation() const
 	return true;
 }
 
+FTickFunction* UTickableTransformConstraint::GetChildHandleTickFunction() const 
+{
+	return GetHandleTickFunction(ChildTRSHandle); 
+}
+
+FTickFunction* UTickableTransformConstraint::GetParentHandleTickFunction() const 
+{
+	return GetHandleTickFunction(ParentTRSHandle); 
+}
+
+FTickFunction* UTickableTransformConstraint::GetHandleTickFunction(const TObjectPtr<UTransformableHandle>& InHandle) const 
+{
+	if (!IsValid(InHandle) || !InHandle->IsValid())
+	{
+		return nullptr;
+	}
+
+	// avoid creating dependencies between functions that are registered in levels that don't leave in the same world 
+	const UObject* PrerequisiteObject = InHandle->GetPrerequisiteObject();
+	if (!PrerequisiteObject || PrerequisiteObject->GetWorld() != GetWorld())
+	{
+		return nullptr;
+	}
+		
+	return InHandle->GetTickFunction();
+}
+
 /** 
  * UTickableTranslationConstraint
  **/
@@ -1161,21 +1384,20 @@ namespace
 // we suppose that both InParentHandle and InChildHandle are safe to use
 bool HasConstraintDependencyWith(UWorld* InWorld, const UTransformableHandle* InParentHandle, const UTransformableHandle* InChildHandle)
 {
-	const FConstraintsManagerController& Controller = FConstraintsManagerController::Get(InWorld);
-	static constexpr bool bSorted = false;
-
-	using HandlePtr = TObjectPtr<UTransformableHandle>;
-	TArray<TObjectPtr<UTransformableHandle>> ParentHandles;
-	
 	using ConstraintPtr = TObjectPtr<UTickableConstraint>;
+	using HandlePtr = TObjectPtr<UTransformableHandle>;
+
+	static constexpr bool bSorted = false;
+	const FConstraintsManagerController& Controller = FConstraintsManagerController::Get(InWorld);
 	const TArray<ConstraintPtr> Constraints = Controller.GetParentConstraints(InParentHandle->GetHash(), bSorted);
 
 	// get parent handles
+	TArray< HandlePtr > ParentHandles;
 	for (const ConstraintPtr& Constraint: Constraints)
 	{
 		if (const UTickableTransformConstraint* TransformConstraint = Cast<UTickableTransformConstraint>(Constraint.Get()))
 		{
-			if (TransformConstraint->ParentTRSHandle)
+			if (IsValid(TransformConstraint->ParentTRSHandle))
 			{
 				ParentHandles.Add(TransformConstraint->ParentTRSHandle);
 			}
@@ -1349,24 +1571,7 @@ bool FTransformConstraintUtils::AddConstraint(
 	}
 
 	// make sure we tick after the parent.
-	const FTickFunction* ParentTickFunction = InParentHandle->GetTickFunction();
-	const FTickFunction* ChildTickFunction = InChildHandle->GetTickFunction();
-	if (ParentTickFunction && (ChildTickFunction != ParentTickFunction))
-	{
-		const TArray<FTickPrerequisite>& ParentPrerequisites =  Constraint->ConstraintTick.GetPrerequisites();
-		if (ParentPrerequisites.IsEmpty())
-		{
-			// if the constraint has no prerex at this stage, this means that the parent tick function
-			// is not registered or can't tick (static meshes for instance) so look for the first parent tick function if any.
-			// In a context of adding several constraints, we want to make sure that the evaluation order is the right one
-			FTickPrerequisite PrimaryPrerex = InParentHandle->GetPrimaryPrerequisite();
-			if (FTickFunction* PotentialFunction = PrimaryPrerex.Get())
-			{
-				UObject* Target = PrimaryPrerex.PrerequisiteObject.Get();
-				Constraint->ConstraintTick.AddPrerequisite(Target, *PotentialFunction);
-			}
-		}
-	}
+	Constraint->EnsurePrimaryDependency();
 
 	// if child handle is the parent of some other constraints, ensure they will tick after that new one
 	TArray<TObjectPtr<UTickableConstraint>> ChildChildConstraints;
@@ -1375,7 +1580,13 @@ bool FTransformConstraintUtils::AddConstraint(
 	{
 		Controller.SetConstraintsDependencies( NewConstraintName, ChildConstraint->GetFName());
 	}
-	
+
+	// warn for possible cycles
+	if (FConstraintCycleChecker::IsCycling(InChildHandle))
+	{
+		UE_LOG(LogTemp, Warning, TEXT("A cycle has been formed while creating %s."), *NewConstraintName.ToString());
+	}
+
 	return true;
 }
 
@@ -1566,15 +1777,11 @@ void FTransformConstraintUtils::GetChildrenConstraints(
 	const UTransformableHandle* InParentHandle,
 	TArray< TObjectPtr<UTickableConstraint> >& OutConstraints)
 {
-	const uint32 ParentHash = InParentHandle->GetHash();
-
-	const FConstraintsManagerController& Controller = FConstraintsManagerController::Get(World);
-
 	using ConstraintPtr = TObjectPtr<UTickableConstraint>;
-	const TArray<ConstraintPtr>& AllConstraints = Controller.GetConstraintsArray();
-
-	const TArray<ConstraintPtr> FilteredConstraints =
-	AllConstraints.FilterByPredicate( [ParentHash](const ConstraintPtr& Constraint)
+	
+	// filter for transform constraints where the InParentHandle is the parent (based on its hash value)
+	const uint32 ParentHash = InParentHandle->GetHash();
+	auto Predicate = [ParentHash](const ConstraintPtr& Constraint)
 	{
 		const UTickableTransformConstraint* TransformConstraint = Cast<UTickableTransformConstraint>(Constraint.Get());
 		if (!TransformConstraint)
@@ -1588,8 +1795,10 @@ void FTransformConstraintUtils::GetChildrenConstraints(
 		}
 
 		return false;
-	} );
+	};
 
+	const FConstraintsManagerController& Controller = FConstraintsManagerController::Get(World);	
+	const TArray<ConstraintPtr> FilteredConstraints = Controller.GetConstraintsByPredicate(Predicate);
 	OutConstraints.Append(FilteredConstraints);
 }
 
diff --git a/Engine/Source/Runtime/Experimental/Animation/Constraints/Public/ConstraintsManager.h b/Engine/Source/Runtime/Experimental/Animation/Constraints/Public/ConstraintsManager.h
--- a/Engine/Source/Runtime/Experimental/Animation/Constraints/Public/ConstraintsManager.h
+++ b/Engine/Source/Runtime/Experimental/Animation/Constraints/Public/ConstraintsManager.h
@@ -86,7 +86,7 @@ public:
 	FConstraintTickFunction ConstraintTick;
 
 	/** Sets the Active value and enable/disable the tick function. */
-	void SetActive(const bool bIsActive);
+	virtual void SetActive(const bool bIsActive);
 	
 	/** Get whether or not it's fully active, it's set to active and all pieces are set up,e.g. objects really exist that are being constrainted*/
 	virtual bool IsFullyActive() const;
@@ -212,22 +212,9 @@ public:
 	/** @todo document */
 	static FConstraintsManagerController& Get(UWorld* InWorld);
 
-	/** @todo document */
+	/** Allocates a new constraint with the constraints manager as the owner. */
 	template< typename TConstraint >
-	TConstraint* AllocateConstraintT(const FName& InBaseName) const
-	{
-		UConstraintsManager* Manager = GetManager();
-		if (!Manager)
-		{
-			return nullptr;
-		}
-		// unique name (we may want to use another approach here to manage uniqueness)
-		const FName Name = MakeUniqueObjectName(Manager, TConstraint::StaticClass(), InBaseName);
-
-		TConstraint* NewConstraint = NewObject<TConstraint>(Manager, Name, RF_Transactional);
-		NewConstraint->Modify();
-		return NewConstraint;
-	}
+	TConstraint* AllocateConstraintT(const FName& InBaseName) const;
 
 	/** Add this constraint to the manager */
 	bool AddConstraint(UTickableConstraint* InConstraint) const;
@@ -244,18 +231,22 @@ public:
 	/** Remove the constraint at the given index. */
 	bool RemoveConstraint(const int32 InConstraintIndex, bool bDoNotCompensate = false) const;
 
-	/** @todo document */
+	/** Returns the constraint based on it's name within the manager's constraints array. */
 	UTickableConstraint* GetConstraint(const FName& InConstraintName) const;
 
-	/** @todo document */
+	/** Returns the constraint based on it's index within the manager's constraints array. */
 	UTickableConstraint* GetConstraint(const int32 InConstraintIndex) const;
 	
 	/** Get read-only access to the array of constraints. */
 	const TArray< TObjectPtr<UTickableConstraint> >& GetConstraintsArray() const;
 
-	/** @todo document */
+	/** Returns manager's constraints array (sorted if needed). */
 	TArray< TObjectPtr<UTickableConstraint> > GetAllConstraints(const bool bSorted = false) const;
 
+	/** Returns a filtered constraints array checking if the predicate for each element is true. */
+	template <typename Predicate>
+	TArray< TObjectPtr<UTickableConstraint> > GetConstraintsByPredicate(Predicate Pred, const bool bSorted = false) const;
+
 	/** Get parent constraints of the specified child. If bSorted is true, then the constraints will be sorted by dependency. */
 	TArray< TObjectPtr<UTickableConstraint> > GetParentConstraints(
 		const uint32 InTargetHash,
diff --git a/Engine/Source/Runtime/Experimental/Animation/Constraints/Public/ConstraintsManager.inl b/Engine/Source/Runtime/Experimental/Animation/Constraints/Public/ConstraintsManager.inl
new file mode 100644
--- /dev/null
+++ b/Engine/Source/Runtime/Experimental/Animation/Constraints/Public/ConstraintsManager.inl
@@ -0,0 +1,62 @@
+﻿// Copyright Epic Games, Inc. All Rights Reserved.
+
+#pragma once
+
+#include "ConstraintsManager.h"
+
+#include "Algo/StableSort.h"
+
+template< typename TConstraint >
+TConstraint* FConstraintsManagerController::AllocateConstraintT(const FName& InBaseName) const
+{
+	UConstraintsManager* Manager = GetManager();
+	if (!Manager)
+	{
+		return nullptr;
+	}
+	// unique name (we may want to use another approach here to manage uniqueness)
+	const FName Name = MakeUniqueObjectName(Manager, TConstraint::StaticClass(), InBaseName);
+
+	TConstraint* NewConstraint = NewObject<TConstraint>(Manager, Name, RF_Transactional);
+	NewConstraint->Modify();
+	return NewConstraint;
+}
+
+template <typename Predicate>
+TArray< TObjectPtr<UTickableConstraint> > FConstraintsManagerController::GetConstraintsByPredicate(
+	Predicate Pred, const bool bSorted) const
+{
+	using ConstraintPtr = TObjectPtr<UTickableConstraint>;
+	
+	UConstraintsManager* Manager = FindManager();
+	if (!Manager)
+	{
+		static const TArray< TObjectPtr<UTickableConstraint> > Empty;
+		return Empty;
+	}
+
+	TArray< ConstraintPtr > Constraints = Manager->Constraints.FilterByPredicate(Pred);
+	if (!bSorted)
+	{
+		return Constraints;
+	}
+
+	// LHS ticks before RHS = LHS is a prerex of RHS 
+	auto TicksBefore = [](const UTickableConstraint& LHS, const UTickableConstraint& RHS)
+	{
+		const TArray<FTickPrerequisite>& RHSPrerex = RHS.ConstraintTick.GetPrerequisites();
+		const FConstraintTickFunction& LHSTickFunction = LHS.ConstraintTick;
+		const bool bIsLHSAPrerexOfRHS = RHSPrerex.ContainsByPredicate([&LHSTickFunction](const FTickPrerequisite& Prerex)
+		{
+			return Prerex.PrerequisiteTickFunction == &LHSTickFunction;
+		});
+		return bIsLHSAPrerexOfRHS;
+	};
+	
+	Algo::StableSort(Constraints, [TicksBefore](const ConstraintPtr& LHS, const ConstraintPtr& RHS)
+	{
+		return TicksBefore(*LHS, *RHS);
+	});
+	
+	return Constraints;
+}
\ No newline at end of file
diff --git a/Engine/Source/Runtime/Experimental/Animation/Constraints/Public/TransformConstraint.h b/Engine/Source/Runtime/Experimental/Animation/Constraints/Public/TransformConstraint.h
--- a/Engine/Source/Runtime/Experimental/Animation/Constraints/Public/TransformConstraint.h
+++ b/Engine/Source/Runtime/Experimental/Animation/Constraints/Public/TransformConstraint.h
@@ -14,9 +14,6 @@ class UTransformableComponentHandle;
 class USceneComponent;
 enum class EMovieSceneTransformChannel : uint32;
 
-using SetTransformFunc = TFunction<void(const FTransform&)>;
-using GetTransformFunc = TFunction<FTransform()>;
-
 /** 
  * UTickableTransformConstraint
  **/
@@ -58,6 +55,9 @@ public:
 	/** Override the evaluate so we can tick our handles*/
 	virtual void Evaluate(bool bTickHandlesAlso = false) const override;
 
+	/** Sets the Active value and enable/disable the tick function. */
+	virtual void SetActive(const bool bIsActive) override;
+	
 	/** The transformable handle representing the parent of that constraint. */
 	UPROPERTY(BlueprintReadWrite, Category = "Handle")
 	TObjectPtr<UTransformableHandle> ParentTRSHandle;
@@ -104,6 +104,15 @@ public:
 	*/
 	virtual void OnHandleModified(UTransformableHandle* InHandle, EHandleEvent InEvent);
 
+	/** Returns the handles tick function (ensuring it lives in the same world). */
+	FTickFunction* GetChildHandleTickFunction() const;
+	FTickFunction* GetParentHandleTickFunction() const;
+
+	/**
+	* Sets up dependencies with the first primary prerequisite available if the parent does not tick.   
+	*/
+	void EnsurePrimaryDependency();
+	
 protected:
 
 	/** Registers/Unregisters useful delegates for both child and parent handles. */
@@ -122,7 +131,7 @@ protected:
 	 * It creates a dependency graph between them so that they tick in the right order when evaluated.   
 	*/
 	void SetupDependencies();
-	
+
 	/** Set the current child's global transform. */
 	void SetChildGlobalTransform(const FTransform& InGlobal) const;
 	
@@ -133,6 +142,12 @@ protected:
 	UPROPERTY()
 	ETransformConstraintType Type = ETransformConstraintType::Parent;
 
+	/** Handle active state modification if needed. */
+	void OnActiveStateChanged() const;
+
+	/** Returns the handle's tick function (ensuring it lives in the same world). */
+	FTickFunction* GetHandleTickFunction(const TObjectPtr<UTransformableHandle>& InHandle) const;
+
 #if WITH_EDITOR
 public:
 	/** Returns the constraint's label used for UI. */
