Chaos : fix potential crash with replication when physics state is destroyed

diff --git a/Engine/Source/Runtime/Experimental/GeometryCollectionEngine/Private/GeometryCollection/GeometryCollectionComponent.cpp b/Engine/Source/Runtime/Experimental/GeometryCollectionEngine/Private/GeometryCollection/GeometryCollectionComponent.cpp
--- a/Engine/Source/Runtime/Experimental/GeometryCollectionEngine/Private/GeometryCollection/GeometryCollectionComponent.cpp
+++ b/Engine/Source/Runtime/Experimental/GeometryCollectionEngine/Private/GeometryCollection/GeometryCollectionComponent.cpp
@@ -1362,6 +1362,15 @@ void ActivateClusters(Chaos::FRigidClustering& Clustering, Chaos::TPBDRigidClust
 	Clustering.DeactivateClusterParticle(Cluster);
 }
 
+void UGeometryCollectionComponent::ResetRepData()
+{
+	ClustersToRep.Reset();
+	RepData.Reset();
+	OneOffActivatedProcessed = 0;
+	VersionProcessed = INDEX_NONE;
+	LastHardsnapTimeInMs = 0;
+}
+
 void UGeometryCollectionComponent::UpdateRepData()
 {
 	using namespace Chaos;
@@ -2566,21 +2575,13 @@ void UGeometryCollectionComponent::OnDestroyPhysicsState()
 		Scene->RemoveObject(PhysicsProxy);
 		InitializationState = ESimulationInitializationState::Unintialized;
 
+		// clear the clusters to rep as the information hold by it is now invalid
+		// we can still call this on the game thread because replication runs with the game thread frozen and will not run while the physics  state is being torned down
+		ResetRepData();
+
 		// Discard the pointer (cleanup happens through the scene or dedicated thread)
 		PhysicsProxy = nullptr;
 	}
-
-	// make sure we also clear the replication data structures
-	if (Chaos::FPhysicsSolver* Solver = GetSolver(*this))
-	{
-		Solver->EnqueueCommandImmediate([this]()
-			{
-				if (ClustersToRep)
-				{
-					ClustersToRep.Reset();
-				}
-			});
-	}
 }
 
 void UGeometryCollectionComponent::SendRenderDynamicData_Concurrent()
diff --git a/Engine/Source/Runtime/Experimental/GeometryCollectionEngine/Public/GeometryCollection/GeometryCollectionComponent.h b/Engine/Source/Runtime/Experimental/GeometryCollectionEngine/Public/GeometryCollection/GeometryCollectionComponent.h
--- a/Engine/Source/Runtime/Experimental/GeometryCollectionEngine/Public/GeometryCollection/GeometryCollectionComponent.h
+++ b/Engine/Source/Runtime/Experimental/GeometryCollectionEngine/Public/GeometryCollection/GeometryCollectionComponent.h
@@ -414,6 +414,13 @@ struct FGeometryCollectionRepData
 	// Just test version to skip having to traverse the whole pose array for replication
 	bool Identical(const FGeometryCollectionRepData* Other, uint32 PortFlags) const;
 	bool NetSerialize(FArchive& Ar, class UPackageMap* Map, bool& bOutSuccess);
+
+	void Reset()
+	{
+		OneOffActivated.Reset();
+		Clusters.Reset();
+	}
+
 };
 
 template<>
@@ -1059,6 +1066,9 @@ protected:
 	/** Called post solve to allow authoritative components to update their replication data */
 	void UpdateRepData();
 
+	/** Clear all rep data, this is required if the physics proxy has been recreated */
+	void ResetRepData();
+
 private:
 
 	bool bRenderStateDirty;
