Chaos : fix field meta data object type check not properly categorizing clustered particles as destruction ones

diff --git a/Engine/Source/Runtime/Experimental/Chaos/Private/PhysicsProxy/PerSolverFieldSystem.cpp b/Engine/Source/Runtime/Experimental/Chaos/Private/PhysicsProxy/PerSolverFieldSystem.cpp
--- a/Engine/Source/Runtime/Experimental/Chaos/Private/PhysicsProxy/PerSolverFieldSystem.cpp
+++ b/Engine/Source/Runtime/Experimental/Chaos/Private/PhysicsProxy/PerSolverFieldSystem.cpp
@@ -457,9 +457,13 @@ void FPerSolverFieldSystem::GetRelevantParticleHandles(
 template<typename ParticleHandleType>
 bool ValidateParticle(const EFieldObjectType ObjectType, const ParticleHandleType& ParticleHandle)
 {
-	return (ObjectType == EFieldObjectType::Field_Object_All) || (ObjectType == EFieldObjectType::Field_Object_Max) ||
-		((ObjectType == EFieldObjectType::Field_Object_Rigid) && (ParticleHandle->GetParticleType() != Chaos::EParticleType::GeometryCollection)) ||
-		((ObjectType == EFieldObjectType::Field_Object_Destruction) && (ParticleHandle->GetParticleType() == Chaos::EParticleType::GeometryCollection));
+	using namespace Chaos;
+	const EParticleType ParticleType = ParticleHandle->GetParticleType();
+	const bool bIsDestructionParticle = (ParticleType == EParticleType::GeometryCollection || ParticleType == EParticleType::Clustered);
+	return (ObjectType == EFieldObjectType::Field_Object_All)
+		|| (ObjectType == EFieldObjectType::Field_Object_Max)
+		|| ((ObjectType == EFieldObjectType::Field_Object_Rigid) && !bIsDestructionParticle)
+		|| ((ObjectType == EFieldObjectType::Field_Object_Destruction) && bIsDestructionParticle);
 }
 
 void FPerSolverFieldSystem::GetFilteredParticleHandles(
