Iris - Removed unused variable.

This also fixes a compile error.

diff --git a/Engine/Source/Runtime/Experimental/Iris/Core/Public/Iris/ReplicationState/Private/IrisFastArraySerializerInternal.h b/Engine/Source/Runtime/Experimental/Iris/Core/Public/Iris/ReplicationState/Private/IrisFastArraySerializerInternal.h
--- a/Engine/Source/Runtime/Experimental/Iris/Core/Public/Iris/ReplicationState/Private/IrisFastArraySerializerInternal.h
+++ b/Engine/Source/Runtime/Experimental/Iris/Core/Public/Iris/ReplicationState/Private/IrisFastArraySerializerInternal.h
@@ -129,8 +129,6 @@ void TIrisFastArrayEditor<FastArrayType>::MarkItemDirty(FastArrayItemType& Item)
 	const int32 Index = (&Item - Array.GetData());
 	if (Array.IsValidIndex(Index))
 	{
-		FNetBitArrayView MemberChangeMask = FIrisFastArraySerializerPrivateAccessor::GetChangeMask(FastArray);
-
 		// If this is a new element make sure it is at the end, otherwise we must dirty array
 		if ((Item.ReplicationID == INDEX_NONE) && (Index != (Array.Num() - 1)))
 		{
