Improved user control and visibility into proximity and connection graph generation for fracture:

 - Added a Proximity tool to Fracture Mode to visualize and change generation settings for GeometryCollection bone proximity.
 - Added more user controls for how proximity is generated, including a "Convex Hull distance" method for deciding contact, which should help include connections that were missed by the default proximity detection method.
 - Made proximity settings live on the GeometryCollection as (non-cooked) attributes, so proximity re-generates with the desired method after further fracturing/clustering.
 - Added an option to automatically convert the proximity graph to a pre-computed connection graph used by simulation.
 - Reduce redundant proximity calculations: Be more consistent in relying on FGeometryCollectionEdit to update proximity as needed, and call 'RequireProximity' instead of 'UpdateProximity' in cases where we expect a valid proximity may already be present.

diff --git a/Engine/Plugins/Experimental/ChaosEditor/Source/FractureEditor/Private/FractureEditorCommands.h b/Engine/Plugins/Experimental/ChaosEditor/Source/FractureEditor/Private/FractureEditorCommands.h
--- a/Engine/Plugins/Experimental/ChaosEditor/Source/FractureEditor/Private/FractureEditorCommands.h
+++ b/Engine/Plugins/Experimental/ChaosEditor/Source/FractureEditor/Private/FractureEditorCommands.h
@@ -92,6 +92,7 @@ class FFractureEditorCommands : public TCommands<FFractureEditorCommands>
 		TSharedPtr< FUICommandInfo > ConvertToMesh;
 		TSharedPtr< FUICommandInfo > Validate;
 		TSharedPtr< FUICommandInfo > MakeConvex;
+		TSharedPtr< FUICommandInfo > Proximity;
 		TSharedPtr< FUICommandInfo > FixTinyGeo;
 
 		// Property Commands
diff --git a/Engine/Plugins/Experimental/ChaosEditor/Source/FractureEditor/Private/FractureEditorModeToolkit.cpp b/Engine/Plugins/Experimental/ChaosEditor/Source/FractureEditor/Private/FractureEditorModeToolkit.cpp
--- a/Engine/Plugins/Experimental/ChaosEditor/Source/FractureEditor/Private/FractureEditorModeToolkit.cpp
+++ b/Engine/Plugins/Experimental/ChaosEditor/Source/FractureEditor/Private/FractureEditorModeToolkit.cpp
@@ -872,6 +872,7 @@ void FFractureEditorModeToolkit::BuildToolPalette(FName PaletteIndex, class FToo
 		ToolbarBuilder.AddToolBarButton(Commands.ConvertToMesh);
 		ToolbarBuilder.AddToolBarButton(Commands.Validate);
 		ToolbarBuilder.AddToolBarButton(Commands.MakeConvex);
+		ToolbarBuilder.AddToolBarButton(Commands.Proximity);
 		ToolbarBuilder.AddToolBarButton(Commands.FixTinyGeo);
 		ToolbarBuilder.AddToolBarButton(Commands.SetInitialDynamicState);
 		ToolbarBuilder.AddToolBarButton(Commands.SetRemoveOnBreak);
diff --git a/Engine/Plugins/Experimental/ChaosEditor/Source/FractureEditor/Private/FractureToolClusterMagnet.cpp b/Engine/Plugins/Experimental/ChaosEditor/Source/FractureEditor/Private/FractureToolClusterMagnet.cpp
--- a/Engine/Plugins/Experimental/ChaosEditor/Source/FractureEditor/Private/FractureToolClusterMagnet.cpp
+++ b/Engine/Plugins/Experimental/ChaosEditor/Source/FractureEditor/Private/FractureToolClusterMagnet.cpp
@@ -172,7 +172,7 @@ bool UFractureToolClusterMagnet::CheckPresenceOfNecessaryAttributes(const FGeome
 TMap<int32, TSet<int32>> UFractureToolClusterMagnet::InitializeConnectivity(const TSet<int32>& TopNodes, FGeometryCollectionPtr GeometryCollection, int32 OperatingLevel) const
 {
 	FGeometryCollectionProximityUtility ProximityUtility(GeometryCollection.Get());
-	ProximityUtility.UpdateProximity();
+	ProximityUtility.RequireProximity();
 
 	TMap<int32, TSet<int32>> ConnectivityMap;
 	for (int32 Index : TopNodes)
diff --git a/Engine/Plugins/Experimental/ChaosEditor/Source/FractureEditor/Private/FractureToolGenerators.cpp b/Engine/Plugins/Experimental/ChaosEditor/Source/FractureEditor/Private/FractureToolGenerators.cpp
--- a/Engine/Plugins/Experimental/ChaosEditor/Source/FractureEditor/Private/FractureToolGenerators.cpp
+++ b/Engine/Plugins/Experimental/ChaosEditor/Source/FractureEditor/Private/FractureToolGenerators.cpp
@@ -685,10 +685,6 @@ AGeometryCollectionActor* UFractureToolGenerateAsset::ConvertActorsToGeometryCol
 	// Add and initialize guids
 	::GeometryCollection::GenerateTemporaryGuids(FracturedGeometryCollection->GetGeometryCollection().Get(), 0 , true);
 
-	// Update proximity graph
-	FGeometryCollectionProximityUtility ProximityUtility(FracturedGeometryCollection->GetGeometryCollection().Get());
-	ProximityUtility.UpdateProximity();
-
 	const UFractureModeSettings* ModeSettings = GetDefault<UFractureModeSettings>();
 	ModeSettings->ApplyDefaultConvexSettings(*FracturedGeometryCollection->GetGeometryCollection());
 
@@ -876,10 +872,6 @@ void UFractureToolResetAsset::Execute(TWeakPtr<FFractureEditorModeToolkit> InToo
 						GeometryCollectionObject->Materials[NewMatNum - 1] = OldMaterials[OldMatNum - 1];
 					}
 				}
-				
-				// Update proximity graph
-				FGeometryCollectionProximityUtility ProximityUtility(GeometryCollection);
-				ProximityUtility.UpdateProximity();
 
 				const UFractureModeSettings* ModeSettings = GetDefault<UFractureModeSettings>();
 				ModeSettings->ApplyDefaultConvexSettings(*GeometryCollection);
diff --git a/Engine/Plugins/Experimental/ChaosEditor/Source/FractureEditor/Private/FractureToolProximity.cpp b/Engine/Plugins/Experimental/ChaosEditor/Source/FractureEditor/Private/FractureToolProximity.cpp
new file mode 100644
--- /dev/null
+++ b/Engine/Plugins/Experimental/ChaosEditor/Source/FractureEditor/Private/FractureToolProximity.cpp
@@ -0,0 +1,213 @@
+// Copyright Epic Games, Inc. All Rights Reserved.
+
+#include "FractureToolProximity.h"
+
+#include "FractureToolContext.h"
+#include "FractureModeSettings.h"
+
+#include "GeometryCollection/GeometryCollectionObject.h"
+#include "GeometryCollection/GeometryCollectionComponent.h"
+#include "GeometryCollection/GeometryCollectionAlgo.h"
+#include "GeometryCollection/GeometryCollectionProximityUtility.h"
+#include "GeometryCollection/GeometryCollectionClusteringUtility.h"
+
+#define LOCTEXT_NAMESPACE "FractureToolProximity"
+
+
+
+UFractureToolProximity::UFractureToolProximity(const FObjectInitializer& ObjInit)
+	: Super(ObjInit)
+{
+	ProximitySettings = NewObject<UFractureProximitySettings>(GetTransientPackage(), UFractureProximitySettings::StaticClass());
+	ProximitySettings->OwnerTool = this;
+}
+
+bool UFractureToolProximity::CanExecute() const
+{
+	if (!IsGeometryCollectionSelected())
+	{
+		return false;
+	}
+
+	return true;
+}
+
+FText UFractureToolProximity::GetDisplayText() const
+{
+	return FText(LOCTEXT("FractureToolProximity", "Update Proximity Generation Settings"));
+}
+
+FText UFractureToolProximity::GetTooltipText() const
+{
+	return FText(LOCTEXT("FractureToolProximityTooltip", "This tool visualizes and updates settings for proximity (contact) graph generation on the geometry collections"));
+}
+
+FSlateIcon UFractureToolProximity::GetToolIcon() const
+{
+	return FSlateIcon("FractureEditorStyle", "FractureEditor.Proximity");
+}
+
+void UFractureToolProximity::RegisterUICommand(FFractureEditorCommands* BindingContext)
+{
+	UI_COMMAND_EXT(BindingContext, UICommandInfo, "Proximity", "Prxmty", "Update (and visualize) the proximity (contact) graph for the bones of geometry collections.", EUserInterfaceActionType::ToggleButton, FInputChord());
+	BindingContext->Proximity = UICommandInfo;
+}
+
+TArray<UObject*> UFractureToolProximity::GetSettingsObjects() const
+{
+	TArray<UObject*> Settings;
+	Settings.Add(ProximitySettings);
+	return Settings;
+}
+
+void UFractureToolProximity::FractureContextChanged()
+{
+	UpdateVisualizations();
+}
+
+void UFractureToolProximity::UpdateVisualizations()
+{
+	TArray<FFractureToolContext> FractureContexts = GetFractureToolContexts();
+
+	ClearVisualizations();
+
+	for (const FFractureToolContext& FractureContext : FractureContexts)
+	{
+		FGeometryCollection& Collection = *FractureContext.GetGeometryCollection();
+
+		if (!Collection.HasAttribute("Proximity", FGeometryCollection::GeometryGroup))
+		{
+			continue;
+		}
+		const TManagedArray<TSet<int32>>& Proximity = Collection.GetAttribute<TSet<int32>>("Proximity", FGeometryCollection::GeometryGroup);
+
+		// Save the geometry collection component for rendering
+		int32 CollectionIdx = VisualizedCollections.Add(FractureContext.GetGeometryCollectionComponent());
+		FCollectionVisInfo& Vis = ProximityVisualizations.Emplace_GetRef();
+		Vis.CollectionIndex = CollectionIdx;
+
+		// Get proximity graph node positions as bounding box centers for each piece of geometry
+		Vis.GeoCenters.SetNum(Collection.NumElements(FGeometryCollection::GeometryGroup));
+		TArray<FTransform> GlobalTransformArray;
+		GeometryCollectionAlgo::GlobalMatrices(Collection.Transform, Collection.Parent, GlobalTransformArray);
+		for (int32 GeoIdx = 0; GeoIdx < Collection.NumElements(FGeometryCollection::GeometryGroup); ++GeoIdx)
+		{
+			int32 TransformIdx = Collection.TransformIndex[GeoIdx];
+			FVector BoxCenter = GlobalTransformArray[TransformIdx].TransformPosition(Collection.BoundingBox[GeoIdx].GetCenter());
+			Vis.GeoCenters[GeoIdx] = BoxCenter;
+		}
+
+		// Collect relevant proximity graph edges to render
+		auto AddEdgesForGeoIdx = [&Proximity, &Vis](int32 GeoIdx, bool bAddAll)
+		{
+			for (int32 GeoNbr : Proximity[GeoIdx])
+			{
+				if (bAddAll || GeoIdx < GeoNbr)
+				{
+					Vis.ProximityEdges.Add(FEdgeVisInfo{ GeoIdx, GeoNbr });
+				}
+			}
+		};
+		if (ProximitySettings->bOnlyShowForSelected)
+		{
+			TArray<int32> SelectedRigids;
+			for (int32 TransformIdx : FractureContext.GetSelection())
+			{
+				SelectedRigids.Reset();
+				FGeometryCollectionClusteringUtility::GetLeafBones(&Collection, TransformIdx, true, SelectedRigids);
+				for (int32 RigidIdx : SelectedRigids)
+				{
+					int32 GeoIdx = Collection.TransformToGeometryIndex[RigidIdx];
+					if (GeoIdx != INDEX_NONE)
+					{
+						AddEdgesForGeoIdx(GeoIdx, true);
+					}
+				}
+			}
+		}
+		else
+		{
+			for (int32 GeoIdx = 0; GeoIdx < Collection.NumElements(FGeometryCollection::GeometryGroup); ++GeoIdx)
+			{
+				AddEdgesForGeoIdx(GeoIdx, false);
+			}
+		}
+	}
+}
+
+void UFractureToolProximity::Render(const FSceneView* View, FViewport* Viewport, FPrimitiveDrawInterface* PDI)
+{
+	if (ProximitySettings->bShowProximity)
+	{
+		for (const FCollectionVisInfo& Vis : ProximityVisualizations)
+		{
+			const FGeometryCollection& Collection = *VisualizedCollections[Vis.CollectionIndex]->GetRestCollection()->GetGeometryCollection();
+			const TManagedArray<FVector3f>* ExplodedVectors = Collection.FindAttributeTyped<FVector3f>("ExplodedVector", FGeometryCollection::TransformGroup);
+			auto GetExplodedOffset = [&ExplodedVectors, &Collection](int32 GeoIdx)
+			{
+				return ExplodedVectors ?
+					(FVector)(*ExplodedVectors)[Collection.TransformIndex[GeoIdx]] :
+					FVector::ZeroVector;
+			};
+			const FTransform WorldTransform = VisualizedCollections[Vis.CollectionIndex]->GetComponentTransform();
+			for (const FEdgeVisInfo& Edge : Vis.ProximityEdges)
+			{
+				FVector P1 = WorldTransform.TransformPosition(Vis.GeoCenters[Edge.A] + GetExplodedOffset(Edge.A));
+				FVector P2 = WorldTransform.TransformPosition(Vis.GeoCenters[Edge.B] + GetExplodedOffset(Edge.B));
+				PDI->DrawLine(P1, P2, FLinearColor::Green, SDPG_Foreground, 0.0f, 0.001f);
+			}
+		}
+	}
+}
+
+void UFractureToolProximity::PostEditChangeChainProperty(struct FPropertyChangedChainEvent& PropertyChangedEvent)
+{
+	if (PropertyChangedEvent.GetMemberPropertyName() == GET_MEMBER_NAME_CHECKED(UFractureProximitySettings, bOnlyShowForSelected))
+	{
+		UpdateVisualizations();
+	}
+}
+
+TArray<FFractureToolContext> UFractureToolProximity::GetFractureToolContexts() const
+{
+	TArray<FFractureToolContext> Contexts;
+	TSet<UGeometryCollectionComponent*> GeomCompSelection;
+	GetSelectedGeometryCollectionComponents(GeomCompSelection);
+
+	for (UGeometryCollectionComponent* GeometryCollectionComponent : GeomCompSelection)
+	{
+		// Generate a context for each selected node
+		FFractureToolContext FullSelection(GeometryCollectionComponent);
+
+		Contexts.Add(FullSelection);
+	}
+
+	return Contexts;
+}
+
+int32 UFractureToolProximity::ExecuteFracture(const FFractureToolContext& FractureContext)
+{
+	if (FractureContext.GetGeometryCollection().IsValid())
+	{
+		FGeometryCollectionProximityPropertiesInterface::FProximityProperties Properties = FractureContext.GetGeometryCollection()->GetProximityProperties();
+		
+		Properties.Method = ProximitySettings->Method;
+		Properties.DistanceThreshold = ProximitySettings->DistanceThreshold;
+		Properties.bUseAsConnectionGraph = ProximitySettings->bUseAsConnectionGraph;
+		
+		FractureContext.GetGeometryCollection()->SetProximityProperties(Properties);
+
+		// Invalidate proximity
+		FGeometryCollectionProximityUtility ProximityUtility(FractureContext.GetGeometryCollection().Get());
+		ProximityUtility.InvalidateProximity();
+	}
+
+	return INDEX_NONE;
+}
+
+void UFractureToolProximity::Setup()
+{
+	Super::Setup();
+}
+
+#undef LOCTEXT_NAMESPACE
diff --git a/Engine/Plugins/Experimental/ChaosEditor/Source/FractureEditor/Private/FractureToolProximity.h b/Engine/Plugins/Experimental/ChaosEditor/Source/FractureEditor/Private/FractureToolProximity.h
new file mode 100644
--- /dev/null
+++ b/Engine/Plugins/Experimental/ChaosEditor/Source/FractureEditor/Private/FractureToolProximity.h
@@ -0,0 +1,108 @@
+// Copyright Epic Games, Inc. All Rights Reserved.
+
+#pragma once
+
+#include "FractureTool.h"
+
+#include "GeometryCollection/GeometryCollectionProximityUtility.h"
+
+#include "FractureToolProximity.generated.h"
+
+class FFractureToolContext;
+
+/** Settings controlling how proximity is detected for geometry collections */
+UCLASS(config = EditorPerProjectUserSettings)
+class UFractureProximitySettings : public UFractureToolSettings
+{
+public:
+
+	GENERATED_BODY()
+
+	UFractureProximitySettings(const FObjectInitializer& ObjInit)
+		: Super(ObjInit)
+	{}
+
+	// Which method to use to decide whether a given piece of geometry is in proximity with another
+	UPROPERTY(EditAnywhere, Category = Automatic)
+	EProximityMethod Method = EProximityMethod::Precise;
+
+	// If hull-based proximity detection is enabled, amount to expand hulls when searching for overlapping neighbors
+	UPROPERTY(EditAnywhere, Category = Automatic, meta = (ClampMin = "0", EditCondition = "Method == EProximityMethod::ConvexHull"))
+	double DistanceThreshold = 1;
+
+	// Whether to automatically transform the proximity graph into a connection graph to be used for simulation
+	UPROPERTY(EditAnywhere, Category = Automatic)
+	bool bUseAsConnectionGraph = false;
+
+	// Whether to display the proximity graph edges
+	UPROPERTY(EditAnywhere, Category = Visualization)
+	bool bShowProximity = true;
+
+	// Whether to only show the proximity graph connections for selected bones
+	UPROPERTY(EditAnywhere, Category = Visualization, meta = (EditCondition = "bShowProximity"))
+	bool bOnlyShowForSelected = false;
+};
+
+
+
+UCLASS(DisplayName = "Proximity Tool", Category = "FractureTools")
+class UFractureToolProximity : public UFractureModalTool
+{
+public:
+	GENERATED_BODY()
+
+	UFractureToolProximity(const FObjectInitializer& ObjInit);
+
+	///
+	/// UFractureModalTool Interface
+	///
+
+	/** This is the Text that will appear on the tool button to execute the tool **/
+	virtual FText GetDisplayText() const override;
+	virtual FText GetTooltipText() const override;
+	/** This is the Text that will appear on the button to execute the fracture **/
+	virtual FText GetApplyText() const override { return FText(NSLOCTEXT("FractureProximity", "ExecuteProximity", "Update Bone Proximity")); }
+	virtual FSlateIcon GetToolIcon() const override;
+	virtual void RegisterUICommand(FFractureEditorCommands* BindingContext) override;
+	virtual TArray<UObject*> GetSettingsObjects() const override;
+	virtual void FractureContextChanged() override;
+	virtual void Render(const FSceneView* View, FViewport* Viewport, FPrimitiveDrawInterface* PDI) override;
+	/** Executes function that generates new geometry. Returns the first new geometry index. */
+	virtual int32 ExecuteFracture(const FFractureToolContext& FractureContext) override;
+	virtual bool CanExecute() const override;
+	virtual void PostEditChangeChainProperty(struct FPropertyChangedChainEvent& PropertyChangedEvent) override;
+
+	/** Gets the UI command info for this command */
+	const TSharedPtr<FUICommandInfo>& GetUICommandInfo() const;
+
+	virtual TArray<FFractureToolContext> GetFractureToolContexts() const override;
+
+	virtual void Setup() override;
+
+private:
+
+	UPROPERTY(EditAnywhere, Category = Proximity)
+	TObjectPtr<UFractureProximitySettings> ProximitySettings;
+
+	void UpdateVisualizations();
+	virtual void ClearVisualizations() override
+	{
+		Super::ClearVisualizations();
+		ProximityVisualizations.Empty();
+	}
+
+	struct FEdgeVisInfo
+	{
+		int32 A, B;
+	};
+
+	struct FCollectionVisInfo
+	{
+		TArray<FEdgeVisInfo> ProximityEdges;
+		TArray<FVector> GeoCenters;
+		int32 CollectionIndex;
+	};
+
+	TArray<FCollectionVisInfo> ProximityVisualizations;
+
+};
diff --git a/Engine/Plugins/Experimental/GeometryCollectionPlugin/Source/GeometryCollectionNodes/Private/Dataflow/GeometryCollectionNodes.cpp b/Engine/Plugins/Experimental/GeometryCollectionPlugin/Source/GeometryCollectionNodes/Private/Dataflow/GeometryCollectionNodes.cpp
--- a/Engine/Plugins/Experimental/GeometryCollectionPlugin/Source/GeometryCollectionNodes/Private/Dataflow/GeometryCollectionNodes.cpp
+++ b/Engine/Plugins/Experimental/GeometryCollectionPlugin/Source/GeometryCollectionNodes/Private/Dataflow/GeometryCollectionNodes.cpp
@@ -202,7 +202,7 @@ void FResetGeometryCollectionDataflowNode::Evaluate(Dataflow::FContext& Context,
 						}
 					}
 				}
-				GeometryCollectionObject->UpdateConvexGeometry();
+				GeometryCollectionObject->UpdateGeometryDependentProperties();
 				GeometryCollectionObject->InitializeMaterials();
 				GeometryCollectionObject->InvalidateCollection();
 
diff --git a/Engine/Plugins/Experimental/PlanarCutPlugin/Source/PlanarCut/Private/PlanarCut.cpp b/Engine/Plugins/Experimental/PlanarCutPlugin/Source/PlanarCut/Private/PlanarCut.cpp
--- a/Engine/Plugins/Experimental/PlanarCutPlugin/Source/PlanarCut/Private/PlanarCut.cpp
+++ b/Engine/Plugins/Experimental/PlanarCutPlugin/Source/PlanarCut/Private/PlanarCut.cpp
@@ -1220,7 +1220,7 @@ int32 MergeBones(
 	FTransform CellsToWorld = FTransform::Identity;
 
 	FGeometryCollectionProximityUtility ProximityUtility(&Collection);
-	ProximityUtility.UpdateProximity();
+	ProximityUtility.RequireProximity();
 	const TManagedArray<TSet<int32>>& Proximity = Collection.GetAttribute<TSet<int32>>("Proximity", FGeometryCollection::GeometryGroup);
 
 	// local array so we can populate it with all transforms if input was empty
diff --git a/Engine/Source/Runtime/Experimental/Chaos/Private/GeometryCollection/Facades/CollectionConnectionGraphFacade.cpp b/Engine/Source/Runtime/Experimental/Chaos/Private/GeometryCollection/Facades/CollectionConnectionGraphFacade.cpp
new file mode 100644
--- /dev/null
+++ b/Engine/Source/Runtime/Experimental/Chaos/Private/GeometryCollection/Facades/CollectionConnectionGraphFacade.cpp
@@ -0,0 +1,46 @@
+// Copyright Epic Games, Inc. All Rights Reserved.
+
+
+#include "GeometryCollection/Facades/CollectionConnectionGraphFacade.h"
+#include "GeometryCollection/TransformCollection.h"
+
+namespace GeometryCollection::Facades
+{
+
+	FCollectionConnectionGraphFacade::FCollectionConnectionGraphFacade(FManagedArrayCollection& InCollection)
+		: ConnectionsAttribute(InCollection, "Connections", FTransformCollection::TransformGroup)
+#if UE_BUILD_DEBUG
+		, ParentAttribute(InCollection, "Parent", FTransformCollection::TransformGroup)
+#endif
+	{
+	}
+
+	bool FCollectionConnectionGraphFacade::IsValid() const
+	{
+		return ConnectionsAttribute.IsValid();
+	}
+
+	void FCollectionConnectionGraphFacade::AddAttributes()
+	{
+		ConnectionsAttribute.Add();
+	}
+	
+	void FCollectionConnectionGraphFacade::ClearAttributes()
+	{
+		ConnectionsAttribute.Remove();
+	}
+
+	void FCollectionConnectionGraphFacade::Connect(int32 BoneA, int32 BoneB)
+	{
+#if UE_BUILD_DEBUG
+		// Expect to always have a parent array and connected bones should always have the same parent
+		checkSlow(ParentAttribute.IsValid() && ParentAttribute.Get()[BoneA] == ParentAttribute.Get()[BoneB]);
+#endif
+		TManagedArray<TSet<int32>>& Connections = ConnectionsAttribute.Modify();
+		Connections[BoneA].Add(BoneB);
+		Connections[BoneB].Add(BoneA);
+	}
+
+};
+
+
diff --git a/Engine/Source/Runtime/Experimental/Chaos/Private/GeometryCollection/GeometryCollection.cpp b/Engine/Source/Runtime/Experimental/Chaos/Private/GeometryCollection/GeometryCollection.cpp
--- a/Engine/Source/Runtime/Experimental/Chaos/Private/GeometryCollection/GeometryCollection.cpp
+++ b/Engine/Source/Runtime/Experimental/Chaos/Private/GeometryCollection/GeometryCollection.cpp
@@ -36,6 +36,7 @@ const FName FGeometryCollection::StatusFlagsAttribute("StatusFlags");
 FGeometryCollection::FGeometryCollection()
 	: FTransformCollection()
 	, FGeometryCollectionConvexPropertiesInterface(this)
+	, FGeometryCollectionProximityPropertiesInterface(this)
 {
 	Construct();
 }
@@ -730,6 +731,7 @@ void FGeometryCollection::Empty()
 void FGeometryCollection::InitializeInterfaces()
 {
 	FGeometryCollectionConvexPropertiesInterface::InitializeInterface();
+	FGeometryCollectionProximityPropertiesInterface::InitializeInterface();
 }
 
 void FGeometryCollection::ReorderElements(FName Group, const TArray<int32>& NewOrder)
@@ -1039,6 +1041,7 @@ void FGeometryCollection::Serialize(Chaos::FChaosArchive& Ar)
 	if (Ar.IsCooking())
 	{
 		FGeometryCollectionConvexPropertiesInterface::CleanInterfaceForCook();
+		FGeometryCollectionProximityPropertiesInterface::CleanInterfaceForCook();
 	}
 
 	Ar.UsingCustomVersion(FFortniteNCBranchObjectVersion::GUID);
diff --git a/Engine/Source/Runtime/Experimental/Chaos/Private/GeometryCollection/GeometryCollectionConvexUtility.cpp b/Engine/Source/Runtime/Experimental/Chaos/Private/GeometryCollection/GeometryCollectionConvexUtility.cpp
--- a/Engine/Source/Runtime/Experimental/Chaos/Private/GeometryCollection/GeometryCollectionConvexUtility.cpp
+++ b/Engine/Source/Runtime/Experimental/Chaos/Private/GeometryCollection/GeometryCollectionConvexUtility.cpp
@@ -1078,7 +1078,7 @@ void CreateNonoverlappingConvexHulls(
 /// Helper to get convex hulls from a geometry collection in the format required by CreateNonoverlappingConvexHulls
 void HullsFromGeometry(
 	FGeometryCollection& Geometry,
-	TArray<FTransform>& GlobalTransformArray,
+	const TArray<FTransform>& GlobalTransformArray,
 	TFunctionRef<bool(int32)> HasCustomConvexFn,
 	TArray<TUniquePtr<Chaos::FConvex>>& Convexes,
 	TArray<FVector>& ConvexPivots,
@@ -1271,43 +1271,63 @@ double ComputeGeometryVolume(
 
 }
 
-FGeometryCollectionConvexUtility::FGeometryCollectionConvexData FGeometryCollectionConvexUtility::CreateNonOverlappingConvexHullData(
-	FGeometryCollection* GeometryCollection, double FracAllowRemove, double SimplificationDistanceThreshold, double CanExceedFraction, EConvexOverlapRemoval OverlapRemovalMethod,
-	double OverlapRemovalShrinkPercent)
+
+UE::GeometryCollectionConvexUtility::FConvexHulls
+FGeometryCollectionConvexUtility::ComputeLeafHulls(FGeometryCollection* GeometryCollection, const TArray<FTransform>& GlobalTransformArray, double SimplificationDistanceThreshold, double OverlapRemovalShrinkPercent)
 {
 	check(GeometryCollection);
 
+
+	UE::GeometryCollectionConvexUtility::FConvexHulls Hulls;
+
 	// Prevent ~100% shrink percentage, because it would not be reversible ...
 	// Note: Could alternatively just disable convex overlap removal in this case
-	OverlapRemovalShrinkPercent = FMath::Min(99.9, OverlapRemovalShrinkPercent);
+	Hulls.OverlapRemovalShrinkPercent = FMath::Min(99.9, OverlapRemovalShrinkPercent);
 
 	TManagedArray<int32>* CustomConvexFlags = GetCustomConvexFlags(GeometryCollection, false);
 	auto ConvexFlagsAlwaysFalse = [](int32) -> bool { return false; };
 	auto ConvexFlagsFromArray = [&CustomConvexFlags](int32 TransformIdx) -> bool { return (bool)(*CustomConvexFlags)[TransformIdx]; };
 	TFunctionRef<bool(int32)> HasCustomConvexFn = (CustomConvexFlags != nullptr) ? (TFunctionRef<bool(int32)>)ConvexFlagsFromArray : (TFunctionRef<bool(int32)>)ConvexFlagsAlwaysFalse;
 
-	bool bHasProximity = GeometryCollection->HasAttribute("Proximity", FGeometryCollection::GeometryGroup);
-	if (!bHasProximity)
+	HullsFromGeometry(*GeometryCollection, GlobalTransformArray, HasCustomConvexFn, Hulls.Hulls, Hulls.Pivots, Hulls.TransformToHullsIndices, GeometryCollection->SimulationType,
+		FGeometryCollection::ESimulationTypes::FST_Rigid, SimplificationDistanceThreshold, Hulls.OverlapRemovalShrinkPercent);
+
+	return Hulls;
+}
+
+FGeometryCollectionConvexUtility::FGeometryCollectionConvexData FGeometryCollectionConvexUtility::CreateNonOverlappingConvexHullData(
+	FGeometryCollection* GeometryCollection, double FracAllowRemove, double SimplificationDistanceThreshold, double CanExceedFraction, EConvexOverlapRemoval OverlapRemovalMethod,
+	double OverlapRemovalShrinkPercentIn, UE::GeometryCollectionConvexUtility::FConvexHulls* PrecomputedLeafHulls)
+{
+	check(GeometryCollection);
+
+	TManagedArray<int32>* CustomConvexFlags = GetCustomConvexFlags(GeometryCollection, false);
+	auto ConvexFlagsAlwaysFalse = [](int32) -> bool { return false; };
+	auto ConvexFlagsFromArray = [&CustomConvexFlags](int32 TransformIdx) -> bool { return (bool)(*CustomConvexFlags)[TransformIdx]; };
+	TFunctionRef<bool(int32)> HasCustomConvexFn = (CustomConvexFlags != nullptr) ? (TFunctionRef<bool(int32)>)ConvexFlagsFromArray : (TFunctionRef<bool(int32)>)ConvexFlagsAlwaysFalse;
+
+	TArray<FTransform> GlobalTransformArray;
+	GeometryCollectionAlgo::GlobalMatrices(GeometryCollection->Transform, GeometryCollection->Parent, GlobalTransformArray);
+
+	UE::GeometryCollectionConvexUtility::FConvexHulls* UseLeafHulls = PrecomputedLeafHulls;
+	UE::GeometryCollectionConvexUtility::FConvexHulls LocalLeafHullsStorage;
+	if (!UseLeafHulls)
 	{
-		FGeometryCollectionProximityUtility ProximityUtility(GeometryCollection);
-		ProximityUtility.UpdateProximity();
+		LocalLeafHullsStorage = ComputeLeafHulls(GeometryCollection, GlobalTransformArray, SimplificationDistanceThreshold, OverlapRemovalShrinkPercentIn);
+		UseLeafHulls = &LocalLeafHullsStorage;
 	}
+
+	FGeometryCollectionProximityUtility ProximityUtility(GeometryCollection);
+	ProximityUtility.RequireProximity(UseLeafHulls);
 	SetVolumeAttributes(GeometryCollection);
 
 	const TManagedArray<TSet<int32>>* GCProximity = GeometryCollection->FindAttribute<TSet<int32>>("Proximity", FGeometryCollection::GeometryGroup);
 	const TManagedArray<float>* Volume = GeometryCollection->FindAttribute<float>("Volume", FGeometryCollection::TransformGroup);
-	TArray<TUniquePtr<Chaos::FConvex>> Convexes;
-	TArray<FVector> ConvexPivots;
-	TArray<TSet<int32>> TransformToConvexIndexArr;
-	TArray<FTransform> GlobalTransformArray;
-	GeometryCollectionAlgo::GlobalMatrices(GeometryCollection->Transform, GeometryCollection->Parent, GlobalTransformArray);
-	HullsFromGeometry(*GeometryCollection, GlobalTransformArray, HasCustomConvexFn, Convexes, ConvexPivots, TransformToConvexIndexArr, GeometryCollection->SimulationType, 
-		FGeometryCollection::ESimulationTypes::FST_Rigid, SimplificationDistanceThreshold, OverlapRemovalShrinkPercent);
 
-	CreateNonoverlappingConvexHulls(Convexes, ConvexPivots, TransformToConvexIndexArr, HasCustomConvexFn, GeometryCollection->SimulationType, FGeometryCollection::ESimulationTypes::FST_Rigid, FGeometryCollection::ESimulationTypes::FST_None,
-		GeometryCollection->Parent, GCProximity, GeometryCollection->TransformIndex, Volume, FracAllowRemove, SimplificationDistanceThreshold, CanExceedFraction, OverlapRemovalMethod, OverlapRemovalShrinkPercent);
+	CreateNonoverlappingConvexHulls(UseLeafHulls->Hulls, UseLeafHulls->Pivots, UseLeafHulls->TransformToHullsIndices, HasCustomConvexFn, GeometryCollection->SimulationType, FGeometryCollection::ESimulationTypes::FST_Rigid, FGeometryCollection::ESimulationTypes::FST_None,
+		GeometryCollection->Parent, GCProximity, GeometryCollection->TransformIndex, Volume, FracAllowRemove, SimplificationDistanceThreshold, CanExceedFraction, OverlapRemovalMethod, UseLeafHulls->OverlapRemovalShrinkPercent);
 
-	TransformHullsToLocal(GlobalTransformArray, Convexes, ConvexPivots, TransformToConvexIndexArr, OverlapRemovalShrinkPercent);
+	TransformHullsToLocal(GlobalTransformArray, UseLeafHulls->Hulls, UseLeafHulls->Pivots, UseLeafHulls->TransformToHullsIndices, UseLeafHulls->OverlapRemovalShrinkPercent);
 
 	if (!GeometryCollection->HasGroup("Convex"))
 	{
@@ -1327,10 +1347,10 @@ FGeometryCollectionConvexUtility::FGeometryCollectionConvexData FGeometryCollect
 
 	TManagedArray<TSet<int32>>& TransformToConvexIndices = GeometryCollection->ModifyAttribute<TSet<int32>>("TransformToConvexIndices", FTransformCollection::TransformGroup);
 	TManagedArray<TUniquePtr<Chaos::FConvex>>& ConvexHull = GeometryCollection->ModifyAttribute<TUniquePtr<Chaos::FConvex>>("ConvexHull", "Convex");
-	TransformToConvexIndices = MoveTemp(TransformToConvexIndexArr);
+	TransformToConvexIndices = MoveTemp(UseLeafHulls->TransformToHullsIndices);
 	GeometryCollection->EmptyGroup("Convex");
-	GeometryCollection->Resize(Convexes.Num(), "Convex");
-	ConvexHull = MoveTemp(Convexes);
+	GeometryCollection->Resize(UseLeafHulls->Hulls.Num(), "Convex");
+	ConvexHull = MoveTemp(UseLeafHulls->Hulls);
 
 	// clear all null and empty hulls
 	TArray<int32> EmptyConvex;
diff --git a/Engine/Source/Runtime/Experimental/Chaos/Private/GeometryCollection/GeometryCollectionProximityPropertiesInterface.cpp b/Engine/Source/Runtime/Experimental/Chaos/Private/GeometryCollection/GeometryCollectionProximityPropertiesInterface.cpp
new file mode 100644
--- /dev/null
+++ b/Engine/Source/Runtime/Experimental/Chaos/Private/GeometryCollection/GeometryCollectionProximityPropertiesInterface.cpp
@@ -0,0 +1,93 @@
+// Copyright Epic Games, Inc. All Rights Reserved.
+
+#include "GeometryCollection/GeometryCollectionProximityPropertiesInterface.h"
+#include "GeometryCollection/GeometryCollection.h"
+
+const FName FGeometryCollectionProximityPropertiesInterface::ProximityPropertiesGroup = "ProximityProperties";
+const FName FGeometryCollectionProximityPropertiesInterface::ProximityDistanceThreshold = "DistanceThreshold";
+const FName FGeometryCollectionProximityPropertiesInterface::ProximityDetectionMethod = "DetectionMethod";
+const FName FGeometryCollectionProximityPropertiesInterface::ProximityAsConnectionGraph = "AsConnectionGraph";
+
+FGeometryCollectionProximityPropertiesInterface::FGeometryCollectionProximityPropertiesInterface(FGeometryCollection* InGeometryCollection)
+	: FManagedArrayInterface(InGeometryCollection)
+{}
+
+void 
+FGeometryCollectionProximityPropertiesInterface::InitializeInterface()
+{
+	if (!ManagedCollection->HasGroup(ProximityPropertiesGroup))
+	{
+		ManagedCollection->AddGroup(ProximityPropertiesGroup);
+	}
+
+	if (!ManagedCollection->HasAttribute(ProximityDetectionMethod, ProximityPropertiesGroup))
+	{
+		ManagedCollection->AddAttribute<int32>(ProximityDetectionMethod, ProximityPropertiesGroup);
+	}
+
+	if (!ManagedCollection->HasAttribute(ProximityDistanceThreshold, ProximityPropertiesGroup))
+	{
+		ManagedCollection->AddAttribute<float>(ProximityDistanceThreshold, ProximityPropertiesGroup);
+	}
+
+	if (!ManagedCollection->HasAttribute(ProximityAsConnectionGraph, ProximityPropertiesGroup))
+	{
+		ManagedCollection->AddAttribute<bool>(ProximityAsConnectionGraph, ProximityPropertiesGroup);
+	}
+}
+
+void 
+FGeometryCollectionProximityPropertiesInterface::CleanInterfaceForCook()
+{
+	RemoveInterfaceAttributes();
+}
+
+void 
+FGeometryCollectionProximityPropertiesInterface::RemoveInterfaceAttributes()
+{
+	ManagedCollection->RemoveGroup(ProximityPropertiesGroup);
+}
+
+FGeometryCollectionProximityPropertiesInterface::FProximityProperties
+FGeometryCollectionProximityPropertiesInterface::GetProximityProperties() const
+{
+	FProximityProperties Properties;
+	const bool bHasProximityProperties = ManagedCollection->NumElements(ProximityPropertiesGroup) > 0;
+	if (bHasProximityProperties)
+	{
+		const TManagedArray<bool>& AsConnectionGraph = ManagedCollection->GetAttribute<bool>(ProximityAsConnectionGraph, ProximityPropertiesGroup);
+		const TManagedArray<int32>& DetectionMethod = ManagedCollection->GetAttribute<int32>(ProximityDetectionMethod, ProximityPropertiesGroup);
+		const TManagedArray<float>& DistanceThreshold = ManagedCollection->GetAttribute<float>(ProximityDistanceThreshold, ProximityPropertiesGroup);
+
+		constexpr int32 DefaultIndex = 0;
+		Properties.bUseAsConnectionGraph = AsConnectionGraph[DefaultIndex];
+		Properties.DistanceThreshold = DistanceThreshold[DefaultIndex];
+		Properties.Method = (EProximityMethod)DetectionMethod[DefaultIndex];
+	}
+	return Properties;
+}
+
+void
+FGeometryCollectionProximityPropertiesInterface::SetProximityProperties(const FProximityProperties& InProximityAttributes)
+{
+	FProximityProperties Property;
+
+	int32 AttributeIndex = 0;
+	const bool bHasProximityProperties = ManagedCollection->NumElements(ProximityPropertiesGroup) > 0;
+	if (!bHasProximityProperties)
+	{
+		if (!ensure(ManagedCollection->HasGroup(ProximityPropertiesGroup)))
+		{
+			InitializeInterface();
+		}
+		AttributeIndex = ManagedCollection->AddElements(1, ProximityPropertiesGroup);
+	}
+
+	TManagedArray<bool>& AsConnectionGraph = ManagedCollection->ModifyAttribute<bool>(ProximityAsConnectionGraph, ProximityPropertiesGroup);
+	TManagedArray<int32>& Method = ManagedCollection->ModifyAttribute<int32>(ProximityDetectionMethod, ProximityPropertiesGroup);
+	TManagedArray<float>& DistanceThreshold = ManagedCollection->ModifyAttribute<float>(ProximityDistanceThreshold, ProximityPropertiesGroup);
+
+	AsConnectionGraph[AttributeIndex] = InProximityAttributes.bUseAsConnectionGraph;
+	DistanceThreshold[AttributeIndex] = InProximityAttributes.DistanceThreshold;
+	Method[AttributeIndex] = (int32)InProximityAttributes.Method;
+}
diff --git a/Engine/Source/Runtime/Experimental/Chaos/Private/GeometryCollection/GeometryCollectionProximityUtility.cpp b/Engine/Source/Runtime/Experimental/Chaos/Private/GeometryCollection/GeometryCollectionProximityUtility.cpp
--- a/Engine/Source/Runtime/Experimental/Chaos/Private/GeometryCollection/GeometryCollectionProximityUtility.cpp
+++ b/Engine/Source/Runtime/Experimental/Chaos/Private/GeometryCollection/GeometryCollectionProximityUtility.cpp
@@ -2,8 +2,13 @@
 
 #include "GeometryCollection/GeometryCollectionProximityUtility.h"
 
-#include "Spatial/SparseDynamicOctree3.h"
+#include "GeometryCollection/Facades/CollectionConnectionGraphFacade.h"
+#include "GeometryCollection/GeometryCollectionClusteringUtility.h"
+#include "GeometryCollection/GeometryCollectionConvexUtility.h"
+#include "Chaos/Convex.h"
+#include "Chaos/GJK.h"
 
+#include "Spatial/SparseDynamicOctree3.h"
 #include "Async/ParallelFor.h"
 #include "GeometryCollection/GeometryCollection.h"
 #include "GeometryCollection/GeometryCollectionAlgo.h"
@@ -110,7 +115,7 @@ struct FBinNormals
 	}
 };
 
-// Per-Geometry pre-computed spatial data
+// Per-Geometry pre-computed spatial data used for computing 'precise' proximity information
 struct FPerGeometrySpatial
 {
 	TArray<TArray<int32>, TFixedAllocator<FBinNormals::NumBins>> Bins;
@@ -132,7 +137,7 @@ struct FPerGeometrySpatial
 	TMap<int32, bool> CandidateContacts; // bool indicates if the contact has been confirmed
 };
 
-// Overall geometry-collection spatial data
+// Overall geometry-collection spatial data for 'Precise method' proximity detection
 struct FGeometryCollectionProximitySpatial
 {
 	TArray<FVector3f> TransformedVertices;
@@ -401,7 +406,94 @@ struct FGeometryCollectionProximitySpatial
 
 };
 
+void BuildProximityFromConvexHulls(FGeometryCollection* Collection, const UE::GeometryCollectionConvexUtility::FConvexHulls& HullData, double DistanceThreshold)
+{
+	double MaxHullDim = DistanceThreshold;
+	Chaos::FAABB3 OverallBounds;
+	for (const TUniquePtr<Chaos::FConvex>& Hull : HullData.Hulls)
+	{
+		const Chaos::FAABB3 HullBounds = Hull->BoundingBox();
+		OverallBounds.GrowToInclude(HullBounds);
+		MaxHullDim = FMath::Max(MaxHullDim, HullBounds.Extents().GetMax() + DistanceThreshold);
+	}
+
+	if (!Collection->HasAttribute("Proximity", FGeometryCollection::GeometryGroup))
+	{
+		const FManagedArrayCollection::FConstructionParameters GeometryDependency(FGeometryCollection::GeometryGroup);
+		Collection->AddAttribute<TSet<int32>>("Proximity", FGeometryCollection::GeometryGroup, GeometryDependency);
+	}
+
+	TManagedArray<TSet<int32>>& Proximity = Collection->ModifyAttribute<TSet<int32>>("Proximity", FGeometryCollection::GeometryGroup);
+
+
+	UE::Geometry::FSparseDynamicOctree3 GeoOctree;
+	GeoOctree.RootDimension = FMath::Min(MaxHullDim * 2, OverallBounds.Extents().GetMax() + DistanceThreshold);
+	FVector3d Center = (FVector3d)OverallBounds.GetCenter();
+	int32 NumGeometry = Collection->NumElements(FGeometryCollection::GeometryGroup);
+	TArray<int32> HullToGeoIdx;
+	HullToGeoIdx.Init(-1, HullData.Hulls.Num());
+	for (int32 GeoIdx = 0; GeoIdx < NumGeometry; ++GeoIdx)
+	{
+		int32 TransformIdx = Collection->TransformIndex[GeoIdx];
+		for (int32 HullIdx : HullData.TransformToHullsIndices[TransformIdx])
+		{
+			checkSlow(HullToGeoIdx[HullIdx] == -1); // these local-space hulls should not be mapped to multiple geometry pieces
+			HullToGeoIdx[HullIdx] = GeoIdx;
+		}
+	}
 
+	TArray<int32> HullIndices;
+	for (int32 GeoIdx = 0; GeoIdx < NumGeometry; ++GeoIdx)
+	{
+		int32 TransformIdx = Collection->TransformIndex[GeoIdx];
+		for (int32 HullIdx : HullData.TransformToHullsIndices[TransformIdx])
+		{
+			const Chaos::FConvex& Hull = *HullData.Hulls[HullIdx];
+			Chaos::FAABB3 ChaosHullBounds = Hull.BoundingBox();
+			ChaosHullBounds.Thicken(DistanceThreshold * .5);
+			UE::Geometry::FAxisAlignedBox3d GeoBounds((FVector3d)ChaosHullBounds.Min(), (FVector3d)ChaosHullBounds.Max());
+			GeoOctree.RangeQuery(GeoBounds, HullIndices);
+			for (int32 CandidateHullIdx : HullIndices)
+			{
+				const Chaos::FConvex& CandidateHull = *HullData.Hulls[CandidateHullIdx];
+				const Chaos::FAABB3 CandidateHullBounds = CandidateHull.BoundingBox();
+				int32 OtherGeoIdx = HullToGeoIdx[CandidateHullIdx];
+
+				if (GeoIdx != OtherGeoIdx && ChaosHullBounds.Intersects(CandidateHullBounds))
+				{
+					// TODO: Note the thickness parameter for GJKIntersection does not work as documented
+					// But if that is fixed, we can replace the below GJKDistance call with this:
+					//	const VectorRegister4Float InitialDirSimd = MakeVectorRegisterFloat(1.f, 0.f, 0.f, 0.f);
+					//	if (GJKIntersectionSameSpaceSimd(Hull, CandidateHull, DistanceThreshold, InitialDirSimd))
+
+					const Chaos::FRigidTransform3 IdentityTransform = Chaos::FRigidTransform3::Identity;
+					Chaos::FReal Distance;
+					Chaos::TVec3<Chaos::FReal> NearestA, NearestB, Normal; // All unused
+					Chaos::EGJKDistanceResult Result = Chaos::GJKDistance<Chaos::FReal>(
+						Chaos::TGJKShape(Hull),
+						Chaos::TGJKShape(CandidateHull),
+						Chaos::GJKDistanceInitialV(Hull, CandidateHull, IdentityTransform),
+						Distance, NearestA, NearestB, Normal);
+					if (Result == Chaos::EGJKDistanceResult::Contact || Result == Chaos::EGJKDistanceResult::DeepContact
+						|| (Result == Chaos::EGJKDistanceResult::Separated && Distance <= DistanceThreshold))
+					{
+						Proximity[GeoIdx].Add(OtherGeoIdx);
+						Proximity[OtherGeoIdx].Add(GeoIdx);
+					}
+				}
+			}
+		}
+		// add all the hulls for a given geometry after intersection-testing them with the hulls already in the octree
+		for (int32 HullIdx : HullData.TransformToHullsIndices[TransformIdx])
+		{
+			const Chaos::FConvex& Hull = *HullData.Hulls[HullIdx];
+			Chaos::FAABB3 ChaosHullBounds = Hull.BoundingBox();
+			ChaosHullBounds.Thicken(DistanceThreshold * .5);
+			UE::Geometry::FAxisAlignedBox3d GeoBounds((FVector3d)ChaosHullBounds.Min(), (FVector3d)ChaosHullBounds.Max());
+			GeoOctree.InsertObject(HullIdx, GeoBounds);
+		}
+	}
+}
 
 }} // namespace UE::GeometryCollectionInternal
 
@@ -411,10 +503,136 @@ FGeometryCollectionProximityUtility::FGeometryCollectionProximityUtility(FGeomet
 	check(Collection);
 }
 
-void FGeometryCollectionProximityUtility::UpdateProximity()
+void FGeometryCollectionProximityUtility::RequireProximity(UE::GeometryCollectionConvexUtility::FConvexHulls* OptionalComputedHulls)
+{
+	bool bHasProximity = Collection->HasAttribute("Proximity", FGeometryCollection::GeometryGroup);
+	if (!bHasProximity)
+	{
+		UpdateProximity(OptionalComputedHulls);
+	}
+}
+
+void FGeometryCollectionProximityUtility::InvalidateProximity()
+{
+	bool bHasProximity = Collection->HasAttribute("Proximity", FGeometryCollection::GeometryGroup);
+	if (bHasProximity)
+	{
+		Collection->RemoveAttribute("Proximity", FGeometryCollection::GeometryGroup);
+	}
+}
+
+void FGeometryCollectionProximityUtility::ClearConnectionGraph()
+{
+	GeometryCollection::Facades::FCollectionConnectionGraphFacade ConnectionsFacade(*Collection);
+	ConnectionsFacade.ClearAttributes();
+}
+
+void FGeometryCollectionProximityUtility::CopyProximityToConnectionGraph()
+{
+	if (!Collection->HasAttribute("Proximity", FGeometryCollection::GeometryGroup))
+	{
+		ClearConnectionGraph();
+		return;
+	}
+
+	const TManagedArray<TSet<int32>>& Proximity = Collection->GetAttribute<TSet<int32>>("Proximity", FGeometryCollection::GeometryGroup);
+	GeometryCollection::Facades::FCollectionConnectionGraphFacade ConnectionsFacade(*Collection);
+	ConnectionsFacade.AddAttributes();
+	TManagedArray<TSet<int32>>& Connections = ConnectionsFacade.ConnectionsAttribute.Modify();
+	Connections.Fill(TSet<int32>());
+
+	int32 NumBones = Connections.Num();
+	TArray<int32> Depths;
+	Depths.SetNumZeroed(NumBones);
+	for (int32 BoneIdx = 0; BoneIdx < NumBones; ++BoneIdx)
+	{
+		if (Collection->SimulationType[BoneIdx] == FGeometryCollection::ESimulationTypes::FST_None)
+		{
+			Depths[BoneIdx] = -1;
+			continue;
+		}
+		int32 Depth = 0, WalkParent = BoneIdx;
+		while (Collection->Parent[WalkParent] != INDEX_NONE)
+		{
+			Depth++;
+			WalkParent = Collection->Parent[WalkParent];
+		}
+		Depths[BoneIdx] = Depth;
+	}
+	TArray<int32> AllLeaves;
+	for (int32 BoneIdx = 0; BoneIdx < NumBones; ++BoneIdx)
+	{
+		if (Collection->SimulationType[BoneIdx] == FGeometryCollection::ESimulationTypes::FST_None)
+		{
+			continue;
+		}
+		int32 BoneDepth = Depths[BoneIdx];
+		int32 BoneParent = Collection->Parent[BoneIdx];
+		AllLeaves.Reset();
+
+		FGeometryCollectionClusteringUtility::GetLeafBones(Collection, BoneIdx, true, AllLeaves);
+		for (int32 LeafBone : AllLeaves)
+		{
+			int32 LeafGeo = Collection->TransformToGeometryIndex[LeafBone];
+			for (int32 NbrGeo : Proximity[LeafGeo])
+			{
+				int32 NbrBone = Collection->TransformIndex[NbrGeo];
+				if (Depths[NbrBone] < BoneDepth)
+				{
+					continue; // cluster is closer to root than us, ignore it
+				}
+				while (NbrBone != INDEX_NONE && Depths[NbrBone] > BoneDepth)
+				{
+					NbrBone = Collection->Parent[NbrBone];
+				}
+				if (NbrBone != INDEX_NONE && Collection->Parent[NbrBone] == BoneParent)
+				{
+					ConnectionsFacade.Connect(BoneIdx, NbrBone);
+				}
+			}
+		}
+	}
+}
+
+void FGeometryCollectionProximityUtility::UpdateProximity(UE::GeometryCollectionConvexUtility::FConvexHulls* OptionalComputedHulls)
 {
 	using namespace UE::GeometryCollectionInternal;
 
-	FGeometryCollectionProximitySpatial Spatial(Collection, ProximityThreshold);
-	Spatial.MoveProximityToCollection(Collection);
+	FGeometryCollectionProximityPropertiesInterface::FProximityProperties Properties = Collection->GetProximityProperties();
+
+	if (Properties.Method == EProximityMethod::ConvexHull)
+	{
+		FGeometryCollectionConvexPropertiesInterface::FConvexCreationProperties ConvexProperties = Collection->GetConvexProperties();
+
+		UE::GeometryCollectionConvexUtility::FConvexHulls LocalComputedHulls;
+		UE::GeometryCollectionConvexUtility::FConvexHulls* UseComputedHulls = OptionalComputedHulls;
+		if (!UseComputedHulls || UseComputedHulls->OverlapRemovalShrinkPercent > 0) // If we don't have precomputed hulls or if they're shrunk, compute new hulls to use for proximity detection
+		{
+			TArray<FTransform> GlobalTransformArray;
+			GeometryCollectionAlgo::GlobalMatrices(Collection->Transform, Collection->Parent, GlobalTransformArray);
+			LocalComputedHulls = FGeometryCollectionConvexUtility::ComputeLeafHulls(Collection, GlobalTransformArray, ConvexProperties.SimplificationThreshold, 0.0f /* Never shrink hulls for proximity detection */);
+			UseComputedHulls = &LocalComputedHulls;
+		}
+
+		UE::GeometryCollectionInternal::BuildProximityFromConvexHulls(Collection, *UseComputedHulls, Properties.DistanceThreshold);
+
+	}
+	else
+	{
+		// This threshold not exposed to the user via Properties.DistanceThreshold because it doesn't behave very well if increased to large values --
+		// the computation would become closer to O(n^2), and the results would likely be confusing/inconsistent
+		constexpr float PreciseProximityThreshold = .01f;
+		FGeometryCollectionProximitySpatial Spatial(Collection, PreciseProximityThreshold);
+		Spatial.MoveProximityToCollection(Collection);
+	}
+
+	if (Properties.bUseAsConnectionGraph)
+	{
+		CopyProximityToConnectionGraph();
+	}
+	else
+	{
+		// TODO: verify that the connection graph was auto-generated from proximity, rather than being custom edited, before removal
+		ClearConnectionGraph();
+	}
 }
diff --git a/Engine/Source/Runtime/Experimental/Chaos/Public/GeometryCollection/Facades/CollectionConnectionGraphFacade.h b/Engine/Source/Runtime/Experimental/Chaos/Public/GeometryCollection/Facades/CollectionConnectionGraphFacade.h
new file mode 100644
--- /dev/null
+++ b/Engine/Source/Runtime/Experimental/Chaos/Public/GeometryCollection/Facades/CollectionConnectionGraphFacade.h
@@ -0,0 +1,43 @@
+// Copyright Epic Games, Inc. All Rights Reserved.
+
+#pragma once
+
+#include "GeometryCollection/ManagedArrayCollection.h"
+#include "GeometryCollection/ManagedArrayAccessor.h"
+
+namespace GeometryCollection::Facades
+{
+
+	/**
+	 * Provides an API for connection-graph related attributes
+	 */
+	class CHAOS_API FCollectionConnectionGraphFacade
+	{
+	public:
+		FCollectionConnectionGraphFacade(FManagedArrayCollection& InCollection);
+
+		/**
+		 * returns true if all the necessary attributes are present
+		 * if not, then the API can be used to create them
+		 */
+		bool IsValid() const;
+
+		/** Add the necessary attributes if they are missing and initialize them if necessary */
+		void AddAttributes();
+
+		/** Remove the attributes */
+		void ClearAttributes();
+
+		/** Connect two bones */
+		void Connect(int32 BoneA, int32 BoneB);
+
+		/** Connections between bones that have the same parent in the hierarchy */
+		TManagedArrayAccessor<TSet<int32>> ConnectionsAttribute;
+
+#if UE_BUILD_DEBUG
+		/** Optional Parent array for validating connections in debug */
+		TManagedArrayAccessor<int32> ParentAttribute;
+#endif 
+	};
+
+}
\ No newline at end of file
diff --git a/Engine/Source/Runtime/Experimental/Chaos/Public/GeometryCollection/GeometryCollection.h b/Engine/Source/Runtime/Experimental/Chaos/Public/GeometryCollection/GeometryCollection.h
--- a/Engine/Source/Runtime/Experimental/Chaos/Public/GeometryCollection/GeometryCollection.h
+++ b/Engine/Source/Runtime/Experimental/Chaos/Public/GeometryCollection/GeometryCollection.h
@@ -7,6 +7,7 @@
 #include "Misc/Crc.h"
 
 #include "GeometryCollection/GeometryCollectionConvexPropertiesInterface.h"
+#include "GeometryCollection/GeometryCollectionProximityPropertiesInterface.h"
 #include "GeometryCollection/ManagedArrayAccessor.h"
 
 namespace Chaos
@@ -14,6 +15,7 @@ namespace Chaos
 	class FChaosArchive;
 }
 class FGeometryCollectionConvexPropertiesInterface;
+class FGeometryCollectionProximityPropertiesInterface;
 
 namespace GeometryCollectionUV
 {
@@ -27,7 +29,8 @@ namespace GeometryCollectionUV
 * FGeometryCollection (FTransformCollection)
 */
 class CHAOS_API FGeometryCollection : public FTransformCollection, 
-	public FGeometryCollectionConvexPropertiesInterface
+	public FGeometryCollectionConvexPropertiesInterface,
+	public FGeometryCollectionProximityPropertiesInterface
 {
 
 public:
diff --git a/Engine/Source/Runtime/Experimental/Chaos/Public/GeometryCollection/GeometryCollectionConvexUtility.h b/Engine/Source/Runtime/Experimental/Chaos/Public/GeometryCollection/GeometryCollectionConvexUtility.h
--- a/Engine/Source/Runtime/Experimental/Chaos/Public/GeometryCollection/GeometryCollectionConvexUtility.h
+++ b/Engine/Source/Runtime/Experimental/Chaos/Public/GeometryCollection/GeometryCollectionConvexUtility.h
@@ -25,6 +25,25 @@ enum class EConvexOverlapRemoval : int32
 	OnlyClustersVsClusters = 3
 };
 
+namespace UE::GeometryCollectionConvexUtility
+{
+	// Used to pass computed convex hulls along with metadata
+	// Note these hulls are typically computed in a shared coordinate space, in contrast to the final hulls on the geometry collection which are in the local space of each bone
+	struct FConvexHulls
+	{
+		TArray<TUniquePtr<Chaos::FConvex>> Hulls;
+
+		// Mapping from geometry collection bones indices to indices in the Hulls array. A Set is used to support multiple hulls per bone.
+		TArray<TSet<int32>> TransformToHullsIndices;
+
+		// Percent the hulls are down-scaled
+		double OverlapRemovalShrinkPercent;
+
+		// The pivots used for scaling by OverlapRemovalShrinkPercent
+		TArray<FVector> Pivots;
+	};
+}
+
 class CHAOS_API FGeometryCollectionConvexUtility
 {
 public:
@@ -53,9 +72,10 @@ public:
 	 @param CanExceedFraction					The fraction by which the convex body volume on a cluster can exceed the volume of the geometry under that cluster (a value of 1 == exceed by 100% == convex hull has 2x the volume of the geometry)
 	 @param OverlapRemovalMethod				If bRemoveOverlaps, control which overlaps are removed
 	 @param OverlapRemovalShrinkPercent			Compute overlaps based on objects shrunk by this percentage, so objects that would not overlap with this value set as their 'Collision Object Reduction Percentage' will not be cut
+	 @param ComputedLeafHullsToModify			Optional pre-computed hulls for geometry of rigid leaves, in the top-level coordinate space of the geometry collection. If passed in, the data will be updated and moved.
 	 */
 	static FGeometryCollectionConvexData CreateNonOverlappingConvexHullData(FGeometryCollection* GeometryCollection, double FractionAllowRemove = .3, double SimplificationDistanceThreshold = 0.0, double CanExceedFraction = .5, 
-		EConvexOverlapRemoval OverlapRemovalMethod = EConvexOverlapRemoval::All, double OverlapRemovalShrinkPercent = 0.0);
+		EConvexOverlapRemoval OverlapRemovalMethod = EConvexOverlapRemoval::All, double OverlapRemovalShrinkPercent = 0.0, UE::GeometryCollectionConvexUtility::FConvexHulls* ComputedLeafHullsToModify = nullptr);
 
 	/** Returns the convex hull of the vertices contained in the specified geometry. */
 	static TUniquePtr<Chaos::FConvex> FindConvexHull(const FGeometryCollection* GeometryCollection, int32 GeometryIndex);
@@ -88,5 +108,11 @@ public:
 	 */
 	static void CopyChildConvexes(const FGeometryCollection* FromCollection, const TArrayView<const int32>& FromTransformIdx, FGeometryCollection* ToCollection, const TArrayView<const int32>& ToTransformIdx, bool bLeafOnly);
 
+	// Compute just the hulls of the leaf / rigid nodes that hold geometry directly, with no cluster hulls and no overlap removal by cutting
+	// This is an initial step of several algorithms: The CreateNonOverlappingConvexHullData function as well as convex-based proximity detection (TODO: and the auto-embed algorithm?)
+	// (TODO: Make auto-embed use this instead of the full hulls?)
+	// @param GlobalTransformArray		GeometryCollection's transforms to global space, as computed by GeometryCollectionAlgo::GlobalMatrices
+	static UE::GeometryCollectionConvexUtility::FConvexHulls ComputeLeafHulls(FGeometryCollection* GeometryCollection, const TArray<FTransform>& GlobalTransformArray, double SimplificationDistanceThreshold = 0.0, double OverlapRemovalShrinkPercent = 0.0);
+
 };
 
diff --git a/Engine/Source/Runtime/Experimental/Chaos/Public/GeometryCollection/GeometryCollectionProximityPropertiesInterface.h b/Engine/Source/Runtime/Experimental/Chaos/Public/GeometryCollection/GeometryCollectionProximityPropertiesInterface.h
new file mode 100644
--- /dev/null
+++ b/Engine/Source/Runtime/Experimental/Chaos/Public/GeometryCollection/GeometryCollectionProximityPropertiesInterface.h
@@ -0,0 +1,46 @@
+// Copyright Epic Games, Inc. All Rights Reserved.
+
+#pragma once
+
+#include "CoreMinimal.h"
+#include "GeometryCollection/ManagedArrayCollection.h"
+#include "GeometryCollection/GeometryCollectionProximityUtility.h"
+
+class FGeometryCollection;
+
+class CHAOS_API FGeometryCollectionProximityPropertiesInterface : public FManagedArrayInterface
+{
+public :
+	typedef FManagedArrayInterface Super;
+	using FManagedArrayInterface::ManagedCollection;
+
+	// Proximity Properties Group Name
+	static const FName ProximityPropertiesGroup;
+	
+	// Attribute: Method to determine proximity
+	static const FName ProximityDetectionMethod;
+	// Attribute: For convex hull proximity, what separation distance can still be considered as 'proximity'
+	static const FName ProximityDistanceThreshold;
+	// Attribute: Whether to use the computed proximity graph as a connection graph
+	static const FName ProximityAsConnectionGraph;
+
+
+	struct FProximityProperties
+	{
+		EProximityMethod Method;
+		float DistanceThreshold = 1.0f;
+		bool bUseAsConnectionGraph = false;
+	};
+
+	FGeometryCollectionProximityPropertiesInterface(FGeometryCollection* InGeometryCollection);
+
+	void InitializeInterface() override;
+
+	void CleanInterfaceForCook() override;
+	
+	void RemoveInterfaceAttributes() override;
+
+	FProximityProperties GetProximityProperties() const;
+	void SetProximityProperties(const FProximityProperties&);
+};
+
diff --git a/Engine/Source/Runtime/Experimental/Chaos/Public/GeometryCollection/GeometryCollectionProximityUtility.h b/Engine/Source/Runtime/Experimental/Chaos/Public/GeometryCollection/GeometryCollectionProximityUtility.h
--- a/Engine/Source/Runtime/Experimental/Chaos/Public/GeometryCollection/GeometryCollectionProximityUtility.h
+++ b/Engine/Source/Runtime/Experimental/Chaos/Public/GeometryCollection/GeometryCollectionProximityUtility.h
@@ -3,18 +3,37 @@
 #pragma once
 
 #include "CoreTypes.h"
+#include "UObject/ObjectMacros.h"
 
 class FGeometryCollection;
+namespace UE::GeometryCollectionConvexUtility
+{
+	struct FConvexHulls;
+}
+
+UENUM()
+enum class EProximityMethod : int32
+{
+	// Precise proximity mode looks for geometry with touching vertices or touching, coplanar, opposite-facing triangles. This works well with geometry fractured using our fracture tools.
+	Precise,
+	// Convex Hull proximity mode looks for geometry with overlapping convex hulls (with an optional offset)
+	ConvexHull
+};
 
 class CHAOS_API FGeometryCollectionProximityUtility
 {
 public:
 	FGeometryCollectionProximityUtility(FGeometryCollection* InCollection);
 
-	void UpdateProximity();
+	void UpdateProximity(UE::GeometryCollectionConvexUtility::FConvexHulls* OptionalComputedHulls = nullptr);
+
+	// Update proximity data if it is not already present
+	void RequireProximity(UE::GeometryCollectionConvexUtility::FConvexHulls* OptionalComputedHulls = nullptr);
+
+	void InvalidateProximity();
 
-	// maximum separation to still count as 'contact'
-	float ProximityThreshold = 0.01f;
+	void CopyProximityToConnectionGraph();
+	void ClearConnectionGraph();
 
 private:
 	FGeometryCollection* Collection;
diff --git a/Engine/Source/Runtime/Experimental/GeometryCollectionEngine/Private/GeometryCollection/GeometryCollectionComponent.cpp b/Engine/Source/Runtime/Experimental/GeometryCollectionEngine/Private/GeometryCollection/GeometryCollectionComponent.cpp
--- a/Engine/Source/Runtime/Experimental/GeometryCollectionEngine/Private/GeometryCollection/GeometryCollectionComponent.cpp
+++ b/Engine/Source/Runtime/Experimental/GeometryCollectionEngine/Private/GeometryCollection/GeometryCollectionComponent.cpp
@@ -2746,7 +2746,7 @@ FGeometryCollectionEdit::~FGeometryCollectionEdit()
 		{
 			if (!bShapeIsUnchanged)
 			{
-				GetRestCollection()->UpdateConvexGeometry();
+				GetRestCollection()->UpdateGeometryDependentProperties();
 			}
 			GetRestCollection()->InvalidateCollection();
 		}
@@ -3126,7 +3126,7 @@ void FScopedColorEdit::SelectBones(GeometryCollection::ESelectionMode SelectionM
 		case GeometryCollection::ESelectionMode::Neighbors:
 		{
 			FGeometryCollectionProximityUtility ProximityUtility(GeometryCollectionPtr.Get());
-			ProximityUtility.UpdateProximity();
+			ProximityUtility.RequireProximity();
 
 			const TManagedArray<int32>& TransformIndex = GeometryCollectionPtr->TransformIndex;
 			const TManagedArray<int32>& TransformToGeometryIndex = GeometryCollectionPtr->TransformToGeometryIndex;
diff --git a/Engine/Source/Runtime/Experimental/GeometryCollectionEngine/Private/GeometryCollection/GeometryCollectionObject.cpp b/Engine/Source/Runtime/Experimental/GeometryCollectionEngine/Private/GeometryCollection/GeometryCollectionObject.cpp
--- a/Engine/Source/Runtime/Experimental/GeometryCollectionEngine/Private/GeometryCollection/GeometryCollectionObject.cpp
+++ b/Engine/Source/Runtime/Experimental/GeometryCollectionEngine/Private/GeometryCollection/GeometryCollectionObject.cpp
@@ -283,6 +283,15 @@ void UGeometryCollection::ValidateSizeSpecificDataDefaults()
 	check(SizeSpecificData.Num());
 }
 
+void UGeometryCollection::UpdateGeometryDependentProperties()
+{
+#if WITH_EDITOR
+	// Note: Currently, computing convex hulls also always computes proximity (if missing) as well as volumes and size.
+	// If adding a condition where we do not compute convex hulls, make sure to still compute proximity, volumes and size here
+	UpdateConvexGeometry();
+#endif
+}
+
 void UGeometryCollection::UpdateConvexGeometry()
 {
 #if WITH_EDITOR
diff --git a/Engine/Source/Runtime/Experimental/GeometryCollectionEngine/Public/GeometryCollection/GeometryCollectionObject.h b/Engine/Source/Runtime/Experimental/GeometryCollectionEngine/Public/GeometryCollection/GeometryCollectionObject.h
--- a/Engine/Source/Runtime/Experimental/GeometryCollectionEngine/Public/GeometryCollection/GeometryCollectionObject.h
+++ b/Engine/Source/Runtime/Experimental/GeometryCollectionEngine/Public/GeometryCollection/GeometryCollectionObject.h
@@ -680,6 +680,11 @@ public:
 	* Update the convex geometry on the collection.
 	*/
 	void UpdateConvexGeometry();
+	
+	/*
+	 * Update properties that depend on the geometry and clustering: Proximity, Convex Hulls, Volume and Size data.
+	 */
+	void UpdateGeometryDependentProperties();
 
 
 	//
