Enable Iris compilation by default. Remains disabled at runtime.

diff --git a/Engine/Source/Programs/UnrealBuildTool/Configuration/TargetRules.cs b/Engine/Source/Programs/UnrealBuildTool/Configuration/TargetRules.cs
--- a/Engine/Source/Programs/UnrealBuildTool/Configuration/TargetRules.cs
+++ b/Engine/Source/Programs/UnrealBuildTool/Configuration/TargetRules.cs
@@ -773,7 +773,7 @@ namespace UnrealBuildTool
 		[RequiresUniqueBuildEnvironment]
 		[CommandLine("-NoUseIris", Value = "false")]
 		[CommandLine("-UseIris", Value = "true")]
-		public bool bUseIris = false;
+		public bool bUseIris = true;
 
 		/// <summary>
 		/// Whether we are compiling editor code or not. Prefer the more explicit bCompileAgainstEditor instead.
diff --git a/Engine/Source/Runtime/Experimental/Iris/Core/Private/Iris/ReplicationSystem/ReplicationOperations.cpp b/Engine/Source/Runtime/Experimental/Iris/Core/Private/Iris/ReplicationSystem/ReplicationOperations.cpp
--- a/Engine/Source/Runtime/Experimental/Iris/Core/Private/Iris/ReplicationSystem/ReplicationOperations.cpp
+++ b/Engine/Source/Runtime/Experimental/Iris/Core/Private/Iris/ReplicationSystem/ReplicationOperations.cpp
@@ -21,6 +21,7 @@
 #include "Iris/Serialization/ObjectNetSerializer.h"
 #include "Misc/CoreMiscDefines.h"
 #include "Net/Core/Trace/NetTrace.h"
+#include "Net/Core/Trace/NetDebugName.h"
 #include "DequantizeAndApplyHelper.h"
 #include "HAL/IConsoleManager.h"
 
diff --git a/Engine/Source/Runtime/Experimental/Iris/Core/Private/Iris/ReplicationSystem/ReplicationReader.cpp b/Engine/Source/Runtime/Experimental/Iris/Core/Private/Iris/ReplicationSystem/ReplicationReader.cpp
--- a/Engine/Source/Runtime/Experimental/Iris/Core/Private/Iris/ReplicationSystem/ReplicationReader.cpp
+++ b/Engine/Source/Runtime/Experimental/Iris/Core/Private/Iris/ReplicationSystem/ReplicationReader.cpp
@@ -6,6 +6,7 @@
 #include "Iris/Core/IrisProfiler.h"
 #include "Iris/Core/IrisDebugging.h"
 #include "Net/Core/Trace/NetTrace.h"
+#include "Net/Core/Trace/NetDebugName.h"
 #include "Iris/DataStream/DataStream.h"
 #include "Iris/ReplicationSystem/DeltaCompression/DeltaCompressionBaselineManager.h"
 #include "Iris/ReplicationSystem/NetHandleManager.h"
diff --git a/Engine/Source/Runtime/Experimental/Iris/Core/Private/Iris/ReplicationSystem/ReplicationWriter.cpp b/Engine/Source/Runtime/Experimental/Iris/Core/Private/Iris/ReplicationSystem/ReplicationWriter.cpp
--- a/Engine/Source/Runtime/Experimental/Iris/Core/Private/Iris/ReplicationSystem/ReplicationWriter.cpp
+++ b/Engine/Source/Runtime/Experimental/Iris/Core/Private/Iris/ReplicationSystem/ReplicationWriter.cpp
@@ -6,6 +6,7 @@
 #include "Iris/Core/IrisLog.h"
 #include "Iris/Core/IrisProfiler.h"
 #include "Net/Core/Trace/NetTrace.h"
+#include "Net/Core/Trace/NetDebugName.h"
 #include "Iris/PacketControl/PacketNotification.h"
 #include "Iris/ReplicationSystem/ReplicationSystemTypes.h"
 #include "Iris/ReplicationSystem/DeltaCompression/DeltaCompressionBaselineManager.h"
