Iris
- Fixed crash in ReplicationStateDescriptorBuilder when building descriptor structs with CustomNetSerializer only occuring if compiling with clang

diff --git a/Engine/Source/Runtime/Experimental/Iris/Core/Private/Iris/ReplicationState/ReplicationStateDescriptorBuilder.cpp b/Engine/Source/Runtime/Experimental/Iris/Core/Private/Iris/ReplicationState/ReplicationStateDescriptorBuilder.cpp
--- a/Engine/Source/Runtime/Experimental/Iris/Core/Private/Iris/ReplicationState/ReplicationStateDescriptorBuilder.cpp
+++ b/Engine/Source/Runtime/Experimental/Iris/Core/Private/Iris/ReplicationState/ReplicationStateDescriptorBuilder.cpp
@@ -1909,7 +1909,7 @@ bool FPropertyReplicationStateDescriptorBuilder::CanStructUseStructNetSerializer
 FRepTag FPropertyReplicationStateDescriptorBuilder::GetRepTagFromProperty(const FMemberCacheEntry& MemberCacheEntry, const FMemberProperty& MemberProperty)
 {
 	const FProperty* Property = MemberProperty.Property;
-	const FName PropertyName = Property->GetFName();
+	const FName PropertyName = Property ? Property->GetFName() : FName();
 	if (PropertyName == NAME_Role || PropertyName == NAME_RemoteRole)
 	{
 		if (MemberCacheEntry.Serializer == &UE_NET_GET_SERIALIZER(FNetRoleNetSerializer))
