Fixed:
- Enabling Collision on a Spline Mesh component with a Static Mesh assigned causes a chrash while executing in a package
- Spline Mesh components did not cook mesh data if collision was disabled

diff --git a/Engine/Source/Runtime/Engine/Private/Components/SplineMeshComponent.cpp b/Engine/Source/Runtime/Engine/Private/Components/SplineMeshComponent.cpp
--- a/Engine/Source/Runtime/Engine/Private/Components/SplineMeshComponent.cpp
+++ b/Engine/Source/Runtime/Engine/Private/Components/SplineMeshComponent.cpp
@@ -1067,7 +1067,7 @@ void USplineMeshComponent::DestroyBodySetup()
 
 void USplineMeshComponent::RecreateCollision()
 {
-	if (GetStaticMesh() && IsCollisionEnabled())
+	if (GetStaticMesh())
 	{
 		if (BodySetup == NULL)
 		{
diff --git a/Engine/Source/Runtime/Experimental/Chaos/Public/Chaos/TriangleMeshImplicitObject.h b/Engine/Source/Runtime/Experimental/Chaos/Public/Chaos/TriangleMeshImplicitObject.h
--- a/Engine/Source/Runtime/Experimental/Chaos/Public/Chaos/TriangleMeshImplicitObject.h
+++ b/Engine/Source/Runtime/Experimental/Chaos/Public/Chaos/TriangleMeshImplicitObject.h
@@ -471,7 +471,7 @@ namespace Chaos
 		: FImplicitObject(EImplicitObject::HasBoundingBox | EImplicitObject::DisableCollisions, ImplicitObjectType::TriangleMesh)
 		, MParticles(MoveTemp(Particles))
 		, MElements(MoveTemp(Elements))
-		, MLocalBoundingBox(MParticles.X(0), MParticles.X(0))
+		, MLocalBoundingBox(FVec3(0), FVec3(0))
 		, MaterialIndices(MoveTemp(InMaterialIndices))
 		, ExternalFaceIndexMap(MoveTemp(InExternalFaceIndexMap))
 		, ExternalVertexIndexMap(MoveTemp(InExternalVertexIndexMap))
