[Chaos]
* ImplicitObject - Removed std functional include (not used)

diff --git a/Engine/Source/Runtime/Experimental/Chaos/Public/Chaos/ImplicitObject.h b/Engine/Source/Runtime/Experimental/Chaos/Public/Chaos/ImplicitObject.h
--- a/Engine/Source/Runtime/Experimental/Chaos/Public/Chaos/ImplicitObject.h
+++ b/Engine/Source/Runtime/Experimental/Chaos/Public/Chaos/ImplicitObject.h
@@ -7,8 +7,6 @@
 #include "Chaos/ImplicitFwd.h"
 #include "Chaos/AABB.h"
 
-#include <functional>
-
 #ifndef TRACK_CHAOS_GEOMETRY
 #define TRACK_CHAOS_GEOMETRY !(UE_BUILD_SHIPPING || UE_BUILD_TEST)
 #endif
