Chaos : fix geometry collection replication failing because the client side component has a authority local role when physics state is being created

diff --git a/Engine/Source/Runtime/Experimental/GeometryCollectionEngine/Private/GeometryCollection/GeometryCollectionComponent.cpp b/Engine/Source/Runtime/Experimental/GeometryCollectionEngine/Private/GeometryCollection/GeometryCollectionComponent.cpp
--- a/Engine/Source/Runtime/Experimental/GeometryCollectionEngine/Private/GeometryCollection/GeometryCollectionComponent.cpp
+++ b/Engine/Source/Runtime/Experimental/GeometryCollectionEngine/Private/GeometryCollection/GeometryCollectionComponent.cpp
@@ -2209,8 +2209,9 @@ void UGeometryCollectionComponent::AsyncPhysicsTickComponent(float DeltaTime, fl
 {
 	Super::AsyncPhysicsTickComponent(DeltaTime, SimTime);
 
-	const ENetRole LocalRole = GetOwnerRole();
-	if (LocalRole == ROLE_Authority)
+	// using net mode for now as using local role seemed to cause other issues at initialization time
+	// we may nee dto to also use local role in the future if the authority is likely to change at runtime
+	if (GetNetMode() != ENetMode::NM_Client)
 	{
 		UpdateRepData();
 	}
@@ -2509,11 +2510,16 @@ void UGeometryCollectionComponent::RegisterAndInitializePhysicsProxy()
 	PhysicsProxy->SetPostPhysicsSyncCallback([this]() { UpdateAttachedChildrenTransform(); }); 
 	if (GetIsReplicated())
 	{
-		const FGeometryCollectionPhysicsProxy::EReplicationMode ReplicationMode = 
-			(GetOwner()->GetLocalRole() == ENetRole::ROLE_Authority) 
-				? FGeometryCollectionPhysicsProxy::EReplicationMode::Server 
-				: FGeometryCollectionPhysicsProxy::EReplicationMode::Client;
-		PhysicsProxy->SetReplicationMode(ReplicationMode);
+		// using net mode and not local role because at this time in the initialization client and server both have an authority local role
+		const ENetMode NetMode = GetNetMode();
+		if (NetMode != NM_Standalone)
+		{
+			const FGeometryCollectionPhysicsProxy::EReplicationMode ReplicationMode =
+				(NetMode == ENetMode::NM_Client)
+				? FGeometryCollectionPhysicsProxy::EReplicationMode::Client
+				: FGeometryCollectionPhysicsProxy::EReplicationMode::Server;
+				PhysicsProxy->SetReplicationMode(ReplicationMode);
+		}
 	}
 
 	FPhysScene_Chaos* Scene = GetInnerChaosScene();
