Chaos : expose ISM Pool so it can be edited in construction scripts

diff --git a/Engine/Source/Runtime/Experimental/GeometryCollectionEngine/Public/GeometryCollection/GeometryCollectionComponent.h b/Engine/Source/Runtime/Experimental/GeometryCollectionEngine/Public/GeometryCollection/GeometryCollectionComponent.h
--- a/Engine/Source/Runtime/Experimental/GeometryCollectionEngine/Public/GeometryCollection/GeometryCollectionComponent.h
+++ b/Engine/Source/Runtime/Experimental/GeometryCollectionEngine/Public/GeometryCollection/GeometryCollectionComponent.h
@@ -1003,7 +1003,7 @@ protected:
 #endif
 	
 	/** ISM pool to use to render the geometry collection - only works for unfractured geometry collections  */
-	UPROPERTY(EditAnywhere, Category = "ChaosPhysics|Rendering", meta = (DisplayName = "ISM Pool"))
+	UPROPERTY(EditAnywhere, BlueprintReadWrite, Category = "ChaosPhysics|Rendering", meta = (DisplayName = "ISM Pool"))
 	TObjectPtr<AGeometryCollectionISMPoolActor> ISMPool;
 
 	int32 ISMPoolMeshGroupIndex = INDEX_NONE;
