Chaos : enable support for root proxy mesh for geometry collection  ( work only with ISM at the moment )

diff --git a/Engine/Source/Runtime/Experimental/Chaos/Private/GeometryCollection/Facades/CollectionHierarchyFacade.cpp b/Engine/Source/Runtime/Experimental/Chaos/Private/GeometryCollection/Facades/CollectionHierarchyFacade.cpp
--- a/Engine/Source/Runtime/Experimental/Chaos/Private/GeometryCollection/Facades/CollectionHierarchyFacade.cpp
+++ b/Engine/Source/Runtime/Experimental/Chaos/Private/GeometryCollection/Facades/CollectionHierarchyFacade.cpp
@@ -26,6 +26,24 @@ namespace Chaos::Facades
 		return LevelAttribute.IsValid() && LevelAttribute.IsPersistent();
 	}
 
+	int32 FCollectionHierarchyFacade::GetRootIndex() const
+	{
+		if (IsValid())
+		{
+			const TManagedArray<int32>& Parents = ParentAttribute.Get();
+			if (Parents.Num() > 0)
+			{
+				int32 RootTransformIndex = 0;
+				while (Parents[RootTransformIndex] != INDEX_NONE)
+				{
+					RootTransformIndex = Parents[RootTransformIndex];
+				}
+				return RootTransformIndex;
+			}
+		}
+		return INDEX_NONE;
+	}
+
 	void FCollectionHierarchyFacade::GenerateLevelAttribute()
 	{
 		check(IsValid());
diff --git a/Engine/Source/Runtime/Experimental/Chaos/Private/PhysicsProxy/GeometryCollectionPhysicsProxy.cpp b/Engine/Source/Runtime/Experimental/Chaos/Private/PhysicsProxy/GeometryCollectionPhysicsProxy.cpp
--- a/Engine/Source/Runtime/Experimental/Chaos/Private/PhysicsProxy/GeometryCollectionPhysicsProxy.cpp
+++ b/Engine/Source/Runtime/Experimental/Chaos/Private/PhysicsProxy/GeometryCollectionPhysicsProxy.cpp
@@ -653,7 +653,15 @@ void FGeometryCollectionPhysicsProxy::Initialize(Chaos::FPBDRigidsEvolutionBase
 	PhysicsThreadCollection.CopyAttribute(*Parameters.RestCollection, /*SrcName=*/"Level", /*DestName=*/"InitialLevel", FTransformCollection::TransformGroup);
 	GameThreadCollection.CopyAttribute(*Parameters.RestCollection, /*SrcName=*/"Level", /*DestName=*/"InitialLevel", FTransformCollection::TransformGroup);
 
-
+	if (Parameters.EnableClustering)
+	{
+		// make sure we set Activate the right way when clustering is enabled ( only root should be enabled at start ) 
+		for (int32 TransformIndex = 0; TransformIndex < NumTransforms; ++TransformIndex)
+		{
+			bool bIsRoot = (GameThreadCollection.Parent[TransformIndex] == INDEX_NONE);
+			GameThreadCollection.Active[TransformIndex] = bIsRoot;
+		}
+	}
 }
 
 
diff --git a/Engine/Source/Runtime/Experimental/Chaos/Public/GeometryCollection/Facades/CollectionHierarchyFacade.h b/Engine/Source/Runtime/Experimental/Chaos/Public/GeometryCollection/Facades/CollectionHierarchyFacade.h
--- a/Engine/Source/Runtime/Experimental/Chaos/Public/GeometryCollection/Facades/CollectionHierarchyFacade.h
+++ b/Engine/Source/Runtime/Experimental/Chaos/Public/GeometryCollection/Facades/CollectionHierarchyFacade.h
@@ -32,6 +32,9 @@ namespace Chaos::Facades
 		/** whether the level attribute is persistent */
 		bool IsLevelAttributePersistent() const;
 
+		/** Get the root index */
+		int32 GetRootIndex() const;
+
 		/** 
 		* Update level attribute for all elements (and create it if it is missing ) 
 		* @Param PersistencePolicy whether to make the attribute persistent or keep the existing state
diff --git a/Engine/Source/Runtime/Experimental/Chaos/Public/GeometryCollectionProxyData.h b/Engine/Source/Runtime/Experimental/Chaos/Public/GeometryCollectionProxyData.h
--- a/Engine/Source/Runtime/Experimental/Chaos/Public/GeometryCollectionProxyData.h
+++ b/Engine/Source/Runtime/Experimental/Chaos/Public/GeometryCollectionProxyData.h
@@ -216,10 +216,10 @@ public:
 
 	struct FState
 	{
-		int16 DynamicState: 8; // need to fit EObjectStateTypeEnum
-		int16 DisabledState: 1;
-		int16 HasInternalClusterParent: 1;
-		int16 DynamicInternalClusterParent: 1;
+		uint16 DynamicState: 8; // need to fit EObjectStateTypeEnum
+		uint16 DisabledState: 1;
+		uint16 HasInternalClusterParent: 1;
+		uint16 DynamicInternalClusterParent: 1;
 		// 5 bits left
 	};
 
diff --git a/Engine/Source/Runtime/Experimental/GeometryCollectionEngine/Private/GeometryCollection/GeometryCollectionComponent.cpp b/Engine/Source/Runtime/Experimental/GeometryCollectionEngine/Private/GeometryCollection/GeometryCollectionComponent.cpp
--- a/Engine/Source/Runtime/Experimental/GeometryCollectionEngine/Private/GeometryCollection/GeometryCollectionComponent.cpp
+++ b/Engine/Source/Runtime/Experimental/GeometryCollectionEngine/Private/GeometryCollection/GeometryCollectionComponent.cpp
@@ -11,6 +11,7 @@
 #include "Components/BoxComponent.h"
 #include "Engine/Engine.h"
 #include "Engine/InstancedStaticMesh.h"
+#include "GeometryCollection/Facades/CollectionHierarchyFacade.h"
 #include "GeometryCollection/GeometryCollectionActor.h"
 #include "GeometryCollection/GeometryCollectionAlgo.h"
 #include "GeometryCollection/GeometryCollectionCache.h"
@@ -3801,10 +3802,6 @@ void UGeometryCollectionComponent::InitializeEmbeddedGeometry()
 		EmbeddedBoneMaps.SetNum(RestCollection->EmbeddedGeometryExemplar.Num());
 		EmbeddedInstanceIndex.Init(INDEX_NONE,RestCollection->GetGeometryCollection()->NumElements(FGeometryCollection::TransformGroup));
 #endif
-
-		CalculateGlobalMatrices();
-		RefreshEmbeddedGeometry();
-		
 	}
 }
 
@@ -3824,6 +3821,7 @@ void UGeometryCollectionComponent::RegisterToISMPool()
 			bool bCanRenderComponent = true;
 			if (RestCollection)
 			{
+				ISMPoolMeshGroupIndex = ISMPoolComp->CreateMeshGroup();
 				if (bChaos_GC_UseISMPoolForNonFracturedParts)
 				{
 					if (RestCollection->GetGeometryCollection())
@@ -3849,7 +3847,6 @@ void UGeometryCollectionComponent::RegisterToISMPool()
 						}
 
 						// now register each mesh 
-						ISMPoolMeshGroupIndex = ISMPoolComp->CreateMeshGroup();
 						for (int32 MeshIndex = 0; MeshIndex < RestCollection->AutoInstanceMeshes.Num(); MeshIndex++)
 						{
 							const FGeometryCollectionAutoInstanceMesh& AutoInstanceMesh = RestCollection->AutoInstanceMeshes[MeshIndex];
@@ -3886,12 +3883,9 @@ void UGeometryCollectionComponent::RegisterToISMPool()
 						// if we use a mesh proxy hide the component for rendering 
 						bCanRenderComponent = false;
 
-						// TODO : store this state in a variable
-
 						FGeometryCollectionStaticMeshInstance StaticMeshInstance;
 						StaticMeshInstance.StaticMesh = RootMeshProxy;
-						// todo : get the mesh index and use it to upadte the mesh transform 
-						ISMPoolComp->AddMeshToGroup(ISMPoolMeshGroupIndex, StaticMeshInstance, 1);
+						ISMPoolRootProxyMeshId = ISMPoolComp->AddMeshToGroup(ISMPoolMeshGroupIndex, StaticMeshInstance, 1);
 					}
 				}
 			}
@@ -3911,6 +3905,7 @@ void UGeometryCollectionComponent::UnregisterFromISMPool()
 		{
 			ISMPoolComp->DestroyMeshGroup(ISMPoolMeshGroupIndex);
 			ISMPoolMeshGroupIndex = INDEX_NONE;
+			ISMPoolRootProxyMeshId = INDEX_NONE;
 		}
 	}
 	SetVisibility(true);
@@ -3937,32 +3932,56 @@ void UGeometryCollectionComponent::RefreshISMPoolInstances()
 							const int32 NumTransforms = RestCollection->NumElements(FGeometryCollection::TransformAttribute);
 							ensure(AutoInstanceMeshIndices->Num() == NumTransforms);
 
-							TArray<FTransform> InstanceTransforms;
 							CalculateGlobalMatrices();
 
 							const FTransform& ComponentTransform = GetComponentTransform();
-							for (int32 MeshIndex = 0; MeshIndex < RestCollection->AutoInstanceMeshes.Num(); MeshIndex++)
+
+							constexpr bool bWorlSpace = true;
+							constexpr bool bMarkRenderStateDirty = true;
+							constexpr bool bTeleport = true;
+
+							const int32 RootIndex = GetRootIndex();
+							//const bool bIsBroken = DynamicCollection ? (DynamicCollection->Children[RootIndex].Num() != Children[RootIndex].Num()) : false;
+							const bool bIsBroken = DynamicCollection ? !DynamicCollection->Active[RootIndex] : false;
+							const bool bHasRootProxyMesh = (ISMPoolRootProxyMeshId != INDEX_NONE);
+
+							if (bHasRootProxyMesh && !bIsBroken)
+							{
+								if (GlobalMatrices.IsValidIndex(RootIndex))
+								{
+									FTransform RootTransform = FTransform(GlobalMatrices[RootIndex]) * ComponentTransform;
+									ISMPoolComp->BatchUpdateInstancesTransforms(ISMPoolMeshGroupIndex, ISMPoolRootProxyMeshId, 0, { RootTransform }, bWorlSpace, bMarkRenderStateDirty, bTeleport);
+								}
+							}
+							else if (bChaos_GC_UseISMPoolForNonFracturedParts)
 							{
-								InstanceTransforms.Reset(NumTransforms); // Allocate for worst case
-								for (int32 TransformIndex = 0; TransformIndex < NumTransforms; TransformIndex++)
+								// make sure this mesh is invisible 
+								// todo : should be event based instead of doing it every frame
+								if (bHasRootProxyMesh && GlobalMatrices.IsValidIndex(RootIndex))
 								{
-									const int32 AutoInstanceMeshIndex = (*AutoInstanceMeshIndices)[TransformIndex];
-									if (AutoInstanceMeshIndex == MeshIndex && Children[TransformIndex].Num() == 0)
+									FTransform RootTransformZeroScale;
+									RootTransformZeroScale.SetIdentityZeroScale();
+									ISMPoolComp->BatchUpdateInstancesTransforms(ISMPoolMeshGroupIndex, ISMPoolRootProxyMeshId, 0, { RootTransformZeroScale }, bWorlSpace, bMarkRenderStateDirty, bTeleport);
+								}
+
+								TArray<FTransform> InstanceTransforms;
+								for (int32 MeshIndex = 0; MeshIndex < RestCollection->AutoInstanceMeshes.Num(); MeshIndex++)
+								{
+									InstanceTransforms.Reset(NumTransforms); // Allocate for worst case
+									for (int32 TransformIndex = 0; TransformIndex < NumTransforms; TransformIndex++)
 									{
-										InstanceTransforms.Add(FTransform(GlobalMatrices[TransformIndex]) * ComponentTransform);
+										const int32 AutoInstanceMeshIndex = (*AutoInstanceMeshIndices)[TransformIndex];
+										if (AutoInstanceMeshIndex == MeshIndex && Children[TransformIndex].Num() == 0)
+										{
+											InstanceTransforms.Add(FTransform(GlobalMatrices[TransformIndex]) * ComponentTransform);
+										}
 									}
+									ISMPoolComp->BatchUpdateInstancesTransforms(ISMPoolMeshGroupIndex, MeshIndex, 0, InstanceTransforms, bWorlSpace, bMarkRenderStateDirty, bTeleport);
 								}
-								constexpr bool bWorlSpace = true;
-								constexpr bool bMarkRenderStateDirty = true;
-								constexpr bool bTeleport = true;
-								ISMPoolComp->BatchUpdateInstancesTransforms(ISMPoolMeshGroupIndex, MeshIndex, 0, InstanceTransforms, bWorlSpace, bMarkRenderStateDirty, bTeleport);
 							}
-				
 						}
 					}
 				}
-
-				// todo : update the mesh proxy if set 
 			}
 		}
 	}
@@ -4298,6 +4317,16 @@ int32 UGeometryCollectionComponent::GetInitialLevel(int32 ItemIndex)
 	return Level;
 }
 
+int32 UGeometryCollectionComponent::GetRootIndex()
+{
+	if (RestCollection && RestCollection->GetGeometryCollection())
+	{
+		Chaos::Facades::FCollectionHierarchyFacade HierarchyFacade(*RestCollection->GetGeometryCollection());
+		return HierarchyFacade.GetRootIndex();
+	}
+	return INDEX_NONE;
+}
+
 void UGeometryCollectionComponent::GetMassAndExtents(int32 ItemIndex, float& OutMass, FBox& OutExtents)
 {
 	using FGeometryCollectionPtr = const TSharedPtr<FGeometryCollection, ESPMode::ThreadSafe>;
diff --git a/Engine/Source/Runtime/Experimental/GeometryCollectionEngine/Private/GeometryCollection/GeometryCollectionISMPoolComponent.cpp b/Engine/Source/Runtime/Experimental/GeometryCollectionEngine/Private/GeometryCollection/GeometryCollectionISMPoolComponent.cpp
--- a/Engine/Source/Runtime/Experimental/GeometryCollectionEngine/Private/GeometryCollection/GeometryCollectionISMPoolComponent.cpp
+++ b/Engine/Source/Runtime/Experimental/GeometryCollectionEngine/Private/GeometryCollection/GeometryCollectionISMPoolComponent.cpp
@@ -73,9 +73,11 @@ int32 FGeometryCollectionISM::AddInstanceGroup(int32 InstanceCount)
 	const FInstanceGroups::FInstanceGroupRange& NewInstanceGroup = InstanceGroups.GetGroup(InstanceGroupIndex);
 	const int32 TotalInstanceCount = NewInstanceGroup.Start + NewInstanceGroup.Count;
 	ISMComponent->PreAllocateInstancesMemory(TotalInstanceCount);
+	FTransform ZeroScaleTransform;
+	ZeroScaleTransform.SetIdentityZeroScale();
 	for (int32 InstanceIndex = NewInstanceGroup.Start; InstanceIndex < TotalInstanceCount; InstanceIndex++)
 	{
-		ISMComponent->AddInstance(FTransform::Identity, true);
+		ISMComponent->AddInstance(ZeroScaleTransform, true);
 	}
 	return InstanceGroupIndex;
 }
diff --git a/Engine/Source/Runtime/Experimental/GeometryCollectionEngine/Public/GeometryCollection/GeometryCollectionComponent.h b/Engine/Source/Runtime/Experimental/GeometryCollectionEngine/Public/GeometryCollection/GeometryCollectionComponent.h
--- a/Engine/Source/Runtime/Experimental/GeometryCollectionEngine/Public/GeometryCollection/GeometryCollectionComponent.h
+++ b/Engine/Source/Runtime/Experimental/GeometryCollectionEngine/Public/GeometryCollection/GeometryCollectionComponent.h
@@ -601,6 +601,10 @@ public:
 	UFUNCTION(BlueprintCallable, Category = "ChaosPhysics")
 	int32 GetInitialLevel(int32 ItemIndex);
 	
+	/** Get the root item index of the hierarchy */
+	UFUNCTION(BlueprintCallable, Category = "ChaosPhysics")
+	int32 GetRootIndex();
+
 	/**
 	* Get mass and extent of a specific piece
 	* @param ItemIndex item index ( from HitResult) of the cluster to get level from
@@ -1007,6 +1011,7 @@ protected:
 	TObjectPtr<AGeometryCollectionISMPoolActor> ISMPool;
 
 	int32 ISMPoolMeshGroupIndex = INDEX_NONE;
+	int32 ISMPoolRootProxyMeshId = INDEX_NONE;
 
 	/** Populate the static geometry structures for the render thread. */
 	void InitConstantData(FGeometryCollectionConstantData* ConstantData) const;
