Re-enable Iris compilation by default. Remains disabled at runtime.
Fix static analysis & nonunity issues.

diff --git a/Engine/Source/Programs/UnrealBuildTool/Configuration/TargetRules.cs b/Engine/Source/Programs/UnrealBuildTool/Configuration/TargetRules.cs
--- a/Engine/Source/Programs/UnrealBuildTool/Configuration/TargetRules.cs
+++ b/Engine/Source/Programs/UnrealBuildTool/Configuration/TargetRules.cs
@@ -782,7 +782,7 @@ namespace UnrealBuildTool
 		[RequiresUniqueBuildEnvironment]
 		[CommandLine("-NoUseIris", Value = "false")]
 		[CommandLine("-UseIris", Value = "true")]
-		public bool bUseIris = false;
+		public bool bUseIris = true;
 
 		/// <summary>
 		/// Whether we are compiling editor code or not. Prefer the more explicit bCompileAgainstEditor instead.
diff --git a/Engine/Source/Runtime/Engine/Private/Net/Experimental/Iris/DataStreamChannel.cpp b/Engine/Source/Runtime/Engine/Private/Net/Experimental/Iris/DataStreamChannel.cpp
--- a/Engine/Source/Runtime/Engine/Private/Net/Experimental/Iris/DataStreamChannel.cpp
+++ b/Engine/Source/Runtime/Engine/Private/Net/Experimental/Iris/DataStreamChannel.cpp
@@ -5,6 +5,7 @@
 #include "Net/Core/Misc/ResizableCircularQueue.h"
 #include "Net/DataBunch.h"
 #include "Net/Core/Trace/NetTrace.h"
+#include "PacketHandler.h"
 
 #if UE_WITH_IRIS
 #include "Iris/IrisConfig.h"
diff --git a/Engine/Source/Runtime/Engine/Private/Net/Iris/ReplicationSystem/ActorReplicationBridge.cpp b/Engine/Source/Runtime/Engine/Private/Net/Iris/ReplicationSystem/ActorReplicationBridge.cpp
--- a/Engine/Source/Runtime/Engine/Private/Net/Iris/ReplicationSystem/ActorReplicationBridge.cpp
+++ b/Engine/Source/Runtime/Engine/Private/Net/Iris/ReplicationSystem/ActorReplicationBridge.cpp
@@ -21,6 +21,7 @@
 #include "Iris/Serialization/IrisObjectReferencePackageMap.h"
 #include "Engine/Engine.h"
 #include "Engine/EngineTypes.h"
+#include "Engine/Level.h"
 #include "Engine/NetConnection.h"
 #include "Engine/NetDriver.h"
 #include "Engine/World.h"
@@ -270,7 +271,7 @@ UE::Net::FNetHandle UActorReplicationBridge::BeginReplication(AActor* Actor, con
 
 	// Set owning connection filtering if actor is only relevant to owner
 	{
-		if (Actor->bOnlyRelevantToOwner & !Actor->bAlwaysRelevant)
+		if (Actor->bOnlyRelevantToOwner && !Actor->bAlwaysRelevant)
 		{
 			GetReplicationSystem()->SetFilter(ActorHandle, ToOwnerFilterHandle);
 		}
@@ -888,7 +889,7 @@ void UActorReplicationBridge::GetActorCreationHeader(const AActor* Actor, UE::Ne
 			Header.LevelReference = GetOrCreateObjectReference(ActorLevel);
 		}
 
-		if (USceneComponent* RootComponent = Actor ? Actor->GetRootComponent() : nullptr)
+		if (USceneComponent* RootComponent = Actor->GetRootComponent())
 		{
 			Header.SpawnInfo.Location = FRepMovement::RebaseOntoZeroOrigin(Actor->GetActorLocation(), Actor);
 			Header.SpawnInfo.Rotation = Actor->GetActorRotation();
diff --git a/Engine/Source/Runtime/Experimental/Iris/Core/Private/Iris/ReplicationState/ReplicationStateStorage.cpp b/Engine/Source/Runtime/Experimental/Iris/Core/Private/Iris/ReplicationState/ReplicationStateStorage.cpp
--- a/Engine/Source/Runtime/Experimental/Iris/Core/Private/Iris/ReplicationState/ReplicationStateStorage.cpp
+++ b/Engine/Source/Runtime/Experimental/Iris/Core/Private/Iris/ReplicationState/ReplicationStateStorage.cpp
@@ -56,31 +56,42 @@ const uint8* FReplicationStateStorage::GetState(uint32 ObjectIndex, EReplication
 
 	if (const FPerObjectInfo* ObjectInfo = GetPerObjectInfoForObject(ObjectIndex))
 	{
-		if (StateType == EReplicationStateType::CurrentSendState)
+		switch (StateType)
 		{
-			return ObjectInfo->StateBuffers[(unsigned)EStateBufferType::SendState];
+			case EReplicationStateType::CurrentSendState:
+			{
+				return ObjectInfo->StateBuffers[(unsigned)EStateBufferType::SendState];
+			}
+			case EReplicationStateType::CurrentRecvState:
+			{
+				return ObjectInfo->StateBuffers[(unsigned)EStateBufferType::RecvState];
+			}
+			default:
+			{
+				checkf(false, TEXT("Unknown EReplicationStateType"));
+				return nullptr;
+			}
 		}
-
-		if (StateType == EReplicationStateType::CurrentRecvState)
-		{
-			return ObjectInfo->StateBuffers[(unsigned)EStateBufferType::RecvState];
-		}
-
-		return nullptr;
 	}
 
 	// Slow path but we don't want to keep PerObjectInfo around unless we have to.
-	if (StateType == EReplicationStateType::CurrentSendState)
+	switch (StateType)
 	{
-		return NetHandleManager->GetReplicatedObjectStateBufferNoCheck(ObjectIndex);
-	}
-	else if (StateType == EReplicationStateType::CurrentSendState)
-	{
-		const Private::FNetHandleManager::FReplicatedObjectData& ReplicatedObjectData = NetHandleManager->GetReplicatedObjectData(ObjectIndex);
-		return ReplicatedObjectData.ReceiveStateBuffer;
+		case EReplicationStateType::CurrentSendState:
+		{
+			return NetHandleManager->GetReplicatedObjectStateBufferNoCheck(ObjectIndex);
+		}
+		case EReplicationStateType::CurrentRecvState:
+		{
+			const Private::FNetHandleManager::FReplicatedObjectData& ReplicatedObjectData = NetHandleManager->GetReplicatedObjectData(ObjectIndex);
+			return ReplicatedObjectData.ReceiveStateBuffer;
+		}
+		default:
+		{
+			checkf(false, TEXT("Unknown EReplicationStateType"));
+			return nullptr;
+		}
 	}
-
-	return nullptr;
 }
 
 uint8* FReplicationStateStorage::AllocBaseline(uint32 ObjectIndex, EReplicationStateType Base)
diff --git a/Engine/Source/Runtime/Experimental/Iris/Core/Private/Iris/ReplicationSystem/AttachmentReplication.cpp b/Engine/Source/Runtime/Experimental/Iris/Core/Private/Iris/ReplicationSystem/AttachmentReplication.cpp
--- a/Engine/Source/Runtime/Experimental/Iris/Core/Private/Iris/ReplicationSystem/AttachmentReplication.cpp
+++ b/Engine/Source/Runtime/Experimental/Iris/Core/Private/Iris/ReplicationSystem/AttachmentReplication.cpp
@@ -1008,26 +1008,6 @@ FNetObjectAttachmentReceiveQueue* FNetObjectAttachmentsReader::GetOrCreateQueue(
 
 		return Queue;
 	}
-
-
-	FNetObjectAttachmentReceiveQueue* Queue = ObjectToQueue.Find(ObjectIndex);
-	if (Queue != nullptr)
-	{
-		return Queue;
-	}
-
-	if (ObjectToQueue.Num() >= AttachmentReplicationCVars::MaxSimultaneousObjectsWithRPCs)
-	{
-		return nullptr;
-	}
-
-	Queue = &ObjectToQueue.Add(ObjectIndex);
-	if (ObjectIndex == 0)
-	{
-		Queue->SetUnreliableQueueCapacity(AttachmentReplicationCVars::ClientToServerUnreliableRPCQueueSize);
-	}
-
-	return Queue;
 }
 
 }
diff --git a/Engine/Source/Runtime/Experimental/Iris/Core/Private/Iris/ReplicationSystem/Conditionals/ReplicationConditionals.cpp b/Engine/Source/Runtime/Experimental/Iris/Core/Private/Iris/ReplicationSystem/Conditionals/ReplicationConditionals.cpp
--- a/Engine/Source/Runtime/Experimental/Iris/Core/Private/Iris/ReplicationSystem/Conditionals/ReplicationConditionals.cpp
+++ b/Engine/Source/Runtime/Experimental/Iris/Core/Private/Iris/ReplicationSystem/Conditionals/ReplicationConditionals.cpp
@@ -442,7 +442,7 @@ bool FReplicationConditionals::ApplyConditionalsToChangeMask(uint32 ReplicatingC
 						const FReplicationStateMemberChangeMaskDescriptor& ChangeMaskDescriptor = ChangeMaskDescriptors[MemberIt];
 						for (uint32 BitIt = ChangeMaskBitOffset + ChangeMaskDescriptor.BitOffset, BitEndIt = BitIt + ChangeMaskDescriptor.BitCount; BitIt != BitEndIt; ++BitIt)
 						{
-							bMaskWasModified = bMaskWasModified | !ChangeMask.GetBit(BitIt);
+							bMaskWasModified |= !ChangeMask.GetBit(BitIt);
 							ChangeMask.SetBit(BitIt);
 						}
 					}
@@ -452,7 +452,7 @@ bool FReplicationConditionals::ApplyConditionalsToChangeMask(uint32 ReplicatingC
 					const FReplicationStateMemberChangeMaskDescriptor& ChangeMaskDescriptor = ChangeMaskDescriptors[MemberIt];
 					for (uint32 BitIt = ChangeMaskBitOffset + ChangeMaskDescriptor.BitOffset, BitEndIt = BitIt + ChangeMaskDescriptor.BitCount; BitIt != BitEndIt; ++BitIt)
 					{
-						bMaskWasModified = bMaskWasModified | ChangeMask.GetBit(BitIt);
+						bMaskWasModified |= ChangeMask.GetBit(BitIt);
 						ChangeMask.ClearBit(BitIt);
 					}
 				}
@@ -481,7 +481,7 @@ bool FReplicationConditionals::ApplyConditionalsToChangeMask(uint32 ReplicatingC
 								const FReplicationStateMemberChangeMaskDescriptor& ChangeMaskDescriptor = ChangeMaskDescriptors[MemberIt];
 								for (uint32 BitIt = CurrentChangeMaskBitOffset + ChangeMaskDescriptor.BitOffset, BitEndIt = BitIt + ChangeMaskDescriptor.BitCount; BitIt != BitEndIt; ++BitIt)
 								{
-									bMaskWasModified = bMaskWasModified | !ChangeMask.GetBit(BitIt);
+									bMaskWasModified |= !ChangeMask.GetBit(BitIt);
 									ChangeMask.SetBit(BitIt);
 								}
 							}
@@ -491,7 +491,7 @@ bool FReplicationConditionals::ApplyConditionalsToChangeMask(uint32 ReplicatingC
 							const FReplicationStateMemberChangeMaskDescriptor& ChangeMaskDescriptor = ChangeMaskDescriptors[MemberIt];
 							for (uint32 BitIt = CurrentChangeMaskBitOffset + ChangeMaskDescriptor.BitOffset, BitEndIt = BitIt + ChangeMaskDescriptor.BitCount; BitIt != BitEndIt; ++BitIt)
 							{
-								bMaskWasModified = bMaskWasModified | ChangeMask.GetBit(BitIt);
+								bMaskWasModified |= ChangeMask.GetBit(BitIt);
 								ChangeMask.ClearBit(BitIt);
 							}
 						}
diff --git a/Engine/Source/Runtime/Experimental/Iris/Core/Private/Iris/ReplicationSystem/DequantizeAndApplyHelper.cpp b/Engine/Source/Runtime/Experimental/Iris/Core/Private/Iris/ReplicationSystem/DequantizeAndApplyHelper.cpp
--- a/Engine/Source/Runtime/Experimental/Iris/Core/Private/Iris/ReplicationSystem/DequantizeAndApplyHelper.cpp
+++ b/Engine/Source/Runtime/Experimental/Iris/Core/Private/Iris/ReplicationSystem/DequantizeAndApplyHelper.cpp
@@ -92,7 +92,7 @@ FDequantizeAndApplyHelper::FContext* FDequantizeAndApplyHelper::Initialize(FNetS
 		ChangeMaskReader.ReadBitStream(ChangeMaskStorage, ChangeMaskBitCount);
 
 		// Cache all ReplicationStates with dirty changes
-		const bool bShouldDequantizeState = ChangeMask.IsAnyBitSet() | (NetSerializationContext.IsInitState() & CurrentDescriptor->IsInitState());
+		const bool bShouldDequantizeState = ChangeMask.IsAnyBitSet() || (NetSerializationContext.IsInitState() && CurrentDescriptor->IsInitState());
 		if (bShouldDequantizeState)
 		{
 			FContext::FStateData& StateData = CachedStateData[CachedStateCount];
@@ -131,7 +131,7 @@ FDequantizeAndApplyHelper::FContext* FDequantizeAndApplyHelper::Initialize(FNetS
 					const FNetBitArrayView UnresolvedChangeMask = MakeNetBitArrayView(UnresolvedChangeMaskData, Protocol->ChangeMaskBitCount);
 					bHasUnresolvableReferences = UnresolvedChangeMask.IsAnyBitSet(CurrentChangeMaskBitOffset, CurrentDescriptor->ChangeMaskBitCount);
 				}
-				bMightHaveUnresolvableInitReferences = Parameters.bHasUnresolvedInitReferences & NetSerializationContext.IsInitState() & CurrentDescriptor->IsInitState();
+				bMightHaveUnresolvableInitReferences = Parameters.bHasUnresolvedInitReferences && NetSerializationContext.IsInitState() && CurrentDescriptor->IsInitState();
 			}
 
 			StateData.bHasUnresolvableReferences = bHasUnresolvableReferences;
diff --git a/Engine/Source/Runtime/Experimental/Iris/Core/Private/Iris/ReplicationSystem/NetBlob/NetRPC.cpp b/Engine/Source/Runtime/Experimental/Iris/Core/Private/Iris/ReplicationSystem/NetBlob/NetRPC.cpp
--- a/Engine/Source/Runtime/Experimental/Iris/Core/Private/Iris/ReplicationSystem/NetBlob/NetRPC.cpp
+++ b/Engine/Source/Runtime/Experimental/Iris/Core/Private/Iris/ReplicationSystem/NetBlob/NetRPC.cpp
@@ -361,7 +361,7 @@ FNetRPC* FNetRPC::Create(UReplicationSystem* ReplicationSystem, const FNetBlobCr
 	uint8* QuantizedBlobState = nullptr;
 
 	// Don't spend CPU cycles on quantizing zero parameters
-	if (BlobDescriptor->InternalSize)
+	if (BlobDescriptor != nullptr && BlobDescriptor->InternalSize)
 	{
 		uint8* StateBuffer = static_cast<uint8*>(GMalloc->Malloc(BlobDescriptor->InternalSize, BlobDescriptor->InternalAlignment));
 		FMemory::Memzero(StateBuffer, BlobDescriptor->InternalSize);
diff --git a/Engine/Source/Runtime/Experimental/Iris/Core/Private/Iris/ReplicationSystem/ObjectReferenceCache.cpp b/Engine/Source/Runtime/Experimental/Iris/Core/Private/Iris/ReplicationSystem/ObjectReferenceCache.cpp
--- a/Engine/Source/Runtime/Experimental/Iris/Core/Private/Iris/ReplicationSystem/ObjectReferenceCache.cpp
+++ b/Engine/Source/Runtime/Experimental/Iris/Core/Private/Iris/ReplicationSystem/ObjectReferenceCache.cpp
@@ -305,7 +305,7 @@ bool FObjectReferenceCache::CreateObjectReferenceInternal(const UObject* Object,
 	CachedObject.NetHandle = NetHandle;
 	CachedObject.RelativePath = PathToken;
 	CachedObject.OuterNetHandle = CreateObjectReferenceHandle(Object->GetOuter());
-	CachedObject.bIsPackage = bIsStatic && Object && !Object->GetOuter();
+	CachedObject.bIsPackage = bIsStatic && !Object->GetOuter();
 	CachedObject.bNoLoad = !CanClientLoadObjectInternal(Object, bIsDynamic); // We are probably going to deal with this as a dependency
 	CachedObject.bIgnoreWhenMissing = !CachedObject.bNoLoad;
 	CachedObject.bIsPending = false;
diff --git a/Engine/Source/Runtime/Experimental/Iris/Core/Private/Iris/ReplicationSystem/ReplicationWriter.cpp b/Engine/Source/Runtime/Experimental/Iris/Core/Private/Iris/ReplicationSystem/ReplicationWriter.cpp
--- a/Engine/Source/Runtime/Experimental/Iris/Core/Private/Iris/ReplicationSystem/ReplicationWriter.cpp
+++ b/Engine/Source/Runtime/Experimental/Iris/Core/Private/Iris/ReplicationSystem/ReplicationWriter.cpp
@@ -651,7 +651,7 @@ void FReplicationWriter::UpdateScope(const FNetBitArrayView& UpdatedScope)
 				if (OwnerInfo.GetState() < EReplicatedObjectState::PendingDestroy)
 				{
 					OwnerInfo.HasDirtySubObjects |= Info.HasDirtyChangeMask;
-					ObjectsWithDirtyChanges.SetBitValue(ObjectData.SubObjectRootIndex, ObjectsWithDirtyChanges.GetBit(ObjectData.SubObjectRootIndex) | Info.HasDirtyChangeMask);
+					ObjectsWithDirtyChanges.SetBitValue(ObjectData.SubObjectRootIndex, ObjectsWithDirtyChanges.GetBit(ObjectData.SubObjectRootIndex) || Info.HasDirtyChangeMask);
 				}
 			}
 		}
@@ -822,7 +822,7 @@ uint32 FReplicationWriter::ScheduleObjects(FScheduleObjectInfo* OutScheduledObje
 		ScheduledObjectInfo.Index = Index;
 		ScheduledObjectInfo.SortKey = UpdatedPriority;
 
-		ScheduledObjectCount += (UpdatedPriority >= FReplicationWriter::SchedulingThresholdPriority) ? 1.f : 0.f;
+		ScheduledObjectCount += (UpdatedPriority >= FReplicationWriter::SchedulingThresholdPriority) ? 1 : 0;
 	};
 
 	// Invoke functor for all updated objects that not are sub objects.
@@ -1211,7 +1211,7 @@ void FReplicationWriter::HandleDroppedRecord<FReplicationWriter::EReplicatedObje
 			if (OwnerInfo.GetState() < EReplicatedObjectState::PendingDestroy)
 			{
 				OwnerInfo.HasDirtySubObjects |= Info.HasDirtyChangeMask;
-				ObjectsWithDirtyChanges.SetBitValue(ObjectData.SubObjectRootIndex, ObjectsWithDirtyChanges.GetBit(ObjectData.SubObjectRootIndex) | OwnerInfo.HasDirtySubObjects);
+				ObjectsWithDirtyChanges.SetBitValue(ObjectData.SubObjectRootIndex, ObjectsWithDirtyChanges.GetBit(ObjectData.SubObjectRootIndex) || OwnerInfo.HasDirtySubObjects);
 			}
 		}
 		else
@@ -1713,16 +1713,16 @@ FReplicationWriter::EWriteObjectStatus FReplicationWriter::WriteObjectInBatch(FN
 	uint8* ReplicatedObjectStateBuffer = NetHandleManager->GetReplicatedObjectStateBufferNoCheck(InternalIndex);
 
 	// Filter out changemasks that are not supposed to be replicated to this connection
-	const bool bNeedToFilterChangeMask = (bIsInitialState | Info.HasDirtyChangeMask) & Info.HasChangemaskFilter;
+	const bool bNeedToFilterChangeMask = (bIsInitialState || Info.HasDirtyChangeMask) && Info.HasChangemaskFilter;
 	if (bNeedToFilterChangeMask)
 	{
 		ApplyFilterToChangeMask(OutBatchInfo.ParentInternalIndex, InternalIndex, Info, ObjectData.Protocol, ReplicatedObjectStateBuffer);
 	}
 
 	const bool bIsObjectIndexForAttachment = IsObjectIndexForOOBAttachment(InternalIndex);
-	const bool bHasState = (bIsInitialState | Info.HasDirtyChangeMask) & !!(WriteObjectFlags & EWriteObjectFlag::WriteObjectFlag_State);
-	const bool bHasAttachments = (Info.HasAttachments | bIsObjectIndexForAttachment);
-	const bool bWriteAttachments = bHasAttachments & !!(WriteObjectFlags & EWriteObjectFlag::WriteObjectFlag_Attachments);
+	const bool bHasState = (bIsInitialState || Info.HasDirtyChangeMask) && !!(WriteObjectFlags & EWriteObjectFlag::WriteObjectFlag_State);
+	const bool bHasAttachments = (Info.HasAttachments || bIsObjectIndexForAttachment);
+	const bool bWriteAttachments = bHasAttachments && !!(WriteObjectFlags & EWriteObjectFlag::WriteObjectFlag_Attachments);
 	BatchEntry.bHasUnsentAttachments = bHasAttachments;
 
 	// Check if we must defer tearoff until after flush
@@ -2033,7 +2033,7 @@ FReplicationWriter::EWriteObjectStatus FReplicationWriter::WriteObjectInBatch(FN
 			if (OutBatchInfo.ObjectInfos.Num() > BatchObjectInfoCount)
 			{
 				const FBatchObjectInfo& SubObjectEntry = OutBatchInfo.ObjectInfos.Last();
-				bHasDirtySubObjects = SubObjectEntry.bHasDirtySubObjects | SubObjectEntry.bHasUnsentAttachments;
+				bHasDirtySubObjects = SubObjectEntry.bHasDirtySubObjects || SubObjectEntry.bHasUnsentAttachments;
 			}
 		}
 
@@ -2555,7 +2555,7 @@ int FReplicationWriter::HandleObjectBatchSuccess(const FBatchInfo& BatchInfo, FR
 			Info.PendingBaselineIndex = BatchObjectInfo.NewBaselineIndex;
 		}
 		
-		const bool bObjectIsStillDirty = BatchObjectInfo.bHasUnsentAttachments | BatchObjectInfo.bHasDirtySubObjects;
+		const bool bObjectIsStillDirty = BatchObjectInfo.bHasUnsentAttachments || BatchObjectInfo.bHasDirtySubObjects;
 		ObjectsWithDirtyChanges.SetBitValue(BatchObjectInfo.InternalIndex, bObjectIsStillDirty);
 
 		// Reset scheduling priority if everything was replicated
@@ -2571,7 +2571,7 @@ int FReplicationWriter::HandleObjectBatchSuccess(const FBatchInfo& BatchInfo, FR
 		}
 
 #if UE_NET_IRIS_CSV_STATS && CSV_PROFILER
-		if (BatchObjectInfo.bSentState & (Info.LastAckedBaselineIndex != FDeltaCompressionBaselineManager::InvalidBaselineIndex))
+		if (BatchObjectInfo.bSentState && (Info.LastAckedBaselineIndex != FDeltaCompressionBaselineManager::InvalidBaselineIndex))
 		{
 			++WriteContext.DeltaCompressedObjectCount;
 		}
diff --git a/Engine/Source/Runtime/Experimental/Iris/Core/Private/Iris/ReplicationSystem/ReplicationWriter.h b/Engine/Source/Runtime/Experimental/Iris/Core/Private/Iris/ReplicationSystem/ReplicationWriter.h
--- a/Engine/Source/Runtime/Experimental/Iris/Core/Private/Iris/ReplicationSystem/ReplicationWriter.h
+++ b/Engine/Source/Runtime/Experimental/Iris/Core/Private/Iris/ReplicationSystem/ReplicationWriter.h
@@ -95,7 +95,7 @@ public:
 		FlushFlags_None				= 0U,
 		FlushFlags_FlushState		= 1U << 0U,											// Make sure that all current state data is acknowledged before we stop replicating the object
 		FlushFlags_FlushReliable	= FlushFlags_FlushState << 1U,						// Make sure that all enqueued Reliable RPCs are delivered before we stop replicating the object
-		FlushFlags_All				= FlushFlags_FlushState | FlushFlags_FlushState,
+		FlushFlags_All				= FlushFlags_FlushState | FlushFlags_FlushReliable,
 		FlushFlags_Default			= FlushFlags_FlushReliable,
 	};
 
diff --git a/Engine/Source/Runtime/Experimental/Iris/Core/Private/Iris/Serialization/PackedVectorNetSerializers.cpp b/Engine/Source/Runtime/Experimental/Iris/Core/Private/Iris/Serialization/PackedVectorNetSerializers.cpp
--- a/Engine/Source/Runtime/Experimental/Iris/Core/Private/Iris/Serialization/PackedVectorNetSerializers.cpp
+++ b/Engine/Source/Runtime/Experimental/Iris/Core/Private/Iris/Serialization/PackedVectorNetSerializers.cpp
@@ -239,7 +239,7 @@ void FPackedVectorNetSerializerBase::Deserialize(uint32 ScaleBitCount, FNetSeria
 			TempValue.Z |= static_cast<uint64>(Reader->ReadBits(32U)) << 32U;
 
 			FVector3d Vector;
-			memcpy(&Vector, &TempValue.X, 3U*sizeof(double));
+			memcpy(&Vector, &TempValue.X, 3U*sizeof(double)); //-V512
 			if (Vector.ContainsNaN())
 			{
 				// While we could detect this at send time it's very likely that a NaN or infinite value
@@ -518,7 +518,7 @@ void FPackedVectorNetSerializerBase::DeserializeDelta(uint32 ScaleBitCount, FNet
 			}
 
 			FVector3d Vector;
-			memcpy(&Vector, &TempValue.X, 3U*sizeof(double));
+			memcpy(&Vector, &TempValue.X, 3U*sizeof(double)); //-V512
 			if (Vector.ContainsNaN())
 			{
 				// While we could detect this at send time it's very likely that a NaN or infinite value
@@ -680,12 +680,12 @@ void FPackedVectorNetSerializerBase::Dequantize(uint32 ScaleBitCount, FNetSerial
 		{
 			if constexpr (std::is_same<SourceType, FVector3d>::value)
 			{
-				memcpy(&Target.X, &Source.X, 3U*sizeof(double));
+				memcpy(&Target.X, &Source.X, 3U*sizeof(double)); //-V512
 			}
 			else
 			{
 				FVector3d Vector;
-				memcpy(&Vector.X, &Source.X, 3U*sizeof(double));
+				memcpy(&Vector.X, &Source.X, 3U*sizeof(double)); //-V512
 				Target = SourceType(Vector);
 			}
 		}
diff --git a/Engine/Source/Runtime/Experimental/Iris/Core/Private/Iris/Serialization/StringNetSerializerUtils.h b/Engine/Source/Runtime/Experimental/Iris/Core/Private/Iris/Serialization/StringNetSerializerUtils.h
--- a/Engine/Source/Runtime/Experimental/Iris/Core/Private/Iris/Serialization/StringNetSerializerUtils.h
+++ b/Engine/Source/Runtime/Experimental/Iris/Core/Private/Iris/Serialization/StringNetSerializerUtils.h
@@ -195,7 +195,7 @@ public:
 				 * A pair of surrogates would be encoded as 6 bytes. So we construct the 4-byte codepoint and encode them as 3 bytes.
 				 * We let the decoder deal with invalid surrogate sequences, such as wrong order or a single surrogate.
 				 */
-				if (IsHighSurrogate(Char) & IsLowSurrogate(NextChar))
+				if (IsHighSurrogate(Char) && IsLowSurrogate(NextChar))
 				{
 					const uint32 Codepoint = GetCodepointFromSurrogates(Char, NextChar);
 					Dest[DestIt + 0] = EncodeType(0x80) | static_cast<EncodeType>(Codepoint >> 14U);
diff --git a/Engine/Source/Runtime/Experimental/Iris/Core/Public/Iris/ReplicationState/Private/IrisFastArraySerializerInternal.h b/Engine/Source/Runtime/Experimental/Iris/Core/Public/Iris/ReplicationState/Private/IrisFastArraySerializerInternal.h
--- a/Engine/Source/Runtime/Experimental/Iris/Core/Public/Iris/ReplicationState/Private/IrisFastArraySerializerInternal.h
+++ b/Engine/Source/Runtime/Experimental/Iris/Core/Public/Iris/ReplicationState/Private/IrisFastArraySerializerInternal.h
@@ -171,7 +171,7 @@ void TIrisFastArrayEditor<FastArrayType>::Add(const FastArrayItemType& ItemEntry
 template <typename FastArrayType>
 typename TIrisFastArrayEditor<FastArrayType>::FastArrayItemType& TIrisFastArrayEditor<FastArrayType>::Edit(int32 ItemIdx)
 {
-	FastArrayItemType& Item = FastArray.GetItemArray()[ItemIdx];
+	FastArrayItemType& Item = FastArray.GetItemArray()[ItemIdx]; //-V758
 	MarkItemDirty(Item);
 	return Item;
 }
