Chaos - fix almost-diagonal inertia rotation calculation
- reduce accumulation error in inertia calculation for triangulated shapes

diff --git a/Engine/Source/Runtime/Experimental/Chaos/Private/Chaos/MassProperties.cpp b/Engine/Source/Runtime/Experimental/Chaos/Private/Chaos/MassProperties.cpp
--- a/Engine/Source/Runtime/Experimental/Chaos/Private/Chaos/MassProperties.cpp
+++ b/Engine/Source/Runtime/Experimental/Chaos/Private/Chaos/MassProperties.cpp
@@ -19,7 +19,7 @@ namespace Chaos
 		TRotation<T,3> FinalRotation;
 
 		// Extract Eigenvalues
-		T OffDiagSize = FMath::Square(Inertia.M[1][0]) + FMath::Square(Inertia.M[2][0]) + FMath::Square(Inertia.M[2][1]);
+		T OffDiagSizeSq = FMath::Square(Inertia.M[1][0]) + FMath::Square(Inertia.M[2][0]) + FMath::Square(Inertia.M[2][1]);
 		FRealDouble Trace = (Inertia.M[0][0] + Inertia.M[1][1] + Inertia.M[2][2]) / 3;
 
 		if (Trace <= UE_SMALL_NUMBER)
@@ -28,13 +28,17 @@ namespace Chaos
 			return TRotation<T,3>::FromIdentity();
 		}
 
-		if ((OffDiagSize / Trace) < UE_SMALL_NUMBER)
+		// If the matrix is "almost diagonal" we are already in local space.
+		// This is especially important when we built our inertia from a surface integral over a 
+		// set of triangles because it will have accumulation errors that give non-zero off-diagonal elements
+		// and we can end up selecting very unexpected axes for our eigenvectors.
+		// This tolerance is pretty arbitrary - what value should we use?
+		if ((OffDiagSizeSq / FMath::Square(Trace)) < UE_SMALL_NUMBER)
 		{
-			// Almost diagonal matrix - we are already in local space.
 			return TRotation<T,3>::FromIdentity();
 		}
 
-		T Size = static_cast<T>(FMath::Sqrt((FMath::Square(Inertia.M[0][0] - Trace) + FMath::Square(Inertia.M[1][1] - Trace) + FMath::Square(Inertia.M[2][2] - Trace) + 2. * OffDiagSize) / 6.));
+		T Size = static_cast<T>(FMath::Sqrt((FMath::Square(Inertia.M[0][0] - Trace) + FMath::Square(Inertia.M[1][1] - Trace) + FMath::Square(Inertia.M[2][2] - Trace) + 2. * OffDiagSizeSq) / 6.));
 		PMatrix<T,3,3> NewMat = (Inertia - PMatrix<T,3,3>::Identity * static_cast<T>(Trace)) * static_cast<T>((1 / Size));
 		T HalfDeterminant = NewMat.Determinant() / 2;
 		T Angle = HalfDeterminant <= -1 ? UE_PI / 3 : (HalfDeterminant >= 1 ? 0 : FMath::Acos(HalfDeterminant) / 3);
@@ -203,28 +207,31 @@ namespace Chaos
 	{
 		check(Density > 0);
 
-		static const PMatrix<T, 3, 3> Standard(2, 1, 1, 2, 1, 2);
-		PMatrix<T, 3, 3> Covariance(0);
+		// Perform the calculation in doubles regardless of input/output type to reduce accumulation errors
+		static const FMatrix33 Standard(2, 1, 1, 2, 1, 2);
+		FMatrix33 Covariance(0);
 		for (const auto& Element : Surfaces)
 		{
-			PMatrix<T, 3, 3> DeltaMatrix(0);
+			FMatrix33 DeltaMatrix(0);
 			for (int32 i = 0; i < Element.Num(); ++i)
 			{
-				TVec DeltaVector = Vertices[Element[i]] - CenterOfMass;
+				FVec3 DeltaVector = FVec3(Vertices[Element[i]] - CenterOfMass);
 				DeltaMatrix.M[0][i] = DeltaVector[0];
 				DeltaMatrix.M[1][i] = DeltaVector[1];
 				DeltaMatrix.M[2][i] = DeltaVector[2];
 			}
-			T Det = DeltaMatrix.M[0][0] * (DeltaMatrix.M[1][1] * DeltaMatrix.M[2][2] - DeltaMatrix.M[1][2] * DeltaMatrix.M[2][1]) -
+			const FReal Det = DeltaMatrix.M[0][0] * (DeltaMatrix.M[1][1] * DeltaMatrix.M[2][2] - DeltaMatrix.M[1][2] * DeltaMatrix.M[2][1]) -
 				DeltaMatrix.M[0][1] * (DeltaMatrix.M[1][0] * DeltaMatrix.M[2][2] - DeltaMatrix.M[1][2] * DeltaMatrix.M[2][0]) +
 				DeltaMatrix.M[0][2] * (DeltaMatrix.M[1][0] * DeltaMatrix.M[2][1] - DeltaMatrix.M[1][1] * DeltaMatrix.M[2][0]);
-			const PMatrix<T, 3, 3> ScaledStandard = Standard * Det;
+			const FMatrix33 ScaledStandard = Standard * Det;
 			Covariance += DeltaMatrix * ScaledStandard * DeltaMatrix.GetTransposed();
 		}
-		T Trace = Covariance.M[0][0] + Covariance.M[1][1] + Covariance.M[2][2];
-		PMatrix<T, 3, 3> TraceMat(Trace, Trace, Trace);
-		OutInertiaTensor = (TraceMat - Covariance) * (1 / (T)120) * Density;
-		OutRotationOfMass = TransformToLocalSpace(OutInertiaTensor);
+		const FReal Trace = Covariance.M[0][0] + Covariance.M[1][1] + Covariance.M[2][2];
+		const FMatrix33 TraceMat(Trace, Trace, Trace);
+		FMatrix33 Inertia = (TraceMat - Covariance) * (Density / FReal(120));
+		const FRotation3 Rotation = TransformToLocalSpace(Inertia);
+		OutInertiaTensor = PMatrix<T, 3, 3>(Inertia);
+		OutRotationOfMass = TRotation<T, 3>(Rotation);
 	}
 
 	template<typename T, typename TSurfaces>
