Chaos
- spheres and capsules now use PBD friction (velocity based dynamic friction not required now)
- set friction stiffness to 1.0
- collision constraint regularization

diff --git a/Engine/Source/Runtime/Experimental/Chaos/Private/Chaos/Collision/PBDCollisionSolver.cpp b/Engine/Source/Runtime/Experimental/Chaos/Private/Chaos/Collision/PBDCollisionSolver.cpp
--- a/Engine/Source/Runtime/Experimental/Chaos/Private/Chaos/Collision/PBDCollisionSolver.cpp
+++ b/Engine/Source/Runtime/Experimental/Chaos/Private/Chaos/Collision/PBDCollisionSolver.cpp
@@ -25,7 +25,7 @@ namespace Chaos
 
 		bool bChaos_PBDCollisionSolver_Position_SolveEnabled = true;
 		float Chaos_PBDCollisionSolver_Position_MinInvMassScale = 0.77f;
-		float Chaos_PBDCollisionSolver_Position_StaticFrictionStiffness = 0.5f;
+		float Chaos_PBDCollisionSolver_Position_StaticFrictionStiffness = 1.0f;
 		float Chaos_PBDCollisionSolver_Position_PositionSolverTolerance = 0.001f;		// cms
 		float Chaos_PBDCollisionSolver_Position_RotationSolverTolerance = 0.001f;		// rads
 
diff --git a/Engine/Source/Runtime/Experimental/Chaos/Private/Chaos/Collision/PBDCollisionSolver.h b/Engine/Source/Runtime/Experimental/Chaos/Private/Chaos/Collision/PBDCollisionSolver.h
--- a/Engine/Source/Runtime/Experimental/Chaos/Private/Chaos/Collision/PBDCollisionSolver.h
+++ b/Engine/Source/Runtime/Experimental/Chaos/Private/Chaos/Collision/PBDCollisionSolver.h
@@ -45,20 +45,8 @@ namespace Chaos
 			const FSolverVec3& InWorldContactTangentV,
 			const FSolverReal InWorldContactDeltaNormal,
 			const FSolverReal InWorldContactDeltaTangentU,
-			const FSolverReal InWorldContactDeltaTangentV);
-
-		void InitContact(
-			const FSolverReal Dt,
-			const FConstraintSolverBody& Body0,
-			const FConstraintSolverBody& Body1,
-			const SolverVectorRegister& InRelativeContactPosition0,
-			const SolverVectorRegister& InRelativeContactPosition1,
-			const SolverVectorRegister& InWorldContactNormal,
-			const SolverVectorRegister& InWorldContactTangentU,
-			const SolverVectorRegister& InWorldContactTangentV,
-			const FSolverReal InWorldContactDeltaNormal,
-			const FSolverReal InWorldContactDeltaTangentU,
-			const FSolverReal InWorldContactDeltaTangentV);
+			const FSolverReal InWorldContactDeltaTangentV,
+			const FSolverReal InRegularization);
 
 		/**
 		 * @brief Initialize the material related properties of the contact
@@ -72,7 +60,7 @@ namespace Chaos
 		/**
 		 * @brief Update the cached mass properties based on the current body transforms
 		*/
-		void UpdateMass(const FConstraintSolverBody& Body0, const FConstraintSolverBody& Body1);
+		void UpdateMass(const FConstraintSolverBody& Body0, const FConstraintSolverBody& Body1, const FSolverReal Regularization);
 
 		/**
 		 * @brief Update the contact mass for the normal correction
@@ -80,15 +68,6 @@ namespace Chaos
 		*/
 		void UpdateMassNormal(const FConstraintSolverBody& Body0, const FConstraintSolverBody& Body1);
 
-		void UpdateMass(
-			const FConstraintSolverBody& Body0,
-			const FConstraintSolverBody& Body1,
-			const SolverVectorRegister& InRelativeContactPosition0,
-			const SolverVectorRegister& InRelativeContactPosition1,
-			const SolverVectorRegister& InWorldContactNormal,
-			const SolverVectorRegister& InWorldContactTangentU,
-			const SolverVectorRegister& InWorldContactTangentV);
-
 		/**
 		 * @brief Calculate the position error at the current transforms
 		 * @param MaxPushOut a limit on the position error for this iteration to prevent initial-penetration explosion (a common PBD problem)
@@ -198,6 +177,11 @@ namespace Chaos
 			State.Stiffness = InStiffness;
 		}
 
+		void SetRegularization(const FSolverReal InRegularization)
+		{
+			State.Regularization = InRegularization;
+		}
+
 		void SetSolverBodies(FSolverBody& SolverBody0, FSolverBody& SolverBody1)
 		{
 			State.SolverBodies[0] = FConstraintSolverBody(SolverBody0);
@@ -241,20 +225,6 @@ namespace Chaos
 			const FSolverReal InWorldContactDeltaTangentU,
 			const FSolverReal InWorldContactDeltaTangentV);
 
-		void SetManifoldPoint(
-			const int32 ManifoldPoiontIndex,
-			const FSolverReal Dt,
-			const FSolverReal InRestitution,
-			const FSolverReal InRestitutionVelocityThreshold,
-			const SolverVectorRegister& InRelativeContactPosition0,
-			const SolverVectorRegister& InRelativeContactPosition1,
-			const SolverVectorRegister& InWorldContactNormal,
-			const SolverVectorRegister& InWorldContactTangentU,
-			const SolverVectorRegister& InWorldContactTangentV,
-			const FSolverReal InWorldContactDeltaNormal,
-			const FSolverReal InWorldContactDeltaTangentU,
-			const FSolverReal InWorldContactDeltaTangentV);
-
 		/**
 		 * @brief Get the first (decaorated) solver body
 		 * The decorator add a possible mass scale
@@ -319,6 +289,7 @@ namespace Chaos
 				, DynamicFriction(0)
 				, VelocityFriction(0)
 				, Stiffness(1)
+				, Regularization(0)
 				, NumManifoldPoints(0)
 				, SolverBodies()
 				, ManifoldPoints()
@@ -332,6 +303,7 @@ namespace Chaos
 				DynamicFriction = 0;
 				VelocityFriction = 0;
 				Stiffness = 1;
+				Regularization = 0;
 				NumManifoldPoints = 0;
 			}
 
@@ -347,6 +319,9 @@ namespace Chaos
 			// Solver stiffness (scales all pushout and impulses)
 			FSolverReal Stiffness;
 
+			// Solver Regularization
+			FSolverReal Regularization;
+
 			// Bodies and contacts
 			int32 NumManifoldPoints;
 			FConstraintSolverBody SolverBodies[MaxConstrainedBodies];
@@ -365,6 +340,7 @@ namespace Chaos
 
 	FORCEINLINE_DEBUGGABLE void CalculatePositionCorrectionNormal(
 		const FSolverReal Stiffness,
+		const FSolverReal Regularization,
 		const FSolverReal ContactDeltaNormal,
 		const FSolverReal ContactMassNormal,
 		const FSolverReal NetPushOutNormal,
@@ -387,8 +363,10 @@ namespace Chaos
 
 	FORCEINLINE_DEBUGGABLE void CalculatePositionCorrectionTangent(
 		const FSolverReal Stiffness,
+		const FSolverReal Regularization,
 		const FSolverReal ContactDeltaTangent,
 		const FSolverReal ContactMassTangent,
+		const FSolverReal NetPushOutTangent,
 		FSolverReal& OutPushOutTangent)
 	{
 		const FSolverReal FrictionStiffness = CVars::Chaos_PBDCollisionSolver_Position_StaticFrictionStiffness;
@@ -445,6 +423,7 @@ namespace Chaos
 
 	FORCEINLINE_DEBUGGABLE void ApplyPositionCorrectionTangential(
 		const FSolverReal Stiffness,
+		const FSolverReal Regularization,
 		const FSolverReal StaticFriction,
 		const FSolverReal DynamicFriction,
 		const FSolverReal MaxFrictionPushOut,
@@ -459,14 +438,18 @@ namespace Chaos
 
 		CalculatePositionCorrectionTangent(
 			Stiffness,
+			Regularization,
 			ContactDeltaTangentU,
 			ManifoldPoint.ContactMassTangentU,
+			ManifoldPoint.NetPushOutTangentU,
 			PushOutTangentU);					// Out
 
 		CalculatePositionCorrectionTangent(
 			Stiffness,
+			Regularization,
 			ContactDeltaTangentV,
 			ManifoldPoint.ContactMassTangentV,
+			ManifoldPoint.NetPushOutTangentV,
 			PushOutTangentV);					// Out
 
 		ManifoldPoint.NetPushOutTangentU += PushOutTangentU;
@@ -502,6 +485,7 @@ namespace Chaos
 
 	FORCEINLINE_DEBUGGABLE void ApplyPositionCorrectionNormal(
 		const FSolverReal Stiffness,
+		const FSolverReal Regularization,
 		const FSolverReal ContactDeltaNormal,
 		FPBDCollisionSolverManifoldPoint& ManifoldPoint,
 		FConstraintSolverBody& Body0,
@@ -511,6 +495,7 @@ namespace Chaos
 
 		CalculatePositionCorrectionNormal(
 			Stiffness,
+			Regularization,
 			ContactDeltaNormal,
 			ManifoldPoint.ContactMassNormal,
 			ManifoldPoint.NetPushOutNormal,
@@ -648,7 +633,8 @@ namespace Chaos
 		const FSolverVec3& InWorldContactTangentV,
 		const FSolverReal InWorldContactDeltaNormal,
 		const FSolverReal InWorldContactDeltaTangentU,
-		const FSolverReal InWorldContactDeltaTangentV)
+		const FSolverReal InWorldContactDeltaTangentV,
+		const FSolverReal InRegularization)
 	{
 		RelativeContactPosition0 = InRelativeContactPosition0;
 		RelativeContactPosition1 = InRelativeContactPosition1;
@@ -669,50 +655,7 @@ namespace Chaos
 		StaticFrictionRatio = FSolverReal(0);
 		WorldContactVelocityTargetNormal = FSolverReal(0);
 
-		UpdateMass(Body0, Body1);
-	}
-
-
-	FORCEINLINE_DEBUGGABLE void FPBDCollisionSolverManifoldPoint::InitContact(
-		const FSolverReal Dt,
-		const FConstraintSolverBody& Body0,
-		const FConstraintSolverBody& Body1,
-		const SolverVectorRegister& InRelativeContactPosition0,
-		const SolverVectorRegister& InRelativeContactPosition1,
-		const SolverVectorRegister& InWorldContactNormal,
-		const SolverVectorRegister& InWorldContactTangentU,
-		const SolverVectorRegister& InWorldContactTangentV,
-		const FSolverReal InWorldContactDeltaNormal,
-		const FSolverReal InWorldContactDeltaTangentU,
-		const FSolverReal InWorldContactDeltaTangentV)
-	{
-		VectorStoreFloat3(InRelativeContactPosition0, &RelativeContactPosition0);
-		VectorStoreFloat3(InRelativeContactPosition1, &RelativeContactPosition1);
-		VectorStoreFloat3(InWorldContactNormal, &WorldContactNormal);
-		VectorStoreFloat3(InWorldContactTangentU, &WorldContactTangentU);
-		VectorStoreFloat3(InWorldContactTangentV, &WorldContactTangentV);
-		WorldContactDeltaNormal = InWorldContactDeltaNormal;
-		WorldContactDeltaTangentU = InWorldContactDeltaTangentU;
-		WorldContactDeltaTangentV = InWorldContactDeltaTangentV;
-
-		NetPushOutNormal = FSolverReal(0);
-		NetPushOutTangentU = FSolverReal(0);
-		NetPushOutTangentV = FSolverReal(0);
-		NetImpulseNormal = FSolverReal(0);
-		NetImpulseTangentU = FSolverReal(0);
-		NetImpulseTangentV = FSolverReal(0);
-
-		StaticFrictionRatio = FSolverReal(0);
-		WorldContactVelocityTargetNormal = FSolverReal(0);
-
-		UpdateMass(
-			Body0, 
-			Body1,
-			InRelativeContactPosition0,
-			InRelativeContactPosition1,
-			InWorldContactNormal,
-			InWorldContactTangentU,
-			InWorldContactTangentV);
+		UpdateMass(Body0, Body1, InRegularization);
 	}
 
 	FORCEINLINE_DEBUGGABLE void FPBDCollisionSolverManifoldPoint::InitMaterial(
@@ -732,11 +675,11 @@ namespace Chaos
 		}
 	}
 
-	FORCEINLINE_DEBUGGABLE void FPBDCollisionSolverManifoldPoint::UpdateMass(const FConstraintSolverBody& Body0, const FConstraintSolverBody& Body1)
+	FORCEINLINE_DEBUGGABLE void FPBDCollisionSolverManifoldPoint::UpdateMass(const FConstraintSolverBody& Body0, const FConstraintSolverBody& Body1, const FSolverReal Regularization)
 	{
-		FSolverReal ContactMassInvNormal = FSolverReal(0);
-		FSolverReal ContactMassInvTangentU = FSolverReal(0);
-		FSolverReal ContactMassInvTangentV = FSolverReal(0);
+		FSolverReal ContactMassInvNormal = Regularization;
+		FSolverReal ContactMassInvTangentU = Regularization;
+		FSolverReal ContactMassInvTangentV = Regularization;
 
 		// These are not used if not initialized below so no need to clear
 		//WorldContactNormalAngular0 = FSolverVec3(0);
@@ -784,67 +727,6 @@ namespace Chaos
 		ContactMassTangentV = (ContactMassInvTangentV > FSolverReal(UE_SMALL_NUMBER)) ? FSolverReal(1) / ContactMassInvTangentV : FSolverReal(0);
 	}
 
-	FORCEINLINE_DEBUGGABLE void FPBDCollisionSolverManifoldPoint::UpdateMass(
-		const FConstraintSolverBody& Body0, 
-		const FConstraintSolverBody& Body1,
-		const SolverVectorRegister& InRelativeContactPosition0,
-		const SolverVectorRegister& InRelativeContactPosition1,
-		const SolverVectorRegister& InWorldContactNormal,
-		const SolverVectorRegister& InWorldContactTangentU,
-		const SolverVectorRegister& InWorldContactTangentV)
-	{
-		FSolverReal ContactMassInvNormal = FSolverReal(0);
-		FSolverReal ContactMassInvTangentU = FSolverReal(0);
-		FSolverReal ContactMassInvTangentV = FSolverReal(0);
-
-		if (Body0.IsDynamic())
-		{
-			const SolverVectorRegister R0xN = VectorCross(InRelativeContactPosition0, InWorldContactNormal);
-			const SolverVectorRegister R0xU = VectorCross(InRelativeContactPosition0, InWorldContactTangentU);
-			const SolverVectorRegister R0xV = VectorCross(InRelativeContactPosition0, InWorldContactTangentV);
-
-			const SolverVectorRegister InvMScale0 = VectorSetFloat1(Body0.InvMScale());
-			const FSolverMatrix33& InvI0 = Body0.SolverBody().InvI();
-
-			const SolverVectorRegister NormalAngular0 = VectorMultiply(VectorTransformVector(R0xN, &InvI0), InvMScale0);
-			const SolverVectorRegister TangentUAngular0 = VectorMultiply(VectorTransformVector(R0xU, &InvI0), InvMScale0);
-			const SolverVectorRegister TangentVAngular0 = VectorMultiply(VectorTransformVector(R0xV, &InvI0), InvMScale0);
-
-			ContactMassInvNormal += VectorGetComponent(VectorDot3(R0xN, NormalAngular0), 0) + Body0.InvM();
-			ContactMassInvTangentU += VectorGetComponent(VectorDot3(R0xU, TangentUAngular0), 0) + Body0.InvM();
-			ContactMassInvTangentV += VectorGetComponent(VectorDot3(R0xV, TangentVAngular0), 0) + Body0.InvM();
-
-			VectorStoreFloat3(NormalAngular0, &WorldContactNormalAngular0);
-			VectorStoreFloat3(TangentUAngular0, &WorldContactTangentUAngular0);
-			VectorStoreFloat3(TangentVAngular0, &WorldContactTangentVAngular0);
-		}
-		if (Body1.IsDynamic())
-		{
-			const SolverVectorRegister R1xN = VectorCross(InRelativeContactPosition1, InWorldContactNormal);
-			const SolverVectorRegister R1xU = VectorCross(InRelativeContactPosition1, InWorldContactTangentU);
-			const SolverVectorRegister R1xV = VectorCross(InRelativeContactPosition1, InWorldContactTangentV);
-
-			const SolverVectorRegister InvMScale1 = VectorSetFloat1(Body1.InvMScale());
-			const FSolverMatrix33& InvI1 = Body1.SolverBody().InvI();
-
-			const SolverVectorRegister NormalAngular1 = VectorMultiply(VectorTransformVector(R1xN, &InvI1), InvMScale1);
-			const SolverVectorRegister TangentUAngular1 = VectorMultiply(VectorTransformVector(R1xU, &InvI1), InvMScale1);
-			const SolverVectorRegister TangentVAngular1 = VectorMultiply(VectorTransformVector(R1xV, &InvI1), InvMScale1);
-
-			ContactMassInvNormal += VectorGetComponent(VectorDot3(R1xN, NormalAngular1), 0) + Body1.InvM();
-			ContactMassInvTangentU += VectorGetComponent(VectorDot3(R1xU, TangentUAngular1), 0) + Body1.InvM();
-			ContactMassInvTangentV += VectorGetComponent(VectorDot3(R1xV, TangentVAngular1), 0) + Body1.InvM();
-
-			VectorStoreFloat3(NormalAngular1, &WorldContactNormalAngular1);
-			VectorStoreFloat3(TangentUAngular1, &WorldContactTangentUAngular1);
-			VectorStoreFloat3(TangentVAngular1, &WorldContactTangentVAngular1);
-		}
-
-		ContactMassNormal = (ContactMassInvNormal > FSolverReal(UE_SMALL_NUMBER)) ? FSolverReal(1) / ContactMassInvNormal : FSolverReal(0);
-		ContactMassTangentU = (ContactMassInvTangentU > FSolverReal(UE_SMALL_NUMBER)) ? FSolverReal(1) / ContactMassInvTangentU : FSolverReal(0);
-		ContactMassTangentV = (ContactMassInvTangentV > FSolverReal(UE_SMALL_NUMBER)) ? FSolverReal(1) / ContactMassInvTangentV : FSolverReal(0);
-	}
-
 	FORCEINLINE_DEBUGGABLE void FPBDCollisionSolverManifoldPoint::UpdateMassNormal(const FConstraintSolverBody& Body0, const FConstraintSolverBody& Body1)
 	{
 		FSolverReal ContactMassInvNormal = FSolverReal(0);
@@ -989,41 +871,8 @@ namespace Chaos
 			InWorldContactTangentV,
 			InWorldContactDeltaNormal,
 			InWorldContactDeltaTangentU,
-			InWorldContactDeltaTangentV);
-
-		State.ManifoldPoints[ManifoldPoiontIndex].InitMaterial(
-			State.SolverBodies[0],
-			State.SolverBodies[1],
-			InRestitution,
-			InRestitutionVelocityThreshold);
-	}
-
-	FORCEINLINE_DEBUGGABLE void FPBDCollisionSolver::SetManifoldPoint(
-		const int32 ManifoldPoiontIndex,
-		const FSolverReal Dt,
-		const FSolverReal InRestitution,
-		const FSolverReal InRestitutionVelocityThreshold,
-		const SolverVectorRegister& InRelativeContactPosition0,
-		const SolverVectorRegister& InRelativeContactPosition1,
-		const SolverVectorRegister& InWorldContactNormal,
-		const SolverVectorRegister& InWorldContactTangentU,
-		const SolverVectorRegister& InWorldContactTangentV,
-		const FSolverReal InWorldContactDeltaNormal,
-		const FSolverReal InWorldContactDeltaTangentU,
-		const FSolverReal InWorldContactDeltaTangentV)
-	{
-		State.ManifoldPoints[ManifoldPoiontIndex].InitContact(
-			FSolverReal(Dt),
-			State.SolverBodies[0],
-			State.SolverBodies[1],
-			InRelativeContactPosition0,
-			InRelativeContactPosition1,
-			InWorldContactNormal,
-			InWorldContactTangentU,
-			InWorldContactTangentV,
-			InWorldContactDeltaNormal,
-			InWorldContactDeltaTangentU,
-			InWorldContactDeltaTangentV);
+			InWorldContactDeltaTangentV, 
+			State.Regularization);
 
 		State.ManifoldPoints[ManifoldPoiontIndex].InitMaterial(
 			State.SolverBodies[0],
@@ -1057,6 +906,7 @@ namespace Chaos
 			{
 				ApplyPositionCorrectionNormal(
 					State.Stiffness,
+					State.Regularization,
 					ContactDeltaNormal,
 					SolverManifoldPoint,
 					Body0,
@@ -1094,6 +944,7 @@ namespace Chaos
 
 					ApplyPositionCorrectionTangential(
 						State.Stiffness,
+						State.Regularization,
 						State.StaticFriction,
 						State.DynamicFriction,
 						FrictionMaxPushOut,
@@ -1128,6 +979,7 @@ namespace Chaos
 			{
 				ApplyPositionCorrectionNormal(
 					State.Stiffness,
+					State.Regularization,
 					ContactDeltaNormal,
 					SolverManifoldPoint,
 					Body0,
diff --git a/Engine/Source/Runtime/Experimental/Chaos/Private/Chaos/Collision/SolverCollisionContainer.cpp b/Engine/Source/Runtime/Experimental/Chaos/Private/Chaos/Collision/SolverCollisionContainer.cpp
--- a/Engine/Source/Runtime/Experimental/Chaos/Private/Chaos/Collision/SolverCollisionContainer.cpp
+++ b/Engine/Source/Runtime/Experimental/Chaos/Private/Chaos/Collision/SolverCollisionContainer.cpp
@@ -33,6 +33,10 @@ namespace Chaos
 		FRealSingle Chaos_PBDCollisionSolver_AutoStiffness_MassRatio2 = 0;
 		FAutoConsoleVariableRef CVarChaosPBDCollisionSolverAutoStiffnessMassRatio1(TEXT("p.Chaos.PBDCollisionSolver.AutoStiffness.MassRatio1"), Chaos_PBDCollisionSolver_AutoStiffness_MassRatio1, TEXT(""));
 		FAutoConsoleVariableRef CVarChaosPBDCollisionSolverAutoStiffnessMassRatio2(TEXT("p.Chaos.PBDCollisionSolver.AutoStiffness.MassRatio2"), Chaos_PBDCollisionSolver_AutoStiffness_MassRatio2, TEXT(""));
+
+		// Collision solver regularization (experimental, disabled by default)
+		FRealSingle Chaos_PBDCollisionSolver_Regularization = 0;
+		FAutoConsoleVariableRef CVarChaosPBDCollisionSolverRegularization(TEXT("p.Chaos.PBDCollisionSolver.Regularization"), Chaos_PBDCollisionSolver_Regularization, TEXT(""));
 	}
 	using namespace CVars;
 
@@ -186,9 +190,7 @@ namespace Chaos
 			const FConstraintSolverBody& Body1 = Solver.SolverBody1();
 
 			// Friction values. Static and Dynamic friction are applied in the position solve for most shapes.
-			// For quadratic shapes, we run dynamic friction in the velocity solve for better rolling behaviour.
 			// We can also run in a mode without static friction at all. This is faster but stacking is not possible.
-			// @todo(chaos): fix static/dynamic friction for quadratic shapes
 			const FSolverReal StaticFriction = FSolverReal(Constraint->GetStaticFriction());
 			const FSolverReal DynamicFriction = FSolverReal(Constraint->GetDynamicFriction());
 			FSolverReal PositionStaticFriction = FSolverReal(0);
@@ -197,14 +199,7 @@ namespace Chaos
 			if (SolverSettings.NumPositionFrictionIterations > 0)
 			{
 				PositionStaticFriction = StaticFriction;
-				if (!Constraint->HasQuadraticShape())
-				{
-					PositionDynamicFriction = DynamicFriction;
-				}
-				else
-				{
-					VelocityDynamicFriction = DynamicFriction;
-				}
+				PositionDynamicFriction = DynamicFriction;
 			}
 			else
 			{
@@ -216,6 +211,9 @@ namespace Chaos
 			const FReal SolverStiffness = CalculateSolverStiffness(Body0.SolverBody(), Body1.SolverBody(), Chaos_PBDCollisionSolver_AutoStiffness_MassRatio1, Chaos_PBDCollisionSolver_AutoStiffness_MassRatio2);
 			Solver.SetStiffness(FSolverReal(SolverStiffness));
 
+			const FReal Regularization = Chaos_PBDCollisionSolver_Regularization / (Dt * Dt);
+			Solver.SetRegularization(FSolverReal(Regularization));
+
 			Solver.SolverBody0().SetInvMScale(Constraint->GetInvMassScale0());
 			Solver.SolverBody0().SetInvIScale(Constraint->GetInvInertiaScale0());
 			Solver.SolverBody1().SetInvMScale(Constraint->GetInvMassScale1());
