Chaos : add method to dynamically set Geometry collection fragment to be anchored ( by index or by box )

diff --git a/Engine/Source/Runtime/Experimental/Chaos/Private/PhysicsProxy/GeometryCollectionPhysicsProxy.cpp b/Engine/Source/Runtime/Experimental/Chaos/Private/PhysicsProxy/GeometryCollectionPhysicsProxy.cpp
--- a/Engine/Source/Runtime/Experimental/Chaos/Private/PhysicsProxy/GeometryCollectionPhysicsProxy.cpp
+++ b/Engine/Source/Runtime/Experimental/Chaos/Private/PhysicsProxy/GeometryCollectionPhysicsProxy.cpp
@@ -1727,6 +1727,48 @@ void FGeometryCollectionPhysicsProxy::BreakActiveClusters_External()
 	}
 }
 
+void FGeometryCollectionPhysicsProxy::SetAnchored_External(int32 Index)
+{
+	check(IsInGameThread());
+	if (Chaos::FPhysicsSolver* RBDSolver = GetSolver<Chaos::FPhysicsSolver>())
+	{
+		RBDSolver->EnqueueCommandImmediate([this, Index, RBDSolver]()
+			{
+				Chaos::FPBDRigidsEvolution* Evolution = RBDSolver->GetEvolution();
+				if (SolverParticleHandles.IsValidIndex(Index))
+				{
+					if (FClusterHandle* ParticleHandle = SolverParticleHandles[Index])
+					{
+						ParticleHandle->SetIsAnchored(true);
+					}
+				}
+			});
+	}
+}
+
+void FGeometryCollectionPhysicsProxy::SetAnchored_External(const FBox& WorldSpaceBox)
+{
+	check(IsInGameThread());
+	if (Chaos::FPhysicsSolver* RBDSolver = GetSolver<Chaos::FPhysicsSolver>())
+	{
+		Chaos::FAABB3 BoundsToCheck(WorldSpaceBox.Min, WorldSpaceBox.Max);
+		RBDSolver->EnqueueCommandImmediate([this, BoundsToCheck, RBDSolver]()
+			{
+				Chaos::FPBDRigidsEvolution* Evolution = RBDSolver->GetEvolution();
+				for (FClusterHandle* ParticleHandle : GetSolverParticleHandles())
+				{
+					if (ParticleHandle)
+					{
+						if (BoundsToCheck.Contains(ParticleHandle->X()))
+						{
+							ParticleHandle->SetIsAnchored(true);
+						}
+					}
+				}
+			});
+	}
+}
+
 void FGeometryCollectionPhysicsProxy::RemoveAllAnchors_External()
 {
 	check(IsInGameThread());
diff --git a/Engine/Source/Runtime/Experimental/Chaos/Public/PhysicsProxy/GeometryCollectionPhysicsProxy.h b/Engine/Source/Runtime/Experimental/Chaos/Public/PhysicsProxy/GeometryCollectionPhysicsProxy.h
--- a/Engine/Source/Runtime/Experimental/Chaos/Public/PhysicsProxy/GeometryCollectionPhysicsProxy.h
+++ b/Engine/Source/Runtime/Experimental/Chaos/Public/PhysicsProxy/GeometryCollectionPhysicsProxy.h
@@ -376,6 +376,8 @@ public:
 	void ApplyImpulseAt_External(FVector Force, FVector WorldLocation);
 	void BreakClusters_External(TArray<FGeometryCollectionItemIndex>&& ItemIndices);
 	void BreakActiveClusters_External();
+	void SetAnchored_External(int32 Index);
+	void SetAnchored_External(const FBox& WorldSpaceBox);
 	void RemoveAllAnchors_External();
 	void ApplyExternalStrain_External(FGeometryCollectionItemIndex ItemIndex, const FVector& WorldLocation, float Radius, int32 PropagationDepth, float PropagationFactor, float StrainValue);
 	void ApplyInternalStrain_External(FGeometryCollectionItemIndex ItemIndex, const FVector& WorldLocation, float Radius, int32 PropagationDepth, float PropagationFactor, float StrainValue);
diff --git a/Engine/Source/Runtime/Experimental/GeometryCollectionEngine/Private/GeometryCollection/GeometryCollectionComponent.cpp b/Engine/Source/Runtime/Experimental/GeometryCollectionEngine/Private/GeometryCollection/GeometryCollectionComponent.cpp
--- a/Engine/Source/Runtime/Experimental/GeometryCollectionEngine/Private/GeometryCollection/GeometryCollectionComponent.cpp
+++ b/Engine/Source/Runtime/Experimental/GeometryCollectionEngine/Private/GeometryCollection/GeometryCollectionComponent.cpp
@@ -4232,6 +4232,22 @@ void UGeometryCollectionComponent::CrumbleActiveClusters()
 	}
 }
 
+void UGeometryCollectionComponent::SetAnchoredByIndex(int32 Index)
+{
+	if (PhysicsProxy)
+	{
+		PhysicsProxy->SetAnchored_External(Index);
+	}
+}
+
+void UGeometryCollectionComponent::SetAnchoredByBox(FBox WorldSpaceBox)
+{
+	if (PhysicsProxy)
+	{
+		PhysicsProxy->SetAnchored_External(WorldSpaceBox);
+	}
+}
+
 void UGeometryCollectionComponent::RemoveAllAnchors()
 {
 	if (PhysicsProxy)
diff --git a/Engine/Source/Runtime/Experimental/GeometryCollectionEngine/Public/GeometryCollection/GeometryCollectionComponent.h b/Engine/Source/Runtime/Experimental/GeometryCollectionEngine/Public/GeometryCollection/GeometryCollectionComponent.h
--- a/Engine/Source/Runtime/Experimental/GeometryCollectionEngine/Public/GeometryCollection/GeometryCollectionComponent.h
+++ b/Engine/Source/Runtime/Experimental/GeometryCollectionEngine/Public/GeometryCollection/GeometryCollectionComponent.h
@@ -552,6 +552,18 @@ public:
 	UFUNCTION(BlueprintCallable, Category = "ChaosPhysics")
 	void CrumbleActiveClusters();
 
+	/**
+	* Set a fragment or cluster to be anchored
+	*/
+	UFUNCTION(BlueprintCallable, Category = "ChaosPhysics")
+	void SetAnchoredByIndex(int32 Index);
+
+	/**
+	* Set all fragments within a world bounding box to be anchored
+	*/
+	UFUNCTION(BlueprintCallable, Category = "ChaosPhysics")
+	void SetAnchoredByBox(FBox WorldSpaceBox);
+
 	/**
 	* this will remove anchors on all the pieces ( inlcuding the static and kinematic initial states ones ) of the geometry colection
 	*/
