Chaos: Fixed global namespace exposure in source files leading to ambiguous symbols in unity builds.

diff --git a/Engine/Plugins/Experimental/ChaosVehiclesPlugin/Source/ChaosVehicles/Private/ChaosWheeledVehicleMovementComponent.cpp b/Engine/Plugins/Experimental/ChaosVehiclesPlugin/Source/ChaosVehicles/Private/ChaosWheeledVehicleMovementComponent.cpp
--- a/Engine/Plugins/Experimental/ChaosVehiclesPlugin/Source/ChaosVehicles/Private/ChaosWheeledVehicleMovementComponent.cpp
+++ b/Engine/Plugins/Experimental/ChaosVehiclesPlugin/Source/ChaosVehicles/Private/ChaosWheeledVehicleMovementComponent.cpp
@@ -30,8 +30,6 @@
 #endif
 #include "PhysicsProxy/SingleParticlePhysicsProxy.h"
 
-using namespace Chaos;
-
 #if VEHICLE_DEBUGGING_ENABLED
 PRAGMA_DISABLE_OPTIMIZATION
 #endif
@@ -297,9 +295,9 @@ void UChaosWheeledVehicleSimulation::UpdateSimulation(float DeltaTime, const FCh
 
 }
 
-bool UChaosWheeledVehicleSimulation::ContainsTraces(const FBox& Box, const TArray<FSuspensionTrace>& SuspensionTrace)
+bool UChaosWheeledVehicleSimulation::ContainsTraces(const FBox& Box, const TArray<Chaos::FSuspensionTrace>& SuspensionTrace)
 {
-	const FAABB3 Aabb(Box.Min, Box.Max);
+	const Chaos::FAABB3 Aabb(Box.Min, Box.Max);
 
 	for (int WheelIdx = 0; WheelIdx < SuspensionTrace.Num(); WheelIdx++)
 	{
@@ -312,7 +310,7 @@ bool UChaosWheeledVehicleSimulation::ContainsTraces(const FBox& Box, const TArra
 	return true;
 }
 
-void UChaosWheeledVehicleSimulation::PerformSuspensionTraces(const TArray<FSuspensionTrace>& SuspensionTrace, FCollisionQueryParams& TraceParams, FCollisionResponseContainer& CollisionResponse, TArray<FWheelTraceParams>& WheelTraceParams)
+void UChaosWheeledVehicleSimulation::PerformSuspensionTraces(const TArray<Chaos::FSuspensionTrace>& SuspensionTrace, FCollisionQueryParams& TraceParams, FCollisionResponseContainer& CollisionResponse, TArray<FWheelTraceParams>& WheelTraceParams)
 {
 	SCOPE_CYCLE_COUNTER(STAT_ChaosVehicle_SuspensionRaycasts);
 
@@ -491,6 +489,8 @@ void UChaosWheeledVehicleSimulation::PerformSuspensionTraces(const TArray<FSuspe
 
 void UChaosWheeledVehicleSimulation::ApplyWheelFrictionForces(float DeltaTime)
 {
+	using namespace Chaos;
+
 	for (int WheelIdx = 0; WheelIdx < PVehicle->Wheels.Num(); WheelIdx++)
 	{
 		auto& PWheel = PVehicle->Wheels[WheelIdx]; // Physics Wheel
@@ -571,6 +571,8 @@ void UChaosWheeledVehicleSimulation::ApplyWheelFrictionForces(float DeltaTime)
 
 void UChaosWheeledVehicleSimulation::ApplySuspensionForces(float DeltaTime)
 {
+	using namespace Chaos;
+
 	TArray<float> SusForces;
 	SusForces.Init(0.f, PVehicle->Suspension.Num());
 
@@ -686,6 +688,8 @@ void UChaosWheeledVehicleSimulation::ApplySuspensionForces(float DeltaTime)
 
 void UChaosWheeledVehicleSimulation::ProcessSteering(const FControlInputs& ControlInputs)
 {
+	using namespace Chaos;
+
 	auto& PSteering = PVehicle->GetSteering();
 
 	for (int WheelIdx = 0; WheelIdx < PVehicle->Wheels.Num(); WheelIdx++)
@@ -726,6 +730,8 @@ void UChaosWheeledVehicleSimulation::ProcessSteering(const FControlInputs& Contr
 
 void UChaosWheeledVehicleSimulation::ApplyInput(const FControlInputs& ControlInputs, float DeltaTime)
 {
+	using namespace Chaos;
+
 	UChaosVehicleSimulation::ApplyInput(ControlInputs, DeltaTime);
 
 	FControlInputs ModifiedInputs = ControlInputs;
@@ -812,6 +818,8 @@ bool UChaosWheeledVehicleSimulation::IsWheelSpinning() const
 
 void UChaosWheeledVehicleSimulation::ProcessMechanicalSimulation(float DeltaTime)
 {
+	using namespace Chaos;
+
 	if (PVehicle->HasEngine())
 	{
 		auto& PEngine = PVehicle->GetEngine();
@@ -862,6 +870,8 @@ void UChaosWheeledVehicleSimulation::ProcessMechanicalSimulation(float DeltaTime
 
 void UChaosWheeledVehicleSimulation::DrawDebug3D()
 {
+	using namespace Chaos;
+
 #if !(UE_BUILD_SHIPPING || UE_BUILD_TEST)
 
 	UChaosVehicleSimulation::DrawDebug3D();
@@ -1353,6 +1363,8 @@ void UChaosWheeledVehicleMovementComponent::DestroyWheels()
 
 void UChaosWheeledVehicleMovementComponent::SetupVehicle(TUniquePtr<Chaos::FSimpleWheeledVehicle>& PVehicle)
 {
+	using namespace Chaos;
+
 	check(PVehicle);
 
 	Super::SetupVehicle(PVehicle);
@@ -1690,6 +1702,8 @@ void UChaosWheeledVehicleMovementComponent::Update(float DeltaTime)
 // Debug
 void UChaosWheeledVehicleMovementComponent::DrawDebug(UCanvas* Canvas, float& YL, float& YPos)
 {
+	using namespace Chaos;
+
 	ensure(IsInGameThread());
 
 #if !(UE_BUILD_SHIPPING || UE_BUILD_TEST)
@@ -2626,6 +2640,8 @@ void UChaosWheeledVehicleMovementComponent::SetTorqueCombineMethod(ETorqueCombin
 
 void UChaosWheeledVehicleMovementComponent::SetDriveTorque(float DriveTorque, int32 WheelIndex)
 {
+	using namespace Chaos;
+
 	SetSleeping(false);
 
 	if (FBodyInstance* TargetInstance = GetBodyInstance())
@@ -2644,6 +2660,8 @@ void UChaosWheeledVehicleMovementComponent::SetDriveTorque(float DriveTorque, in
 
 void UChaosWheeledVehicleMovementComponent::SetBrakeTorque(float BrakeTorque, int32 WheelIndex)
 {
+	using namespace Chaos;
+
 	if (FBodyInstance* TargetInstance = GetBodyInstance())
 	{
 		FPhysicsCommand::ExecuteWrite(TargetInstance->ActorHandle, [&](const FPhysicsActorHandle& Chassis)
diff --git a/Engine/Source/Runtime/Engine/Private/Collision/CollisionConversions.cpp b/Engine/Source/Runtime/Engine/Private/Collision/CollisionConversions.cpp
--- a/Engine/Source/Runtime/Engine/Private/Collision/CollisionConversions.cpp
+++ b/Engine/Source/Runtime/Engine/Private/Collision/CollisionConversions.cpp
@@ -62,6 +62,8 @@ DECLARE_CYCLE_STAT(TEXT("SetHitResultFromShapeAndFaceIndex"), STAT_CollisionSetH
 /* Validate Normal of OutResult. We're on hunt for invalid normal */
 static void CheckHitResultNormal(const FHitResult& OutResult, const TCHAR* Message, const FVector& Start=FVector::ZeroVector, const FVector& End = FVector::ZeroVector, const FPhysicsGeometry* Geom=nullptr, const FVector& NormalPreSafeNormalize = FVector::ZeroVector)
 {
+	using namespace ChaosInterface;
+
 	if(!OutResult.bStartPenetrating && !OutResult.Normal.IsNormalized())
 	{
 		UE_LOG(LogPhysics, Warning, TEXT("(%s) Non-normalized OutResult.Normal from hit conversion: %s (Component- %s)"), Message, *OutResult.Normal.ToString(), *OutResult.Component->GetFullName());
@@ -119,6 +121,8 @@ static FVector FindGeomOpposingNormal(ECollisionShapeType QueryGeomType, const T
 /** Set info in the HitResult (Actor, Component, PhysMaterial, BoneName, Item) based on the supplied shape and face index */
 static void SetHitResultFromShapeAndFaceIndex(const FPhysicsShape& Shape,  const FPhysicsActor& Actor, const uint32 FaceIndex, const FVector& HitLocation, FHitResult& OutResult, bool bReturnPhysMat)
 {
+	using namespace ChaosInterface;
+
 	SCOPE_CYCLE_COUNTER(STAT_CollisionSetHitResultFromShapeAndFaceIndex);
 
 	const int32 ShapeIndex = Shape.GetShapeIndex();
@@ -233,6 +237,8 @@ const FPhysicsActor* GetGTActor(const Chaos::FGeometryParticleHandle* PTActor)
 template <typename THitLocation>
 EConvertQueryResult ConvertQueryImpactHitImp(const UWorld* World, const THitLocation& Hit, FHitResult& OutResult, float CheckLength, const FCollisionFilterData& QueryFilter, const FVector& StartLoc, const FVector& EndLoc, const FPhysicsGeometry* Geom, const FTransform& QueryTM, bool bReturnFaceIndex, bool bReturnPhysMat)
 {
+	using namespace ChaosInterface;
+
 	SCOPE_CYCLE_COUNTER(STAT_ConvertQueryImpactHit);
 
 #if WITH_EDITOR
@@ -396,7 +402,7 @@ EConvertQueryResult ConvertQueryImpactHit(const UWorld* World, const FHitLocatio
 {
 	return ConvertQueryImpactHitImp(World, Hit, OutResult, CheckLength, QueryFilter, StartLoc, EndLoc, Geom, QueryTM, bReturnFaceIndex, bReturnPhysMat);
 }
-EConvertQueryResult ConvertQueryImpactHit(const UWorld* World, const FPTLocationHit& Hit, FHitResult& OutResult, float CheckLength, const FCollisionFilterData& QueryFilter, const FVector& StartLoc, const FVector& EndLoc, const FPhysicsGeometry* Geom, const FTransform& QueryTM, bool bReturnFaceIndex, bool bReturnPhysMat)
+EConvertQueryResult ConvertQueryImpactHit(const UWorld* World, const ChaosInterface::FPTLocationHit& Hit, FHitResult& OutResult, float CheckLength, const FCollisionFilterData& QueryFilter, const FVector& StartLoc, const FVector& EndLoc, const FPhysicsGeometry* Geom, const FTransform& QueryTM, bool bReturnFaceIndex, bool bReturnPhysMat)
 {
 	return ConvertQueryImpactHitImp(World, Hit, OutResult, CheckLength, QueryFilter, StartLoc, EndLoc, Geom, QueryTM, bReturnFaceIndex, bReturnPhysMat);
 }
@@ -461,17 +467,19 @@ EConvertQueryResult ConvertTraceResults(bool& OutHasValidBlockingHit, const UWor
 
 template EConvertQueryResult ConvertTraceResults<FHitSweep>(bool& OutHasValidBlockingHit, const UWorld* World, int32 NumHits, FHitSweep* Hits, float CheckLength, const FCollisionFilterData& QueryFilter, TArray<FHitResult>& OutHits, const FVector& StartLoc, const FVector& EndLoc, const FPhysicsGeometry& Geom, const FTransform& QueryTM, float MaxDistance, bool bReturnFaceIndex, bool bReturnPhysMat);
 template EConvertQueryResult ConvertTraceResults<FHitSweep>(bool& OutHasValidBlockingHit, const UWorld* World, int32 NumHits, FHitSweep* Hits, float CheckLength, const FCollisionFilterData& QueryFilter, FHitResult& OutHit, const FVector& StartLoc, const FVector& EndLoc, const FPhysicsGeometry& Geom, const FTransform& QueryTM, float MaxDistance, bool bReturnFaceIndex, bool bReturnPhysMat);
-template EConvertQueryResult ConvertTraceResults<FPTSweepHit>(bool& OutHasValidBlockingHit, const UWorld* World, int32 NumHits, FPTSweepHit* Hits, float CheckLength, const FCollisionFilterData& QueryFilter, TArray<FHitResult>& OutHits, const FVector& StartLoc, const FVector& EndLoc, const FPhysicsGeometry& Geom, const FTransform& QueryTM, float MaxDistance, bool bReturnFaceIndex, bool bReturnPhysMat);
-template EConvertQueryResult ConvertTraceResults<FPTSweepHit>(bool& OutHasValidBlockingHit, const UWorld* World, int32 NumHits, FPTSweepHit* Hits, float CheckLength, const FCollisionFilterData& QueryFilter, FHitResult& OutHit, const FVector& StartLoc, const FVector& EndLoc, const FPhysicsGeometry& Geom, const FTransform& QueryTM, float MaxDistance, bool bReturnFaceIndex, bool bReturnPhysMat);
+template EConvertQueryResult ConvertTraceResults<ChaosInterface::FPTSweepHit>(bool& OutHasValidBlockingHit, const UWorld* World, int32 NumHits, ChaosInterface::FPTSweepHit* Hits, float CheckLength, const FCollisionFilterData& QueryFilter, TArray<FHitResult>& OutHits, const FVector& StartLoc, const FVector& EndLoc, const FPhysicsGeometry& Geom, const FTransform& QueryTM, float MaxDistance, bool bReturnFaceIndex, bool bReturnPhysMat);
+template EConvertQueryResult ConvertTraceResults<ChaosInterface::FPTSweepHit>(bool& OutHasValidBlockingHit, const UWorld* World, int32 NumHits, ChaosInterface::FPTSweepHit* Hits, float CheckLength, const FCollisionFilterData& QueryFilter, FHitResult& OutHit, const FVector& StartLoc, const FVector& EndLoc, const FPhysicsGeometry& Geom, const FTransform& QueryTM, float MaxDistance, bool bReturnFaceIndex, bool bReturnPhysMat);
 template EConvertQueryResult ConvertTraceResults<FHitRaycast>(bool& OutHasValidBlockingHit, const UWorld* World, int32 NumHits, FHitRaycast* Hits, float CheckLength, const FCollisionFilterData& QueryFilter, TArray<FHitResult>& OutHits, const FVector& StartLoc, const FVector& EndLoc, const FPhysicsGeometry& Geom, const FTransform& QueryTM, float MaxDistance, bool bReturnFaceIndex, bool bReturnPhysMat);
 template EConvertQueryResult ConvertTraceResults<FHitRaycast>(bool& OutHasValidBlockingHit, const UWorld* World, int32 NumHits, FHitRaycast* Hits, float CheckLength, const FCollisionFilterData& QueryFilter, FHitResult& OutHit, const FVector& StartLoc, const FVector& EndLoc, const FPhysicsGeometry& Geom, const FTransform& QueryTM, float MaxDistance, bool bReturnFaceIndex, bool bReturnPhysMat);
-template EConvertQueryResult ConvertTraceResults<FPTRaycastHit>(bool& OutHasValidBlockingHit, const UWorld* World, int32 NumHits, FPTRaycastHit* Hits, float CheckLength, const FCollisionFilterData& QueryFilter, TArray<FHitResult>& OutHits, const FVector& StartLoc, const FVector& EndLoc, const FPhysicsGeometry& Geom, const FTransform& QueryTM, float MaxDistance, bool bReturnFaceIndex, bool bReturnPhysMat);
-template EConvertQueryResult ConvertTraceResults<FPTRaycastHit>(bool& OutHasValidBlockingHit, const UWorld* World, int32 NumHits, FPTRaycastHit* Hits, float CheckLength, const FCollisionFilterData& QueryFilter, FHitResult& OutHit, const FVector& StartLoc, const FVector& EndLoc, const FPhysicsGeometry& Geom, const FTransform& QueryTM, float MaxDistance, bool bReturnFaceIndex, bool bReturnPhysMat);
+template EConvertQueryResult ConvertTraceResults<ChaosInterface::FPTRaycastHit>(bool& OutHasValidBlockingHit, const UWorld* World, int32 NumHits, ChaosInterface::FPTRaycastHit* Hits, float CheckLength, const FCollisionFilterData& QueryFilter, TArray<FHitResult>& OutHits, const FVector& StartLoc, const FVector& EndLoc, const FPhysicsGeometry& Geom, const FTransform& QueryTM, float MaxDistance, bool bReturnFaceIndex, bool bReturnPhysMat);
+template EConvertQueryResult ConvertTraceResults<ChaosInterface::FPTRaycastHit>(bool& OutHasValidBlockingHit, const UWorld* World, int32 NumHits, ChaosInterface::FPTRaycastHit* Hits, float CheckLength, const FCollisionFilterData& QueryFilter, FHitResult& OutHit, const FVector& StartLoc, const FVector& EndLoc, const FPhysicsGeometry& Geom, const FTransform& QueryTM, float MaxDistance, bool bReturnFaceIndex, bool bReturnPhysMat);
 
 /** Util to convert an overlapped shape into a sweep hit result, returns whether it was a blocking hit. */
 template <typename THitLocation>
 static bool ConvertOverlappedShapeToImpactHit(const UWorld* World, const THitLocation& Hit, const FVector& StartLoc, const FVector& EndLoc, FHitResult& OutResult, const FPhysicsGeometry& Geom, const FTransform& QueryTM, const FCollisionFilterData& QueryFilter, bool bReturnPhysMat)
 {
+	using namespace ChaosInterface;
+
 	SCOPE_CYCLE_COUNTER(STAT_CollisionConvertOverlapToHit);
 
 	const FPhysicsShape* pHitShape = GetShape(Hit);
@@ -586,6 +594,8 @@ static bool ConvertOverlappedShapeToImpactHit(const UWorld* World, const THitLoc
 
 void ConvertQueryOverlap(const FPhysicsShape& Shape, const FPhysicsActor& Actor, FOverlapResult& OutOverlap, const FCollisionFilterData& QueryFilter)
 {
+	using namespace ChaosInterface;
+
 	// Grab actor/component
 	
 	// Try body instance
@@ -648,6 +658,8 @@ static void AddUniqueOverlap(TArray<FOverlapResult>& OutOverlaps, const FOverlap
 
 bool IsBlocking(const FPhysicsShape& Shape, const FCollisionFilterData& QueryFilter)
 {
+	using namespace ChaosInterface;
+
 	// See if this is a 'blocking' hit
 	const FCollisionFilterData ShapeFilter = GetQueryFilterData(Shape);
 	const ECollisionQueryHitType HitType = FCollisionQueryFilterCallback::CalcQueryHitType(QueryFilter, ShapeFilter);
@@ -757,7 +769,7 @@ bool ConvertOverlapResults(int32 NumOverlaps, FHitOverlap* OverlapResults, const
 	return ConvertOverlapResultsImp(NumOverlaps, OverlapResults, QueryFilter, OutOverlaps);
 }
 
-bool ConvertOverlapResults(int32 NumOverlaps, FPTOverlapHit* OverlapResults, const FCollisionFilterData& QueryFilter, TArray<FOverlapResult>& OutOverlaps)
+bool ConvertOverlapResults(int32 NumOverlaps, ChaosInterface::FPTOverlapHit* OverlapResults, const FCollisionFilterData& QueryFilter, TArray<FOverlapResult>& OutOverlaps)
 {
 	return ConvertOverlapResultsImp(NumOverlaps, OverlapResults, QueryFilter, OutOverlaps);
 }
\ No newline at end of file
diff --git a/Engine/Source/Runtime/Engine/Private/Collision/CollisionConversions.h b/Engine/Source/Runtime/Engine/Private/Collision/CollisionConversions.h
--- a/Engine/Source/Runtime/Engine/Private/Collision/CollisionConversions.h
+++ b/Engine/Source/Runtime/Engine/Private/Collision/CollisionConversions.h
@@ -31,7 +31,7 @@ enum class EConvertQueryResult
  * @return	Whether result passed NaN/Inf checks.
  */
 EConvertQueryResult ConvertQueryImpactHit(const UWorld* World, const FHitLocation& PHit, FHitResult& OutResult, float CheckLength, const FCollisionFilterData& QueryFilter, const FVector& StartLoc, const FVector& EndLoc, const FPhysicsGeometry* Geom, const FTransform& QueryTM, bool bReturnFaceIndex, bool bReturnPhysMat);
-EConvertQueryResult ConvertQueryImpactHit(const UWorld* World, const FPTLocationHit& PHit, FHitResult& OutResult, float CheckLength, const FCollisionFilterData& QueryFilter, const FVector& StartLoc, const FVector& EndLoc, const FPhysicsGeometry* Geom, const FTransform& QueryTM, bool bReturnFaceIndex, bool bReturnPhysMat);
+EConvertQueryResult ConvertQueryImpactHit(const UWorld* World, const ChaosInterface::FPTLocationHit& PHit, FHitResult& OutResult, float CheckLength, const FCollisionFilterData& QueryFilter, const FVector& StartLoc, const FVector& EndLoc, const FPhysicsGeometry* Geom, const FTransform& QueryTM, bool bReturnFaceIndex, bool bReturnPhysMat);
 
 /** 
  * Util to convert physX sweep results to unreal hit results and add to array
@@ -83,7 +83,7 @@ bool IsBlocking(const FPhysicsShape& PShape, const FCollisionFilterData& QueryFi
  * @return	OutOverlaps	Converted data
  */
 bool ConvertOverlapResults(int32 NumOverlaps, FHitOverlap* POverlapResults, const FCollisionFilterData& QueryFilter, TArray<FOverlapResult>& OutOverlaps);
-bool ConvertOverlapResults(int32 NumOverlaps, FPTOverlapHit* POverlapResults, const FCollisionFilterData& QueryFilter, TArray<FOverlapResult>& OutOverlaps);
+bool ConvertOverlapResults(int32 NumOverlaps, ChaosInterface::FPTOverlapHit* POverlapResults, const FCollisionFilterData& QueryFilter, TArray<FOverlapResult>& OutOverlaps);
 
 
 struct FCompareFHitResultTime
diff --git a/Engine/Source/Runtime/Engine/Private/Collision/SceneQuery.cpp b/Engine/Source/Runtime/Engine/Private/Collision/SceneQuery.cpp
--- a/Engine/Source/Runtime/Engine/Private/Collision/SceneQuery.cpp
+++ b/Engine/Source/Runtime/Engine/Private/Collision/SceneQuery.cpp
@@ -114,9 +114,9 @@ struct FRaycastSQAdditionalInputs
 	}
 };
 
-void LowLevelRaycast(FPhysScene& Scene, const FVector& Start, const FVector& Dir, float DeltaMag, FPhysicsHitCallback<FHitRaycast>& HitBuffer, EHitFlags OutputFlags, FQueryFlags QueryFlags, const FCollisionFilterData& Filter, const FQueryFilterData& QueryFilterData, ICollisionQueryFilterCallbackBase* QueryCallback, const FQueryDebugParams& DebugParams);
-void LowLevelSweep(FPhysScene& Scene, const FPhysicsGeometry& Geom, const FTransform& StartTM, const FVector& Dir, float DeltaMag, FPhysicsHitCallback<FHitSweep>& HitBuffer, EHitFlags OutputFlags, FQueryFlags QueryFlags, const FCollisionFilterData& Filter, const FQueryFilterData& QueryFilterData, ICollisionQueryFilterCallbackBase* QueryCallback, const FQueryDebugParams& DebugParams);
-void LowLevelOverlap(FPhysScene& Scene, const FPhysicsGeometry& Geom, const FTransform& GeomPose, FPhysicsHitCallback<FHitOverlap>& HitBuffer, FQueryFlags QueryFlags, const FCollisionFilterData& Filter, const FQueryFilterData& QueryFilterData, ICollisionQueryFilterCallbackBase* QueryCallback, const FQueryDebugParams& DebugParams);
+void LowLevelRaycast(FPhysScene& Scene, const FVector& Start, const FVector& Dir, float DeltaMag, FPhysicsHitCallback<FHitRaycast>& HitBuffer, EHitFlags OutputFlags, FQueryFlags QueryFlags, const FCollisionFilterData& Filter, const ChaosInterface::FQueryFilterData& QueryFilterData, ICollisionQueryFilterCallbackBase* QueryCallback, const ChaosInterface::FQueryDebugParams& DebugParams);
+void LowLevelSweep(FPhysScene& Scene, const FPhysicsGeometry& Geom, const FTransform& StartTM, const FVector& Dir, float DeltaMag, FPhysicsHitCallback<FHitSweep>& HitBuffer, EHitFlags OutputFlags, FQueryFlags QueryFlags, const FCollisionFilterData& Filter, const ChaosInterface::FQueryFilterData& QueryFilterData, ICollisionQueryFilterCallbackBase* QueryCallback, const ChaosInterface::FQueryDebugParams& DebugParams);
+void LowLevelOverlap(FPhysScene& Scene, const FPhysicsGeometry& Geom, const FTransform& GeomPose, FPhysicsHitCallback<FHitOverlap>& HitBuffer, FQueryFlags QueryFlags, const FCollisionFilterData& Filter, const ChaosInterface::FQueryFilterData& QueryFilterData, ICollisionQueryFilterCallbackBase* QueryCallback, const ChaosInterface::FQueryDebugParams& DebugParams);
 
 template <typename InHitType, ESweepOrRay InGeometryQuery, ESingleMultiOrTest InSingleMultiOrTest>
 struct TSQTraits
@@ -159,6 +159,7 @@ struct TSQTraits
 	template <typename TGeomInputs, ESweepOrRay T = GeometryQuery>
 	static typename TEnableIf<T == ESweepOrRay::Raycast, void>::Type SceneTrace(FPhysScene& Scene, const TGeomInputs& GeomInputs, const FVector& Dir, float DeltaMag, const FTransform& StartTM, THitBuffer& HitBuffer, EHitFlags OutputFlags, EQueryFlags QueryFlags, const FCollisionFilterData& FilterData, const FCollisionQueryParams& Params, ICollisionQueryFilterCallbackBase* QueryCallback)
 	{
+		using namespace ChaosInterface;
 		FQueryFilterData QueryFilterData = MakeQueryFilterData(FilterData, QueryFlags, Params);
 		FQueryDebugParams DebugParams;
 #if !(UE_BUILD_TEST || UE_BUILD_SHIPPING)
@@ -171,6 +172,7 @@ struct TSQTraits
 	template <typename TGeomInputs, ESweepOrRay T = GeometryQuery>
 	static typename TEnableIf<T == ESweepOrRay::Sweep, void>::Type SceneTrace(FPhysScene& Scene, const TGeomInputs& GeomInputs, const FVector& Dir, float DeltaMag, const FTransform& StartTM, THitBuffer& HitBuffer, EHitFlags OutputFlags, EQueryFlags QueryFlags, const FCollisionFilterData& FilterData, const FCollisionQueryParams& Params, ICollisionQueryFilterCallbackBase* QueryCallback)
 	{
+		using namespace ChaosInterface;
 		FQueryFilterData QueryFilterData = MakeQueryFilterData(FilterData, QueryFlags, Params);
 		FQueryDebugParams DebugParams;
 #if !(UE_BUILD_TEST || UE_BUILD_SHIPPING)
@@ -337,6 +339,8 @@ EThreadQueryContext GetThreadQueryContext(const Chaos::FPhysicsSolver& Solver)
 template <typename Traits, typename TGeomInputs>
 bool TSceneCastCommonImp(const UWorld* World, typename Traits::TOutHits& OutHits, const TGeomInputs& GeomInputs, const FVector Start, const FVector End, ECollisionChannel TraceChannel, const struct FCollisionQueryParams& Params, const struct FCollisionResponseParams& ResponseParams, const struct FCollisionObjectQueryParams& ObjectParams)
 {
+	using namespace ChaosInterface;
+
 	FScopeCycleCounter Counter(Params.StatId);
 	STARTQUERYTIMER();
 
@@ -457,6 +461,8 @@ bool TSceneCastCommon(const UWorld* World, typename Traits::TOutHits& OutHits, c
 
 bool FGenericPhysicsInterface::RaycastTest(const UWorld* World, const FVector Start, const FVector End, ECollisionChannel TraceChannel, const struct FCollisionQueryParams& Params, const struct FCollisionResponseParams& ResponseParams, const struct FCollisionObjectQueryParams& ObjectParams)
 {
+	using namespace ChaosInterface;
+
 	SCOPE_CYCLE_COUNTER(STAT_Collision_SceneQueryTotal);
 	SCOPE_CYCLE_COUNTER(STAT_Collision_RaycastAny);
 	CSV_SCOPED_TIMING_STAT(SceneQuery, RaycastTest);
@@ -470,6 +476,8 @@ bool FGenericPhysicsInterface::RaycastTest(const UWorld* World, const FVector St
 
 bool FGenericPhysicsInterface::RaycastSingle(const UWorld* World, struct FHitResult& OutHit, const FVector Start, const FVector End, ECollisionChannel TraceChannel, const struct FCollisionQueryParams& Params, const struct FCollisionResponseParams& ResponseParams, const struct FCollisionObjectQueryParams& ObjectParams)
 {
+	using namespace ChaosInterface;
+
 	SCOPE_CYCLE_COUNTER(STAT_Collision_SceneQueryTotal);
 	SCOPE_CYCLE_COUNTER(STAT_Collision_RaycastSingle);
 	CSV_SCOPED_TIMING_STAT(SceneQuery, RaycastSingle);
@@ -483,6 +491,8 @@ bool FGenericPhysicsInterface::RaycastSingle(const UWorld* World, struct FHitRes
 
 bool FGenericPhysicsInterface::RaycastMulti(const UWorld* World, TArray<struct FHitResult>& OutHits, const FVector& Start, const FVector& End, ECollisionChannel TraceChannel, const struct FCollisionQueryParams& Params, const struct FCollisionResponseParams& ResponseParams, const struct FCollisionObjectQueryParams& ObjectParams)
 {
+	using namespace ChaosInterface;
+
 	SCOPE_CYCLE_COUNTER(STAT_Collision_SceneQueryTotal);
 	SCOPE_CYCLE_COUNTER(STAT_Collision_RaycastMultiple);
 	CSV_SCOPED_TIMING_STAT(SceneQuery, RaycastMultiple);
@@ -497,6 +507,8 @@ bool FGenericPhysicsInterface::RaycastMulti(const UWorld* World, TArray<struct F
 
 bool FGenericPhysicsInterface::GeomSweepTest(const UWorld* World, const struct FCollisionShape& CollisionShape, const FQuat& Rot, FVector Start, FVector End, ECollisionChannel TraceChannel, const struct FCollisionQueryParams& Params, const struct FCollisionResponseParams& ResponseParams, const struct FCollisionObjectQueryParams& ObjectParams)
 {
+	using namespace ChaosInterface;
+
 	SCOPE_CYCLE_COUNTER(STAT_Collision_SceneQueryTotal);
 	SCOPE_CYCLE_COUNTER(STAT_Collision_GeomSweepAny);
 	CSV_SCOPED_TIMING_STAT(SceneQuery, GeomSweepTest);
@@ -510,6 +522,8 @@ bool FGenericPhysicsInterface::GeomSweepTest(const UWorld* World, const struct F
 
 bool FGenericPhysicsInterface::GeomSweepSingle(const UWorld* World, const struct FCollisionShape& CollisionShape, const FQuat& Rot, FHitResult& OutHit, FVector Start, FVector End, ECollisionChannel TraceChannel, const struct FCollisionQueryParams& Params, const struct FCollisionResponseParams& ResponseParams, const struct FCollisionObjectQueryParams& ObjectParams)
 {
+	using namespace ChaosInterface;
+
 	SCOPE_CYCLE_COUNTER(STAT_Collision_SceneQueryTotal);
 	SCOPE_CYCLE_COUNTER(STAT_Collision_GeomSweepSingle);
 	CSV_SCOPED_TIMING_STAT(SceneQuery, GeomSweepSingle);
@@ -522,6 +536,8 @@ bool FGenericPhysicsInterface::GeomSweepSingle(const UWorld* World, const struct
 template<>
 bool FGenericPhysicsInterface::GeomSweepMulti(const UWorld* World, const FPhysicsGeometryCollection& InGeom, const FQuat& InGeomRot, TArray<FHitResult>& OutHits, FVector Start, FVector End, ECollisionChannel TraceChannel, const FCollisionQueryParams& Params, const FCollisionResponseParams& ResponseParams, const FCollisionObjectQueryParams& ObjectParams /*= FCollisionObjectQueryParams::DefaultObjectQueryParam*/)
 {
+	using namespace ChaosInterface;
+
 	SCOPE_CYCLE_COUNTER(STAT_Collision_SceneQueryTotal);
 	SCOPE_CYCLE_COUNTER(STAT_Collision_GeomSweepMultiple);
 	CSV_SCOPED_TIMING_STAT(SceneQuery, GeomSweepMultiple);
@@ -534,6 +550,8 @@ bool FGenericPhysicsInterface::GeomSweepMulti(const UWorld* World, const FPhysic
 template<>
 bool FGenericPhysicsInterface::GeomSweepMulti(const UWorld* World, const FCollisionShape& InGeom, const FQuat& InGeomRot, TArray<FHitResult>& OutHits, FVector Start, FVector End, ECollisionChannel TraceChannel, const FCollisionQueryParams& Params, const FCollisionResponseParams& ResponseParams, const FCollisionObjectQueryParams& ObjectParams /*= FCollisionObjectQueryParams::DefaultObjectQueryParam*/)
 {
+	using namespace ChaosInterface;
+
 	SCOPE_CYCLE_COUNTER(STAT_Collision_SceneQueryTotal);
 	SCOPE_CYCLE_COUNTER(STAT_Collision_GeomSweepMultiple);
 	CSV_SCOPED_TIMING_STAT(SceneQuery, GeomSweepMultiple);
@@ -561,6 +579,8 @@ namespace EQueryInfo
 template <typename TOverlapHit, EQueryInfo::Type InfoType, typename TCollisionAnalyzerType>
 bool GeomOverlapMultiImp(const UWorld* World, const FPhysicsGeometry& Geom, const TCollisionAnalyzerType& CollisionAnalyzerType, const FTransform& GeomPose, TArray<FOverlapResult>& OutOverlaps, ECollisionChannel TraceChannel, const struct FCollisionQueryParams& Params, const struct FCollisionResponseParams& ResponseParams, const struct FCollisionObjectQueryParams& ObjectParams)
 {
+	using namespace ChaosInterface;
+
 	FScopeCycleCounter Counter(Params.StatId);
 
 	if ((World == NULL) || (World->GetPhysicsScene() == NULL))
@@ -645,6 +665,8 @@ bool GeomOverlapMultiImp(const UWorld* World, const FPhysicsGeometry& Geom, cons
 template <EQueryInfo::Type InfoType, typename TCollisionAnalyzerType>
 bool GeomOverlapMultiHelper(const UWorld* World, const FPhysicsGeometry& Geom, const TCollisionAnalyzerType& CollisionAnalyzerType, const FTransform& GeomPose, TArray<FOverlapResult>& OutOverlaps, ECollisionChannel TraceChannel, const struct FCollisionQueryParams& Params, const struct FCollisionResponseParams& ResponseParams, const struct FCollisionObjectQueryParams& ObjectParams)
 {
+	using namespace ChaosInterface;
+
 	if ((World == NULL) || (World->GetPhysicsScene() == NULL))
 	{
 		return false;
diff --git a/Engine/Source/Runtime/Engine/Private/Collision/SceneQueryLowLevel.cpp b/Engine/Source/Runtime/Engine/Private/Collision/SceneQueryLowLevel.cpp
--- a/Engine/Source/Runtime/Engine/Private/Collision/SceneQueryLowLevel.cpp
+++ b/Engine/Source/Runtime/Engine/Private/Collision/SceneQueryLowLevel.cpp
@@ -73,7 +73,7 @@ void FinalizeCapture(FPhysTestSerializer& Serializer) {}
 
 namespace
 {
-	void SweepSQCaptureHelper(float QueryDurationSeconds, const FChaosSQAccelerator& SQAccelerator, FPhysScene& Scene, const FPhysicsGeometry& QueryGeom, const FTransform& StartTM, const FVector& Dir, float DeltaMag, const FPhysicsHitCallback<FHitSweep>& HitBuffer, EHitFlags OutputFlags, FQueryFlags QueryFlags, const FCollisionFilterData& Filter, const FQueryFilterData& QueryFilterData, ICollisionQueryFilterCallbackBase* QueryCallback, const FQueryDebugParams& DebugParams)
+	void SweepSQCaptureHelper(float QueryDurationSeconds, const FChaosSQAccelerator& SQAccelerator, FPhysScene& Scene, const FPhysicsGeometry& QueryGeom, const FTransform& StartTM, const FVector& Dir, float DeltaMag, const FPhysicsHitCallback<FHitSweep>& HitBuffer, EHitFlags OutputFlags, FQueryFlags QueryFlags, const FCollisionFilterData& Filter, const ChaosInterface::FQueryFilterData& QueryFilterData, ICollisionQueryFilterCallbackBase* QueryCallback, const ChaosInterface::FQueryDebugParams& DebugParams)
 	{
 		float QueryDurationMicro = QueryDurationSeconds * 1000.0 * 1000.0;
 		if (((SerializeSQs && QueryDurationMicro > SerializeSQs)) && IsInGameThread())
@@ -107,7 +107,7 @@ namespace
 		}
 	}
 
-	void RaycastSQCaptureHelper(float QueryDurationSeconds, const FChaosSQAccelerator& SQAccelerator, FPhysScene& Scene, const FVector& Start, const FVector& Dir, float DeltaMag, const FPhysicsHitCallback<FHitRaycast>& HitBuffer, EHitFlags OutputFlags, FQueryFlags QueryFlags, const FCollisionFilterData& Filter, const FQueryFilterData& QueryFilterData, ICollisionQueryFilterCallbackBase* QueryCallback, const FQueryDebugParams& DebugParams)
+	void RaycastSQCaptureHelper(float QueryDurationSeconds, const FChaosSQAccelerator& SQAccelerator, FPhysScene& Scene, const FVector& Start, const FVector& Dir, float DeltaMag, const FPhysicsHitCallback<FHitRaycast>& HitBuffer, EHitFlags OutputFlags, FQueryFlags QueryFlags, const FCollisionFilterData& Filter, const ChaosInterface::FQueryFilterData& QueryFilterData, ICollisionQueryFilterCallbackBase* QueryCallback, const ChaosInterface::FQueryDebugParams& DebugParams)
 	{
 		float QueryDurationMicro = QueryDurationSeconds * 1000.0 * 1000.0;
 		if (((!!SerializeSQs && QueryDurationMicro > SerializeSQs)) && IsInGameThread())
@@ -141,7 +141,7 @@ namespace
 		}
 	}
 
-	void OverlapSQCaptureHelper(float QueryDurationSeconds, const FChaosSQAccelerator& SQAccelerator, FPhysScene& Scene, const FPhysicsGeometry& QueryGeom, const FTransform& GeomPose, const FPhysicsHitCallback<FHitOverlap>& HitBuffer, FQueryFlags QueryFlags, const FCollisionFilterData& Filter, const FQueryFilterData& QueryFilterData, ICollisionQueryFilterCallbackBase* QueryCallback, const FQueryDebugParams& DebugParams)
+	void OverlapSQCaptureHelper(float QueryDurationSeconds, const FChaosSQAccelerator& SQAccelerator, FPhysScene& Scene, const FPhysicsGeometry& QueryGeom, const FTransform& GeomPose, const FPhysicsHitCallback<FHitOverlap>& HitBuffer, FQueryFlags QueryFlags, const FCollisionFilterData& Filter, const ChaosInterface::FQueryFilterData& QueryFilterData, ICollisionQueryFilterCallbackBase* QueryCallback, const ChaosInterface::FQueryDebugParams& DebugParams)
 	{
 		float QueryDurationMicro = QueryDurationSeconds * 1000.0 * 1000.0;
 		if (((!!SerializeSQs && QueryDurationMicro > SerializeSQs)) && IsInGameThread())
@@ -177,7 +177,7 @@ namespace
 }
 
 template <typename THitRaycast>
-void LowLevelRaycastImp(FPhysScene& Scene, const FVector& Start, const FVector& Dir, float DeltaMag, FPhysicsHitCallback<THitRaycast>& HitBuffer, EHitFlags OutputFlags, FQueryFlags QueryFlags, const FCollisionFilterData& Filter, const FQueryFilterData& QueryFilterData, ICollisionQueryFilterCallbackBase* QueryCallback, const FQueryDebugParams& DebugParams)
+void LowLevelRaycastImp(FPhysScene& Scene, const FVector& Start, const FVector& Dir, float DeltaMag, FPhysicsHitCallback<THitRaycast>& HitBuffer, EHitFlags OutputFlags, FQueryFlags QueryFlags, const FCollisionFilterData& Filter, const ChaosInterface::FQueryFilterData& QueryFilterData, ICollisionQueryFilterCallbackBase* QueryCallback, const ChaosInterface::FQueryDebugParams& DebugParams)
 {
 	constexpr bool bGTData = std::is_same<THitRaycast, FHitRaycast>::value;
 	const auto SolverAccelerationStructure = bGTData ? Scene.GetSpacialAcceleration() : Scene.GetSolver()->GetInternalAccelerationStructure_Internal();
@@ -200,22 +200,22 @@ void LowLevelRaycastImp(FPhysScene& Scene, const FVector& Start, const FVector&
 	}
 }
 
-void LowLevelRaycast(FPhysScene& Scene, const FVector& Start, const FVector& Dir, float DeltaMag, FPhysicsHitCallback<FHitRaycast>& HitBuffer, EHitFlags OutputFlags, FQueryFlags QueryFlags, const FCollisionFilterData& Filter, const FQueryFilterData& QueryFilterData, ICollisionQueryFilterCallbackBase* QueryCallback, const FQueryDebugParams& DebugParams)
+void LowLevelRaycast(FPhysScene& Scene, const FVector& Start, const FVector& Dir, float DeltaMag, FPhysicsHitCallback<FHitRaycast>& HitBuffer, EHitFlags OutputFlags, FQueryFlags QueryFlags, const FCollisionFilterData& Filter, const ChaosInterface::FQueryFilterData& QueryFilterData, ICollisionQueryFilterCallbackBase* QueryCallback, const ChaosInterface::FQueryDebugParams& DebugParams)
 {
 	LowLevelRaycastImp(Scene, Start, Dir, DeltaMag, HitBuffer, OutputFlags, QueryFlags, Filter, QueryFilterData, QueryCallback, DebugParams);
 }
 
-void LowLevelRaycast(FPhysScene& Scene, const FVector& Start, const FVector& Dir, float DeltaMag, FPhysicsHitCallback<FPTRaycastHit>& HitBuffer, EHitFlags OutputFlags, FQueryFlags QueryFlags, const FCollisionFilterData& Filter, const FQueryFilterData& QueryFilterData, ICollisionQueryFilterCallbackBase* QueryCallback, const FQueryDebugParams& DebugParams)
+void LowLevelRaycast(FPhysScene& Scene, const FVector& Start, const FVector& Dir, float DeltaMag, FPhysicsHitCallback<ChaosInterface::FPTRaycastHit>& HitBuffer, EHitFlags OutputFlags, FQueryFlags QueryFlags, const FCollisionFilterData& Filter, const ChaosInterface::FQueryFilterData& QueryFilterData, ICollisionQueryFilterCallbackBase* QueryCallback, const ChaosInterface::FQueryDebugParams& DebugParams)
 {
 	LowLevelRaycastImp(Scene, Start, Dir, DeltaMag, HitBuffer, OutputFlags, QueryFlags, Filter, QueryFilterData, QueryCallback, DebugParams);
 }
 
 template <typename THitSweep>
-void LowLevelSweepImp(FPhysScene& Scene, const FPhysicsGeometry& QueryGeom, const FTransform& StartTM, const FVector& Dir, float DeltaMag, FPhysicsHitCallback<THitSweep>& HitBuffer, EHitFlags OutputFlags, FQueryFlags QueryFlags, const FCollisionFilterData& Filter, const FQueryFilterData& QueryFilterData, ICollisionQueryFilterCallbackBase* QueryCallback, const FQueryDebugParams& DebugParams)
+void LowLevelSweepImp(FPhysScene& Scene, const FPhysicsGeometry& QueryGeom, const FTransform& StartTM, const FVector& Dir, float DeltaMag, FPhysicsHitCallback<THitSweep>& HitBuffer, EHitFlags OutputFlags, FQueryFlags QueryFlags, const FCollisionFilterData& Filter, const ChaosInterface::FQueryFilterData& QueryFilterData, ICollisionQueryFilterCallbackBase* QueryCallback, const ChaosInterface::FQueryDebugParams& DebugParams)
 {
 	constexpr bool bGTData = std::is_same<THitSweep, FHitSweep>::value;
 	const auto SolverAccelerationStructure = bGTData ? Scene.GetSpacialAcceleration() : Scene.GetSolver()->GetInternalAccelerationStructure_Internal();
-	if (SolverAccelerationStructure)
+	if(SolverAccelerationStructure)
 	{
 		FChaosSQAccelerator SQAccelerator(*SolverAccelerationStructure);
 		{
@@ -226,9 +226,9 @@ void LowLevelSweepImp(FPhysScene& Scene, const FPhysicsGeometry& QueryGeom, cons
 				SQAccelerator.Sweep(QueryGeom, StartTM, Dir, DeltaMag, HitBuffer, OutputFlags, QueryFilterData, *QueryCallback, DebugParams);
 			}
 
-			if constexpr (bGTData)
+			if constexpr(bGTData)
 			{
-				if (!!SerializeSQs && !!EnableSweepSQCapture)
+				if(!!SerializeSQs && !!EnableSweepSQCapture)
 				{
 					SweepSQCaptureHelper(Time, SQAccelerator, Scene, QueryGeom, StartTM, Dir, DeltaMag, HitBuffer, OutputFlags, QueryFlags, Filter, QueryFilterData, QueryCallback, DebugParams);
 				}
@@ -237,22 +237,22 @@ void LowLevelSweepImp(FPhysScene& Scene, const FPhysicsGeometry& QueryGeom, cons
 	}
 }
 
-void LowLevelSweep(FPhysScene& Scene, const FPhysicsGeometry& QueryGeom, const FTransform& StartTM, const FVector& Dir, float DeltaMag, FPhysicsHitCallback<FHitSweep>& HitBuffer, EHitFlags OutputFlags, FQueryFlags QueryFlags, const FCollisionFilterData& Filter, const FQueryFilterData& QueryFilterData, ICollisionQueryFilterCallbackBase* QueryCallback, const FQueryDebugParams& DebugParams)
+void LowLevelSweep(FPhysScene& Scene, const FPhysicsGeometry& QueryGeom, const FTransform& StartTM, const FVector& Dir, float DeltaMag, FPhysicsHitCallback<FHitSweep>& HitBuffer, EHitFlags OutputFlags, FQueryFlags QueryFlags, const FCollisionFilterData& Filter, const ChaosInterface::FQueryFilterData& QueryFilterData, ICollisionQueryFilterCallbackBase* QueryCallback, const ChaosInterface::FQueryDebugParams& DebugParams)
 {
 	return LowLevelSweepImp(Scene, QueryGeom, StartTM, Dir, DeltaMag, HitBuffer, OutputFlags, QueryFlags, Filter, QueryFilterData, QueryCallback, DebugParams);
 }
 
-void LowLevelSweep(FPhysScene& Scene, const FPhysicsGeometry& QueryGeom, const FTransform& StartTM, const FVector& Dir, float DeltaMag, FPhysicsHitCallback<FPTSweepHit>& HitBuffer, EHitFlags OutputFlags, FQueryFlags QueryFlags, const FCollisionFilterData& Filter, const FQueryFilterData& QueryFilterData, ICollisionQueryFilterCallbackBase* QueryCallback, const FQueryDebugParams& DebugParams)
+void LowLevelSweep(FPhysScene& Scene, const FPhysicsGeometry& QueryGeom, const FTransform& StartTM, const FVector& Dir, float DeltaMag, FPhysicsHitCallback<ChaosInterface::FPTSweepHit>& HitBuffer, EHitFlags OutputFlags, FQueryFlags QueryFlags, const FCollisionFilterData& Filter, const ChaosInterface::FQueryFilterData& QueryFilterData, ICollisionQueryFilterCallbackBase* QueryCallback, const ChaosInterface::FQueryDebugParams& DebugParams)
 {
 	return LowLevelSweepImp(Scene, QueryGeom, StartTM, Dir, DeltaMag, HitBuffer, OutputFlags, QueryFlags, Filter, QueryFilterData, QueryCallback, DebugParams);
 }
 
 template <typename THitOverlap>
-void LowLevelOverlapImp(FPhysScene& Scene, const FPhysicsGeometry& QueryGeom, const FTransform& GeomPose, FPhysicsHitCallback<THitOverlap>& HitBuffer, FQueryFlags QueryFlags, const FCollisionFilterData& Filter, const FQueryFilterData& QueryFilterData, ICollisionQueryFilterCallbackBase* QueryCallback, const FQueryDebugParams& DebugParams)
+void LowLevelOverlapImp(FPhysScene& Scene, const FPhysicsGeometry& QueryGeom, const FTransform& GeomPose, FPhysicsHitCallback<THitOverlap>& HitBuffer, FQueryFlags QueryFlags, const FCollisionFilterData& Filter, const ChaosInterface::FQueryFilterData& QueryFilterData, ICollisionQueryFilterCallbackBase* QueryCallback, const ChaosInterface::FQueryDebugParams& DebugParams)
 {
 	constexpr bool bGTData = std::is_same<THitOverlap, FHitOverlap>::value;
 	const auto SolverAccelerationStructure = bGTData ? Scene.GetSpacialAcceleration() : Scene.GetSolver()->GetInternalAccelerationStructure_Internal();
-	if (SolverAccelerationStructure)
+	if(SolverAccelerationStructure)
 	{
 		FChaosSQAccelerator SQAccelerator(*SolverAccelerationStructure);
 		double Time = 0.0;
@@ -261,9 +261,9 @@ void LowLevelOverlapImp(FPhysScene& Scene, const FPhysicsGeometry& QueryGeom, co
 			SQAccelerator.Overlap(QueryGeom, GeomPose, HitBuffer, QueryFilterData, *QueryCallback, DebugParams);
 		}
 
-		if constexpr (bGTData)
+		if constexpr(bGTData)
 		{
-			if (!!SerializeSQs && !!EnableOverlapSQCapture)
+			if(!!SerializeSQs && !!EnableOverlapSQCapture)
 			{
 				OverlapSQCaptureHelper(Time, SQAccelerator, Scene, QueryGeom, GeomPose, HitBuffer, QueryFlags, Filter, QueryFilterData, QueryCallback, DebugParams);
 			}
@@ -271,12 +271,12 @@ void LowLevelOverlapImp(FPhysScene& Scene, const FPhysicsGeometry& QueryGeom, co
 	}
 }
 
-void LowLevelOverlap(FPhysScene& Scene, const FPhysicsGeometry& QueryGeom, const FTransform& GeomPose, FPhysicsHitCallback<FHitOverlap>& HitBuffer, FQueryFlags QueryFlags, const FCollisionFilterData& Filter, const FQueryFilterData& QueryFilterData, ICollisionQueryFilterCallbackBase* QueryCallback, const FQueryDebugParams& DebugParams)
+void LowLevelOverlap(FPhysScene& Scene, const FPhysicsGeometry& QueryGeom, const FTransform& GeomPose, FPhysicsHitCallback<FHitOverlap>& HitBuffer, FQueryFlags QueryFlags, const FCollisionFilterData& Filter, const ChaosInterface::FQueryFilterData& QueryFilterData, ICollisionQueryFilterCallbackBase* QueryCallback, const ChaosInterface::FQueryDebugParams& DebugParams)
 {
 	return LowLevelOverlapImp(Scene, QueryGeom, GeomPose, HitBuffer, QueryFlags, Filter, QueryFilterData, QueryCallback, DebugParams);
 }
 
-void LowLevelOverlap(FPhysScene& Scene, const FPhysicsGeometry& QueryGeom, const FTransform& GeomPose, FPhysicsHitCallback<FPTOverlapHit>& HitBuffer, FQueryFlags QueryFlags, const FCollisionFilterData& Filter, const FQueryFilterData& QueryFilterData, ICollisionQueryFilterCallbackBase* QueryCallback, const FQueryDebugParams& DebugParams)
+void LowLevelOverlap(FPhysScene& Scene, const FPhysicsGeometry& QueryGeom, const FTransform& GeomPose, FPhysicsHitCallback<ChaosInterface::FPTOverlapHit>& HitBuffer, FQueryFlags QueryFlags, const FCollisionFilterData& Filter, const ChaosInterface::FQueryFilterData& QueryFilterData, ICollisionQueryFilterCallbackBase* QueryCallback, const ChaosInterface::FQueryDebugParams& DebugParams)
 {
 	return LowLevelOverlapImp(Scene, QueryGeom, GeomPose, HitBuffer, QueryFlags, Filter, QueryFilterData, QueryCallback, DebugParams);
 }
diff --git a/Engine/Source/Runtime/Engine/Private/PhysicsEngine/Experimental/PhysInterface_Chaos.cpp b/Engine/Source/Runtime/Engine/Private/PhysicsEngine/Experimental/PhysInterface_Chaos.cpp
--- a/Engine/Source/Runtime/Engine/Private/PhysicsEngine/Experimental/PhysInterface_Chaos.cpp
+++ b/Engine/Source/Runtime/Engine/Private/PhysicsEngine/Experimental/PhysInterface_Chaos.cpp
@@ -84,7 +84,7 @@ DECLARE_DWORD_COUNTER_STAT(TEXT("(ASync) Shapes"), STAT_NumShapesAsync, STATGROU
 
 ECollisionShapeType GetGeometryType(const Chaos::FPerShapeData& Shape)
 {
-	return GetType(*Shape.GetGeometry());
+	return ChaosInterface::GetType(*Shape.GetGeometry());
 }
 
 Chaos::FChaosPhysicsMaterial* GetMaterialFromInternalFaceIndex(const FPhysicsShape& Shape, const FPhysicsActor& Actor, uint32 InternalFaceIndex)
@@ -115,6 +115,8 @@ Chaos::FChaosPhysicsMaterial* GetMaterialFromInternalFaceIndex(const FPhysicsSha
 
 Chaos::FChaosPhysicsMaterial* GetMaterialFromInternalFaceIndexAndHitLocation(const FPhysicsShape& Shape, const FPhysicsActor& Actor, uint32 InternalFaceIndex, const FVector& HitLocation)
 {
+	using namespace ChaosInterface;
+
 	{
 		SCOPE_CYCLE_COUNTER(STAT_QueryPhysicalMaterialMaskHit);
 
@@ -981,6 +983,8 @@ void FinishSceneStat()
 
 bool FPhysInterface_Chaos::LineTrace_Geom(FHitResult& OutHit, const FBodyInstance* InInstance, const FVector& WorldStart, const FVector& WorldEnd, bool bTraceComplex, bool bExtractPhysMaterial)
 {
+	using namespace ChaosInterface;
+
 	// Need an instance to trace against
 	check(InInstance);
 
@@ -1077,6 +1081,8 @@ bool FPhysInterface_Chaos::LineTrace_Geom(FHitResult& OutHit, const FBodyInstanc
 
 bool FPhysInterface_Chaos::Sweep_Geom(FHitResult& OutHit, const FBodyInstance* InInstance, const FVector& InStart, const FVector& InEnd, const FQuat& InShapeRotation, const FCollisionShape& InShape, bool bSweepComplex)
 {
+	using namespace ChaosInterface;
+
 	bool bSweepHit = false;
 
 	if (InShape.IsNearlyZero())
diff --git a/Engine/Source/Runtime/Engine/Public/Physics/Experimental/ChaosInterfaceWrapper.h b/Engine/Source/Runtime/Engine/Public/Physics/Experimental/ChaosInterfaceWrapper.h
--- a/Engine/Source/Runtime/Engine/Public/Physics/Experimental/ChaosInterfaceWrapper.h
+++ b/Engine/Source/Runtime/Engine/Public/Physics/Experimental/ChaosInterfaceWrapper.h
@@ -41,11 +41,11 @@ UPhysicalMaterial* GetUserData(const Chaos::FChaosPhysicsMaterial& Material);
 
 }
 
-void LowLevelRaycast(FPhysScene& Scene, const FVector& Start, const FVector& Dir, float DeltaMag, FPhysicsHitCallback<FHitRaycast>& HitBuffer, EHitFlags OutputFlags, FQueryFlags QueryFlags, const FCollisionFilterData& Filter, const FQueryFilterData& QueryFilterData, ICollisionQueryFilterCallbackBase* QueryCallback, const FQueryDebugParams& DebugParams = FQueryDebugParams());
-void LowLevelRaycast(FPhysScene& Scene, const FVector& Start, const FVector& Dir, float DeltaMag, FPhysicsHitCallback<FPTRaycastHit>& HitBuffer, EHitFlags OutputFlags, FQueryFlags QueryFlags, const FCollisionFilterData& Filter, const FQueryFilterData& QueryFilterData, ICollisionQueryFilterCallbackBase* QueryCallback, const FQueryDebugParams& DebugParams = FQueryDebugParams());
+void LowLevelRaycast(FPhysScene& Scene, const FVector& Start, const FVector& Dir, float DeltaMag, FPhysicsHitCallback<FHitRaycast>& HitBuffer, EHitFlags OutputFlags, FQueryFlags QueryFlags, const FCollisionFilterData& Filter, const ChaosInterface::FQueryFilterData& QueryFilterData, ICollisionQueryFilterCallbackBase* QueryCallback, const ChaosInterface::FQueryDebugParams& DebugParams = {});
+void LowLevelRaycast(FPhysScene& Scene, const FVector& Start, const FVector& Dir, float DeltaMag, FPhysicsHitCallback<ChaosInterface::FPTRaycastHit>& HitBuffer, EHitFlags OutputFlags, FQueryFlags QueryFlags, const FCollisionFilterData& Filter, const ChaosInterface::FQueryFilterData& QueryFilterData, ICollisionQueryFilterCallbackBase* QueryCallback, const ChaosInterface::FQueryDebugParams& DebugParams = {});
 
-void LowLevelSweep(FPhysScene& Scene, const FPhysicsGeometry& Geom, const FTransform& StartTM, const FVector& Dir, float DeltaMag, FPhysicsHitCallback<FHitSweep>& HitBuffer, EHitFlags OutputFlags, FQueryFlags QueryFlags, const FCollisionFilterData& Filter, const FQueryFilterData& QueryFilterData, ICollisionQueryFilterCallbackBase* QueryCallback, const FQueryDebugParams& DebugParams = FQueryDebugParams());
-void LowLevelSweep(FPhysScene& Scene, const FPhysicsGeometry& Geom, const FTransform& StartTM, const FVector& Dir, float DeltaMag, FPhysicsHitCallback<FPTSweepHit>& HitBuffer, EHitFlags OutputFlags, FQueryFlags QueryFlags, const FCollisionFilterData& Filter, const FQueryFilterData& QueryFilterData, ICollisionQueryFilterCallbackBase* QueryCallback, const FQueryDebugParams& DebugParams = FQueryDebugParams());
+void LowLevelSweep(FPhysScene& Scene, const FPhysicsGeometry& Geom, const FTransform& StartTM, const FVector& Dir, float DeltaMag, FPhysicsHitCallback<FHitSweep>& HitBuffer, EHitFlags OutputFlags, FQueryFlags QueryFlags, const FCollisionFilterData& Filter, const ChaosInterface::FQueryFilterData& QueryFilterData, ICollisionQueryFilterCallbackBase* QueryCallback, const ChaosInterface::FQueryDebugParams& DebugParams = {});
+void LowLevelSweep(FPhysScene& Scene, const FPhysicsGeometry& Geom, const FTransform& StartTM, const FVector& Dir, float DeltaMag, FPhysicsHitCallback<ChaosInterface::FPTSweepHit>& HitBuffer, EHitFlags OutputFlags, FQueryFlags QueryFlags, const FCollisionFilterData& Filter, const ChaosInterface::FQueryFilterData& QueryFilterData, ICollisionQueryFilterCallbackBase* QueryCallback, const ChaosInterface::FQueryDebugParams& DebugParams = {});
 
-void LowLevelOverlap(FPhysScene& Scene, const FPhysicsGeometry& Geom, const FTransform& GeomPose, FPhysicsHitCallback<FHitOverlap>& HitBuffer, FQueryFlags QueryFlags, const FCollisionFilterData& Filter, const FQueryFilterData& QueryFilterData, ICollisionQueryFilterCallbackBase* QueryCallback, const FQueryDebugParams& DebugParams = FQueryDebugParams());
-void LowLevelOverlap(FPhysScene& Scene, const FPhysicsGeometry& Geom, const FTransform& GeomPose, FPhysicsHitCallback<FPTOverlapHit>& HitBuffer, FQueryFlags QueryFlags, const FCollisionFilterData& Filter, const FQueryFilterData& QueryFilterData, ICollisionQueryFilterCallbackBase* QueryCallback, const FQueryDebugParams& DebugParams = FQueryDebugParams());
+void LowLevelOverlap(FPhysScene& Scene, const FPhysicsGeometry& Geom, const FTransform& GeomPose, FPhysicsHitCallback<FHitOverlap>& HitBuffer, FQueryFlags QueryFlags, const FCollisionFilterData& Filter, const ChaosInterface::FQueryFilterData& QueryFilterData, ICollisionQueryFilterCallbackBase* QueryCallback, const ChaosInterface::FQueryDebugParams& DebugParams = {});
+void LowLevelOverlap(FPhysScene& Scene, const FPhysicsGeometry& Geom, const FTransform& GeomPose, FPhysicsHitCallback<ChaosInterface::FPTOverlapHit>& HitBuffer, FQueryFlags QueryFlags, const FCollisionFilterData& Filter, const ChaosInterface::FQueryFilterData& QueryFilterData, ICollisionQueryFilterCallbackBase* QueryCallback, const ChaosInterface::FQueryDebugParams& DebugParams = {});
diff --git a/Engine/Source/Runtime/Engine/Public/Physics/Experimental/PhysInterface_Chaos.h b/Engine/Source/Runtime/Engine/Public/Physics/Experimental/PhysInterface_Chaos.h
--- a/Engine/Source/Runtime/Engine/Public/Physics/Experimental/PhysInterface_Chaos.h
+++ b/Engine/Source/Runtime/Engine/Public/Physics/Experimental/PhysInterface_Chaos.h
@@ -138,22 +138,22 @@ FORCEINLINE float GetHalfHeight(const Chaos::FCapsule& Capsule)
 }
 */
 
-FORCEINLINE void DrawOverlappingTris(const UWorld* World, const FLocationHit& Hit, const Chaos::FImplicitObject& Geom, const FTransform& QueryTM)
+FORCEINLINE void DrawOverlappingTris(const UWorld* World, const ChaosInterface::FLocationHit& Hit, const Chaos::FImplicitObject& Geom, const FTransform& QueryTM)
 {
 	//TODO_SQ_IMPLEMENTATION
 }
 
-FORCEINLINE void ComputeZeroDistanceImpactNormalAndPenetration(const UWorld* World, const FLocationHit& Hit, const Chaos::FImplicitObject& Geom, const FTransform& QueryTM, FHitResult& OutResult)
+FORCEINLINE void ComputeZeroDistanceImpactNormalAndPenetration(const UWorld* World, const ChaosInterface::FLocationHit& Hit, const Chaos::FImplicitObject& Geom, const FTransform& QueryTM, FHitResult& OutResult)
 {
 	//TODO_SQ_IMPLEMENTATION
 }
 
-FORCEINLINE void DrawOverlappingTris(const UWorld* World, const FPTLocationHit& Hit, const Chaos::FImplicitObject& Geom, const FTransform& QueryTM)
+FORCEINLINE void DrawOverlappingTris(const UWorld* World, const ChaosInterface::FPTLocationHit& Hit, const Chaos::FImplicitObject& Geom, const FTransform& QueryTM)
 {
 	//TODO_SQ_IMPLEMENTATION
 }
 
-FORCEINLINE void ComputeZeroDistanceImpactNormalAndPenetration(const UWorld* World, const FPTLocationHit& Hit, const Chaos::FImplicitObject& Geom, const FTransform& QueryTM, FHitResult& OutResult)
+FORCEINLINE void ComputeZeroDistanceImpactNormalAndPenetration(const UWorld* World, const ChaosInterface::FPTLocationHit& Hit, const Chaos::FImplicitObject& Geom, const FTransform& QueryTM, FHitResult& OutResult)
 {
 	//TODO_SQ_IMPLEMENTATION
 }
diff --git a/Engine/Source/Runtime/Experimental/Chaos/Private/Chaos/BVHParticles.cpp b/Engine/Source/Runtime/Experimental/Chaos/Private/Chaos/BVHParticles.cpp
--- a/Engine/Source/Runtime/Experimental/Chaos/Private/Chaos/BVHParticles.cpp
+++ b/Engine/Source/Runtime/Experimental/Chaos/Private/Chaos/BVHParticles.cpp
@@ -3,67 +3,65 @@
 #include "Chaos/BVHParticles.h"
 #include "Chaos/BoundingVolumeHierarchy.h"
 
-using namespace Chaos;
-
-FBVHParticles::FBVHParticles()
-	: FParticles()
-	, MBVH(new TBoundingVolumeHierarchy<FParticles, TArray<int32>>(*this, CollisionParticlesBVHDepth))
+namespace Chaos
 {
-}
+	FBVHParticles::FBVHParticles()
+		: FParticles()
+		, MBVH(new TBoundingVolumeHierarchy<FParticles, TArray<int32>>(*this, CollisionParticlesBVHDepth))
+	{}
 
-FBVHParticles::FBVHParticles(FBVHParticles&& Other)
-	: FParticles(MoveTemp(Other))
-	, MBVH(new TBoundingVolumeHierarchy<FParticles, TArray<int32>>(MoveTemp(*Other.MBVH)))
-{
-}
+	FBVHParticles::FBVHParticles(FBVHParticles&& Other)
+		: FParticles(MoveTemp(Other))
+		, MBVH(new TBoundingVolumeHierarchy<FParticles, TArray<int32>>(MoveTemp(*Other.MBVH)))
+	{}
 
-FBVHParticles::FBVHParticles(FParticles&& Other)
-	: FParticles(MoveTemp(Other))
-	, MBVH(new TBoundingVolumeHierarchy<FParticles, TArray<int32>>(*this, CollisionParticlesBVHDepth))
-{
-}
+	FBVHParticles::FBVHParticles(FParticles&& Other)
+		: FParticles(MoveTemp(Other))
+		, MBVH(new TBoundingVolumeHierarchy<FParticles, TArray<int32>>(*this, CollisionParticlesBVHDepth))
+	{}
 
-FBVHParticles::~FBVHParticles()
-{
-    delete MBVH;
-}
+	FBVHParticles::~FBVHParticles()
+	{
+		delete MBVH;
+	}
 
-FBVHParticles& FBVHParticles::operator=(const FBVHParticles& Other)
-{
-	*this = FBVHParticles(Other);
-	return *this;
-}
+	FBVHParticles& FBVHParticles::operator=(const FBVHParticles& Other)
+	{
+		*this = FBVHParticles(Other);
+		return *this;
+	}
 
-FBVHParticles& FBVHParticles::operator=(FBVHParticles&& Other)
-{
-	*MBVH = MoveTemp(*Other.MBVH);
-	FParticles::operator=(static_cast<FParticles&&>(Other));
-	return *this;
-}
+	FBVHParticles& FBVHParticles::operator=(FBVHParticles&& Other)
+	{
+		*MBVH = MoveTemp(*Other.MBVH);
+		FParticles::operator=(static_cast<FParticles&&>(Other));
+		return *this;
+	}
 
-FBVHParticles::FBVHParticles(const FBVHParticles& Other)
-	: FParticles()
-{
-	AddParticles(Other.Size());
-	for (int32 i = Other.Size() - 1; 0 <= i; i--)
+	FBVHParticles::FBVHParticles(const FBVHParticles& Other)
+		: FParticles()
 	{
-		X(i) = Other.X(i);
+		AddParticles(Other.Size());
+		for(int32 i = Other.Size() - 1; 0 <= i; i--)
+		{
+			X(i) = Other.X(i);
+		}
+		MBVH = new TBoundingVolumeHierarchy<FParticles, TArray<int32>>(*this, CollisionParticlesBVHDepth);
 	}
-	MBVH = new TBoundingVolumeHierarchy<FParticles, TArray<int32>>(*this, CollisionParticlesBVHDepth);
-}
 
-void FBVHParticles::UpdateAccelerationStructures()
-{
-	MBVH->UpdateHierarchy();
-}
+	void FBVHParticles::UpdateAccelerationStructures()
+	{
+		MBVH->UpdateHierarchy();
+	}
 
-const TArray<int32> FBVHParticles::FindAllIntersections(const FAABB3& Object) const
-{
-	return MBVH->FindAllIntersections(Object);
-}
+	const TArray<int32> FBVHParticles::FindAllIntersections(const FAABB3& Object) const
+	{
+		return MBVH->FindAllIntersections(Object);
+	}
 
-void FBVHParticles::Serialize(FChaosArchive& Ar)
-{
-	FParticles::Serialize(Ar);
-	Ar << *MBVH;
+	void FBVHParticles::Serialize(FChaosArchive& Ar)
+	{
+		FParticles::Serialize(Ar);
+		Ar << *MBVH;
+	}
 }
diff --git a/Engine/Source/Runtime/Experimental/Chaos/Private/Chaos/BoundingVolumeHierarchy.cpp b/Engine/Source/Runtime/Experimental/Chaos/Private/Chaos/BoundingVolumeHierarchy.cpp
--- a/Engine/Source/Runtime/Experimental/Chaos/Private/Chaos/BoundingVolumeHierarchy.cpp
+++ b/Engine/Source/Runtime/Experimental/Chaos/Private/Chaos/BoundingVolumeHierarchy.cpp
@@ -5,8 +5,6 @@
 #include "Chaos/BoundingVolumeUtilities.h"
 #include "ChaosStats.h"
 
-using namespace Chaos;
-
 namespace Chaos
 {
 	// Bounds expansion (for speculative contact creation) is limited by this factor multiplied by the objects size.
@@ -17,679 +15,678 @@ namespace Chaos
 	// NOTE: This is only used by objects that have CCD disabled.
 	FRealSingle Chaos_Bounds_MaxInflationScale = 2.f;
 	FAutoConsoleVariableRef CVarBoundsMaxInflatiuonScale(TEXT("p.Chaos.MaxInflationScale"), Chaos_Bounds_MaxInflationScale, TEXT("A limit on the bounds used to detect collisions when CCD is disabled. The bounds limit is this scale multiplied by the object's max dimension"));
-}
 
-template<class OBJECT_ARRAY, class LEAF_TYPE, class T, int d>
-TBoundingVolumeHierarchy<OBJECT_ARRAY, LEAF_TYPE, T, d>::TBoundingVolumeHierarchy(const OBJECT_ARRAY& Objects, const int32 MaxLevels, const bool bUseVelocity, const T Dt)
-    : MObjects(&Objects), MMaxLevels(MaxLevels)
-{
-	if (GetObjectCount(Objects) > 0)
+
+	template<class OBJECT_ARRAY, class LEAF_TYPE, class T, int d>
+	TBoundingVolumeHierarchy<OBJECT_ARRAY, LEAF_TYPE, T, d>::TBoundingVolumeHierarchy(const OBJECT_ARRAY& Objects, const int32 MaxLevels, const bool bUseVelocity, const T Dt)
+		: MObjects(&Objects), MMaxLevels(MaxLevels)
 	{
-		UpdateHierarchy(DefaultAllowMultipleSplitting, bUseVelocity, DefaultDt);
+		if(GetObjectCount(Objects) > 0)
+		{
+			UpdateHierarchy(DefaultAllowMultipleSplitting, bUseVelocity, DefaultDt);
+		}
 	}
-}
 
-template<class OBJECT_ARRAY, class LEAF_TYPE, class T, int d>
-TBoundingVolumeHierarchy<OBJECT_ARRAY, LEAF_TYPE, T, d>::TBoundingVolumeHierarchy(const OBJECT_ARRAY& Objects, const TArray<uint32>& ActiveIndices, const int32 MaxLevels, const bool bUseVelocity, const T Dt)
-	: MObjects(&Objects), MMaxLevels(MaxLevels)
-{
-	if (GetObjectCount(Objects) > 0)
+	template<class OBJECT_ARRAY, class LEAF_TYPE, class T, int d>
+	TBoundingVolumeHierarchy<OBJECT_ARRAY, LEAF_TYPE, T, d>::TBoundingVolumeHierarchy(const OBJECT_ARRAY& Objects, const TArray<uint32>& ActiveIndices, const int32 MaxLevels, const bool bUseVelocity, const T Dt)
+		: MObjects(&Objects), MMaxLevels(MaxLevels)
 	{
-		UpdateHierarchy(ActiveIndices, DefaultAllowMultipleSplitting, bUseVelocity, DefaultDt);
+		if(GetObjectCount(Objects) > 0)
+		{
+			UpdateHierarchy(ActiveIndices, DefaultAllowMultipleSplitting, bUseVelocity, DefaultDt);
+		}
 	}
-}
 
-template<class OBJECT_ARRAY, class LEAF_TYPE, class T, int d>
-const TAABB<T, d>& TBoundingVolumeHierarchy<OBJECT_ARRAY, LEAF_TYPE, T, d>::GetWorldSpaceBoundingBox(const TGeometryParticles<T, d>& InParticles, const int32 Index)
-{
-	return Chaos::GetWorldSpaceBoundingBox(InParticles, Index, MWorldSpaceBoxes);
-}
+	template<class OBJECT_ARRAY, class LEAF_TYPE, class T, int d>
+	const TAABB<T, d>& TBoundingVolumeHierarchy<OBJECT_ARRAY, LEAF_TYPE, T, d>::GetWorldSpaceBoundingBox(const TGeometryParticles<T, d>& InParticles, const int32 Index)
+	{
+		return Chaos::GetWorldSpaceBoundingBox(InParticles, Index, MWorldSpaceBoxes);
+	}
 
 #if !UE_BUILD_SHIPPING
-template<typename T, int d>
-void DrawNodeRecursive(ISpacialDebugDrawInterface<T>* InInterface, const TBVHNode<T, d>& InNode, const TArray<TBVHNode<T,d>>& InAllNodes)
-{
-	TAABB<T, d> Box(InNode.MMin, InNode.MMax);
+	template<typename T, int d>
+	void DrawNodeRecursive(ISpacialDebugDrawInterface<T>* InInterface, const TBVHNode<T, d>& InNode, const TArray<TBVHNode<T, d>>& InAllNodes)
+	{
+		TAABB<T, d> Box(InNode.MMin, InNode.MMax);
 
-	InInterface->Box(Box, {1.0f, 1.0f, 1.0f}, 1.0f);
+		InInterface->Box(Box, { 1.0f, 1.0f, 1.0f }, 1.0f);
 
-	for(const int32 Child : InNode.MChildren)
-	{
-		DrawNodeRecursive(InInterface, InAllNodes[Child], InAllNodes);
+		for(const int32 Child : InNode.MChildren)
+		{
+			DrawNodeRecursive(InInterface, InAllNodes[Child], InAllNodes);
+		}
 	}
-}
 
-template<class OBJECT_ARRAY, class LEAF_TYPE, class T, int d>
-void Chaos::TBoundingVolumeHierarchy<OBJECT_ARRAY, LEAF_TYPE, T, d>::DebugDraw(ISpacialDebugDrawInterface<T>* InInterface) const
-{
-	if(Elements.Num() > 0)
+	template<class OBJECT_ARRAY, class LEAF_TYPE, class T, int d>
+	void Chaos::TBoundingVolumeHierarchy<OBJECT_ARRAY, LEAF_TYPE, T, d>::DebugDraw(ISpacialDebugDrawInterface<T>* InInterface) const
 	{
-		DrawNodeRecursive(InInterface, Elements[0], Elements);
+		if(Elements.Num() > 0)
+		{
+			DrawNodeRecursive(InInterface, Elements[0], Elements);
+		}
 	}
-}
 #endif
 
-template<class OBJECT_ARRAY, class T, int d, typename TPayloadType>
-TArray<int32> MakeNewLeaf(const OBJECT_ARRAY& Objects, const TArray<int32>& AllObjects, const TArray<int32>* /*Type*/)
-{
-	return TArray<int32>(AllObjects);
-}
+	template<class OBJECT_ARRAY, class T, int d, typename TPayloadType>
+	TArray<int32> MakeNewLeaf(const OBJECT_ARRAY& Objects, const TArray<int32>& AllObjects, const TArray<int32>* /*Type*/)
+	{
+		return TArray<int32>(AllObjects);
+	}
 
-template<class OBJECT_ARRAY, class T, int d, typename TPayloadType>
-TBoundingVolume<TPayloadType, T, d> MakeNewLeaf(const OBJECT_ARRAY& Objects, const TArray<int32>& AllObjects, const TBoundingVolume<TPayloadType, T, d>* /*Type*/)
-{
+	template<class OBJECT_ARRAY, class T, int d, typename TPayloadType>
+	TBoundingVolume<TPayloadType, T, d> MakeNewLeaf(const OBJECT_ARRAY& Objects, const TArray<int32>& AllObjects, const TBoundingVolume<TPayloadType, T, d>* /*Type*/)
+	{
 #if CHAOS_PARTICLEHANDLE_TODO
-	// @todo(mlentine): This is pretty dirty but should work for now.
-	return TBoundingVolume<OBJECT_ARRAY, TPayloadType, T, d>(Objects, reinterpret_cast<const TArray<uint32>&>(AllObjects));
+		// @todo(mlentine): This is pretty dirty but should work for now.
+		return TBoundingVolume<OBJECT_ARRAY, TPayloadType, T, d>(Objects, reinterpret_cast<const TArray<uint32>&>(AllObjects));
 #else
-	TArray<TSOAView<OBJECT_ARRAY>> SOAs;
-	SOAs.Add(TSOAView<OBJECT_ARRAY>(const_cast<OBJECT_ARRAY*>(&Objects)));	//todo: this const_cast is here because bvh is holding on to things as const but giving it out as non const. Not sure if we should change API or not
-	//return TBoundingVolume<OBJECT_ARRAY, T, d>(MakeParticleView<OBJECT_ARRAY>({ {&Objects} }));
-	return TBoundingVolume<TPayloadType, T, d>(MakeParticleView<OBJECT_ARRAY>(MoveTemp(SOAs)));
+		TArray<TSOAView<OBJECT_ARRAY>> SOAs;
+		SOAs.Add(TSOAView<OBJECT_ARRAY>(const_cast<OBJECT_ARRAY*>(&Objects)));	//todo: this const_cast is here because bvh is holding on to things as const but giving it out as non const. Not sure if we should change API or not
+		//return TBoundingVolume<OBJECT_ARRAY, T, d>(MakeParticleView<OBJECT_ARRAY>({ {&Objects} }));
+		return TBoundingVolume<TPayloadType, T, d>(MakeParticleView<OBJECT_ARRAY>(MoveTemp(SOAs)));
 #endif
-}
-
-template<class OBJECT_ARRAY, class LEAF_TYPE, class T, int d>
-void TBoundingVolumeHierarchy<OBJECT_ARRAY, LEAF_TYPE, T, d>::UpdateHierarchy(const bool bAllowMultipleSplitting, const bool bUseVelocity, const T Dt)
-{
-	const int32 NumObjects = GetObjectCount(*MObjects);
-	
-	MScratchAllObjects.Reset();
-	Leafs.Reset();
+	}
 
-	MGlobalObjects.Reset();
-	for (int32 Idx = 0; Idx < NumObjects; ++Idx)	//todo: check for disabled objects if possible? (this is what BV does)
+	template<class OBJECT_ARRAY, class LEAF_TYPE, class T, int d>
+	void TBoundingVolumeHierarchy<OBJECT_ARRAY, LEAF_TYPE, T, d>::UpdateHierarchy(const bool bAllowMultipleSplitting, const bool bUseVelocity, const T Dt)
 	{
-		if (IsDisabled(*MObjects, Idx))
-		{
-			continue;
-		}
+		const int32 NumObjects = GetObjectCount(*MObjects);
 
-		if (HasBoundingBox(*MObjects, Idx))
-		{
-			MScratchAllObjects.Add(Idx);
-		}
-		else
+		MScratchAllObjects.Reset();
+		Leafs.Reset();
+
+		MGlobalObjects.Reset();
+		for(int32 Idx = 0; Idx < NumObjects; ++Idx)	//todo: check for disabled objects if possible? (this is what BV does)
 		{
-			MGlobalObjects.Add(Idx);
+			if(IsDisabled(*MObjects, Idx))
+			{
+				continue;
+			}
+
+			if(HasBoundingBox(*MObjects, Idx))
+			{
+				MScratchAllObjects.Add(Idx);
+			}
+			else
+			{
+				MGlobalObjects.Add(Idx);
+			}
 		}
-	}
 
-	UpdateHierarchyImp(MScratchAllObjects, bAllowMultipleSplitting, bUseVelocity, Dt);
-}
+		UpdateHierarchyImp(MScratchAllObjects, bAllowMultipleSplitting, bUseVelocity, Dt);
+	}
 
-template<class OBJECT_ARRAY, class LEAF_TYPE, class T, int d>
-void TBoundingVolumeHierarchy<OBJECT_ARRAY, LEAF_TYPE, T, d>::UpdateHierarchy(const TArray<uint32>& ActiveIndices, const bool bAllowMultipleSplitting, const bool bUseVelocity, const T Dt)
-{
-	const int32 NumObjects = ActiveIndices.Num();
+	template<class OBJECT_ARRAY, class LEAF_TYPE, class T, int d>
+	void TBoundingVolumeHierarchy<OBJECT_ARRAY, LEAF_TYPE, T, d>::UpdateHierarchy(const TArray<uint32>& ActiveIndices, const bool bAllowMultipleSplitting, const bool bUseVelocity, const T Dt)
+	{
+		const int32 NumObjects = ActiveIndices.Num();
 
-	MScratchAllObjects.Reset();
-	MScratchAllObjects.Reserve(NumObjects);
+		MScratchAllObjects.Reset();
+		MScratchAllObjects.Reserve(NumObjects);
 
-	MGlobalObjects.Reset();
-	for(uint32 Idx : ActiveIndices)
-	{
-		check(!IsDisabled(*MObjects, Idx));
-		if (HasBoundingBox(*MObjects, Idx))
-		{
-			MScratchAllObjects.Add(Idx);
-		}
-		else
+		MGlobalObjects.Reset();
+		for(uint32 Idx : ActiveIndices)
 		{
-			MGlobalObjects.Add(Idx);
+			check(!IsDisabled(*MObjects, Idx));
+			if(HasBoundingBox(*MObjects, Idx))
+			{
+				MScratchAllObjects.Add(Idx);
+			}
+			else
+			{
+				MGlobalObjects.Add(Idx);
+			}
 		}
-	}
 
-	UpdateHierarchyImp(MScratchAllObjects, bAllowMultipleSplitting, bUseVelocity, Dt);
-}
-
-template<class OBJECT_ARRAY, class LEAF_TYPE, class T, int d>
-void TBoundingVolumeHierarchy<OBJECT_ARRAY, LEAF_TYPE, T, d>::UpdateHierarchyImp(const TArray<int32>& AllObjects, const bool bAllowMultipleSplitting, const bool bUseVelocity, const T Dt)
-{
-	Elements.Reset();
-	MWorldSpaceBoxes.Reset();
+		UpdateHierarchyImp(MScratchAllObjects, bAllowMultipleSplitting, bUseVelocity, Dt);
+	}
 
-	if (!AllObjects.Num())
+	template<class OBJECT_ARRAY, class LEAF_TYPE, class T, int d>
+	void TBoundingVolumeHierarchy<OBJECT_ARRAY, LEAF_TYPE, T, d>::UpdateHierarchyImp(const TArray<int32>& AllObjects, const bool bAllowMultipleSplitting, const bool bUseVelocity, const T Dt)
 	{
-		return;
-	}
-	ComputeAllWorldSpaceBoundingBoxes(*MObjects, AllObjects, bUseVelocity, Dt, MWorldSpaceBoxes);
+		Elements.Reset();
+		MWorldSpaceBoxes.Reset();
 
-	int32 Axis;
-	const TAABB<T, d> GlobalBox = ComputeGlobalBoxAndSplitAxis(*MObjects, AllObjects, MWorldSpaceBoxes, bAllowMultipleSplitting, Axis);
+		if(!AllObjects.Num())
+		{
+			return;
+		}
+		ComputeAllWorldSpaceBoundingBoxes(*MObjects, AllObjects, bUseVelocity, Dt, MWorldSpaceBoxes);
 
-	TBVHNode<T,d> RootNode;
-	RootNode.MMin = GlobalBox.Min();
-	RootNode.MMax = GlobalBox.Max();
-	RootNode.MAxis = Axis;
+		int32 Axis;
+		const TAABB<T, d> GlobalBox = ComputeGlobalBoxAndSplitAxis(*MObjects, AllObjects, MWorldSpaceBoxes, bAllowMultipleSplitting, Axis);
+
+		TBVHNode<T, d> RootNode;
+		RootNode.MMin = GlobalBox.Min();
+		RootNode.MMax = GlobalBox.Max();
+		RootNode.MAxis = Axis;
 #if !UE_BUILD_SHIPPING && !UE_BUILD_TEST
-	RootNode.LeafIndex = -1;
+		RootNode.LeafIndex = -1;
 #endif
-	Elements.Add(RootNode);
-	if (AllObjects.Num() > MIN_NUM_OBJECTS) // TODO(mlentine): What is a good number to stop at?
-	{
-		int32 StartIndex = GenerateNextLevel(GlobalBox.Min(), GlobalBox.Max(), AllObjects, Axis, 1, bAllowMultipleSplitting);
-		for (int32 i = 0; i < (Axis == -1 ? 8 : 2); i++)
+		Elements.Add(RootNode);
+		if(AllObjects.Num() > MIN_NUM_OBJECTS) // TODO(mlentine): What is a good number to stop at?
 		{
-			Elements[0].MChildren.Add(StartIndex + i);
+			int32 StartIndex = GenerateNextLevel(GlobalBox.Min(), GlobalBox.Max(), AllObjects, Axis, 1, bAllowMultipleSplitting);
+			for(int32 i = 0; i < (Axis == -1 ? 8 : 2); i++)
+			{
+				Elements[0].MChildren.Add(StartIndex + i);
+			}
 		}
+		else
+		{
+			Elements[0].LeafIndex = Leafs.Num();
+			Leafs.Add(MakeNewLeaf<OBJECT_ARRAY, T, d, TPayloadType>(*MObjects, AllObjects, static_cast<LEAF_TYPE*>(nullptr)));
+		}
+		check(Elements[0].LeafIndex < Leafs.Num() || Elements[0].MChildren.Num() > 0);
+		UE_LOG(LogChaos, Verbose, TEXT("Generated Tree with %d Nodes"), Elements.Num());
+		//PrintTree("", &Elements[0]);
 	}
-	else
+
+	template<class OBJECT_ARRAY, typename TPayloadType, class T, int d>
+	TArray<int32> FindAllIntersectionsLeafHelper(const TArray<int32>& Leaf, const TVector<T, d>& Point)
 	{
-		Elements[0].LeafIndex = Leafs.Num();
-		Leafs.Add(MakeNewLeaf<OBJECT_ARRAY, T, d, TPayloadType>(*MObjects, AllObjects, static_cast<LEAF_TYPE*>(nullptr)));
+		return Leaf;
 	}
-	check(Elements[0].LeafIndex < Leafs.Num() || Elements[0].MChildren.Num() > 0);
-	UE_LOG(LogChaos, Verbose, TEXT("Generated Tree with %d Nodes"), Elements.Num());
-	//PrintTree("", &Elements[0]);
-}
-
-template<class OBJECT_ARRAY, typename TPayloadType, class T, int d>
-TArray<int32> FindAllIntersectionsLeafHelper(const TArray<int32>& Leaf, const TVector<T, d>& Point)
-{
-	return Leaf;
-}
 
-template<class OBJECT_ARRAY,typename TPayloadType, class T, int d>
-TArray<int32> FindAllIntersectionsLeafHelper(const TBoundingVolume<TPayloadType, T, d>& Leaf, const TVector<T, d>& Point)
-{
+	template<class OBJECT_ARRAY, typename TPayloadType, class T, int d>
+	TArray<int32> FindAllIntersectionsLeafHelper(const TBoundingVolume<TPayloadType, T, d>& Leaf, const TVector<T, d>& Point)
+	{
 #if CHAOS_PARTICLEHANDLE_TODO
-	return Leaf.FindAllIntersections(Point);
+		return Leaf.FindAllIntersections(Point);
 #else
-	check(false);
-	return TArray<int32>();
-#endif
-}
-
-template<class OBJECT_ARRAY, typename TPayloadType, class T, int d>
-TArray<int32> TBoundingVolumeHierarchy<OBJECT_ARRAY, TPayloadType, T, d>::FindAllIntersectionsHelper(const TBVHNode<T,d>& MyNode, const TVector<T, d>& Point) const
-{
-	TAABB<T, d> MBox(MyNode.MMin, MyNode.MMax);
-	if (MBox.SignedDistance(Point) > 0)
-	{
+		check(false);
 		return TArray<int32>();
+#endif
 	}
-	if (MyNode.MChildren.Num() == 0)
+
+	template<class OBJECT_ARRAY, typename TPayloadType, class T, int d>
+	TArray<int32> TBoundingVolumeHierarchy<OBJECT_ARRAY, TPayloadType, T, d>::FindAllIntersectionsHelper(const TBVHNode<T, d>& MyNode, const TVector<T, d>& Point) const
 	{
-		return FindAllIntersectionsLeafHelper<OBJECT_ARRAY, TPayloadType, T, d>(Leafs[MyNode.LeafIndex], Point);
+		TAABB<T, d> MBox(MyNode.MMin, MyNode.MMax);
+		if(MBox.SignedDistance(Point) > 0)
+		{
+			return TArray<int32>();
+		}
+		if(MyNode.MChildren.Num() == 0)
+		{
+			return FindAllIntersectionsLeafHelper<OBJECT_ARRAY, TPayloadType, T, d>(Leafs[MyNode.LeafIndex], Point);
+		}
+		const TVector<T, d> ObjectCenter = MBox.Center();
+		int32 Child = 0;
+		if(MyNode.MAxis >= 0)
+		{
+			if(Point[MyNode.MAxis] > ObjectCenter[MyNode.MAxis])
+				Child += 1;
+		}
+		else
+		{
+			if(Point[0] > ObjectCenter[0])
+				Child += 1;
+			if(Point[1] > ObjectCenter[1])
+				Child += 2;
+			if(Point[2] > ObjectCenter[2])
+				Child += 4;
+		}
+		return FindAllIntersectionsHelper(Elements[MyNode.MChildren[Child]], Point);
 	}
-	const TVector<T, d> ObjectCenter = MBox.Center();
-	int32 Child = 0;
-	if (MyNode.MAxis >= 0)
+
+	template <typename T, int d>
+	bool IntersectsHelper(const TAABB<T, d>& WorldSpaceBox, const TAABB<T, d>& LocalBox)
 	{
-		if (Point[MyNode.MAxis] > ObjectCenter[MyNode.MAxis])
-			Child += 1;
+		return WorldSpaceBox.Intersects(LocalBox);
 	}
-	else
+
+	template <typename T, int d>
+	bool IntersectsHelper(const TAABB<T, d>& WorldSpaceBox, const TSpatialRay<T, d>& Ray)
 	{
-		if (Point[0] > ObjectCenter[0])
-			Child += 1;
-		if (Point[1] > ObjectCenter[1])
-			Child += 2;
-		if (Point[2] > ObjectCenter[2])
-			Child += 4;
+		//Very slow, should look at using slab algorithm
+		const FBox VolumeBox(WorldSpaceBox.Min(), WorldSpaceBox.Max());
+		FVector Hit;
+		FVector Normal;
+		FRealSingle Time;
+		return FMath::LineExtentBoxIntersection(VolumeBox, Ray.Start, Ray.End, FVector::ZeroVector, Hit, Normal, Time);
+
+		//return WorldSpaceBox.FindClosestIntersection(Ray.Start, Ray.End, 0).Second;
 	}
-	return FindAllIntersectionsHelper(Elements[MyNode.MChildren[Child]], Point);
-}
-
-template <typename T, int d>
-bool IntersectsHelper(const TAABB<T, d>& WorldSpaceBox, const TAABB<T, d>& LocalBox)
-{
-	return WorldSpaceBox.Intersects(LocalBox);
-}
-
-template <typename T, int d>
-bool IntersectsHelper(const TAABB<T, d>& WorldSpaceBox, const TSpatialRay<T, d>& Ray)
-{
-	//Very slow, should look at using slab algorithm
-	const FBox VolumeBox(WorldSpaceBox.Min(), WorldSpaceBox.Max());
-	FVector Hit;
-	FVector Normal;
-	FRealSingle Time;
-	return FMath::LineExtentBoxIntersection(VolumeBox, Ray.Start, Ray.End, FVector::ZeroVector, Hit, Normal, Time);
-
-	//return WorldSpaceBox.FindClosestIntersection(Ray.Start, Ray.End, 0).Second;
-}
 
-template <typename OBJECT_ARRAY>
-struct TSpecializeParticlesHelper
-{
-	template <typename T, int d, typename QUERY_OBJECT>
-	static void AccumulateChildrenResults(TArray<int32>& AccumIntersectionList, const TArray<int32>& PotentialChildren, const QUERY_OBJECT& QueryObject, const TMap<int32, TAABB<T,d>>& WorldSpaceBoxes)
+	template <typename OBJECT_ARRAY>
+	struct TSpecializeParticlesHelper
 	{
-		for (int32 ChildIndex : PotentialChildren)
+		template <typename T, int d, typename QUERY_OBJECT>
+		static void AccumulateChildrenResults(TArray<int32>& AccumIntersectionList, const TArray<int32>& PotentialChildren, const QUERY_OBJECT& QueryObject, const TMap<int32, TAABB<T, d>>& WorldSpaceBoxes)
 		{
-			if (IntersectsHelper(WorldSpaceBoxes[ChildIndex], QueryObject)) //todo(ocohen): actually just a single point so should call Contains
+			for(int32 ChildIndex : PotentialChildren)
 			{
-				AccumIntersectionList.Add(ChildIndex);
+				if(IntersectsHelper(WorldSpaceBoxes[ChildIndex], QueryObject)) //todo(ocohen): actually just a single point so should call Contains
+				{
+					AccumIntersectionList.Add(ChildIndex);
+				}
 			}
 		}
-	}
-};
+	};
 
-int32 CheckBox = 1;
-FAutoConsoleVariableRef CVarCheckBox(TEXT("p.checkbox"), CheckBox, TEXT(""));
+	int32 CheckBox = 1;
+	FAutoConsoleVariableRef CVarCheckBox(TEXT("p.checkbox"), CheckBox, TEXT(""));
 
-template <typename T, int d>
-struct TSpecializeParticlesHelper<TParticles<T,d>>
-{
-	template <typename QUERY_OBJECT>
-	static void AccumulateChildrenResults(TArray<int32>& AccumIntersectionList, const TArray<int32>& PotentialChildren, const QUERY_OBJECT& QueryObject, const TMap<int32, TAABB<T, d>>& WorldSpaceBoxes)
+	template <typename T, int d>
+	struct TSpecializeParticlesHelper<TParticles<T, d>>
 	{
-		if (CheckBox)
+		template <typename QUERY_OBJECT>
+		static void AccumulateChildrenResults(TArray<int32>& AccumIntersectionList, const TArray<int32>& PotentialChildren, const QUERY_OBJECT& QueryObject, const TMap<int32, TAABB<T, d>>& WorldSpaceBoxes)
 		{
-			for (int32 ChildIndex : PotentialChildren)
+			if(CheckBox)
 			{
-				if (IntersectsHelper(WorldSpaceBoxes[ChildIndex], QueryObject))	//todo(ocohen): actually just a single point so should call Contains
+				for(int32 ChildIndex : PotentialChildren)
 				{
-					AccumIntersectionList.Add(ChildIndex);
+					if(IntersectsHelper(WorldSpaceBoxes[ChildIndex], QueryObject))	//todo(ocohen): actually just a single point so should call Contains
+					{
+						AccumIntersectionList.Add(ChildIndex);
+					}
 				}
 			}
+			else
+			{
+				AccumIntersectionList.Append(PotentialChildren);
+			}
 		}
-		else
-		{
-			AccumIntersectionList.Append(PotentialChildren);
-		}
-	}
-};
+	};
 
-template<class OBJECT_ARRAY, typename TPayloadType, typename T, int d, typename QUERY_OBJECT>
-TArray<int32> FindAllIntersectionsLeafHelper(const TArray<int32>& Leaf, const QUERY_OBJECT& QueryObject)
-{
-	return Leaf;
-}
+	template<class OBJECT_ARRAY, typename TPayloadType, typename T, int d, typename QUERY_OBJECT>
+	TArray<int32> FindAllIntersectionsLeafHelper(const TArray<int32>& Leaf, const QUERY_OBJECT& QueryObject)
+	{
+		return Leaf;
+	}
 
-template<class OBJECT_ARRAY, typename TPayloadType, class T, int d, typename QUERY_OBJECT>
-TArray<int32> FindAllIntersectionsLeafHelper(const TBoundingVolume<TPayloadType, T, d>& Leaf, const QUERY_OBJECT& QueryObject)
-{
+	template<class OBJECT_ARRAY, typename TPayloadType, class T, int d, typename QUERY_OBJECT>
+	TArray<int32> FindAllIntersectionsLeafHelper(const TBoundingVolume<TPayloadType, T, d>& Leaf, const QUERY_OBJECT& QueryObject)
+	{
 #if CHAOS_PARTICLEHANDLE_TODO
-	return Leaf.FindAllIntersectionsImp(QueryObject);
+		return Leaf.FindAllIntersectionsImp(QueryObject);
 #else
-	return TArray<int32>();
+		return TArray<int32>();
 #endif
-}
-
-int32 FindAllIntersectionsSingleThreaded = 1;
-FAutoConsoleVariableRef CVarFindAllIntersectionsSingleThreaded(TEXT("p.FindAllIntersectionsSingleThreaded"), FindAllIntersectionsSingleThreaded, TEXT(""));
-
-template<class OBJECT_ARRAY, class LEAF_TYPE, class T, int d>
-template <typename QUERY_OBJECT>
-void TBoundingVolumeHierarchy<OBJECT_ARRAY, LEAF_TYPE, T, d>::FindAllIntersectionsHelperRecursive(const TBVHNode<T,d>& MyNode, const QUERY_OBJECT& QueryObject, TArray<int32>& AccumulateElements) const
-{
-	TAABB<T, d> MBox(MyNode.MMin, MyNode.MMax);
-	if (!IntersectsHelper(MBox, QueryObject))
-	{
-		return;
-	}
-	if (MyNode.MChildren.Num() == 0)
-	{
-		TSpecializeParticlesHelper<OBJECT_ARRAY>::AccumulateChildrenResults(AccumulateElements, FindAllIntersectionsLeafHelper<OBJECT_ARRAY, TPayloadType, T, d>(Leafs[MyNode.LeafIndex], QueryObject), QueryObject, MWorldSpaceBoxes);
-		return;
 	}
-	
-	int32 NumChildren = MyNode.MChildren.Num();
-	for (int32 Child = 0; Child < NumChildren; ++Child)
-	{
-		FindAllIntersectionsHelperRecursive(Elements[MyNode.MChildren[Child]], QueryObject, AccumulateElements);
-	}
-}
-
-int32 UseAccumulationArray = 1;
-FAutoConsoleVariableRef CVarUseAccumulationArray(TEXT("p.UseAccumulationArray"), UseAccumulationArray, TEXT(""));
 
-template<class OBJECT_ARRAY, class LEAF_TYPE, class T, int d>
-TArray<int32> TBoundingVolumeHierarchy<OBJECT_ARRAY, LEAF_TYPE, T, d>::FindAllIntersectionsHelper(const TBVHNode<T,d>& MyNode, const TAABB<T, d>& ObjectBox) const
-{
-	TArray<int32> IntersectionList;
-	FindAllIntersectionsHelperRecursive(MyNode, ObjectBox, IntersectionList);
-
-	IntersectionList.Sort();
+	int32 FindAllIntersectionsSingleThreaded = 1;
+	FAutoConsoleVariableRef CVarFindAllIntersectionsSingleThreaded(TEXT("p.FindAllIntersectionsSingleThreaded"), FindAllIntersectionsSingleThreaded, TEXT(""));
 
-	for (int32 i = IntersectionList.Num() - 1; i > 0; i--)
+	template<class OBJECT_ARRAY, class LEAF_TYPE, class T, int d>
+	template <typename QUERY_OBJECT>
+	void TBoundingVolumeHierarchy<OBJECT_ARRAY, LEAF_TYPE, T, d>::FindAllIntersectionsHelperRecursive(const TBVHNode<T, d>& MyNode, const QUERY_OBJECT& QueryObject, TArray<int32>& AccumulateElements) const
 	{
-		if (IntersectionList[i] == IntersectionList[i - 1])
+		TAABB<T, d> MBox(MyNode.MMin, MyNode.MMax);
+		if(!IntersectsHelper(MBox, QueryObject))
+		{
+			return;
+		}
+		if(MyNode.MChildren.Num() == 0)
 		{
-			IntersectionList.RemoveAtSwap(i, 1, false);
+			TSpecializeParticlesHelper<OBJECT_ARRAY>::AccumulateChildrenResults(AccumulateElements, FindAllIntersectionsLeafHelper<OBJECT_ARRAY, TPayloadType, T, d>(Leafs[MyNode.LeafIndex], QueryObject), QueryObject, MWorldSpaceBoxes);
+			return;
+		}
+
+		int32 NumChildren = MyNode.MChildren.Num();
+		for(int32 Child = 0; Child < NumChildren; ++Child)
+		{
+			FindAllIntersectionsHelperRecursive(Elements[MyNode.MChildren[Child]], QueryObject, AccumulateElements);
 		}
 	}
 
-	return IntersectionList;
-}
+	int32 UseAccumulationArray = 1;
+	FAutoConsoleVariableRef CVarUseAccumulationArray(TEXT("p.UseAccumulationArray"), UseAccumulationArray, TEXT(""));
 
-template<class OBJECT_ARRAY, class LEAF_TYPE, class T, int d>
-TArray<int32> TBoundingVolumeHierarchy<OBJECT_ARRAY, LEAF_TYPE, T, d>::FindAllIntersectionsHelper(const TBVHNode<T, d>& MyNode, const TSpatialRay<T, d>& Ray) const
-{
-	TArray<int32> IntersectionList;
-	FindAllIntersectionsHelperRecursive(MyNode, Ray, IntersectionList);
+	template<class OBJECT_ARRAY, class LEAF_TYPE, class T, int d>
+	TArray<int32> TBoundingVolumeHierarchy<OBJECT_ARRAY, LEAF_TYPE, T, d>::FindAllIntersectionsHelper(const TBVHNode<T, d>& MyNode, const TAABB<T, d>& ObjectBox) const
+	{
+		TArray<int32> IntersectionList;
+		FindAllIntersectionsHelperRecursive(MyNode, ObjectBox, IntersectionList);
 
-	IntersectionList.Sort();
+		IntersectionList.Sort();
 
-	for (int32 i = IntersectionList.Num() - 1; i > 0; i--)
-	{
-		if (IntersectionList[i] == IntersectionList[i - 1])
+		for(int32 i = IntersectionList.Num() - 1; i > 0; i--)
 		{
-			IntersectionList.RemoveAtSwap(i, 1, false);
+			if(IntersectionList[i] == IntersectionList[i - 1])
+			{
+				IntersectionList.RemoveAtSwap(i, 1, false);
+			}
 		}
+
+		return IntersectionList;
 	}
 
-	return IntersectionList;
-}
+	template<class OBJECT_ARRAY, class LEAF_TYPE, class T, int d>
+	TArray<int32> TBoundingVolumeHierarchy<OBJECT_ARRAY, LEAF_TYPE, T, d>::FindAllIntersectionsHelper(const TBVHNode<T, d>& MyNode, const TSpatialRay<T, d>& Ray) const
+	{
+		TArray<int32> IntersectionList;
+		FindAllIntersectionsHelperRecursive(MyNode, Ray, IntersectionList);
 
-template<class OBJECT_ARRAY, class LEAF_TYPE, class T, int d>
-TArray<int32> TBoundingVolumeHierarchy<OBJECT_ARRAY, LEAF_TYPE, T, d>::FindAllIntersections(const TGeometryParticles<T, d>& InParticles, const int32 i) const
-{
-	return FindAllIntersections(Chaos::GetWorldSpaceBoundingBox(InParticles, i, MWorldSpaceBoxes));
-}
+		IntersectionList.Sort();
 
-template <int d>
-struct FSplitCount
-{
-	FSplitCount()
-	{
-		for (int i = 0; i < d; ++i)
+		for(int32 i = IntersectionList.Num() - 1; i > 0; i--)
 		{
-			Neg[i] = 0;
-			Pos[i] = 0;
+			if(IntersectionList[i] == IntersectionList[i - 1])
+			{
+				IntersectionList.RemoveAtSwap(i, 1, false);
+			}
 		}
-	}
 
-	int32 Neg[d];
-	int32 Pos[d];
-};
-
-template <typename T, int d>
-void AccumulateNextLevelCount(const TAABB<T, d>& Box, const TVector<T, d>& MidPoint, FSplitCount<d>& Counts)
-{
-	//todo(ocohen): particles min = max so avoid extra work
-	for (int32 i = 0; i < d; ++i)
-	{
-		Counts.Neg[i] += (Box.Max()[i] < MidPoint[i] || Box.Min()[i] < MidPoint[i]) ? 1 : 0;
-		Counts.Pos[i] += (Box.Min()[i] > MidPoint[i] || Box.Max()[i] > MidPoint[i]) ? 1 : 0;
+		return IntersectionList;
 	}
-}
-template<class OBJECT_ARRAY, class LEAF_TYPE, class T, int d>
-int32 TBoundingVolumeHierarchy<OBJECT_ARRAY, LEAF_TYPE, T, d>::GenerateNextLevel(const TVector<T, d>& GlobalMin, const TVector<T, d>& GlobalMax, const TArray<int32>& Objects, const int32 Axis, const int32 Level, const bool AllowMultipleSplitting)
-{
-	if (Axis == -1)
+
+	template<class OBJECT_ARRAY, class LEAF_TYPE, class T, int d>
+	TArray<int32> TBoundingVolumeHierarchy<OBJECT_ARRAY, LEAF_TYPE, T, d>::FindAllIntersections(const TGeometryParticles<T, d>& InParticles, const int32 i) const
 	{
-		return GenerateNextLevel(GlobalMin, GlobalMax, Objects, Level);
+		return FindAllIntersections(Chaos::GetWorldSpaceBoundingBox(InParticles, i, MWorldSpaceBoxes));
 	}
-	
-	FSplitCount<d> Counts[2];
-	TArray<TBVHNode<T,d>> LocalElements;
-	TArray<TArray<int32>> LocalObjects;
-	TArray<LEAF_TYPE> LocalLeafs;
-	LocalElements.SetNum(2);
-	LocalObjects.SetNum(2);
-	LocalLeafs.SetNum(2);
-	TAABB<T, d> GlobalBox(GlobalMin, GlobalMax);
-	const TVector<T, d> WorldCenter = GlobalBox.Center();
-	const TVector<T, d> MinCenterSearch = TAABB<T, d>(GlobalMin, WorldCenter).Center();
-	const TVector<T, d> MaxCenterSearch = TAABB<T, d>(WorldCenter, GlobalMax).Center();
-
-	for (int32 i = 0; i < Objects.Num(); ++i)
+
+	template <int d>
+	struct FSplitCount
 	{
-		check(Objects[i] >= 0 && Objects[i] < GetObjectCount(*MObjects));
-		const TAABB<T, d>& ObjectBox = Chaos::GetWorldSpaceBoundingBox(*MObjects, Objects[i], MWorldSpaceBoxes);
-		const TVector<T, d> ObjectCenter = ObjectBox.Center();
-		bool MinA = false, MaxA = false;
-		if (ObjectBox.Min()[Axis] < WorldCenter[Axis])
-		{
-			MinA = true;
-		}
-		if (ObjectBox.Max()[Axis] >= WorldCenter[Axis])
+		FSplitCount()
 		{
-			MaxA = true;
-		}
-		ensure(MinA || MaxA); // If this is not true we have a nan and the object gets ignored
-		if (MinA)
-		{
-			LocalObjects[0].Add(Objects[i]);
-			AccumulateNextLevelCount(ObjectBox, MinCenterSearch, Counts[0]);
+			for(int i = 0; i < d; ++i)
+			{
+				Neg[i] = 0;
+				Pos[i] = 0;
+			}
 		}
-		if (MaxA)
+
+		int32 Neg[d];
+		int32 Pos[d];
+	};
+
+	template <typename T, int d>
+	void AccumulateNextLevelCount(const TAABB<T, d>& Box, const TVector<T, d>& MidPoint, FSplitCount<d>& Counts)
+	{
+		//todo(ocohen): particles min = max so avoid extra work
+		for(int32 i = 0; i < d; ++i)
 		{
-			LocalObjects[1].Add(Objects[i]);
-			AccumulateNextLevelCount(ObjectBox, MaxCenterSearch, Counts[1]);
+			Counts.Neg[i] += (Box.Max()[i] < MidPoint[i] || Box.Min()[i] < MidPoint[i]) ? 1 : 0;
+			Counts.Pos[i] += (Box.Min()[i] > MidPoint[i] || Box.Max()[i] > MidPoint[i]) ? 1 : 0;
 		}
 	}
-	PhysicsParallelFor(2, [&](int32 i) {
-		TVector<T, d> Min = GlobalBox.Min();
-		TVector<T, d> Max = GlobalBox.Max();
-		if (i == 0)
-			Max[Axis] = WorldCenter[Axis];
-		else
-			Min[Axis] = WorldCenter[Axis];
-		LocalElements[i].MMin = Min;
-		LocalElements[i].MMax = Max;
-		LocalElements[i].MAxis = -1;
+	template<class OBJECT_ARRAY, class LEAF_TYPE, class T, int d>
+	int32 TBoundingVolumeHierarchy<OBJECT_ARRAY, LEAF_TYPE, T, d>::GenerateNextLevel(const TVector<T, d>& GlobalMin, const TVector<T, d>& GlobalMax, const TArray<int32>& Objects, const int32 Axis, const int32 Level, const bool AllowMultipleSplitting)
+	{
+		if(Axis == -1)
+		{
+			return GenerateNextLevel(GlobalMin, GlobalMax, Objects, Level);
+		}
+
+		FSplitCount<d> Counts[2];
+		TArray<TBVHNode<T, d>> LocalElements;
+		TArray<TArray<int32>> LocalObjects;
+		TArray<LEAF_TYPE> LocalLeafs;
+		LocalElements.SetNum(2);
+		LocalObjects.SetNum(2);
+		LocalLeafs.SetNum(2);
+		TAABB<T, d> GlobalBox(GlobalMin, GlobalMax);
+		const TVector<T, d> WorldCenter = GlobalBox.Center();
+		const TVector<T, d> MinCenterSearch = TAABB<T, d>(GlobalMin, WorldCenter).Center();
+		const TVector<T, d> MaxCenterSearch = TAABB<T, d>(WorldCenter, GlobalMax).Center();
+
+		for(int32 i = 0; i < Objects.Num(); ++i)
+		{
+			check(Objects[i] >= 0 && Objects[i] < GetObjectCount(*MObjects));
+			const TAABB<T, d>& ObjectBox = Chaos::GetWorldSpaceBoundingBox(*MObjects, Objects[i], MWorldSpaceBoxes);
+			const TVector<T, d> ObjectCenter = ObjectBox.Center();
+			bool MinA = false, MaxA = false;
+			if(ObjectBox.Min()[Axis] < WorldCenter[Axis])
+			{
+				MinA = true;
+			}
+			if(ObjectBox.Max()[Axis] >= WorldCenter[Axis])
+			{
+				MaxA = true;
+			}
+			ensure(MinA || MaxA); // If this is not true we have a nan and the object gets ignored
+			if(MinA)
+			{
+				LocalObjects[0].Add(Objects[i]);
+				AccumulateNextLevelCount(ObjectBox, MinCenterSearch, Counts[0]);
+			}
+			if(MaxA)
+			{
+				LocalObjects[1].Add(Objects[i]);
+				AccumulateNextLevelCount(ObjectBox, MaxCenterSearch, Counts[1]);
+			}
+		}
+		PhysicsParallelFor(2, [&](int32 i)
+		{
+			TVector<T, d> Min = GlobalBox.Min();
+			TVector<T, d> Max = GlobalBox.Max();
+			if(i == 0)
+				Max[Axis] = WorldCenter[Axis];
+			else
+				Min[Axis] = WorldCenter[Axis];
+			LocalElements[i].MMin = Min;
+			LocalElements[i].MMax = Max;
+			LocalElements[i].MAxis = -1;
 #if !UE_BUILD_SHIPPING && !UE_BUILD_TEST
-		LocalElements[i].LeafIndex = -1;
+			LocalElements[i].LeafIndex = -1;
 #endif
-		if (LocalObjects[i].Num() > MIN_NUM_OBJECTS && Level < MMaxLevels && LocalObjects[i].Num() < Objects.Num())
-		{
-			//we pick the axis that gives us the most culled even in the case when it goes in the wrong direction (i.e the biggest min)
-			int32 BestAxis = 0;
-			int32 MaxCulled = 0;
-			for (int32 LocalAxis = 0; LocalAxis < d; ++LocalAxis)
+			if(LocalObjects[i].Num() > MIN_NUM_OBJECTS && Level < MMaxLevels && LocalObjects[i].Num() < Objects.Num())
 			{
-				int32 CulledWorstCase = FMath::Min(Counts[i].Neg[LocalAxis], Counts[i].Pos[LocalAxis]);
-				if (CulledWorstCase > MaxCulled)
+				//we pick the axis that gives us the most culled even in the case when it goes in the wrong direction (i.e the biggest min)
+				int32 BestAxis = 0;
+				int32 MaxCulled = 0;
+				for(int32 LocalAxis = 0; LocalAxis < d; ++LocalAxis)
 				{
-					MaxCulled = CulledWorstCase;
-					BestAxis = LocalAxis;
+					int32 CulledWorstCase = FMath::Min(Counts[i].Neg[LocalAxis], Counts[i].Pos[LocalAxis]);
+					if(CulledWorstCase > MaxCulled)
+					{
+						MaxCulled = CulledWorstCase;
+						BestAxis = LocalAxis;
+					}
 				}
-			}
 
-			//todo(ocohen): use multi split when counts are very close
-			int32 NextAxis = BestAxis;
-			LocalElements[i].MAxis = NextAxis;
-			int32 StartIndex = GenerateNextLevel(LocalElements[i].MMin, LocalElements[i].MMax, LocalObjects[i], NextAxis, Level + 1, AllowMultipleSplitting);
-			for (int32 j = 0; j < (NextAxis == -1 ? 8 : 2); j++)
+				//todo(ocohen): use multi split when counts are very close
+				int32 NextAxis = BestAxis;
+				LocalElements[i].MAxis = NextAxis;
+				int32 StartIndex = GenerateNextLevel(LocalElements[i].MMin, LocalElements[i].MMax, LocalObjects[i], NextAxis, Level + 1, AllowMultipleSplitting);
+				for(int32 j = 0; j < (NextAxis == -1 ? 8 : 2); j++)
+				{
+					LocalElements[i].MChildren.Add(StartIndex + j);
+				}
+			}
+			else
 			{
-				LocalElements[i].MChildren.Add(StartIndex + j);
+				LocalLeafs[i] = MakeNewLeaf<OBJECT_ARRAY, T, d, TPayloadType>(*MObjects, LocalObjects[i], static_cast<LEAF_TYPE*>(nullptr));
 			}
-		}
-		else
-		{
-			LocalLeafs[i] = MakeNewLeaf<OBJECT_ARRAY, T, d, TPayloadType>(*MObjects, LocalObjects[i], static_cast<LEAF_TYPE*>(nullptr));
-		}
-	});
-	CriticalSection.Lock();
-	for (int i = 0; i < 2; ++i)
-	{
-		if (!LocalElements[i].MChildren.Num())
+		});
+		CriticalSection.Lock();
+		for(int i = 0; i < 2; ++i)
 		{
-			LocalElements[i].LeafIndex = Leafs.Num();
-			Leafs.Add(MoveTemp(LocalLeafs[i]));
+			if(!LocalElements[i].MChildren.Num())
+			{
+				LocalElements[i].LeafIndex = Leafs.Num();
+				Leafs.Add(MoveTemp(LocalLeafs[i]));
+			}
+			check(LocalElements[i].LeafIndex < Leafs.Num() || LocalElements[i].MChildren.Num() > 0);
 		}
-		check(LocalElements[i].LeafIndex < Leafs.Num() || LocalElements[i].MChildren.Num() > 0);
+		int32 MinElem = Elements.Num();
+		Elements.Append(LocalElements);
+		CriticalSection.Unlock();
+		return MinElem;
 	}
-	int32 MinElem = Elements.Num();
-	Elements.Append(LocalElements);
-	CriticalSection.Unlock();
-	return MinElem;
-}
 
-template<class OBJECT_ARRAY, class LEAF_TYPE, class T, int d>
-int32 TBoundingVolumeHierarchy<OBJECT_ARRAY, LEAF_TYPE, T, d>::GenerateNextLevel(const TVector<T, d>& GlobalMin, const TVector<T, d>& GlobalMax, const TArray<int32>& Objects, const int32 Level)
-{
-	TArray<TBVHNode<T,d>> LocalElements;
-	TArray<TArray<int32>> LocalObjects;
-	TArray<LEAF_TYPE> LocalLeafs;
-	LocalElements.SetNum(8);
-	LocalObjects.SetNum(8);
-	LocalLeafs.SetNum(8);
-	TAABB<T, d> GlobalBox(GlobalMin, GlobalMax);
-	const TVector<T, d> WorldCenter = GlobalBox.Center();
-	for (int32 i = 0; i < Objects.Num(); ++i)
+	template<class OBJECT_ARRAY, class LEAF_TYPE, class T, int d>
+	int32 TBoundingVolumeHierarchy<OBJECT_ARRAY, LEAF_TYPE, T, d>::GenerateNextLevel(const TVector<T, d>& GlobalMin, const TVector<T, d>& GlobalMax, const TArray<int32>& Objects, const int32 Level)
 	{
-		check(Objects[i] >= 0 && Objects[i] < GetObjectCount(*MObjects));
-		const TAABB<T, d>& ObjectBox = Chaos::GetWorldSpaceBoundingBox(*MObjects, Objects[i], MWorldSpaceBoxes);
-		const TVector<T, d> ObjectCenter = ObjectBox.Center();
-		bool MinX = false, MaxX = false, MinY = false, MaxY = false, MinZ = false, MaxZ = false;
-		if (ObjectBox.Min()[0] < WorldCenter[0])
-		{
-			MinX = true;
-		}
-		if (ObjectBox.Max()[0] >= WorldCenter[0])
-		{
-			MaxX = true;
-		}
-		if (ObjectBox.Min()[1] < WorldCenter[1])
-		{
-			MinY = true;
-		}
-		if (ObjectBox.Max()[1] >= WorldCenter[1])
-		{
-			MaxY = true;
-		}
-		if (ObjectBox.Min()[2] < WorldCenter[2])
-		{
-			MinZ = true;
-		}
-		if (ObjectBox.Max()[2] >= WorldCenter[2])
-		{
-			MaxZ = true;
-		}
-		check(MinX || MaxX);
-		check(MinY || MaxY);
-		check(MinZ || MaxZ);
-		if (MinX && MinY && MinZ)
-		{
-			LocalObjects[0].Add(Objects[i]);
-		}
-		if (MaxX && MinY && MinZ)
-		{
-			LocalObjects[1].Add(Objects[i]);
-		}
-		if (MinX && MaxY && MinZ)
-		{
-			LocalObjects[2].Add(Objects[i]);
-		}
-		if (MaxX && MaxY && MinZ)
-		{
-			LocalObjects[3].Add(Objects[i]);
-		}
-		if (MinX && MinY && MaxZ)
-		{
-			LocalObjects[4].Add(Objects[i]);
-		}
-		if (MaxX && MinY && MaxZ)
-		{
-			LocalObjects[5].Add(Objects[i]);
-		}
-		if (MinX && MaxY && MaxZ)
-		{
-			LocalObjects[6].Add(Objects[i]);
-		}
-		if (MaxX && MaxY && MaxZ)
-		{
-			LocalObjects[7].Add(Objects[i]);
-		}
-	}
-	PhysicsParallelFor(8, [&](int32 i) {
-		TVector<T, d> Min = GlobalBox.Min();
-		TVector<T, d> Max = GlobalBox.Max();
-		if (i % 2 == 0)
-			Max[0] = WorldCenter[0];
-		else
-			Min[0] = WorldCenter[0];
-		if ((i / 2) % 2 == 0)
-			Max[1] = WorldCenter[1];
-		else
-			Min[1] = WorldCenter[1];
-		if (i < 4)
-			Max[2] = WorldCenter[2];
-		else
-			Min[2] = WorldCenter[2];
-		LocalElements[i].MMin = Min;
-		LocalElements[i].MMax = Max;
-		LocalElements[i].MAxis = -1;
-#if !UE_BUILD_SHIPPING && !UE_BUILD_TEST
-		LocalElements[i].LeafIndex = -1;
-#endif
-		if (LocalObjects[i].Num() > MIN_NUM_OBJECTS && Level < MMaxLevels && LocalObjects[i].Num() < Objects.Num())
-		{
-			TAABB<T, d> LocalBox(Min, Max);
-			int32 NextAxis = 0;
-			TVector<T, d> LocalExtents = LocalBox.Extents();
-			if (LocalExtents[2] > LocalExtents[0] && LocalExtents[2] > LocalExtents[1])
+		TArray<TBVHNode<T, d>> LocalElements;
+		TArray<TArray<int32>> LocalObjects;
+		TArray<LEAF_TYPE> LocalLeafs;
+		LocalElements.SetNum(8);
+		LocalObjects.SetNum(8);
+		LocalLeafs.SetNum(8);
+		TAABB<T, d> GlobalBox(GlobalMin, GlobalMax);
+		const TVector<T, d> WorldCenter = GlobalBox.Center();
+		for(int32 i = 0; i < Objects.Num(); ++i)
+		{
+			check(Objects[i] >= 0 && Objects[i] < GetObjectCount(*MObjects));
+			const TAABB<T, d>& ObjectBox = Chaos::GetWorldSpaceBoundingBox(*MObjects, Objects[i], MWorldSpaceBoxes);
+			const TVector<T, d> ObjectCenter = ObjectBox.Center();
+			bool MinX = false, MaxX = false, MinY = false, MaxY = false, MinZ = false, MaxZ = false;
+			if(ObjectBox.Min()[0] < WorldCenter[0])
 			{
-				NextAxis = 2;
+				MinX = true;
 			}
-			else if (LocalExtents[1] > LocalExtents[0])
+			if(ObjectBox.Max()[0] >= WorldCenter[0])
 			{
-				NextAxis = 1;
+				MaxX = true;
 			}
-			if (LocalExtents[NextAxis] < LocalExtents[(NextAxis + 1) % 3] * 1.25 && LocalExtents[NextAxis] < LocalExtents[(NextAxis + 2) % 3] * 1.25 && LocalObjects[i].Num() > 4 * MIN_NUM_OBJECTS)
+			if(ObjectBox.Min()[1] < WorldCenter[1])
 			{
-				NextAxis = -1;
+				MinY = true;
 			}
-			LocalElements[i].MAxis = NextAxis;
-			int32 StartIndex = GenerateNextLevel(LocalElements[i].MMin, LocalElements[i].MMax, LocalObjects[i], NextAxis, Level + 1, true);
-			for (int32 j = 0; j < (NextAxis == -1 ? 8 : 2); j++)
+			if(ObjectBox.Max()[1] >= WorldCenter[1])
 			{
-				LocalElements[i].MChildren.Add(StartIndex + j);
+				MaxY = true;
+			}
+			if(ObjectBox.Min()[2] < WorldCenter[2])
+			{
+				MinZ = true;
+			}
+			if(ObjectBox.Max()[2] >= WorldCenter[2])
+			{
+				MaxZ = true;
+			}
+			check(MinX || MaxX);
+			check(MinY || MaxY);
+			check(MinZ || MaxZ);
+			if(MinX && MinY && MinZ)
+			{
+				LocalObjects[0].Add(Objects[i]);
+			}
+			if(MaxX && MinY && MinZ)
+			{
+				LocalObjects[1].Add(Objects[i]);
+			}
+			if(MinX && MaxY && MinZ)
+			{
+				LocalObjects[2].Add(Objects[i]);
+			}
+			if(MaxX && MaxY && MinZ)
+			{
+				LocalObjects[3].Add(Objects[i]);
+			}
+			if(MinX && MinY && MaxZ)
+			{
+				LocalObjects[4].Add(Objects[i]);
+			}
+			if(MaxX && MinY && MaxZ)
+			{
+				LocalObjects[5].Add(Objects[i]);
+			}
+			if(MinX && MaxY && MaxZ)
+			{
+				LocalObjects[6].Add(Objects[i]);
+			}
+			if(MaxX && MaxY && MaxZ)
+			{
+				LocalObjects[7].Add(Objects[i]);
 			}
 		}
-		else
-		{
-			LocalLeafs[i] = MakeNewLeaf<OBJECT_ARRAY, T, d, TPayloadType>(*MObjects, LocalObjects[i], static_cast<LEAF_TYPE*>(nullptr));
-		}
-	});
-	CriticalSection.Lock();
-	for (int i = 0; i < 8; ++i)
-	{
-		if (!LocalElements[i].MChildren.Num())
+		PhysicsParallelFor(8, [&](int32 i)
+		{
+			TVector<T, d> Min = GlobalBox.Min();
+			TVector<T, d> Max = GlobalBox.Max();
+			if(i % 2 == 0)
+				Max[0] = WorldCenter[0];
+			else
+				Min[0] = WorldCenter[0];
+			if((i / 2) % 2 == 0)
+				Max[1] = WorldCenter[1];
+			else
+				Min[1] = WorldCenter[1];
+			if(i < 4)
+				Max[2] = WorldCenter[2];
+			else
+				Min[2] = WorldCenter[2];
+			LocalElements[i].MMin = Min;
+			LocalElements[i].MMax = Max;
+			LocalElements[i].MAxis = -1;
+#if !UE_BUILD_SHIPPING && !UE_BUILD_TEST
+			LocalElements[i].LeafIndex = -1;
+#endif
+			if(LocalObjects[i].Num() > MIN_NUM_OBJECTS && Level < MMaxLevels && LocalObjects[i].Num() < Objects.Num())
+			{
+				TAABB<T, d> LocalBox(Min, Max);
+				int32 NextAxis = 0;
+				TVector<T, d> LocalExtents = LocalBox.Extents();
+				if(LocalExtents[2] > LocalExtents[0] && LocalExtents[2] > LocalExtents[1])
+				{
+					NextAxis = 2;
+				}
+				else if(LocalExtents[1] > LocalExtents[0])
+				{
+					NextAxis = 1;
+				}
+				if(LocalExtents[NextAxis] < LocalExtents[(NextAxis + 1) % 3] * 1.25 && LocalExtents[NextAxis] < LocalExtents[(NextAxis + 2) % 3] * 1.25 && LocalObjects[i].Num() > 4 * MIN_NUM_OBJECTS)
+				{
+					NextAxis = -1;
+				}
+				LocalElements[i].MAxis = NextAxis;
+				int32 StartIndex = GenerateNextLevel(LocalElements[i].MMin, LocalElements[i].MMax, LocalObjects[i], NextAxis, Level + 1, true);
+				for(int32 j = 0; j < (NextAxis == -1 ? 8 : 2); j++)
+				{
+					LocalElements[i].MChildren.Add(StartIndex + j);
+				}
+			}
+			else
+			{
+				LocalLeafs[i] = MakeNewLeaf<OBJECT_ARRAY, T, d, TPayloadType>(*MObjects, LocalObjects[i], static_cast<LEAF_TYPE*>(nullptr));
+			}
+		});
+		CriticalSection.Lock();
+		for(int i = 0; i < 8; ++i)
 		{
-			LocalElements[i].LeafIndex = Leafs.Num();
-			Leafs.Add(MoveTemp(LocalLeafs[i]));
+			if(!LocalElements[i].MChildren.Num())
+			{
+				LocalElements[i].LeafIndex = Leafs.Num();
+				Leafs.Add(MoveTemp(LocalLeafs[i]));
+			}
+			check(LocalElements[i].LeafIndex < Leafs.Num() || LocalElements[i].MChildren.Num() > 0);
 		}
-		check(LocalElements[i].LeafIndex < Leafs.Num() || LocalElements[i].MChildren.Num() > 0);
+		int32 MinElem = Elements.Num();
+		Elements.Append(LocalElements);
+		CriticalSection.Unlock();
+		return MinElem;
 	}
-	int32 MinElem = Elements.Num();
-	Elements.Append(LocalElements);
-	CriticalSection.Unlock();
-	return MinElem;
-}
 
-namespace Chaos
-{
-template <typename OBJECT_ARRAY, typename TPayloadType, typename T, int d>
-void FixupLeafObj(const OBJECT_ARRAY& Objects, TArray<TBoundingVolume<TPayloadType, T, d>>& Leafs)
-{
-#if CHAOS_PARTICLEHANDLE_TODO
-	for (TBoundingVolume<OBJECT_ARRAY, TPayloadType, T, d>& Leaf : Leafs)
+	template <typename OBJECT_ARRAY, typename TPayloadType, typename T, int d>
+	void FixupLeafObj(const OBJECT_ARRAY& Objects, TArray<TBoundingVolume<TPayloadType, T, d>>& Leafs)
 	{
-		Leaf.SetObjects(Objects);
-	}
+#if CHAOS_PARTICLEHANDLE_TODO
+		for(TBoundingVolume<OBJECT_ARRAY, TPayloadType, T, d>& Leaf : Leafs)
+		{
+			Leaf.SetObjects(Objects);
+		}
 #endif
-}
+	}
 
-template <typename OBJECT_ARRAY>
-void FixupLeafObj(const OBJECT_ARRAY& Objects, TArray<TArray<int32>>& Leafs)
-{
-}
+	template <typename OBJECT_ARRAY>
+	void FixupLeafObj(const OBJECT_ARRAY& Objects, TArray<TArray<int32>>& Leafs)
+	{}
 
-template<class OBJECT_ARRAY, class LEAF_TYPE, class T, int d>
-void TBoundingVolumeHierarchy<OBJECT_ARRAY, LEAF_TYPE, T, d>::Serialize(FArchive& Ar)
-{
-	Ar << MGlobalObjects;
-	TBox<T, d>::SerializeAsAABBs(Ar, MWorldSpaceBoxes);
-	Ar << MMaxLevels << Elements << Leafs;
-	if (Ar.IsLoading())
+	template<class OBJECT_ARRAY, class LEAF_TYPE, class T, int d>
+	void TBoundingVolumeHierarchy<OBJECT_ARRAY, LEAF_TYPE, T, d>::Serialize(FArchive& Ar)
 	{
-		FixupLeafObj(*MObjects, Leafs);
+		Ar << MGlobalObjects;
+		TBox<T, d>::SerializeAsAABBs(Ar, MWorldSpaceBoxes);
+		Ar << MMaxLevels << Elements << Leafs;
+		if(Ar.IsLoading())
+		{
+			FixupLeafObj(*MObjects, Leafs);
+		}
 	}
 }
-}
 
 //template class CHAOS_API Chaos::TBoundingVolume<int32, FReal, 3>;
-template class Chaos::TBoundingVolumeHierarchy<TArray<Chaos::TSphere<FReal, 3>*>, TArray<int32>, FReal, 3>;
-template class Chaos::TBoundingVolumeHierarchy<Chaos::TPBDRigidParticles<FReal, 3>, TArray<int32>, FReal, 3>;
-template class Chaos::TBoundingVolumeHierarchy<Chaos::FParticles, TArray<int32>, FReal, 3>;
-template class Chaos::TBoundingVolumeHierarchy<Chaos::TGeometryParticles<FReal, 3>, TArray<int32>, FReal, 3>;
-template class Chaos::TBoundingVolumeHierarchy<Chaos::TPBDRigidParticles<FReal, 3>, TBoundingVolume<TPBDRigidParticleHandle<FReal,3>*, FReal, 3>, FReal, 3>;
-template class Chaos::TBoundingVolumeHierarchy<Chaos::TGeometryParticles<FReal, 3>, TBoundingVolume<TGeometryParticleHandle<FReal,3>*, FReal, 3>, FReal, 3>;
-template class Chaos::TBoundingVolumeHierarchy<TArray<TUniquePtr<Chaos::FImplicitObject>>, TArray<int32>, FReal, 3>;
+template class Chaos::TBoundingVolumeHierarchy<TArray<Chaos::TSphere<Chaos::FReal, 3>*>, TArray<int32>, Chaos::FReal, 3>;
+template class Chaos::TBoundingVolumeHierarchy<Chaos::TPBDRigidParticles<Chaos::FReal, 3>, TArray<int32>, Chaos::FReal, 3>;
+template class Chaos::TBoundingVolumeHierarchy<Chaos::FParticles, TArray<int32>, Chaos::FReal, 3>;
+template class Chaos::TBoundingVolumeHierarchy<Chaos::TGeometryParticles<Chaos::FReal, 3>, TArray<int32>, Chaos::FReal, 3>;
+template class Chaos::TBoundingVolumeHierarchy<Chaos::TPBDRigidParticles<Chaos::FReal, 3>, Chaos::TBoundingVolume<Chaos::TPBDRigidParticleHandle<Chaos::FReal, 3>*, Chaos::FReal, 3>, Chaos::FReal, 3>;
+template class Chaos::TBoundingVolumeHierarchy<Chaos::TGeometryParticles<Chaos::FReal, 3>, Chaos::TBoundingVolume<Chaos::TGeometryParticleHandle<Chaos::FReal, 3>*, Chaos::FReal, 3>, Chaos::FReal, 3>;
+template class Chaos::TBoundingVolumeHierarchy<TArray<TUniquePtr<Chaos::FImplicitObject>>, TArray<int32>, Chaos::FReal, 3>;
diff --git a/Engine/Source/Runtime/Experimental/Chaos/Private/Chaos/DebugDrawQueue.cpp b/Engine/Source/Runtime/Experimental/Chaos/Private/Chaos/DebugDrawQueue.cpp
--- a/Engine/Source/Runtime/Experimental/Chaos/Private/Chaos/DebugDrawQueue.cpp
+++ b/Engine/Source/Runtime/Experimental/Chaos/Private/Chaos/DebugDrawQueue.cpp
@@ -4,33 +4,35 @@
 #include "Misc/ScopeLock.h"
 
 #if CHAOS_DEBUG_DRAW
-using namespace Chaos;
 
-void FDebugDrawQueue::SetConsumerActive(void* Consumer, bool bConsumerActive)
+namespace Chaos
 {
-	FScopeLock Lock(&ConsumersCS);
-
-	if (bConsumerActive)
-	{
-		Consumers.AddUnique(Consumer);
-	}
-	else
+	void FDebugDrawQueue::SetConsumerActive(void* Consumer, bool bConsumerActive)
 	{
-		Consumers.Remove(Consumer);
+		FScopeLock Lock(&ConsumersCS);
+	
+		if(bConsumerActive)
+		{
+			Consumers.AddUnique(Consumer);
+		}
+		else
+		{
+			Consumers.Remove(Consumer);
+		}
+	
+		NumConsumers = Consumers.Num();
 	}
-
-	NumConsumers = Consumers.Num();
-}
-
-FDebugDrawQueue& FDebugDrawQueue::GetInstance()
-{
-	static FDebugDrawQueue* PSingleton = nullptr;
-	if (PSingleton == nullptr)
+	
+	FDebugDrawQueue& FDebugDrawQueue::GetInstance()
 	{
-		static FDebugDrawQueue Singleton;
-		PSingleton = &Singleton;
+		static FDebugDrawQueue* PSingleton = nullptr;
+		if(PSingleton == nullptr)
+		{
+			static FDebugDrawQueue Singleton;
+			PSingleton = &Singleton;
+		}
+		return *PSingleton;
 	}
-	return *PSingleton;
 }
 
 #endif
diff --git a/Engine/Source/Runtime/Experimental/Chaos/Private/Chaos/FFT.cpp b/Engine/Source/Runtime/Experimental/Chaos/Private/Chaos/FFT.cpp
--- a/Engine/Source/Runtime/Experimental/Chaos/Private/Chaos/FFT.cpp
+++ b/Engine/Source/Runtime/Experimental/Chaos/Private/Chaos/FFT.cpp
@@ -2,239 +2,240 @@
 // Parts of this file are adapted from PhysBAM under Copyright http://physbam.stanford.edu/links/backhistdisclaimcopy.html
 #include "Chaos/FFT.h"
 
-using namespace Chaos;
-
-void ConjugateAll(const FFFT3::FUniformGrid& Grid, FFFT3::FArrayNDOfComplex& Velocity, const int32 z)
+namespace Chaos
 {
-	for (int i = 1; i < Grid.Counts()[0] / 2; ++i)
-	{
-		Velocity(Grid.Counts()[0] - i, 0, z) = Velocity(i, 0, z).Conjugated();
-		Velocity(Grid.Counts()[0] - i, Grid.Counts()[1] / 2, z) = Velocity(i, Grid.Counts()[1] / 2, z).Conjugated();
-	}
-	for (int i = 1; i < Grid.Counts()[1] / 2; ++i)
+	void ConjugateAll(const FFFT3::FUniformGrid& Grid, FFFT3::FArrayNDOfComplex& Velocity, const int32 z)
 	{
-		Velocity(0, Grid.Counts()[1] - i, z) = Velocity(0, i, z).Conjugated();
-	}
-	for (int i = 1; i < Grid.Counts()[0]; ++i)
-		for (int j = 1; j < Grid.Counts()[1] / 2; ++j)
+		for(int i = 1; i < Grid.Counts()[0] / 2; ++i)
 		{
-			Velocity(Grid.Counts()[0] - i, Grid.Counts()[1] - j, z) = Velocity(i, j, z).Conjugated();
+			Velocity(Grid.Counts()[0] - i, 0, z) = Velocity(i, 0, z).Conjugated();
+			Velocity(Grid.Counts()[0] - i, Grid.Counts()[1] / 2, z) = Velocity(i, Grid.Counts()[1] / 2, z).Conjugated();
 		}
-}
-void EnforceSymmetry(const FFFT3::FUniformGrid& Grid, FFFT3::FArrayNDOfComplex& Velocity)
-{
-	Velocity[0].MakeReal();
-	Velocity(Grid.Counts()[0] / 2, 0, 0).MakeReal();
-	Velocity(0, Grid.Counts()[1] / 2, 0).MakeReal();
-	Velocity(Grid.Counts()[0] / 2, Grid.Counts()[1] / 2, 0).MakeReal();
-	Velocity(0, 0, Grid.Counts()[2] / 2).MakeReal();
-	Velocity(Grid.Counts()[0] / 2, 0, Grid.Counts()[2] / 2).MakeReal();
-	Velocity(0, Grid.Counts()[1] / 2, Grid.Counts()[2] / 2).MakeReal();
-	Velocity(Grid.Counts()[0] / 2, Grid.Counts()[1] / 2, Grid.Counts()[2] / 2).MakeReal();
-	ConjugateAll(Grid, Velocity, 0);
-	ConjugateAll(Grid, Velocity, Grid.Counts()[2] / 2);
-}
-void FFFT3::MakeDivergenceFree(const FUniformGrid& Grid, FArrayNDOfComplex& u, FArrayNDOfComplex& v, FArrayNDOfComplex& w)
-{
-	FVec3 Coefficients = ((FReal)2 * (FReal)UE_PI) / Grid.DomainSize();
-	for (int i = 1; i <= Grid.Counts()[0] / 2; ++i)
-		u(i, 0, 0) = FComplex(0, 0);
-	for (int i = 0; i < Grid.Counts()[0]; ++i)
-	{
-		FReal k1 = Coefficients[0] * (i <= Grid.Counts()[0] / 2 ? (FReal)i : (FReal)i - (FReal)(Grid.Counts()[0]));
-		for (int j = 1; j <= Grid.Counts()[1] / 2; ++j)
+		for(int i = 1; i < Grid.Counts()[1] / 2; ++i)
 		{
-			FReal k2 = Coefficients[1] * (FReal)j;
-			FReal OneOverKSquared = 1 / (k1 * k1 + k2 * k2);
-			FComplex correction = (k1 * u(i, j, 0) + k2 * v(i, j, 0)) * OneOverKSquared;
-			u(i, j, 0) -= correction * k1;
-			v(i, j, 0) -= correction * k2;
+			Velocity(0, Grid.Counts()[1] - i, z) = Velocity(0, i, z).Conjugated();
 		}
-	}
-	// volume
-	for (int i = 0; i < Grid.Counts()[0]; ++i)
-	{
-		FReal k1 = Coefficients[0] * (i <= Grid.Counts()[0] / 2 ? (FReal)i : (FReal)i - (FReal)Grid.Counts()[0]);
-		for (int j = 0; j < Grid.Counts()[1]; ++j)
-		{
-			FReal k2 = Coefficients[1] * (j <= Grid.Counts()[1] / 2 ? (FReal)j : (FReal)j - (FReal)Grid.Counts()[1]);
-			for (int k = 1; k <= Grid.Counts()[2] / 2; ++k)
+		for(int i = 1; i < Grid.Counts()[0]; ++i)
+			for(int j = 1; j < Grid.Counts()[1] / 2; ++j)
 			{
-				FReal k3 = Coefficients[2] * (FReal)k;
-				FReal OneOverKSquared = 1 / (k1 * k1 + k2 * k2 + k3 * k3);
-				FComplex correction = (k1 * u(i, j, k) + k2 * v(i, j, k) + k3 * w(i, j, k)) * OneOverKSquared;
-				u(i, j, k) -= correction * k1;
-				v(i, j, k) -= correction * k2;
-				w(i, j, k) -= correction * k3;
+				Velocity(Grid.Counts()[0] - i, Grid.Counts()[1] - j, z) = Velocity(i, j, z).Conjugated();
 			}
-		}
 	}
-	EnforceSymmetry(Grid, u);
-	EnforceSymmetry(Grid, v);
-	EnforceSymmetry(Grid, w);
-}
-
-// Numerical Receipes
-void exchange(FReal& f1, FReal& f2)
-{
-	FReal tmp = f1;
-	f1 = f2;
-	f2 = tmp;
-}
-
-template<int d>
-void NRFourn(const int isign, const TVector<int32, d>& counts, TArray<FReal>& Data)
-{
-	int idim;
-	unsigned long i1, i2, i3, i2rev, i3rev, ip1, ip2, ip3, ifp1, ifp2;
-	unsigned long ibit, k1, k2, n, nprev, nrem, ntot;
-	FReal tempi, tempr;
-	double theta, wi, wpi, wpr, wr, wtemp;
-
-	ntot = counts.Product();
-	nprev = 1;
-	for (idim = d - 1; idim >= 0; idim--)
+	void EnforceSymmetry(const FFFT3::FUniformGrid& Grid, FFFT3::FArrayNDOfComplex& Velocity)
 	{
-		n = counts[idim];
-		nrem = ntot / (n * nprev);
-		ip1 = nprev << 1;
-		ip2 = ip1 * n;
-		ip3 = ip2 * nrem;
-		i2rev = 1;
-		for (i2 = 1; i2 <= ip2; i2 += ip1)
+		Velocity[0].MakeReal();
+		Velocity(Grid.Counts()[0] / 2, 0, 0).MakeReal();
+		Velocity(0, Grid.Counts()[1] / 2, 0).MakeReal();
+		Velocity(Grid.Counts()[0] / 2, Grid.Counts()[1] / 2, 0).MakeReal();
+		Velocity(0, 0, Grid.Counts()[2] / 2).MakeReal();
+		Velocity(Grid.Counts()[0] / 2, 0, Grid.Counts()[2] / 2).MakeReal();
+		Velocity(0, Grid.Counts()[1] / 2, Grid.Counts()[2] / 2).MakeReal();
+		Velocity(Grid.Counts()[0] / 2, Grid.Counts()[1] / 2, Grid.Counts()[2] / 2).MakeReal();
+		ConjugateAll(Grid, Velocity, 0);
+		ConjugateAll(Grid, Velocity, Grid.Counts()[2] / 2);
+	}
+	void FFFT3::MakeDivergenceFree(const FUniformGrid& Grid, FArrayNDOfComplex& u, FArrayNDOfComplex& v, FArrayNDOfComplex& w)
+	{
+		FVec3 Coefficients = ((FReal)2 * (FReal)UE_PI) / Grid.DomainSize();
+		for(int i = 1; i <= Grid.Counts()[0] / 2; ++i)
+			u(i, 0, 0) = FComplex(0, 0);
+		for(int i = 0; i < Grid.Counts()[0]; ++i)
 		{
-			if (i2 < i2rev)
-				for (i1 = i2; i1 <= i2 + ip1 - 2; i1 += 2)
-					for (i3 = i1; i3 <= ip3; i3 += ip2)
-					{
-						i3rev = i2rev + i3 - i2;
-						exchange(Data[i3 - 1], Data[i3rev - 1]);
-						exchange(Data[i3], Data[i3rev]);
-					}
-			ibit = ip2 >> 1;
-			while (ibit >= ip1 && i2rev > ibit)
+			FReal k1 = Coefficients[0] * (i <= Grid.Counts()[0] / 2 ? (FReal)i : (FReal)i - (FReal)(Grid.Counts()[0]));
+			for(int j = 1; j <= Grid.Counts()[1] / 2; ++j)
 			{
-				i2rev -= ibit;
-				ibit >>= 1;
+				FReal k2 = Coefficients[1] * (FReal)j;
+				FReal OneOverKSquared = 1 / (k1 * k1 + k2 * k2);
+				FComplex correction = (k1 * u(i, j, 0) + k2 * v(i, j, 0)) * OneOverKSquared;
+				u(i, j, 0) -= correction * k1;
+				v(i, j, 0) -= correction * k2;
 			}
-			i2rev += ibit;
 		}
-		ifp1 = ip1;
-		while (ifp1 < ip2)
+		// volume
+		for(int i = 0; i < Grid.Counts()[0]; ++i)
 		{
-			ifp2 = ifp1 << 1;
-			theta = isign * 6.28318530717959 / (FReal)(ifp2 / ip1);
-			wtemp = sin(.5 * theta);
-			wpr = -2 * wtemp * wtemp;
-			wpi = sin(theta);
-			wr = 1;
-			wi = 0;
-			for (i3 = 1; i3 <= ifp1; i3 += ip1)
+			FReal k1 = Coefficients[0] * (i <= Grid.Counts()[0] / 2 ? (FReal)i : (FReal)i - (FReal)Grid.Counts()[0]);
+			for(int j = 0; j < Grid.Counts()[1]; ++j)
 			{
-				for (i1 = i3; i1 <= i3 + ip1 - 2; i1 += 2)
-					for (i2 = i1; i2 <= ip3; i2 += ifp2)
-					{
-						k1 = i2;
-						k2 = k1 + ifp1;
-						tempr = (FReal)wr * Data[k2 - 1] - (FReal)wi * Data[k2];
-						tempi = (FReal)wr * Data[k2] + (FReal)wi * Data[k2 - 1];
-						Data[k2 - 1] = Data[k1 - 1] - tempr;
-						Data[k2] = Data[k1] - tempi;
-						Data[k1 - 1] += tempr;
-						Data[k1] += tempi;
-					}
-				wr = (wtemp = wr) * wpr - wi * wpi + wr;
-				wi = wi * wpr + wtemp * wpi + wi;
+				FReal k2 = Coefficients[1] * (j <= Grid.Counts()[1] / 2 ? (FReal)j : (FReal)j - (FReal)Grid.Counts()[1]);
+				for(int k = 1; k <= Grid.Counts()[2] / 2; ++k)
+				{
+					FReal k3 = Coefficients[2] * (FReal)k;
+					FReal OneOverKSquared = 1 / (k1 * k1 + k2 * k2 + k3 * k3);
+					FComplex correction = (k1 * u(i, j, k) + k2 * v(i, j, k) + k3 * w(i, j, k)) * OneOverKSquared;
+					u(i, j, k) -= correction * k1;
+					v(i, j, k) -= correction * k2;
+					w(i, j, k) -= correction * k3;
+				}
 			}
-			ifp1 = ifp2;
 		}
-		nprev *= n;
+		EnforceSymmetry(Grid, u);
+		EnforceSymmetry(Grid, v);
+		EnforceSymmetry(Grid, w);
 	}
-}
 
-void InverseTransformHelper(const FFFT3::FUniformGrid& Grid, TArrayND<FVec3, 3>& Velocity, const FFFT3::FArrayNDOfComplex& u, const int32 index, const bool Normalize)
-{
-	int32 Size = Grid.Counts().Product();
-	TArray<FReal> Data;
-	Data.SetNum(2 * Size);
-	int32 k = 0;
-	for (int32 i = 0; i < Grid.Counts()[0]; ++i)
+	// Numerical Receipes
+	void exchange(FReal& f1, FReal& f2)
+	{
+		FReal tmp = f1;
+		f1 = f2;
+		f2 = tmp;
+	}
+
+	template<int d>
+	void NRFourn(const int isign, const TVector<int32, d>& counts, TArray<FReal>& Data)
 	{
-		int32 negi = (i == 0) ? 0 : Grid.Counts()[0] - i;
-		for (int32 j = 0; j < Grid.Counts()[1]; ++j)
+		int idim;
+		unsigned long i1, i2, i3, i2rev, i3rev, ip1, ip2, ip3, ifp1, ifp2;
+		unsigned long ibit, k1, k2, n, nprev, nrem, ntot;
+		FReal tempi, tempr;
+		double theta, wi, wpi, wpr, wr, wtemp;
+
+		ntot = counts.Product();
+		nprev = 1;
+		for(idim = d - 1; idim >= 0; idim--)
 		{
-			int32 negj = (j == 0) ? 0 : Grid.Counts()[1] - j;
-			for (int32 ij = 0; ij <= Grid.Counts()[2] / 2; ++ij)
+			n = counts[idim];
+			nrem = ntot / (n * nprev);
+			ip1 = nprev << 1;
+			ip2 = ip1 * n;
+			ip3 = ip2 * nrem;
+			i2rev = 1;
+			for(i2 = 1; i2 <= ip2; i2 += ip1)
 			{
-				Data[k++] = (FReal)u(i, j, ij).Real();
-				Data[k++] = (FReal)u(i, j, ij).Imaginary();
+				if(i2 < i2rev)
+					for(i1 = i2; i1 <= i2 + ip1 - 2; i1 += 2)
+						for(i3 = i1; i3 <= ip3; i3 += ip2)
+						{
+							i3rev = i2rev + i3 - i2;
+							exchange(Data[i3 - 1], Data[i3rev - 1]);
+							exchange(Data[i3], Data[i3rev]);
+						}
+				ibit = ip2 >> 1;
+				while(ibit >= ip1 && i2rev > ibit)
+				{
+					i2rev -= ibit;
+					ibit >>= 1;
+				}
+				i2rev += ibit;
 			}
-			for (int32 ij = Grid.Counts()[2] / 2 + 1; ij < Grid.Counts()[2]; ++ij)
+			ifp1 = ip1;
+			while(ifp1 < ip2)
 			{
-				int32 negij = Grid.Counts()[2] - ij;
-				Data[k++] = (FReal)u(negi, negj, negij).Real();
-				Data[k++] = -(FReal)u(negi, negj, negij).Imaginary();
+				ifp2 = ifp1 << 1;
+				theta = isign * 6.28318530717959 / (FReal)(ifp2 / ip1);
+				wtemp = sin(.5 * theta);
+				wpr = -2 * wtemp * wtemp;
+				wpi = sin(theta);
+				wr = 1;
+				wi = 0;
+				for(i3 = 1; i3 <= ifp1; i3 += ip1)
+				{
+					for(i1 = i3; i1 <= i3 + ip1 - 2; i1 += 2)
+						for(i2 = i1; i2 <= ip3; i2 += ifp2)
+						{
+							k1 = i2;
+							k2 = k1 + ifp1;
+							tempr = (FReal)wr * Data[k2 - 1] - (FReal)wi * Data[k2];
+							tempi = (FReal)wr * Data[k2] + (FReal)wi * Data[k2 - 1];
+							Data[k2 - 1] = Data[k1 - 1] - tempr;
+							Data[k2] = Data[k1] - tempi;
+							Data[k1 - 1] += tempr;
+							Data[k1] += tempi;
+						}
+					wr = (wtemp = wr) * wpr - wi * wpi + wr;
+					wi = wi * wpr + wtemp * wpi + wi;
+				}
+				ifp1 = ifp2;
 			}
+			nprev *= n;
 		}
 	}
-	NRFourn<3>(1, Grid.Counts(), Data);
-	k = 0;
-	const FReal multiplier = Normalize ? ((FReal)1 / (FReal)Size) : 1;
-	for (int32 i = 0; i < Grid.Counts()[0]; ++i)
+
+	void InverseTransformHelper(const FFFT3::FUniformGrid& Grid, TArrayND<FVec3, 3>& Velocity, const FFFT3::FArrayNDOfComplex& u, const int32 index, const bool Normalize)
 	{
-		for (int32 j = 0; j < Grid.Counts()[1]; ++j)
+		int32 Size = Grid.Counts().Product();
+		TArray<FReal> Data;
+		Data.SetNum(2 * Size);
+		int32 k = 0;
+		for(int32 i = 0; i < Grid.Counts()[0]; ++i)
 		{
-			for (int32 ij = 0; ij < Grid.Counts()[2]; ++ij)
+			int32 negi = (i == 0) ? 0 : Grid.Counts()[0] - i;
+			for(int32 j = 0; j < Grid.Counts()[1]; ++j)
 			{
-				Velocity(i, j, ij)[index] = Data[k++] * multiplier;
-				k++;
+				int32 negj = (j == 0) ? 0 : Grid.Counts()[1] - j;
+				for(int32 ij = 0; ij <= Grid.Counts()[2] / 2; ++ij)
+				{
+					Data[k++] = (FReal)u(i, j, ij).Real();
+					Data[k++] = (FReal)u(i, j, ij).Imaginary();
+				}
+				for(int32 ij = Grid.Counts()[2] / 2 + 1; ij < Grid.Counts()[2]; ++ij)
+				{
+					int32 negij = Grid.Counts()[2] - ij;
+					Data[k++] = (FReal)u(negi, negj, negij).Real();
+					Data[k++] = -(FReal)u(negi, negj, negij).Imaginary();
+				}
+			}
+		}
+		NRFourn<3>(1, Grid.Counts(), Data);
+		k = 0;
+		const FReal multiplier = Normalize ? ((FReal)1 / (FReal)Size) : 1;
+		for(int32 i = 0; i < Grid.Counts()[0]; ++i)
+		{
+			for(int32 j = 0; j < Grid.Counts()[1]; ++j)
+			{
+				for(int32 ij = 0; ij < Grid.Counts()[2]; ++ij)
+				{
+					Velocity(i, j, ij)[index] = Data[k++] * multiplier;
+					k++;
+				}
 			}
 		}
 	}
-}
 
-void FFFT3::InverseTransform(const FUniformGrid& Grid, TArrayND<FVec3, 3>& Velocity, const FArrayNDOfComplex& u, const FArrayNDOfComplex& v, const FArrayNDOfComplex& w, const bool Normalize)
-{
-	InverseTransformHelper(Grid, Velocity, u, 0, Normalize);
-	InverseTransformHelper(Grid, Velocity, v, 1, Normalize);
-	InverseTransformHelper(Grid, Velocity, w, 2, Normalize);
-}
+	void FFFT3::InverseTransform(const FUniformGrid& Grid, TArrayND<FVec3, 3>& Velocity, const FArrayNDOfComplex& u, const FArrayNDOfComplex& v, const FArrayNDOfComplex& w, const bool Normalize)
+	{
+		InverseTransformHelper(Grid, Velocity, u, 0, Normalize);
+		InverseTransformHelper(Grid, Velocity, v, 1, Normalize);
+		InverseTransformHelper(Grid, Velocity, w, 2, Normalize);
+	}
 
-void TransformHelper(const FFFT3::FUniformGrid& Grid, const TArrayND<FVec3, 3>& Velocity, FFFT3::FArrayNDOfComplex& u, const int32 index)
-{
-	int32 Size = Grid.Counts().Product();
-	TArray<FReal> Data;
-	Data.SetNum(2 * Size);
-	int32 k = 0;
-	for (int32 i = 0; i < Grid.Counts()[0]; ++i)
+	void TransformHelper(const FFFT3::FUniformGrid& Grid, const TArrayND<FVec3, 3>& Velocity, FFFT3::FArrayNDOfComplex& u, const int32 index)
 	{
-		for (int32 j = 0; j < Grid.Counts()[1]; ++j)
+		int32 Size = Grid.Counts().Product();
+		TArray<FReal> Data;
+		Data.SetNum(2 * Size);
+		int32 k = 0;
+		for(int32 i = 0; i < Grid.Counts()[0]; ++i)
 		{
-			for (int32 ij = 0; ij < Grid.Counts()[2]; ++ij)
+			for(int32 j = 0; j < Grid.Counts()[1]; ++j)
 			{
-				Data[k++] = (FReal)Velocity(i, j, ij)[index];
-				Data[k++] = 0.f;
+				for(int32 ij = 0; ij < Grid.Counts()[2]; ++ij)
+				{
+					Data[k++] = (FReal)Velocity(i, j, ij)[index];
+					Data[k++] = 0.f;
+				}
 			}
 		}
-	}
-	NRFourn<3>(-1, Grid.Counts(), Data);
-	k = 0;
-	for (int32 i = 0; i < Grid.Counts()[0]; ++i)
-	{
-		for (int32 j = 0; j < Grid.Counts()[1]; ++j)
+		NRFourn<3>(-1, Grid.Counts(), Data);
+		k = 0;
+		for(int32 i = 0; i < Grid.Counts()[0]; ++i)
 		{
-			for (int32 ij = 0; ij <= Grid.Counts()[2] / 2; ++ij)
+			for(int32 j = 0; j < Grid.Counts()[1]; ++j)
 			{
-				u(i, j, ij) = FComplex(Data[k], Data[k + 1]);
-				k += 2;
+				for(int32 ij = 0; ij <= Grid.Counts()[2] / 2; ++ij)
+				{
+					u(i, j, ij) = FComplex(Data[k], Data[k + 1]);
+					k += 2;
+				}
+				k += Grid.Counts()[2] - 2;
 			}
-			k += Grid.Counts()[2] - 2;
 		}
 	}
-}
-void FFFT3::Transform(const FUniformGrid& Grid, const TArrayND<FVec3, 3>& Velocity, FArrayNDOfComplex& u, FArrayNDOfComplex& v, FArrayNDOfComplex& w)
-{
-	TransformHelper(Grid, Velocity, u, 0);
-	TransformHelper(Grid, Velocity, v, 1);
-	TransformHelper(Grid, Velocity, w, 2);
+	void FFFT3::Transform(const FUniformGrid& Grid, const TArrayND<FVec3, 3>& Velocity, FArrayNDOfComplex& u, FArrayNDOfComplex& v, FArrayNDOfComplex& w)
+	{
+		TransformHelper(Grid, Velocity, u, 0);
+		TransformHelper(Grid, Velocity, v, 1);
+		TransformHelper(Grid, Velocity, w, 2);
+	}
 }
diff --git a/Engine/Source/Runtime/Experimental/Chaos/Private/Chaos/Framework/Parallel.cpp b/Engine/Source/Runtime/Experimental/Chaos/Private/Chaos/Framework/Parallel.cpp
--- a/Engine/Source/Runtime/Experimental/Chaos/Private/Chaos/Framework/Parallel.cpp
+++ b/Engine/Source/Runtime/Experimental/Chaos/Private/Chaos/Framework/Parallel.cpp
@@ -4,8 +4,6 @@
 #include "Async/ParallelFor.h"
 #include "Framework/Threading.h"
 
-using namespace Chaos;
-
 namespace Chaos
 {
 	CHAOS_API int32 GSingleThreadedPhysics = 0;
diff --git a/Engine/Source/Runtime/Experimental/Chaos/Private/Chaos/ImplicitObject.cpp b/Engine/Source/Runtime/Experimental/Chaos/Private/Chaos/ImplicitObject.cpp
--- a/Engine/Source/Runtime/Experimental/Chaos/Private/Chaos/ImplicitObject.cpp
+++ b/Engine/Source/Runtime/Experimental/Chaos/Private/Chaos/ImplicitObject.cpp
@@ -19,361 +19,361 @@
 #include "UObject/DestructionObjectVersion.h"
 #include "UObject/ReleaseObjectVersion.h"
 
-using namespace Chaos;
-
-FImplicitObject::FImplicitObject(int32 Flags, EImplicitObjectType InType)
-    : Type(InType)
-	, CollisionType(InType)
-	, Margin(0.0f)
-    , bIsConvex(!!(Flags & EImplicitObject::IsConvex))
-    , bDoCollide(!(Flags & EImplicitObject::DisableCollisions))
-    , bHasBoundingBox(!!(Flags & EImplicitObject::HasBoundingBox))
-#if TRACK_CHAOS_GEOMETRY
-    , bIsTracked(false)
-#endif
-{
-}
-
-FImplicitObject::~FImplicitObject()
+namespace Chaos
 {
+	FImplicitObject::FImplicitObject(int32 Flags, EImplicitObjectType InType)
+		: Type(InType)
+		, CollisionType(InType)
+		, Margin(0.0f)
+		, bIsConvex(!!(Flags& EImplicitObject::IsConvex))
+		, bDoCollide(!(Flags& EImplicitObject::DisableCollisions))
+		, bHasBoundingBox(!!(Flags& EImplicitObject::HasBoundingBox))
 #if TRACK_CHAOS_GEOMETRY
-	if (bIsTracked)
+		, bIsTracked(false)
+#endif
 	{
-		FTrackedGeometryManager::Get().RemoveGeometry(this);
 	}
+
+	FImplicitObject::~FImplicitObject()
+	{
+#if TRACK_CHAOS_GEOMETRY
+		if(bIsTracked)
+		{
+			FTrackedGeometryManager::Get().RemoveGeometry(this);
+		}
 #endif
-}
+	}
 
 #if TRACK_CHAOS_GEOMETRY
 
-void FImplicitObject::Track(TSerializablePtr<FImplicitObject> This, const FString& DebugInfo)
-{
-	if (ensure(This.Get() == this))
+	void FImplicitObject::Track(TSerializablePtr<FImplicitObject> This, const FString& DebugInfo)
 	{
-		if (!bIsTracked)
+		if(ensure(This.Get() == this))
 		{
-			bIsTracked = true;
-			FTrackedGeometryManager::Get().AddGeometry(This, DebugInfo);
+			if(!bIsTracked)
+			{
+				bIsTracked = true;
+				FTrackedGeometryManager::Get().AddGeometry(This, DebugInfo);
+			}
 		}
 	}
-}
 
 #endif
 
-EImplicitObjectType FImplicitObject::GetType() const
-{
-	return Type;
-}
-
-EImplicitObjectType FImplicitObject::GetCollisionType() const
-{
-	return CollisionType;
-}
-
-bool FImplicitObject::IsValidGeometry() const
-{
-	return true;
-}
-
-TUniquePtr<FImplicitObject> FImplicitObject::Copy() const
-{
-	check(false);
-	return nullptr;
-}
-
-TUniquePtr<FImplicitObject> FImplicitObject::CopyWithScale(const FVec3& Scale) const
-{
-	check(false);
-	return nullptr;
-}
+	EImplicitObjectType FImplicitObject::GetType() const
+	{
+		return Type;
+	}
 
-bool FImplicitObject::IsUnderlyingUnion() const
-{
-	return Type == ImplicitObjectType::Union;
-}
+	EImplicitObjectType FImplicitObject::GetCollisionType() const
+	{
+		return CollisionType;
+	}
 
-FReal FImplicitObject::SignedDistance(const FVec3& x) const
-{
-	FVec3 Normal;
-	return PhiWithNormal(x, Normal);
-}
+	bool FImplicitObject::IsValidGeometry() const
+	{
+		return true;
+	}
 
-FVec3 FImplicitObject::Normal(const FVec3& x) const
-{
-	FVec3 Normal;
-	PhiWithNormal(x, Normal);
-	return Normal;
-}
+	TUniquePtr<FImplicitObject> FImplicitObject::Copy() const
+	{
+		check(false);
+		return nullptr;
+	}
 
-const FAABB3 FImplicitObject::BoundingBox() const
-{
-	check(false);
-	static const FAABB3 Unbounded(FVec3(-FLT_MAX), FVec3(FLT_MAX));
-	return Unbounded;
-}
+	TUniquePtr<FImplicitObject> FImplicitObject::CopyWithScale(const FVec3& Scale) const
+	{
+		check(false);
+		return nullptr;
+	}
 
-// @todo(mlentine): This is a lot of duplication from the collisions code that should be reduced
-Pair<FVec3, bool> FImplicitObject::FindDeepestIntersection(const FImplicitObject* Other, const FBVHParticles* Particles, const FMatrix33& OtherToLocalTransform, const FReal Thickness) const
-{
-	// Do analytics
-	// @todo(mlentine): Should we do a convex pass here?
-	if (!Particles)
+	bool FImplicitObject::IsUnderlyingUnion() const
 	{
-		return MakePair(FVec3(0), false);
+		return Type == ImplicitObjectType::Union;
 	}
-	FVec3 Point;
-	FReal Phi = Thickness;
-	if (HasBoundingBox())
+
+	FReal FImplicitObject::SignedDistance(const FVec3& x) const
 	{
-		FAABB3 ImplicitBox = BoundingBox().TransformedAABB((FMatrix)OtherToLocalTransform.Inverse());
-		ImplicitBox.Thicken(Thickness);
-		TArray<int32> PotentialParticles = Particles->FindAllIntersections(ImplicitBox);
-		for (int32 i : PotentialParticles)
-		{
-			FVec3 LocalPoint = OtherToLocalTransform.TransformPosition(Particles->X(i));
-			FReal LocalPhi = SignedDistance(LocalPoint);
-			if (LocalPhi < Phi)
-			{
-				Phi = LocalPhi;
-				Point = Particles->X(i);
-			}
-		}
+		FVec3 Normal;
+		return PhiWithNormal(x, Normal);
 	}
-	else
+
+	FVec3 FImplicitObject::Normal(const FVec3& x) const
 	{
-		return FindDeepestIntersection(Other, static_cast<const FParticles*>(Particles), OtherToLocalTransform, Thickness);
+		FVec3 Normal;
+		PhiWithNormal(x, Normal);
+		return Normal;
 	}
-	return MakePair(Point, Phi < Thickness);
-}
 
-Pair<FVec3, bool> FImplicitObject::FindDeepestIntersection(const FImplicitObject* Other, const FParticles* Particles, const FMatrix33& OtherToLocalTransform, const FReal Thickness) const
-{
-	// Do analytics
-	// @todo(mlentine): Should we do a convex pass here?
-	if (!Particles)
+	const FAABB3 FImplicitObject::BoundingBox() const
 	{
-		return MakePair(FVec3(0), false);
+		check(false);
+		static const FAABB3 Unbounded(FVec3(-FLT_MAX), FVec3(FLT_MAX));
+		return Unbounded;
 	}
-	FVec3 Point;
-	FReal Phi = Thickness;
-	int32 NumParticles = Particles->Size();
-	for (int32 i = 0; i < NumParticles; ++i)
+
+	// @todo(mlentine): This is a lot of duplication from the collisions code that should be reduced
+	Pair<FVec3, bool> FImplicitObject::FindDeepestIntersection(const FImplicitObject* Other, const FBVHParticles* Particles, const FMatrix33& OtherToLocalTransform, const FReal Thickness) const
 	{
-		FVec3 LocalPoint = OtherToLocalTransform.TransformPosition(Particles->X(i));
-		FReal LocalPhi = SignedDistance(LocalPoint);
-		if (LocalPhi < Phi)
+		// Do analytics
+		// @todo(mlentine): Should we do a convex pass here?
+		if(!Particles)
+		{
+			return MakePair(FVec3(0), false);
+		}
+		FVec3 Point;
+		FReal Phi = Thickness;
+		if(HasBoundingBox())
+		{
+			FAABB3 ImplicitBox = BoundingBox().TransformedAABB((FMatrix)OtherToLocalTransform.Inverse());
+			ImplicitBox.Thicken(Thickness);
+			TArray<int32> PotentialParticles = Particles->FindAllIntersections(ImplicitBox);
+			for(int32 i : PotentialParticles)
+			{
+				FVec3 LocalPoint = OtherToLocalTransform.TransformPosition(Particles->X(i));
+				FReal LocalPhi = SignedDistance(LocalPoint);
+				if(LocalPhi < Phi)
+				{
+					Phi = LocalPhi;
+					Point = Particles->X(i);
+				}
+			}
+		}
+		else
 		{
-			Phi = LocalPhi;
-			Point = Particles->X(i);
+			return FindDeepestIntersection(Other, static_cast<const FParticles*>(Particles), OtherToLocalTransform, Thickness);
 		}
+		return MakePair(Point, Phi < Thickness);
 	}
-	return MakePair(Point, Phi < Thickness);
-}
 
-Pair<FVec3, bool> FImplicitObject::FindClosestIntersection(const FVec3& StartPoint, const FVec3& EndPoint, const FReal Thickness) const
-{
-	constexpr FReal Epsilon = (FReal)1e-4;
-	constexpr FReal EpsilonSquared = Epsilon * Epsilon;
-
-	//Consider 0 thickness with Start sitting on abs(Phi) < Epsilon. This is a common case; for example a particle sitting perfectly on a floor. In this case intersection could return false.
-	//If start is in this fuzzy region we simply return that spot snapped onto the surface. This is valid because low precision means we don't really know where we are, so let's take the cheapest option
-	//If end is in this fuzzy region it is also a valid hit. However, there could be multiple hits between start and end and since we want the first one, we can't simply return this point.
-	//As such we move end away from start (and out of the fuzzy region) so that we always get a valid intersection if no earlier ones exist
-	//When Thickness > 0 the same idea applies, but we must consider Phi = (Thickness - Epsilon, Thickness + Epsilon)
-	FVec3 Normal;
-	const FReal Phi = PhiWithNormal(StartPoint, Normal);
-	if (FMath::IsNearlyEqual(Phi, Thickness, Epsilon))
+	Pair<FVec3, bool> FImplicitObject::FindDeepestIntersection(const FImplicitObject* Other, const FParticles* Particles, const FMatrix33& OtherToLocalTransform, const FReal Thickness) const
 	{
-		return MakePair(FVec3(StartPoint - Normal * Phi), true); //snap to surface
+		// Do analytics
+		// @todo(mlentine): Should we do a convex pass here?
+		if(!Particles)
+		{
+			return MakePair(FVec3(0), false);
+		}
+		FVec3 Point;
+		FReal Phi = Thickness;
+		int32 NumParticles = Particles->Size();
+		for(int32 i = 0; i < NumParticles; ++i)
+		{
+			FVec3 LocalPoint = OtherToLocalTransform.TransformPosition(Particles->X(i));
+			FReal LocalPhi = SignedDistance(LocalPoint);
+			if(LocalPhi < Phi)
+			{
+				Phi = LocalPhi;
+				Point = Particles->X(i);
+			}
+		}
+		return MakePair(Point, Phi < Thickness);
 	}
 
-	FVec3 ModifiedEnd = EndPoint;
+	Pair<FVec3, bool> FImplicitObject::FindClosestIntersection(const FVec3& StartPoint, const FVec3& EndPoint, const FReal Thickness) const
 	{
-		const FVec3 OriginalStartToEnd = (EndPoint - StartPoint);
-		const FReal OriginalLength2 = OriginalStartToEnd.SizeSquared();
-		if (OriginalLength2 < EpsilonSquared)
+		constexpr FReal Epsilon = (FReal)1e-4;
+		constexpr FReal EpsilonSquared = Epsilon * Epsilon;
+
+		//Consider 0 thickness with Start sitting on abs(Phi) < Epsilon. This is a common case; for example a particle sitting perfectly on a floor. In this case intersection could return false.
+		//If start is in this fuzzy region we simply return that spot snapped onto the surface. This is valid because low precision means we don't really know where we are, so let's take the cheapest option
+		//If end is in this fuzzy region it is also a valid hit. However, there could be multiple hits between start and end and since we want the first one, we can't simply return this point.
+		//As such we move end away from start (and out of the fuzzy region) so that we always get a valid intersection if no earlier ones exist
+		//When Thickness > 0 the same idea applies, but we must consider Phi = (Thickness - Epsilon, Thickness + Epsilon)
+		FVec3 Normal;
+		const FReal Phi = PhiWithNormal(StartPoint, Normal);
+		if(FMath::IsNearlyEqual(Phi, Thickness, Epsilon))
 		{
-			return MakePair(FVec3(0), false); //start was not close to surface, and end is very close to start so no hit
+			return MakePair(FVec3(StartPoint - Normal * Phi), true); //snap to surface
 		}
 
-		FVec3 EndNormal;
-		const FReal EndPhi = PhiWithNormal(EndPoint, EndNormal);
-		if (FMath::IsNearlyEqual(EndPhi, Thickness, Epsilon))
+		FVec3 ModifiedEnd = EndPoint;
 		{
-			//We want to push End out of the fuzzy region. Moving along the normal direction is best since direction could be nearly parallel with fuzzy band
-			//To ensure an intersection, we must go along the normal, but in the same general direction as the ray.
-			const FVec3 OriginalDir = OriginalStartToEnd / FMath::Sqrt(OriginalLength2);
-			const FReal Dot = FVec3::DotProduct(OriginalDir, EndNormal);
-			if (FMath::IsNearlyZero(Dot, Epsilon))
+			const FVec3 OriginalStartToEnd = (EndPoint - StartPoint);
+			const FReal OriginalLength2 = OriginalStartToEnd.SizeSquared();
+			if(OriginalLength2 < EpsilonSquared)
 			{
-				//End is in the fuzzy region, and the direction from start to end is nearly parallel with this fuzzy band, so we should just return End since no other hits will occur
-				return MakePair(FVec3(EndPoint - Normal * Phi), true); //snap to surface
+				return MakePair(FVec3(0), false); //start was not close to surface, and end is very close to start so no hit
 			}
-			else
+
+			FVec3 EndNormal;
+			const FReal EndPhi = PhiWithNormal(EndPoint, EndNormal);
+			if(FMath::IsNearlyEqual(EndPhi, Thickness, Epsilon))
 			{
-				ModifiedEnd = EndPoint + 2.f * Epsilon * FMath::Sign(Dot) * EndNormal; //get out of fuzzy region
+				//We want to push End out of the fuzzy region. Moving along the normal direction is best since direction could be nearly parallel with fuzzy band
+				//To ensure an intersection, we must go along the normal, but in the same general direction as the ray.
+				const FVec3 OriginalDir = OriginalStartToEnd / FMath::Sqrt(OriginalLength2);
+				const FReal Dot = FVec3::DotProduct(OriginalDir, EndNormal);
+				if(FMath::IsNearlyZero(Dot, Epsilon))
+				{
+					//End is in the fuzzy region, and the direction from start to end is nearly parallel with this fuzzy band, so we should just return End since no other hits will occur
+					return MakePair(FVec3(EndPoint - Normal * Phi), true); //snap to surface
+				}
+				else
+				{
+					ModifiedEnd = EndPoint + 2.f * Epsilon * FMath::Sign(Dot) * EndNormal; //get out of fuzzy region
+				}
 			}
 		}
-	}
 
-	return FindClosestIntersectionImp(StartPoint, ModifiedEnd, Thickness);
-}
+		return FindClosestIntersectionImp(StartPoint, ModifiedEnd, Thickness);
+	}
 
-FRealSingle ClosestIntersectionStepSizeMultiplier = 0.5f;
-FAutoConsoleVariableRef CVarClosestIntersectionStepSizeMultiplier(TEXT("p.ClosestIntersectionStepSizeMultiplier"), ClosestIntersectionStepSizeMultiplier, TEXT("When raycasting we use this multiplier to substep the travel distance along the ray. Smaller number gives better accuracy at higher cost"));
+	FRealSingle ClosestIntersectionStepSizeMultiplier = 0.5f;
+	FAutoConsoleVariableRef CVarClosestIntersectionStepSizeMultiplier(TEXT("p.ClosestIntersectionStepSizeMultiplier"), ClosestIntersectionStepSizeMultiplier, TEXT("When raycasting we use this multiplier to substep the travel distance along the ray. Smaller number gives better accuracy at higher cost"));
 
-Pair<FVec3, bool> FImplicitObject::FindClosestIntersectionImp(const FVec3& StartPoint, const FVec3& EndPoint, const FReal Thickness) const
-{
-	FReal Epsilon = (FReal)1e-4;
+	Pair<FVec3, bool> FImplicitObject::FindClosestIntersectionImp(const FVec3& StartPoint, const FVec3& EndPoint, const FReal Thickness) const
+	{
+		FReal Epsilon = (FReal)1e-4;
 
-	FVec3 Ray = EndPoint - StartPoint;
-	FReal Length = Ray.Size();
-	FVec3 Direction = Ray.GetUnsafeNormal(); //this is safe because StartPoint and EndPoint were already tested to be far enough away. In the case where ModifiedEnd is pushed, we push it along the direction so it can only get farther
-	FVec3 EndNormal;
-	const FReal EndPhi = PhiWithNormal(EndPoint, EndNormal);
-	FVec3 ClosestPoint = StartPoint;
+		FVec3 Ray = EndPoint - StartPoint;
+		FReal Length = Ray.Size();
+		FVec3 Direction = Ray.GetUnsafeNormal(); //this is safe because StartPoint and EndPoint were already tested to be far enough away. In the case where ModifiedEnd is pushed, we push it along the direction so it can only get farther
+		FVec3 EndNormal;
+		const FReal EndPhi = PhiWithNormal(EndPoint, EndNormal);
+		FVec3 ClosestPoint = StartPoint;
 
-	FVec3 Normal;
-	FReal Phi = PhiWithNormal(ClosestPoint, Normal);
+		FVec3 Normal;
+		FReal Phi = PhiWithNormal(ClosestPoint, Normal);
 
-	while (Phi > Thickness + Epsilon)
-	{
-		ClosestPoint += Direction * (Phi - Thickness) * (FReal)ClosestIntersectionStepSizeMultiplier;
-		if ((ClosestPoint - StartPoint).Size() > Length)
+		while(Phi > Thickness + Epsilon)
 		{
-			if (EndPhi < Thickness + Epsilon)
+			ClosestPoint += Direction * (Phi - Thickness) * (FReal)ClosestIntersectionStepSizeMultiplier;
+			if((ClosestPoint - StartPoint).Size() > Length)
 			{
-				return MakePair(FVec3(EndPoint + EndNormal * (-EndPhi + Thickness)), true);
+				if(EndPhi < Thickness + Epsilon)
+				{
+					return MakePair(FVec3(EndPoint + EndNormal * (-EndPhi + Thickness)), true);
+				}
+				return MakePair(FVec3(0), false);
 			}
-			return MakePair(FVec3(0), false);
-		}
-		// If the Change is too small we want to nudge it forward. This makes it possible to miss intersections very close to the surface but is more efficient and shouldn't matter much.
-		if ((Phi - Thickness) < (FReal)1e-2)
-		{
-			ClosestPoint += Direction * (FReal)1e-2;
-			if ((ClosestPoint - StartPoint).Size() > Length)
+			// If the Change is too small we want to nudge it forward. This makes it possible to miss intersections very close to the surface but is more efficient and shouldn't matter much.
+			if((Phi - Thickness) < (FReal)1e-2)
 			{
-				if (EndPhi < Thickness + Epsilon)
+				ClosestPoint += Direction * (FReal)1e-2;
+				if((ClosestPoint - StartPoint).Size() > Length)
 				{
-					return MakePair(FVec3(EndPoint + EndNormal * (-EndPhi + Thickness)), true);
+					if(EndPhi < Thickness + Epsilon)
+					{
+						return MakePair(FVec3(EndPoint + EndNormal * (-EndPhi + Thickness)), true);
+					}
+					else
+					{
+						return MakePair(FVec3(0), false);
+					}
 				}
-				else
+			}
+			FReal NewPhi = PhiWithNormal(ClosestPoint, Normal);
+			if(NewPhi >= Phi)
+			{
+				if(EndPhi < Thickness + Epsilon)
 				{
-					return MakePair(FVec3(0), false);
+					return MakePair(FVec3(EndPoint + EndNormal * (-EndPhi + Thickness)), true);
 				}
+				return MakePair(FVec3(0), false);
 			}
+			Phi = NewPhi;
 		}
-		FReal NewPhi = PhiWithNormal(ClosestPoint, Normal);
-		if (NewPhi >= Phi)
+		if(Phi < Thickness + Epsilon)
 		{
-			if (EndPhi < Thickness + Epsilon)
-			{
-				return MakePair(FVec3(EndPoint + EndNormal * (-EndPhi + Thickness)), true);
-			}
-			return MakePair(FVec3(0), false);
+			ClosestPoint += Normal * (-Phi + Thickness);
 		}
-		Phi = NewPhi;
-	}
-	if (Phi < Thickness + Epsilon)
-	{
-		ClosestPoint += Normal * (-Phi + Thickness);
+		return MakePair(ClosestPoint, true);
 	}
-	return MakePair(ClosestPoint, true);
-}
 
-void FImplicitObject::FindAllIntersectingObjects(TArray<Pair<const FImplicitObject*, FRigidTransform3>>& Out, const FAABB3& LocalBounds) const
-{
-	if (!HasBoundingBox() || LocalBounds.Intersects(BoundingBox()))
+	void FImplicitObject::FindAllIntersectingObjects(TArray<Pair<const FImplicitObject*, FRigidTransform3>>& Out, const FAABB3& LocalBounds) const
 	{
-		Out.Add(MakePair(this, FRigidTransform3(FVec3(0), FRotation3::FromElements(FVec3(0), (FReal)1))));
+		if(!HasBoundingBox() || LocalBounds.Intersects(BoundingBox()))
+		{
+			Out.Add(MakePair(this, FRigidTransform3(FVec3(0), FRotation3::FromElements(FVec3(0), (FReal)1))));
+		}
 	}
-}
 
-FArchive& FImplicitObject::SerializeLegacyHelper(FArchive& Ar, TUniquePtr<FImplicitObject>& Value)
-{
-	bool bExists = Value.Get() != nullptr;
-	Ar << bExists;
-	if (bExists)
+	FArchive& FImplicitObject::SerializeLegacyHelper(FArchive& Ar, TUniquePtr<FImplicitObject>& Value)
 	{
-		if (Ar.IsLoading())
-		{
-			uint8 ObjectType;
-			Ar << ObjectType;
-			switch ((EImplicitObjectType)ObjectType)
-			{
-			case ImplicitObjectType::Sphere: { Value = TUniquePtr<TSphere<FReal,3>>(new TSphere<FReal, 3>()); break; }
-			case ImplicitObjectType::Box: { Value = TUniquePtr<TBox<FReal,3>>(new TBox<FReal, 3>()); break; }
-			case ImplicitObjectType::Plane: { Value = TUniquePtr<TPlane<FReal,3>>(new TPlane<FReal, 3>()); break; }
-			case ImplicitObjectType::LevelSet: { Value = TUniquePtr<FLevelSet>(new FLevelSet()); break; }
-			default: check(false);
-			}
-		}
-		else
+		bool bExists = Value.Get() != nullptr;
+		Ar << bExists;
+		if(bExists)
 		{
-			if (Value->Type == ImplicitObjectType::Sphere || Value->Type == ImplicitObjectType::Box || Value->Type == ImplicitObjectType::Plane || Value->Type == ImplicitObjectType::LevelSet)
+			if(Ar.IsLoading())
 			{
-				Ar << Value->Type;
+				uint8 ObjectType;
+				Ar << ObjectType;
+				switch((EImplicitObjectType)ObjectType)
+				{
+				case ImplicitObjectType::Sphere: { Value = TUniquePtr<TSphere<FReal, 3>>(new TSphere<FReal, 3>()); break; }
+				case ImplicitObjectType::Box: { Value = TUniquePtr<TBox<FReal, 3>>(new TBox<FReal, 3>()); break; }
+				case ImplicitObjectType::Plane: { Value = TUniquePtr<TPlane<FReal, 3>>(new TPlane<FReal, 3>()); break; }
+				case ImplicitObjectType::LevelSet: { Value = TUniquePtr<FLevelSet>(new FLevelSet()); break; }
+				default: check(false);
+				}
 			}
 			else
 			{
-				check(false); //should not be serializing this out
+				if(Value->Type == ImplicitObjectType::Sphere || Value->Type == ImplicitObjectType::Box || Value->Type == ImplicitObjectType::Plane || Value->Type == ImplicitObjectType::LevelSet)
+				{
+					Ar << Value->Type;
+				}
+				else
+				{
+					check(false); //should not be serializing this out
+				}
 			}
+			Ar << *Value;
 		}
-		Ar << *Value;
+		return Ar;
 	}
-	return Ar;
-}
 
-void FImplicitObject::SerializeImp(FArchive& Ar)
-{
-	Ar.UsingCustomVersion(FDestructionObjectVersion::GUID);
-	if (Ar.CustomVer(FDestructionObjectVersion::GUID) >= FDestructionObjectVersion::ChaosArchiveAdded)
+	void FImplicitObject::SerializeImp(FArchive& Ar)
 	{
-		Ar << bIsConvex << bDoCollide;
-	}
+		Ar.UsingCustomVersion(FDestructionObjectVersion::GUID);
+		if(Ar.CustomVer(FDestructionObjectVersion::GUID) >= FDestructionObjectVersion::ChaosArchiveAdded)
+		{
+			Ar << bIsConvex << bDoCollide;
+		}
 
-	if (Ar.CustomVer(FDestructionObjectVersion::GUID) <= FDestructionObjectVersion::ImplicitObjectDoCollideAttribute)
-	{
-		bDoCollide = true;
-	}
+		if(Ar.CustomVer(FDestructionObjectVersion::GUID) <= FDestructionObjectVersion::ImplicitObjectDoCollideAttribute)
+		{
+			bDoCollide = true;
+		}
 
-	Ar.UsingCustomVersion(FReleaseObjectVersion::GUID);
-	if (Ar.CustomVer(FReleaseObjectVersion::GUID) > FReleaseObjectVersion::CustomImplicitCollisionType)
-	{
-		Ar << CollisionType;
+		Ar.UsingCustomVersion(FReleaseObjectVersion::GUID);
+		if(Ar.CustomVer(FReleaseObjectVersion::GUID) > FReleaseObjectVersion::CustomImplicitCollisionType)
+		{
+			Ar << CollisionType;
+		}
+		else
+		{
+			CollisionType = Type;
+		}
+
+		// NOTE: Do not serialize Margin in FImplicitObject base class - this is handled by derived types that use it
 	}
-	else
+
+	void FImplicitObject::Serialize(FChaosArchive& Ar)
 	{
-		CollisionType = Type;
+		SerializeImp(Ar);
 	}
 
-	// NOTE: Do not serialize Margin in FImplicitObject base class - this is handled by derived types that use it
-}
-
-void FImplicitObject::Serialize(FChaosArchive& Ar)
-{
-	SerializeImp(Ar);
-}
-
-const FName FImplicitObject::GetTypeName(const EImplicitObjectType InType)
-{
-	static const FName SphereName = TEXT("Sphere");
-	static const FName BoxName = TEXT("Box");
-	static const FName PlaneName = TEXT("Plane");
-	static const FName CapsuleName = TEXT("Capsule");
-	static const FName TransformedName = TEXT("Transformed");
-	static const FName UnionName = TEXT("Union");
-	static const FName LevelSetName = TEXT("LevelSet");
-	static const FName UnknownName = TEXT("Unknown");
-	static const FName ConvexName = TEXT("Convex");
-	static const FName TaperedCylinderName = TEXT("TaperedCylinder");
-	static const FName CylinderName = TEXT("Cylinder");
-	static const FName TriangleMeshName = TEXT("TriangleMesh");
-	static const FName HeightFieldName = TEXT("HeightField");
-	static const FName TaperedCapsuleName = TEXT("TaperedCapsule");
-
-	switch (GetInnerType(InType))
+	const FName FImplicitObject::GetTypeName(const EImplicitObjectType InType)
 	{
+		static const FName SphereName = TEXT("Sphere");
+		static const FName BoxName = TEXT("Box");
+		static const FName PlaneName = TEXT("Plane");
+		static const FName CapsuleName = TEXT("Capsule");
+		static const FName TransformedName = TEXT("Transformed");
+		static const FName UnionName = TEXT("Union");
+		static const FName LevelSetName = TEXT("LevelSet");
+		static const FName UnknownName = TEXT("Unknown");
+		static const FName ConvexName = TEXT("Convex");
+		static const FName TaperedCylinderName = TEXT("TaperedCylinder");
+		static const FName CylinderName = TEXT("Cylinder");
+		static const FName TriangleMeshName = TEXT("TriangleMesh");
+		static const FName HeightFieldName = TEXT("HeightField");
+		static const FName TaperedCapsuleName = TEXT("TaperedCapsule");
+
+		switch(GetInnerType(InType))
+		{
 		case ImplicitObjectType::Sphere: return SphereName;
 		case ImplicitObjectType::Box: return BoxName;
 		case ImplicitObjectType::Plane: return PlaneName;
@@ -388,68 +388,69 @@ const FName FImplicitObject::GetTypeName(const EImplicitObjectType InType)
 		case ImplicitObjectType::TriangleMesh: return TriangleMeshName;
 		case ImplicitObjectType::HeightField: return HeightFieldName;
 		case ImplicitObjectType::TaperedCapsule: return TaperedCapsuleName;
+		}
+		return NAME_None;
 	}
-	return NAME_None;
-}
-
-FImplicitObject* FImplicitObject::SerializationFactory(FChaosArchive& Ar, FImplicitObject* Obj)
-{
-	int8 ObjectType = Ar.IsLoading() ? 0 : (int8)Obj->Type;
-	Ar << ObjectType;
 
-	Ar.UsingCustomVersion(FExternalPhysicsCustomObjectVersion::GUID);
-	if (Ar.CustomVer(FExternalPhysicsCustomObjectVersion::GUID) >= FExternalPhysicsCustomObjectVersion::ScaledGeometryIsConcrete)
+	FImplicitObject* FImplicitObject::SerializationFactory(FChaosArchive& Ar, FImplicitObject* Obj)
 	{
-		if (IsScaled(ObjectType))
+		int8 ObjectType = Ar.IsLoading() ? 0 : (int8)Obj->Type;
+		Ar << ObjectType;
+
+		Ar.UsingCustomVersion(FExternalPhysicsCustomObjectVersion::GUID);
+		if(Ar.CustomVer(FExternalPhysicsCustomObjectVersion::GUID) >= FExternalPhysicsCustomObjectVersion::ScaledGeometryIsConcrete)
+		{
+			if(IsScaled(ObjectType))
+			{
+				EImplicitObjectType InnerType = GetInnerType(ObjectType);
+				switch(InnerType)
+				{
+				case ImplicitObjectType::Convex: if(Ar.IsLoading()){ return new TImplicitObjectScaled<FConvex>(); } break;
+				case ImplicitObjectType::TriangleMesh: if(Ar.IsLoading()){ return new TImplicitObjectScaled<FTriangleMeshImplicitObject>(); } break;
+				default: check(false);
+				}
+
+				return nullptr;
+			}
+		}
+
+		if(IsInstanced(ObjectType))
 		{
 			EImplicitObjectType InnerType = GetInnerType(ObjectType);
-			switch (InnerType)
+			switch(InnerType)
 			{
-			case ImplicitObjectType::Convex: if (Ar.IsLoading()) { return new TImplicitObjectScaled<FConvex>(); } break;
-			case ImplicitObjectType::TriangleMesh: if (Ar.IsLoading()) { return new TImplicitObjectScaled<FTriangleMeshImplicitObject>(); } break;
+			case ImplicitObjectType::Convex: if(Ar.IsLoading()) { return new TImplicitObjectInstanced<FConvex>(); } break;
+			case ImplicitObjectType::TriangleMesh: if(Ar.IsLoading()) { return new TImplicitObjectInstanced<FTriangleMeshImplicitObject>(); } break;
 			default: check(false);
 			}
 
 			return nullptr;
 		}
-	}
 
-	if(IsInstanced(ObjectType))
-	{
-		EImplicitObjectType InnerType = GetInnerType(ObjectType);
-		switch(InnerType)
+		switch((EImplicitObjectType)ObjectType)
+		{
+		case ImplicitObjectType::Sphere: if(Ar.IsLoading()) { return new TSphere<FReal, 3>(); } break;
+		case ImplicitObjectType::Box: if(Ar.IsLoading()) { return new TBox<FReal, 3>(); } break;
+		case ImplicitObjectType::Plane: if(Ar.IsLoading()) { return new TPlane<FReal, 3>(); } break;
+		case ImplicitObjectType::Capsule: if(Ar.IsLoading()) { return new FCapsule(); } break;
+		case ImplicitObjectType::Transformed: if(Ar.IsLoading()) { return new TImplicitObjectTransformed<FReal, 3>(); } break;
+		case ImplicitObjectType::Union: if(Ar.IsLoading()) { return new FImplicitObjectUnion(); } break;
+		case ImplicitObjectType::UnionClustered: if(Ar.IsLoading()) { return new FImplicitObjectUnionClustered(); } break;
+		case ImplicitObjectType::LevelSet: if(Ar.IsLoading()) { return new FLevelSet(); } break;
+		case ImplicitObjectType::Convex: if(Ar.IsLoading()) { return new FConvex(); } break;
+		case ImplicitObjectType::TaperedCylinder: if(Ar.IsLoading()) { return new FTaperedCylinder(); } break;
+		case ImplicitObjectType::TaperedCapsule: if(Ar.IsLoading()) { return new FTaperedCapsule(); } break;
+		case ImplicitObjectType::TriangleMesh: if(Ar.IsLoading()) { return new FTriangleMeshImplicitObject(); } break;
+		case ImplicitObjectType::DEPRECATED_Scaled:
 		{
-		case ImplicitObjectType::Convex: if(Ar.IsLoading()) { return new TImplicitObjectInstanced<FConvex>(); } break;
-		case ImplicitObjectType::TriangleMesh: if(Ar.IsLoading()) { return new TImplicitObjectInstanced<FTriangleMeshImplicitObject>(); } break;
-		default: check(false);
+			ensure(Ar.IsLoading() && (Ar.CustomVer(FExternalPhysicsCustomObjectVersion::GUID) < FExternalPhysicsCustomObjectVersion::ScaledGeometryIsConcrete));
+			return new TImplicitObjectScaledGeneric<FReal, 3>();
+		}
+		case ImplicitObjectType::HeightField: if(Ar.IsLoading()) { return new FHeightField(); } break;
+		case ImplicitObjectType::Cylinder: if(Ar.IsLoading()) { return new FCylinder(); } break;
+		default:
+			check(false);
 		}
-
 		return nullptr;
 	}
-
-	switch ((EImplicitObjectType)ObjectType)
-	{
-	case ImplicitObjectType::Sphere: if (Ar.IsLoading()) { return new TSphere<FReal, 3>(); } break;
-	case ImplicitObjectType::Box: if (Ar.IsLoading()) { return new TBox<FReal, 3>(); } break;
-	case ImplicitObjectType::Plane: if (Ar.IsLoading()) { return new TPlane<FReal, 3>(); } break;
-	case ImplicitObjectType::Capsule: if (Ar.IsLoading()) { return new FCapsule(); } break;
-	case ImplicitObjectType::Transformed: if (Ar.IsLoading()) { return new TImplicitObjectTransformed<FReal, 3>(); } break;
-	case ImplicitObjectType::Union: if (Ar.IsLoading()) { return new FImplicitObjectUnion(); } break;
-	case ImplicitObjectType::UnionClustered: if (Ar.IsLoading()) { return new FImplicitObjectUnionClustered(); } break;
-	case ImplicitObjectType::LevelSet: if (Ar.IsLoading()) { return new FLevelSet(); } break;
-	case ImplicitObjectType::Convex: if (Ar.IsLoading()) { return new FConvex(); } break;
-	case ImplicitObjectType::TaperedCylinder: if (Ar.IsLoading()) { return new FTaperedCylinder(); } break;
-	case ImplicitObjectType::TaperedCapsule: if (Ar.IsLoading()) { return new FTaperedCapsule(); } break;
-	case ImplicitObjectType::TriangleMesh: if (Ar.IsLoading()) { return new FTriangleMeshImplicitObject(); } break;
-	case ImplicitObjectType::DEPRECATED_Scaled:
-	{
-		ensure(Ar.IsLoading() && (Ar.CustomVer(FExternalPhysicsCustomObjectVersion::GUID) < FExternalPhysicsCustomObjectVersion::ScaledGeometryIsConcrete));
-		return new TImplicitObjectScaledGeneric<FReal, 3>();
-	}
-	case ImplicitObjectType::HeightField: if (Ar.IsLoading()) { return new FHeightField(); } break;
-	case ImplicitObjectType::Cylinder: if (Ar.IsLoading()) { return new FCylinder(); } break;
-	default:
-		check(false);
-	}
-	return nullptr;
 }
diff --git a/Engine/Source/Runtime/Experimental/Chaos/Private/Chaos/PBDRigidDynamicSpringConstraints.cpp b/Engine/Source/Runtime/Experimental/Chaos/Private/Chaos/PBDRigidDynamicSpringConstraints.cpp
--- a/Engine/Source/Runtime/Experimental/Chaos/Private/Chaos/PBDRigidDynamicSpringConstraints.cpp
+++ b/Engine/Source/Runtime/Experimental/Chaos/Private/Chaos/PBDRigidDynamicSpringConstraints.cpp
@@ -5,216 +5,217 @@
 #include "Chaos/Particle/ParticleUtilities.h"
 #include "Chaos/Utilities.h"
 
-using namespace Chaos;
-
-FParticlePair FPBDRigidDynamicSpringConstraintHandle::GetConstrainedParticles() const
+namespace Chaos
 {
-	return ConcreteContainer()->GetConstrainedParticles(ConstraintIndex);
-}
+	FParticlePair FPBDRigidDynamicSpringConstraintHandle::GetConstrainedParticles() const
+	{
+		return ConcreteContainer()->GetConstrainedParticles(ConstraintIndex);
+	}
 
-// Note: this is called outside the constraint solver loop (it generates constraints and can merge 
-// islands - the equivalent of collision detection) and therefore reads directly from Particles, 
-// rather than from SolverBodies.
-void FPBDRigidDynamicSpringConstraints::UpdatePositionBasedState(const FReal Dt)
-{
-	const int32 NumConstraints = Constraints.Num();
-	for (int32 ConstraintIndex = 0; ConstraintIndex < NumConstraints; ++ConstraintIndex)
+	// Note: this is called outside the constraint solver loop (it generates constraints and can merge 
+	// islands - the equivalent of collision detection) and therefore reads directly from Particles, 
+	// rather than from SolverBodies.
+	void FPBDRigidDynamicSpringConstraints::UpdatePositionBasedState(const FReal Dt)
 	{
-		FGeometryParticleHandle* Static0 = Constraints[ConstraintIndex][0];
-		FGeometryParticleHandle* Static1 = Constraints[ConstraintIndex][1];
-		FPBDRigidParticleHandle* PBDRigid0 = Static0->CastToRigidParticle();
-		FPBDRigidParticleHandle* PBDRigid1 = Static1->CastToRigidParticle();
-		const bool bIsRigidDynamic0 = PBDRigid0 && PBDRigid0->ObjectState() == EObjectStateType::Dynamic;
-		const bool bIsRigidDynamic1 = PBDRigid1 && PBDRigid1->ObjectState() == EObjectStateType::Dynamic;
-
-		// Do not create springs between objects with no geometry
-		if (!Static0->Geometry() || !Static1->Geometry())
+		const int32 NumConstraints = Constraints.Num();
+		for(int32 ConstraintIndex = 0; ConstraintIndex < NumConstraints; ++ConstraintIndex)
 		{
-			continue;
-		}
+			FGeometryParticleHandle* Static0 = Constraints[ConstraintIndex][0];
+			FGeometryParticleHandle* Static1 = Constraints[ConstraintIndex][1];
+			FPBDRigidParticleHandle* PBDRigid0 = Static0->CastToRigidParticle();
+			FPBDRigidParticleHandle* PBDRigid1 = Static1->CastToRigidParticle();
+			const bool bIsRigidDynamic0 = PBDRigid0 && PBDRigid0->ObjectState() == EObjectStateType::Dynamic;
+			const bool bIsRigidDynamic1 = PBDRigid1 && PBDRigid1->ObjectState() == EObjectStateType::Dynamic;
+
+			// Do not create springs between objects with no geometry
+			if(!Static0->Geometry() || !Static1->Geometry())
+			{
+				continue;
+			}
 
-		const FRotation3 Q0 = bIsRigidDynamic0 ? PBDRigid0->Q() : Static0->R();
-		const FRotation3 Q1 = bIsRigidDynamic1 ? PBDRigid1->Q() : Static1->R();
-		const FVec3& P0 = bIsRigidDynamic0 ? PBDRigid0->P() : Static0->X();
-		const FVec3& P1 = bIsRigidDynamic1 ? PBDRigid1->P() : Static1->X();
+			const FRotation3 Q0 = bIsRigidDynamic0 ? PBDRigid0->Q() : Static0->R();
+			const FRotation3 Q1 = bIsRigidDynamic1 ? PBDRigid1->Q() : Static1->R();
+			const FVec3& P0 = bIsRigidDynamic0 ? PBDRigid0->P() : Static0->X();
+			const FVec3& P1 = bIsRigidDynamic1 ? PBDRigid1->P() : Static1->X();
 
-		// Delete constraints
-		const int32 NumSprings = SpringDistances[ConstraintIndex].Num();
-		for (int32 SpringIndex = NumSprings - 1; SpringIndex >= 0; --SpringIndex)
-		{
-			const FVec3& Distance0 = Distances[ConstraintIndex][SpringIndex][0];
-			const FVec3& Distance1 = Distances[ConstraintIndex][SpringIndex][1];
-			const FVec3 WorldSpaceX1 = Q0.RotateVector(Distance0) + P0;
-			const FVec3 WorldSpaceX2 = Q1.RotateVector(Distance1) + P1;
-			const FVec3 Difference = WorldSpaceX2 - WorldSpaceX1;
-			FReal Distance = Difference.Size();
-			if (Distance > CreationThreshold * 2)
+			// Delete constraints
+			const int32 NumSprings = SpringDistances[ConstraintIndex].Num();
+			for(int32 SpringIndex = NumSprings - 1; SpringIndex >= 0; --SpringIndex)
 			{
-				Distances[ConstraintIndex].RemoveAtSwap(SpringIndex);
-				SpringDistances[ConstraintIndex].RemoveAtSwap(SpringIndex);
+				const FVec3& Distance0 = Distances[ConstraintIndex][SpringIndex][0];
+				const FVec3& Distance1 = Distances[ConstraintIndex][SpringIndex][1];
+				const FVec3 WorldSpaceX1 = Q0.RotateVector(Distance0) + P0;
+				const FVec3 WorldSpaceX2 = Q1.RotateVector(Distance1) + P1;
+				const FVec3 Difference = WorldSpaceX2 - WorldSpaceX1;
+				FReal Distance = Difference.Size();
+				if(Distance > CreationThreshold * 2)
+				{
+					Distances[ConstraintIndex].RemoveAtSwap(SpringIndex);
+					SpringDistances[ConstraintIndex].RemoveAtSwap(SpringIndex);
+				}
 			}
-		}
 
-		if (SpringDistances[ConstraintIndex].Num() == MaxSprings)
-		{
-			continue;
-		}
+			if(SpringDistances[ConstraintIndex].Num() == MaxSprings)
+			{
+				continue;
+			}
 
-		FRigidTransform3 Transform1(P0, Q0);
-		FRigidTransform3 Transform2(P1, Q1);
+			FRigidTransform3 Transform1(P0, Q0);
+			FRigidTransform3 Transform2(P1, Q1);
 
-		// Create constraints
-		if (Static0->Geometry()->HasBoundingBox() && Static1->Geometry()->HasBoundingBox())
-		{
-			// Matrix multiplication is reversed intentionally to be compatible with unreal
-			FAABB3 Box1 = Static0->Geometry()->BoundingBox().TransformedAABB(Transform1 * Transform2.Inverse());
-			Box1.Thicken(CreationThreshold);
-			FAABB3 Box2 = Static1->Geometry()->BoundingBox();
-			Box2.Thicken(CreationThreshold);
-			if (!Box1.Intersects(Box2))
+			// Create constraints
+			if(Static0->Geometry()->HasBoundingBox() && Static1->Geometry()->HasBoundingBox())
+			{
+				// Matrix multiplication is reversed intentionally to be compatible with unreal
+				FAABB3 Box1 = Static0->Geometry()->BoundingBox().TransformedAABB(Transform1 * Transform2.Inverse());
+				Box1.Thicken(CreationThreshold);
+				FAABB3 Box2 = Static1->Geometry()->BoundingBox();
+				Box2.Thicken(CreationThreshold);
+				if(!Box1.Intersects(Box2))
+				{
+					continue;
+				}
+			}
+			const FVec3 Midpoint = (P0 + P1) / (FReal)2.;
+			FVec3 Normal1;
+			const FReal Phi1 = Static0->Geometry()->PhiWithNormal(Transform1.InverseTransformPosition(Midpoint), Normal1);
+			Normal1 = Transform2.TransformVector(Normal1);
+			FVec3 Normal2;
+			const FReal Phi2 = Static1->Geometry()->PhiWithNormal(Transform2.InverseTransformPosition(Midpoint), Normal2);
+			Normal2 = Transform2.TransformVector(Normal2);
+			if((Phi1 + Phi2) > CreationThreshold)
 			{
 				continue;
 			}
+			FVec3 Location0 = Midpoint - Phi1 * Normal1;
+			FVec3 Location1 = Midpoint - Phi2 * Normal2;
+			TVec2<FVec3> Distance;
+			Distance[0] = Q0.Inverse().RotateVector(Location0 - P0);
+			Distance[1] = Q0.Inverse().RotateVector(Location1 - P1);
+			Distances[ConstraintIndex].Add(MoveTemp(Distance));
+			SpringDistances[ConstraintIndex].Add((Location0 - Location1).Size());
 		}
-		const FVec3 Midpoint = (P0 + P1) / (FReal)2.;
-		FVec3 Normal1;
-		const FReal Phi1 = Static0->Geometry()->PhiWithNormal(Transform1.InverseTransformPosition(Midpoint), Normal1);
-		Normal1 = Transform2.TransformVector(Normal1);
-		FVec3 Normal2;
-		const FReal Phi2 = Static1->Geometry()->PhiWithNormal(Transform2.InverseTransformPosition(Midpoint), Normal2);
-		Normal2 = Transform2.TransformVector(Normal2);
-		if ((Phi1 + Phi2) > CreationThreshold)
-		{
-			continue;
-		}
-		FVec3 Location0 = Midpoint - Phi1 * Normal1;
-		FVec3 Location1 = Midpoint - Phi2 * Normal2;
-		TVec2<FVec3> Distance;
-		Distance[0] = Q0.Inverse().RotateVector(Location0 - P0);
-		Distance[1] = Q0.Inverse().RotateVector(Location1 - P1);
-		Distances[ConstraintIndex].Add(MoveTemp(Distance));
-		SpringDistances[ConstraintIndex].Add((Location0 - Location1).Size());
 	}
-}
 
-void FPBDRigidDynamicSpringConstraints::AddConstraintsToGraph(FPBDIslandManager& IslandManager)
-{
-	IslandManager.AddContainerConstraints(*this);
-}
+	void FPBDRigidDynamicSpringConstraints::AddConstraintsToGraph(FPBDIslandManager& IslandManager)
+	{
+		IslandManager.AddContainerConstraints(*this);
+	}
 
-void FPBDRigidDynamicSpringConstraints::AddBodies(FSolverBodyContainer& SolverBodyContainer)
-{
-	for (int32 ConstraintIndex = 0; ConstraintIndex < NumConstraints(); ++ConstraintIndex)
+	void FPBDRigidDynamicSpringConstraints::AddBodies(FSolverBodyContainer& SolverBodyContainer)
 	{
-		AddBodies(ConstraintIndex, SolverBodyContainer);
+		for(int32 ConstraintIndex = 0; ConstraintIndex < NumConstraints(); ++ConstraintIndex)
+		{
+			AddBodies(ConstraintIndex, SolverBodyContainer);
+		}
 	}
-}
 
-void FPBDRigidDynamicSpringConstraints::ScatterOutput(const FReal Dt)
-{
-	for (int32 ConstraintIndex = 0; ConstraintIndex < NumConstraints(); ++ConstraintIndex)
+	void FPBDRigidDynamicSpringConstraints::ScatterOutput(const FReal Dt)
 	{
-		ConstraintSolverBodies[ConstraintIndex] =
+		for(int32 ConstraintIndex = 0; ConstraintIndex < NumConstraints(); ++ConstraintIndex)
 		{
-			nullptr,
-			nullptr
-		};
+			ConstraintSolverBodies[ConstraintIndex] =
+			{
+				nullptr,
+				nullptr
+			};
+		}
 	}
-}
 
-void FPBDRigidDynamicSpringConstraints::ApplyPositionConstraints(const FReal Dt, const int32 It, const int32 NumIts)
-{
-	for (int32 ConstraintIndex = 0; ConstraintIndex < NumConstraints(); ++ConstraintIndex)
+	void FPBDRigidDynamicSpringConstraints::ApplyPositionConstraints(const FReal Dt, const int32 It, const int32 NumIts)
 	{
-		ApplySingle(Dt, ConstraintIndex);
+		for(int32 ConstraintIndex = 0; ConstraintIndex < NumConstraints(); ++ConstraintIndex)
+		{
+			ApplySingle(Dt, ConstraintIndex);
+		}
 	}
-}
 
-void FPBDRigidDynamicSpringConstraints::AddBodies(const TArrayView<int32>& ConstraintIndices, FSolverBodyContainer& SolverBodyContainer)
-{
-	for (int32 ConstraintIndex : ConstraintIndices)
+	void FPBDRigidDynamicSpringConstraints::AddBodies(const TArrayView<int32>& ConstraintIndices, FSolverBodyContainer& SolverBodyContainer)
 	{
-		AddBodies(ConstraintIndex, SolverBodyContainer);
+		for(int32 ConstraintIndex : ConstraintIndices)
+		{
+			AddBodies(ConstraintIndex, SolverBodyContainer);
+		}
 	}
-}
 
-void FPBDRigidDynamicSpringConstraints::ScatterOutput(const TArrayView<int32>& ConstraintIndices, const FReal Dt)
-{
-	for (int32 ConstraintIndex : ConstraintIndices)
+	void FPBDRigidDynamicSpringConstraints::ScatterOutput(const TArrayView<int32>& ConstraintIndices, const FReal Dt)
 	{
-		ConstraintSolverBodies[ConstraintIndex] = 
-		{ 
-			nullptr, 
-			nullptr 
-		};
+		for(int32 ConstraintIndex : ConstraintIndices)
+		{
+			ConstraintSolverBodies[ConstraintIndex] =
+			{
+				nullptr,
+				nullptr
+			};
+		}
 	}
-}
 
-void FPBDRigidDynamicSpringConstraints::ApplyPositionConstraints(const TArrayView<int32>& ConstraintIndices, const FReal Dt, const int32 It, const int32 NumIts)
-{
-	for (int32 ConstraintIndex : ConstraintIndices)
+	void FPBDRigidDynamicSpringConstraints::ApplyPositionConstraints(const TArrayView<int32>& ConstraintIndices, const FReal Dt, const int32 It, const int32 NumIts)
 	{
-		ApplySingle(Dt, ConstraintIndex);
+		for(int32 ConstraintIndex : ConstraintIndices)
+		{
+			ApplySingle(Dt, ConstraintIndex);
+		}
 	}
-}
 
-void FPBDRigidDynamicSpringConstraints::AddBodies(const int32 ConstraintIndex, FSolverBodyContainer& SolverBodyContainer)
-{
-	ConstraintSolverBodies[ConstraintIndex] =
+	void FPBDRigidDynamicSpringConstraints::AddBodies(const int32 ConstraintIndex, FSolverBodyContainer& SolverBodyContainer)
 	{
-		SolverBodyContainer.FindOrAdd(Constraints[ConstraintIndex][0]),
-		SolverBodyContainer.FindOrAdd(Constraints[ConstraintIndex][1])
-	};
-}
+		ConstraintSolverBodies[ConstraintIndex] =
+		{
+			SolverBodyContainer.FindOrAdd(Constraints[ConstraintIndex][0]),
+			SolverBodyContainer.FindOrAdd(Constraints[ConstraintIndex][1])
+		};
+	}
 
-FVec3 FPBDRigidDynamicSpringConstraints::GetDelta(const FVec3& WorldSpaceX1, const FVec3& WorldSpaceX2, const int32 ConstraintIndex, const int32 SpringIndex) const
-{
-	FSolverBody& Body0 = *ConstraintSolverBodies[ConstraintIndex][0];
-	FSolverBody& Body1 = *ConstraintSolverBodies[ConstraintIndex][1];
-	check(Body0.IsDynamic() || Body1.IsDynamic());
+	FVec3 FPBDRigidDynamicSpringConstraints::GetDelta(const FVec3& WorldSpaceX1, const FVec3& WorldSpaceX2, const int32 ConstraintIndex, const int32 SpringIndex) const
+	{
+		FSolverBody& Body0 = *ConstraintSolverBodies[ConstraintIndex][0];
+		FSolverBody& Body1 = *ConstraintSolverBodies[ConstraintIndex][1];
+		check(Body0.IsDynamic() || Body1.IsDynamic());
 
-	const FVec3 Difference = WorldSpaceX2 - WorldSpaceX1;
-	FReal Distance = Difference.Size();
-	check(Distance > 1e-7);
+		const FVec3 Difference = WorldSpaceX2 - WorldSpaceX1;
+		FReal Distance = Difference.Size();
+		check(Distance > 1e-7);
 
-	const FVec3 Direction = Difference / Distance;
-	const FVec3 Delta = (Distance - SpringDistances[ConstraintIndex][SpringIndex]) * Direction;
-	return Stiffness * Delta / (Body0.InvM() + Body1.InvM());
-}
+		const FVec3 Direction = Difference / Distance;
+		const FVec3 Delta = (Distance - SpringDistances[ConstraintIndex][SpringIndex]) * Direction;
+		return Stiffness * Delta / (Body0.InvM() + Body1.InvM());
+	}
 
-void FPBDRigidDynamicSpringConstraints::ApplySingle(const FReal Dt, int32 ConstraintIndex) const
-{
-	check(ConstraintSolverBodies[ConstraintIndex][0] != nullptr);
-	check(ConstraintSolverBodies[ConstraintIndex][1] != nullptr);
-	FSolverBody& Body0 = *ConstraintSolverBodies[ConstraintIndex][0];
-	FSolverBody& Body1 = *ConstraintSolverBodies[ConstraintIndex][1];
-	check(Body0.IsDynamic() || Body1.IsDynamic());
-
-	const int32 NumSprings = SpringDistances[ConstraintIndex].Num();
-	for (int32 SpringIndex = 0; SpringIndex < NumSprings; ++SpringIndex)
+	void FPBDRigidDynamicSpringConstraints::ApplySingle(const FReal Dt, int32 ConstraintIndex) const
 	{
-		const FVec3& Distance0 = Distances[ConstraintIndex][SpringIndex][0];
-		const FVec3& Distance1 = Distances[ConstraintIndex][SpringIndex][1];
-		const FVec3 WorldSpaceX1 = Body0.CorrectedQ().RotateVector(Distance0) + Body0.CorrectedP();
-		const FVec3 WorldSpaceX2 = Body1.CorrectedQ().RotateVector(Distance1) + Body1.CorrectedP();
-		const FVec3 Delta = GetDelta(WorldSpaceX1, WorldSpaceX2, ConstraintIndex, SpringIndex);
+		check(ConstraintSolverBodies[ConstraintIndex][0] != nullptr);
+		check(ConstraintSolverBodies[ConstraintIndex][1] != nullptr);
+		FSolverBody& Body0 = *ConstraintSolverBodies[ConstraintIndex][0];
+		FSolverBody& Body1 = *ConstraintSolverBodies[ConstraintIndex][1];
+		check(Body0.IsDynamic() || Body1.IsDynamic());
 
-		if (Body0.IsDynamic())
+		const int32 NumSprings = SpringDistances[ConstraintIndex].Num();
+		for(int32 SpringIndex = 0; SpringIndex < NumSprings; ++SpringIndex)
 		{
-			const FVec3 Radius = WorldSpaceX1 - Body0.P();
-			const FVec3 DX = Body0.InvM() * Delta;
-			const FVec3 DR = Body0.InvI() * FVec3::CrossProduct(Radius, Delta);
-			Body0.ApplyTransformDelta(DX, DR);
-			Body0.UpdateRotationDependentState();
-		}
+			const FVec3& Distance0 = Distances[ConstraintIndex][SpringIndex][0];
+			const FVec3& Distance1 = Distances[ConstraintIndex][SpringIndex][1];
+			const FVec3 WorldSpaceX1 = Body0.CorrectedQ().RotateVector(Distance0) + Body0.CorrectedP();
+			const FVec3 WorldSpaceX2 = Body1.CorrectedQ().RotateVector(Distance1) + Body1.CorrectedP();
+			const FVec3 Delta = GetDelta(WorldSpaceX1, WorldSpaceX2, ConstraintIndex, SpringIndex);
 
-		if (Body1.IsDynamic())
-		{
-			const FVec3 Radius = WorldSpaceX2 - Body1.P();
-			const FVec3 DX = Body1.InvM() * -Delta;
-			const FVec3 DR = Body1.InvI() * FVec3::CrossProduct(Radius, -Delta);
-			Body1.ApplyTransformDelta(DX, DR);
-			Body1.UpdateRotationDependentState();
+			if(Body0.IsDynamic())
+			{
+				const FVec3 Radius = WorldSpaceX1 - Body0.P();
+				const FVec3 DX = Body0.InvM() * Delta;
+				const FVec3 DR = Body0.InvI() * FVec3::CrossProduct(Radius, Delta);
+				Body0.ApplyTransformDelta(DX, DR);
+				Body0.UpdateRotationDependentState();
+			}
+
+			if(Body1.IsDynamic())
+			{
+				const FVec3 Radius = WorldSpaceX2 - Body1.P();
+				const FVec3 DX = Body1.InvM() * -Delta;
+				const FVec3 DR = Body1.InvI() * FVec3::CrossProduct(Radius, -Delta);
+				Body1.ApplyTransformDelta(DX, DR);
+				Body1.UpdateRotationDependentState();
+			}
 		}
 	}
 }
diff --git a/Engine/Source/Runtime/Experimental/Chaos/Private/Chaos/RigidParticles.cpp b/Engine/Source/Runtime/Experimental/Chaos/Private/Chaos/RigidParticles.cpp
--- a/Engine/Source/Runtime/Experimental/Chaos/Private/Chaos/RigidParticles.cpp
+++ b/Engine/Source/Runtime/Experimental/Chaos/Private/Chaos/RigidParticles.cpp
@@ -1,23 +1,24 @@
 // Copyright Epic Games, Inc. All Rights Reserved.
 #include "Chaos/RigidParticles.h"
 
-using namespace Chaos;
-
 //Note this has to be in the cpp to avoid allocating/freeing across DLLs
 
-template<typename T, int d>
-void TRigidParticles<T, d>::CollisionParticlesInitIfNeeded(const int32 Index)
+namespace Chaos
 {
-	if (MCollisionParticles[Index] == nullptr)
+	template<typename T, int d>
+	void TRigidParticles<T, d>::CollisionParticlesInitIfNeeded(const int32 Index)
 	{
-		MCollisionParticles[Index] = MakeUnique<TBVHParticles<T, d>>();
+		if(MCollisionParticles[Index] == nullptr)
+		{
+			MCollisionParticles[Index] = MakeUnique<TBVHParticles<T, d>>();
+		}
 	}
-}
 
-template<typename T, int d>
-void TRigidParticles<T, d>::SetCollisionParticles(const int32 Index, TParticles<T, d>&& Points)
-{
-	MCollisionParticles[Index] = MakeUnique<TBVHParticles<T, d>>(MoveTemp(Points));
+	template<typename T, int d>
+	void TRigidParticles<T, d>::SetCollisionParticles(const int32 Index, TParticles<T, d>&& Points)
+	{
+		MCollisionParticles[Index] = MakeUnique<TBVHParticles<T, d>>(MoveTemp(Points));
+	}
 }
 
 template class Chaos::TRigidParticles<Chaos::FReal, 3>;
diff --git a/Engine/Source/Runtime/Experimental/Chaos/Private/Chaos/SegmentMesh.cpp b/Engine/Source/Runtime/Experimental/Chaos/Private/Chaos/SegmentMesh.cpp
--- a/Engine/Source/Runtime/Experimental/Chaos/Private/Chaos/SegmentMesh.cpp
+++ b/Engine/Source/Runtime/Experimental/Chaos/Private/Chaos/SegmentMesh.cpp
@@ -3,110 +3,111 @@
 #include "Chaos/SegmentMesh.h"
 #include "Chaos/Particles.h"
 
-using namespace Chaos;
-
-FSegmentMesh::FSegmentMesh(TArray<TVec2<int32>>&& Elements)
-    : MElements(MoveTemp(Elements))
+namespace Chaos
 {
-	// Check for degenerate edges.
-	for (const TVec2<int32>& edge : MElements)
+	FSegmentMesh::FSegmentMesh(TArray<TVec2<int32>>&& Elements)
+		: MElements(MoveTemp(Elements))
 	{
-		check(edge[0] != edge[1]);
+		// Check for degenerate edges.
+		for(const TVec2<int32>& edge : MElements)
+		{
+			check(edge[0] != edge[1]);
+		}
 	}
-}
 
-FSegmentMesh::~FSegmentMesh()
-{}
+	FSegmentMesh::~FSegmentMesh()
+	{}
 
-void FSegmentMesh::_ClearAuxStructures()
-{
-	MPointToEdgeMap.Empty();
-	MPointToNeighborsMap.Empty();
-}
-
-void FSegmentMesh::Init(const TArray<TVec2<int32>>& Elements)
-{
-	_ClearAuxStructures();
-	MElements = Elements;
-	// Check for degenerate edges.
-	for (const TVec2<int32>& edge : MElements)
+	void FSegmentMesh::_ClearAuxStructures()
 	{
-		check(edge[0] != edge[1]);
+		MPointToEdgeMap.Empty();
+		MPointToNeighborsMap.Empty();
 	}
-}
 
-void FSegmentMesh::Init(TArray<TVec2<int32>>&& Elements)
-{
-	_ClearAuxStructures();
-	MElements = MoveTempIfPossible(Elements);
-	// Check for degenerate edges.
-	for (const TVec2<int32>& edge : MElements)
+	void FSegmentMesh::Init(const TArray<TVec2<int32>>& Elements)
 	{
-		check(edge[0] != edge[1]);
+		_ClearAuxStructures();
+		MElements = Elements;
+		// Check for degenerate edges.
+		for(const TVec2<int32>& edge : MElements)
+		{
+			check(edge[0] != edge[1]);
+		}
 	}
-}
 
-const TMap<int32, TSet<int32>>& FSegmentMesh::GetPointToNeighborsMap() const
-{
-	if (!MPointToNeighborsMap.Num())
-		_UpdatePointToNeighborsMap();
-	return MPointToNeighborsMap;
-}
+	void FSegmentMesh::Init(TArray<TVec2<int32>>&& Elements)
+	{
+		_ClearAuxStructures();
+		MElements = MoveTempIfPossible(Elements);
+		// Check for degenerate edges.
+		for(const TVec2<int32>& edge : MElements)
+		{
+			check(edge[0] != edge[1]);
+		}
+	}
 
-void FSegmentMesh::_UpdatePointToNeighborsMap() const
-{
-	MPointToNeighborsMap.Reset();
-	MPointToNeighborsMap.Reserve(MElements.Num() * 2);
-	for (const TVec2<int32>& edge : MElements)
+	const TMap<int32, TSet<int32>>& FSegmentMesh::GetPointToNeighborsMap() const
 	{
-		MPointToNeighborsMap.FindOrAdd(edge[0]).Add(edge[1]);
-		MPointToNeighborsMap.FindOrAdd(edge[1]).Add(edge[0]);
-		// Paranoia:
-		check(MPointToNeighborsMap.Find(edge[0]) != nullptr);
-		check(MPointToNeighborsMap.Find(edge[1]) != nullptr);
-		check(MPointToNeighborsMap.Find(edge[0])->Find(edge[1]) != nullptr);
-		check(MPointToNeighborsMap.Find(edge[1])->Find(edge[0]) != nullptr);
+		if(!MPointToNeighborsMap.Num())
+			_UpdatePointToNeighborsMap();
+		return MPointToNeighborsMap;
 	}
-}
 
-const TMap<int32, TArray<int32>>& FSegmentMesh::GetPointToEdges() const
-{
-	if (!MPointToEdgeMap.Num())
-		_UpdatePointToEdgesMap();
-	return MPointToEdgeMap;
-}
+	void FSegmentMesh::_UpdatePointToNeighborsMap() const
+	{
+		MPointToNeighborsMap.Reset();
+		MPointToNeighborsMap.Reserve(MElements.Num() * 2);
+		for(const TVec2<int32>& edge : MElements)
+		{
+			MPointToNeighborsMap.FindOrAdd(edge[0]).Add(edge[1]);
+			MPointToNeighborsMap.FindOrAdd(edge[1]).Add(edge[0]);
+			// Paranoia:
+			check(MPointToNeighborsMap.Find(edge[0]) != nullptr);
+			check(MPointToNeighborsMap.Find(edge[1]) != nullptr);
+			check(MPointToNeighborsMap.Find(edge[0])->Find(edge[1]) != nullptr);
+			check(MPointToNeighborsMap.Find(edge[1])->Find(edge[0]) != nullptr);
+		}
+	}
 
-void FSegmentMesh::_UpdatePointToEdgesMap() const
-{
-	MPointToEdgeMap.Reset();
-	MPointToEdgeMap.Reserve(MElements.Num() * 2);
-	for (int32 i = 0; i < MElements.Num(); i++)
+	const TMap<int32, TArray<int32>>& FSegmentMesh::GetPointToEdges() const
 	{
-		const TVec2<int32>& edge = MElements[i];
-		MPointToEdgeMap.FindOrAdd(edge[0]).Add(i);
-		MPointToEdgeMap.FindOrAdd(edge[1]).Add(i);
+		if(!MPointToEdgeMap.Num())
+			_UpdatePointToEdgesMap();
+		return MPointToEdgeMap;
 	}
-}
 
-TArray<FReal> FSegmentMesh::GetEdgeLengths(const TParticles<FReal, 3>& InParticles, const bool lengthSquared) const
-{
-	TArray<FReal> lengths;
-	lengths.AddUninitialized(MElements.Num());
-	if (lengthSquared)
+	void FSegmentMesh::_UpdatePointToEdgesMap() const
 	{
-		for (int32 i = 0; i < MElements.Num(); i++)
+		MPointToEdgeMap.Reset();
+		MPointToEdgeMap.Reserve(MElements.Num() * 2);
+		for(int32 i = 0; i < MElements.Num(); i++)
 		{
 			const TVec2<int32>& edge = MElements[i];
-			lengths[i] = (InParticles.X(edge[0]) - InParticles.X(edge[1])).SizeSquared();
+			MPointToEdgeMap.FindOrAdd(edge[0]).Add(i);
+			MPointToEdgeMap.FindOrAdd(edge[1]).Add(i);
 		}
 	}
-	else
+
+	TArray<FReal> FSegmentMesh::GetEdgeLengths(const TParticles<FReal, 3>& InParticles, const bool lengthSquared) const
 	{
-		for (int32 i = 0; i < MElements.Num(); i++)
+		TArray<FReal> lengths;
+		lengths.AddUninitialized(MElements.Num());
+		if(lengthSquared)
 		{
-			const TVec2<int32>& edge = MElements[i];
-			lengths[i] = (InParticles.X(edge[0]) - InParticles.X(edge[1])).Size();
+			for(int32 i = 0; i < MElements.Num(); i++)
+			{
+				const TVec2<int32>& edge = MElements[i];
+				lengths[i] = (InParticles.X(edge[0]) - InParticles.X(edge[1])).SizeSquared();
+			}
+		}
+		else
+		{
+			for(int32 i = 0; i < MElements.Num(); i++)
+			{
+				const TVec2<int32>& edge = MElements[i];
+				lengths[i] = (InParticles.X(edge[0]) - InParticles.X(edge[1])).Size();
+			}
 		}
+		return lengths;
 	}
-	return lengths;
-}
\ No newline at end of file
+}
diff --git a/Engine/Source/Runtime/Experimental/Chaos/Private/Chaos/SmokeEvolution.cpp b/Engine/Source/Runtime/Experimental/Chaos/Private/Chaos/SmokeEvolution.cpp
--- a/Engine/Source/Runtime/Experimental/Chaos/Private/Chaos/SmokeEvolution.cpp
+++ b/Engine/Source/Runtime/Experimental/Chaos/Private/Chaos/SmokeEvolution.cpp
@@ -4,41 +4,46 @@
 #include "Chaos/Framework/Parallel.h"
 #include "Chaos/PerCellBoundaryConditions.h"
 
-using namespace Chaos;
-
-template<class T, int d>
-void TSmokeEvolution<T, d>::AdvanceOneTimeStep(const T Dt)
+namespace Chaos
 {
-	TPerCellBoundaryConditions<T, d> BoundaryRule(MCollisionParticles, MSourceParticles);
-
-	TArrayFaceND<T, d> VelocityN = MVelocity.Copy();
-	TArrayND<T, d> DensityN = MDensity.Copy();
-	PhysicsParallelFor(MGrid.GetNumCells(), [&](int32 Index) {
-		auto CellIndex = MGrid.GetIndex(Index);
-		MAdvectionRule(MGrid, MDensity, DensityN, VelocityN, Dt, CellIndex);
-	});
-	for (int32 i = 0; i < d; ++i)
+	template<class T, int d>
+	void TSmokeEvolution<T, d>::AdvanceOneTimeStep(const T Dt)
 	{
-		TVector<T, d> HalfDx = TVector<T, d>::AxisVector(i) * (MGrid.Dx()[i] / 2);
-		TUniformGrid<T, d> DualGrid(MGrid.MinCorner() - HalfDx, MGrid.MaxCorner() + HalfDx, MGrid.Counts() + TVector<int32, d>::AxisVector(i));
-		PhysicsParallelFor(DualGrid.GetNumCells(), [&](int32 Index) {
-			auto CellIndex = DualGrid.GetIndex(Index);
-			MConvectionRule(DualGrid, MVelocity.GetComponent(i), VelocityN.GetComponent(i), VelocityN, Dt, CellIndex);
+		TPerCellBoundaryConditions<T, d> BoundaryRule(MCollisionParticles, MSourceParticles);
+
+		TArrayFaceND<T, d> VelocityN = MVelocity.Copy();
+		TArrayND<T, d> DensityN = MDensity.Copy();
+		PhysicsParallelFor(MGrid.GetNumCells(), [&](int32 Index)
+		{
+			auto CellIndex = MGrid.GetIndex(Index);
+			MAdvectionRule(MGrid, MDensity, DensityN, VelocityN, Dt, CellIndex);
 		});
-	}
-	PhysicsParallelFor(MGrid.GetNumFaces(), [&](int32 Index) {
-		auto FaceIndex = MGrid.GetFaceIndex(Index);
-		for (auto ForceRule : MForceRules)
+		for(int32 i = 0; i < d; ++i)
 		{
-			ForceRule(MGrid, MVelocity, Dt, FaceIndex);
+			TVector<T, d> HalfDx = TVector<T, d>::AxisVector(i) * (MGrid.Dx()[i] / 2);
+			TUniformGrid<T, d> DualGrid(MGrid.MinCorner() - HalfDx, MGrid.MaxCorner() + HalfDx, MGrid.Counts() + TVector<int32, d>::AxisVector(i));
+			PhysicsParallelFor(DualGrid.GetNumCells(), [&](int32 Index)
+			{
+				auto CellIndex = DualGrid.GetIndex(Index);
+				MConvectionRule(DualGrid, MVelocity.GetComponent(i), VelocityN.GetComponent(i), VelocityN, Dt, CellIndex);
+			});
 		}
-		BoundaryRule.ApplyNeumann(MGrid, MNeumann, MVelocity, Dt, FaceIndex);
-	});
-	PhysicsParallelFor(MGrid.GetNumCells(), [&](int32 Index) {
-		auto CellIndex = MGrid.GetIndex(Index);
-		BoundaryRule.ApplyDirichlet(MGrid, MDirichlet, MDensity, Dt, CellIndex);
-	});
-	MProjectionRule(MGrid, MVelocity, MDirichlet, MNeumann, Dt);
+		PhysicsParallelFor(MGrid.GetNumFaces(), [&](int32 Index)
+		{
+			auto FaceIndex = MGrid.GetFaceIndex(Index);
+			for(auto ForceRule : MForceRules)
+			{
+				ForceRule(MGrid, MVelocity, Dt, FaceIndex);
+			}
+			BoundaryRule.ApplyNeumann(MGrid, MNeumann, MVelocity, Dt, FaceIndex);
+		});
+		PhysicsParallelFor(MGrid.GetNumCells(), [&](int32 Index)
+		{
+			auto CellIndex = MGrid.GetIndex(Index);
+			BoundaryRule.ApplyDirichlet(MGrid, MDirichlet, MDensity, Dt, CellIndex);
+		});
+		MProjectionRule(MGrid, MVelocity, MDirichlet, MNeumann, Dt);
+	}
 }
 
 template class Chaos::TSmokeEvolution<Chaos::FReal, 3>;
diff --git a/Engine/Source/Runtime/Experimental/Chaos/Private/Chaos/Transform.cpp b/Engine/Source/Runtime/Experimental/Chaos/Private/Chaos/Transform.cpp
--- a/Engine/Source/Runtime/Experimental/Chaos/Private/Chaos/Transform.cpp
+++ b/Engine/Source/Runtime/Experimental/Chaos/Private/Chaos/Transform.cpp
@@ -3,17 +3,17 @@
 #include "Chaos/Transform.h"
 #include "Chaos/Real.h"
 
-using namespace Chaos;
-
-PMatrix<FRealSingle, 4, 4> TRigidTransform<FRealSingle, 3>::operator*(const PMatrix<FRealSingle, 4, 4>& Matrix) const
+namespace Chaos
 {
-	// LWC_TODO: Perf pessimization
-	return ToMatrixNoScale() * static_cast<const UE::Math::TMatrix<FRealSingle>&>(Matrix);
-}
+	PMatrix<FRealSingle, 4, 4> TRigidTransform<FRealSingle, 3>::operator*(const PMatrix<FRealSingle, 4, 4>& Matrix) const
+	{
+		// LWC_TODO: Perf pessimization
+		return ToMatrixNoScale() * static_cast<const UE::Math::TMatrix<FRealSingle>&>(Matrix);
+	}
 
-PMatrix<FRealDouble, 4, 4> TRigidTransform<FRealDouble, 3>::operator*(const PMatrix<FRealDouble, 4, 4>& Matrix) const
-{
-	// LWC_TODO: Perf pessimization
-	return ToMatrixNoScale() * static_cast<const UE::Math::TMatrix<FRealDouble>&>(Matrix);
+	PMatrix<FRealDouble, 4, 4> TRigidTransform<FRealDouble, 3>::operator*(const PMatrix<FRealDouble, 4, 4>& Matrix) const
+	{
+		// LWC_TODO: Perf pessimization
+		return ToMatrixNoScale() * static_cast<const UE::Math::TMatrix<FRealDouble>&>(Matrix);
+	}
 }
-
diff --git a/Engine/Source/Runtime/Experimental/Chaos/Private/Chaos/UniformGrid.cpp b/Engine/Source/Runtime/Experimental/Chaos/Private/Chaos/UniformGrid.cpp
--- a/Engine/Source/Runtime/Experimental/Chaos/Private/Chaos/UniformGrid.cpp
+++ b/Engine/Source/Runtime/Experimental/Chaos/Private/Chaos/UniformGrid.cpp
@@ -5,385 +5,403 @@
 #include "Chaos/ArrayND.h"
 #include "Chaos/Core.h"
 
-using namespace Chaos;
-
-template<class T_SCALAR, class T>
-T_SCALAR LinearlyInterpolate1D(const T_SCALAR& Prev, const T_SCALAR& Next, const T Alpha)
+namespace Chaos
 {
-	return  Prev + (Next - Prev)*Alpha;
-}
-
-template<class T_SCALAR, class T, int d>
-T_SCALAR LinearlyInterpolateHelper(const TArrayND<T_SCALAR, d>& ScalarN, const TVector<int32, d>& CellPrev, const TVector<T, d>& Alpha)
-{
-	check(false);
-}
+	template<class T_SCALAR, class T>
+	T_SCALAR LinearlyInterpolate1D(const T_SCALAR& Prev, const T_SCALAR& Next, const T Alpha)
+	{
+		return  Prev + (Next - Prev) * Alpha;
+	}
 
-template<class T_SCALAR, class T>
-T_SCALAR LinearlyInterpolateHelper(const TArrayND<T_SCALAR, 2>& ScalarN, const TVector<int32, 2>& CellPrev, const TVector<T, 2>& Alpha)
-{
-	const T_SCALAR interpx1 = LinearlyInterpolate1D(ScalarN(CellPrev), ScalarN(CellPrev + TVector<int32, 2>({1, 0})), Alpha[0]);
-	const T_SCALAR interpx2 = LinearlyInterpolate1D(ScalarN(CellPrev + TVector<int32, 2>({0, 1})), ScalarN(CellPrev + TVector<int32, 2>({1, 1})), Alpha[0]);
-	return LinearlyInterpolate1D(interpx1, interpx2, Alpha[1]);
-}
+	template<class T_SCALAR, class T, int d>
+	T_SCALAR LinearlyInterpolateHelper(const TArrayND<T_SCALAR, d>& ScalarN, const TVector<int32, d>& CellPrev, const TVector<T, d>& Alpha)
+	{
+		check(false);
+	}
 
-template<class T_SCALAR, class T>
-T_SCALAR LinearlyInterpolateHelper(const TArrayND<T_SCALAR, 3>& ScalarN, const TVector<int32, 3>& CellPrev, const TVector<T, 3>& Alpha)
-{
-	const T_SCALAR interpx1 = LinearlyInterpolate1D(ScalarN(CellPrev), ScalarN(CellPrev + TVector<int32, 3>(1, 0, 0)), Alpha[0]);
-	const T_SCALAR interpx2 = LinearlyInterpolate1D(ScalarN(CellPrev + TVector<int32, 3>(0, 1, 0)), ScalarN(CellPrev + TVector<int32, 3>(1, 1, 0)), Alpha[0]);
-	const T_SCALAR interpx3 = LinearlyInterpolate1D(ScalarN(CellPrev + TVector<int32, 3>(0, 0, 1)), ScalarN(CellPrev + TVector<int32, 3>(1, 0, 1)), Alpha[0]);
-	const T_SCALAR interpx4 = LinearlyInterpolate1D(ScalarN(CellPrev + TVector<int32, 3>(0, 1, 1)), ScalarN(CellPrev + TVector<int32, 3>(1, 1, 1)), Alpha[0]);
-	const T_SCALAR interpy1 = LinearlyInterpolate1D(interpx1, interpx2, Alpha[1]);
-	const T_SCALAR interpy2 = LinearlyInterpolate1D(interpx3, interpx4, Alpha[1]);
-	return LinearlyInterpolate1D(interpy1, interpy2, Alpha[2]);
-}
+	template<class T_SCALAR, class T>
+	T_SCALAR LinearlyInterpolateHelper(const TArrayND<T_SCALAR, 2>& ScalarN, const TVector<int32, 2>& CellPrev, const TVector<T, 2>& Alpha)
+	{
+		const T_SCALAR interpx1 = LinearlyInterpolate1D(ScalarN(CellPrev), ScalarN(CellPrev + TVector<int32, 2>({ 1, 0 })), Alpha[0]);
+		const T_SCALAR interpx2 = LinearlyInterpolate1D(ScalarN(CellPrev + TVector<int32, 2>({ 0, 1 })), ScalarN(CellPrev + TVector<int32, 2>({ 1, 1 })), Alpha[0]);
+		return LinearlyInterpolate1D(interpx1, interpx2, Alpha[1]);
+	}
 
-template<class T, int d>
-template<class T_SCALAR>
-T_SCALAR TUniformGridBase<T, d>::LinearlyInterpolate(const TArrayND<T_SCALAR, d>& ScalarN, const TVector<T, d>& X) const
-{
-	TVector<int32, d> XCell = Cell(X);
-	TVector<T, d> XCenter = Location(XCell);
-	TVector<int32, d> CellPrev;
-	for (int32 i = 0; i < d; ++i)
+	template<class T_SCALAR, class T>
+	T_SCALAR LinearlyInterpolateHelper(const TArrayND<T_SCALAR, 3>& ScalarN, const TVector<int32, 3>& CellPrev, const TVector<T, 3>& Alpha)
 	{
-		CellPrev[i] = X[i] > XCenter[i] ? XCell[i] : XCell[i] - 1;
+		const T_SCALAR interpx1 = LinearlyInterpolate1D(ScalarN(CellPrev), ScalarN(CellPrev + TVector<int32, 3>(1, 0, 0)), Alpha[0]);
+		const T_SCALAR interpx2 = LinearlyInterpolate1D(ScalarN(CellPrev + TVector<int32, 3>(0, 1, 0)), ScalarN(CellPrev + TVector<int32, 3>(1, 1, 0)), Alpha[0]);
+		const T_SCALAR interpx3 = LinearlyInterpolate1D(ScalarN(CellPrev + TVector<int32, 3>(0, 0, 1)), ScalarN(CellPrev + TVector<int32, 3>(1, 0, 1)), Alpha[0]);
+		const T_SCALAR interpx4 = LinearlyInterpolate1D(ScalarN(CellPrev + TVector<int32, 3>(0, 1, 1)), ScalarN(CellPrev + TVector<int32, 3>(1, 1, 1)), Alpha[0]);
+		const T_SCALAR interpy1 = LinearlyInterpolate1D(interpx1, interpx2, Alpha[1]);
+		const T_SCALAR interpy2 = LinearlyInterpolate1D(interpx3, interpx4, Alpha[1]);
+		return LinearlyInterpolate1D(interpy1, interpy2, Alpha[2]);
 	}
-	TVector<T, d> Alpha = (X - Location(CellPrev)) / MDx;
-	// Clamp correctly when on boarder
-	for (int32 i = 0; i < d; ++i)
+
+	template<class T, int d>
+	template<class T_SCALAR>
+	T_SCALAR TUniformGridBase<T, d>::LinearlyInterpolate(const TArrayND<T_SCALAR, d>& ScalarN, const TVector<T, d>& X) const
 	{
-		if (CellPrev[i] == -1)
+		TVector<int32, d> XCell = Cell(X);
+		TVector<T, d> XCenter = Location(XCell);
+		TVector<int32, d> CellPrev;
+		for(int32 i = 0; i < d; ++i)
 		{
-			CellPrev[i] = 0;
-			Alpha[i] = 0;
+			CellPrev[i] = X[i] > XCenter[i] ? XCell[i] : XCell[i] - 1;
 		}
-		if (CellPrev[i] == Counts()[i] - 1)
+		TVector<T, d> Alpha = (X - Location(CellPrev)) / MDx;
+		// Clamp correctly when on boarder
+		for(int32 i = 0; i < d; ++i)
 		{
-			CellPrev[i] = Counts()[i] - 2;
-			Alpha[i] = 1;
+			if(CellPrev[i] == -1)
+			{
+				CellPrev[i] = 0;
+				Alpha[i] = 0;
+			}
+			if(CellPrev[i] == Counts()[i] - 1)
+			{
+				CellPrev[i] = Counts()[i] - 2;
+				Alpha[i] = 1;
+			}
 		}
+		return LinearlyInterpolateHelper(ScalarN, CellPrev, Alpha);
 	}
-	return LinearlyInterpolateHelper(ScalarN, CellPrev, Alpha);
-}
 
-template<class T, int d>
-T TUniformGridBase<T, d>::LinearlyInterpolateComponent(const TArrayND<T, d>& ScalarNComponent, const TVector<T, d>& X, const int32 Axis) const
-{
-	TVector<int32, d> CellCounts = Counts() + TVector<int32, d>::AxisVector(Axis);
-	TVector<int32, d> FaceIndex = Face(X, Axis);
-	TVector<T, d> XCenter = Location(MakePair(Axis, FaceIndex));
-	TVector<int32, d> FacePrev;
-	for (int32 i = 0; i < d; ++i)
+	template<class T, int d>
+	T TUniformGridBase<T, d>::LinearlyInterpolateComponent(const TArrayND<T, d>& ScalarNComponent, const TVector<T, d>& X, const int32 Axis) const
 	{
-		FacePrev[i] = X[i] > XCenter[i] ? FaceIndex[i] : FaceIndex[i] - 1;
-	}
-	TVector<T, d> Alpha = (X - Location(MakePair(Axis, FacePrev))) / MDx;
-	// Clamp correctly when on boarder
-	for (int32 i = 0; i < d; ++i)
-	{
-		if (FacePrev[i] == -1)
+		TVector<int32, d> CellCounts = Counts() + TVector<int32, d>::AxisVector(Axis);
+		TVector<int32, d> FaceIndex = Face(X, Axis);
+		TVector<T, d> XCenter = Location(MakePair(Axis, FaceIndex));
+		TVector<int32, d> FacePrev;
+		for(int32 i = 0; i < d; ++i)
 		{
-			FacePrev[i] = 0;
-			Alpha[i] = 0;
+			FacePrev[i] = X[i] > XCenter[i] ? FaceIndex[i] : FaceIndex[i] - 1;
 		}
-		if (FacePrev[i] == CellCounts[i] - 1)
+		TVector<T, d> Alpha = (X - Location(MakePair(Axis, FacePrev))) / MDx;
+		// Clamp correctly when on boarder
+		for(int32 i = 0; i < d; ++i)
 		{
-			FacePrev[i] = CellCounts[i] - 2;
-			Alpha[i] = 1;
+			if(FacePrev[i] == -1)
+			{
+				FacePrev[i] = 0;
+				Alpha[i] = 0;
+			}
+			if(FacePrev[i] == CellCounts[i] - 1)
+			{
+				FacePrev[i] = CellCounts[i] - 2;
+				Alpha[i] = 1;
+			}
 		}
+		return LinearlyInterpolateHelper(ScalarNComponent, FacePrev, Alpha);
 	}
-	return LinearlyInterpolateHelper(ScalarNComponent, FacePrev, Alpha);
-}
 
-template<class T, int d>
-TVector<T, d> TUniformGridBase<T, d>::LinearlyInterpolate(const TArrayFaceND<T, d>& ScalarN, const TVector<T, d>& X) const
-{
-	TVector<T, d> Result;
-	for (int32 i = 0; i < d; ++i)
+	template<class T, int d>
+	TVector<T, d> TUniformGridBase<T, d>::LinearlyInterpolate(const TArrayFaceND<T, d>& ScalarN, const TVector<T, d>& X) const
 	{
-		Result[i] = LinearlyInterpolateComponent(ScalarN.GetComponent(i), X, i);
+		TVector<T, d> Result;
+		for(int32 i = 0; i < d; ++i)
+		{
+			Result[i] = LinearlyInterpolateComponent(ScalarN.GetComponent(i), X, i);
+		}
+		return Result;
 	}
-	return Result;
-}
 
-template<class T, int d>
-TVector<T, d> TUniformGridBase<T, d>::LinearlyInterpolate(const TArrayFaceND<T, d>& ScalarN, const TVector<T, d>& X, const Pair<int32, TVector<int32, d>> Index) const
-{
-	TVector<T, d> Result;
-	for (int32 i = 0; i < d; ++i)
+	template<class T, int d>
+	TVector<T, d> TUniformGridBase<T, d>::LinearlyInterpolate(const TArrayFaceND<T, d>& ScalarN, const TVector<T, d>& X, const Pair<int32, TVector<int32, d>> Index) const
 	{
-		if (i == Index.First)
+		TVector<T, d> Result;
+		for(int32 i = 0; i < d; ++i)
 		{
-			Result[i] = ScalarN(Index);
+			if(i == Index.First)
+			{
+				Result[i] = ScalarN(Index);
+			}
+			else
+			{
+				Result[i] = LinearlyInterpolateComponent(ScalarN.GetComponent(i), X, Index.First);
+			}
 		}
-		else
+		return Result;
+	}
+
+	template<class T, int d>
+	TVector<int32, d> TUniformGrid<T, d>::GetIndex(const int32 Index) const
+	{
+		TVector<int32, d> NDIndex;
+		int32 product = 1, Remainder = Index;
+		for(int32 i = 0; i < d; ++i)
 		{
-			Result[i] = LinearlyInterpolateComponent(ScalarN.GetComponent(i), X, Index.First);
+			product *= MCells[i];
 		}
+		for(int32 i = 0; i < d; ++i)
+		{
+			product /= MCells[i];
+			NDIndex[i] = Remainder / product;
+			Remainder -= NDIndex[i] * product;
+		}
+		return NDIndex;
 	}
-	return Result;
-}
 
-template<class T, int d>
-TVector<int32, d> TUniformGrid<T, d>::GetIndex(const int32 Index) const
-{
-	TVector<int32, d> NDIndex;
-	int32 product = 1, Remainder = Index;
-	for (int32 i = 0; i < d; ++i)
+	template<class T, int d>
+	TVector<T, d> TUniformGrid<T, d>::Clamp(const TVector<T, d>& X) const
 	{
-		product *= MCells[i];
+		TVector<T, d> Result;
+		for(int32 i = 0; i < d; ++i)
+		{
+			if(X[i] > MMaxCorner[i])
+				Result[i] = MMaxCorner[i];
+			else if(X[i] < MMinCorner[i])
+				Result[i] = MMinCorner[i];
+			else
+				Result[i] = X[i];
+		}
+		return Result;
 	}
-	for (int32 i = 0; i < d; ++i)
+
+	template<class T, int d>
+	TVector<T, d> TUniformGrid<T, d>::ClampMinusHalf(const TVector<T, d>& X) const
 	{
-		product /= MCells[i];
-		NDIndex[i] = Remainder / product;
-		Remainder -= NDIndex[i] * product;
+		TVector<T, d> Result;
+		TVector<T, d> Max = MMaxCorner - MDx * 0.5;
+		TVector<T, d> Min = MMinCorner + MDx * 0.5;
+		for(int32 i = 0; i < d; ++i)
+		{
+			if(X[i] > Max[i])
+				Result[i] = Max[i];
+			else if(X[i] < Min[i])
+				Result[i] = Min[i];
+			else
+				Result[i] = X[i];
+		}
+		return Result;
 	}
-	return NDIndex;
-}
 
-template<class T, int d>
-TVector<T, d> TUniformGrid<T, d>::Clamp(const TVector<T, d>& X) const
-{
-	TVector<T, d> Result;
-	for (int32 i = 0; i < d; ++i)
+	template<class T>
+	TVector<int32, 3> TUniformGrid<T, 3>::GetIndex(const int32 Index) const
 	{
-		if (X[i] > MMaxCorner[i])
-			Result[i] = MMaxCorner[i];
-		else if (X[i] < MMinCorner[i])
-			Result[i] = MMinCorner[i];
-		else
-			Result[i] = X[i];
+		int32 Remainder;
+		TVector<int32, 3> NDIndex;
+		NDIndex[0] = Index / (MCells[1] * MCells[2]);
+		Remainder = Index - NDIndex[0] * MCells[1] * MCells[2];
+		NDIndex[1] = Remainder / MCells[2];
+		Remainder = Remainder - NDIndex[1] * MCells[2];
+		NDIndex[2] = Remainder;
+		return NDIndex;
 	}
-	return Result;
-}
 
-template<class T, int d>
-TVector<T, d> TUniformGrid<T, d>::ClampMinusHalf(const TVector<T, d>& X) const
-{
-	TVector<T, d> Result;
-	TVector<T, d> Max = MMaxCorner - MDx * 0.5;
-	TVector<T, d> Min = MMinCorner + MDx * 0.5;
-	for (int32 i = 0; i < d; ++i)
+	template<class T>
+	Pair<int32, TVector<int32, 3>> TUniformGrid<T, 3>::GetFaceIndex(int32 Index) const
 	{
-		if (X[i] > Max[i])
-			Result[i] = Max[i];
-		else if (X[i] < Min[i])
-			Result[i] = Min[i];
-		else
-			Result[i] = X[i];
+		int32 Remainder;
+		int32 NumXFaces = (MCells + TVector<int32, 3>::AxisVector(0)).Product();
+		int32 NumYFaces = (MCells + TVector<int32, 3>::AxisVector(1)).Product();
+		int32 Axis = 0;
+		if(Index > NumXFaces)
+		{
+			Axis = 1;
+			Index -= NumXFaces;
+			if(Index > NumYFaces)
+			{
+				Axis = 2;
+				Index -= NumYFaces;
+			}
+		}
+		TVector<int32, 3> Faces = MCells + TVector<int32, 3>::AxisVector(Axis);
+		TVector<int32, 3> NDIndex;
+		NDIndex[0] = Index / (Faces[1] * Faces[2]);
+		Remainder = Index - NDIndex[0] * Faces[1] * Faces[2];
+		NDIndex[1] = Remainder / Faces[2];
+		Remainder = Remainder - NDIndex[1] * Faces[2];
+		NDIndex[2] = Remainder;
+		return MakePair(Axis, NDIndex);
 	}
-	return Result;
-}
 
-template<class T>
-TVector<int32, 3> TUniformGrid<T, 3>::GetIndex(const int32 Index) const
-{
-	int32 Remainder;
-	TVector<int32, 3> NDIndex;
-	NDIndex[0] = Index / (MCells[1] * MCells[2]);
-	Remainder = Index - NDIndex[0] * MCells[1] * MCells[2];
-	NDIndex[1] = Remainder / MCells[2];
-	Remainder = Remainder - NDIndex[1] * MCells[2];
-	NDIndex[2] = Remainder;
-	return NDIndex;
-}
+	template<class T>
+	TVector<int32, 3> TUniformGrid<T, 3>::ClampIndex(const TVector<int32, 3>& Index) const
+	{
+		TVector<int32, 3> Result;
+		Result[0] = Index[0] >= MCells[0] ? MCells[0] - 1 : (Index[0] < 0 ? 0 : Index[0]);
+		Result[1] = Index[1] >= MCells[1] ? MCells[1] - 1 : (Index[1] < 0 ? 0 : Index[1]);
+		Result[2] = Index[2] >= MCells[2] ? MCells[2] - 1 : (Index[2] < 0 ? 0 : Index[2]);
+		return Result;
+	}
 
-template<class T>
-Pair<int32, TVector<int32, 3>> TUniformGrid<T, 3>::GetFaceIndex(int32 Index) const
-{
-	int32 Remainder;
-	int32 NumXFaces = (MCells + TVector<int32, 3>::AxisVector(0)).Product();
-	int32 NumYFaces = (MCells + TVector<int32, 3>::AxisVector(1)).Product();
-	int32 Axis = 0;
-	if (Index > NumXFaces)
+	template<class T>
+	TVector<T, 3> TUniformGrid<T, 3>::Clamp(const TVector<T, 3>& X) const
 	{
-		Axis = 1;
-		Index -= NumXFaces;
-		if (Index > NumYFaces)
-		{
-			Axis = 2;
-			Index -= NumYFaces;
-		}
+		TVector<T, 3> Result;
+		Result[0] = X[0] > MMaxCorner[0] ? MMaxCorner[0] : (X[0] < MMinCorner[0] ? MMinCorner[0] : X[0]);
+		Result[1] = X[1] > MMaxCorner[1] ? MMaxCorner[1] : (X[1] < MMinCorner[1] ? MMinCorner[1] : X[1]);
+		Result[2] = X[2] > MMaxCorner[2] ? MMaxCorner[2] : (X[2] < MMinCorner[2] ? MMinCorner[2] : X[2]);
+		return Result;
 	}
-	TVector<int32, 3> Faces = MCells + TVector<int32, 3>::AxisVector(Axis);
-	TVector<int32, 3> NDIndex;
-	NDIndex[0] = Index / (Faces[1] * Faces[2]);
-	Remainder = Index - NDIndex[0] * Faces[1] * Faces[2];
-	NDIndex[1] = Remainder / Faces[2];
-	Remainder = Remainder - NDIndex[1] * Faces[2];
-	NDIndex[2] = Remainder;
-	return MakePair(Axis, NDIndex);
-}
 
-template<class T>
-TVector<int32, 3> TUniformGrid<T, 3>::ClampIndex(const TVector<int32, 3>& Index) const
-{
-	TVector<int32, 3> Result;
-	Result[0] = Index[0] >= MCells[0] ? MCells[0] - 1 : (Index[0] < 0 ? 0 : Index[0]);
-	Result[1] = Index[1] >= MCells[1] ? MCells[1] - 1 : (Index[1] < 0 ? 0 : Index[1]);
-	Result[2] = Index[2] >= MCells[2] ? MCells[2] - 1 : (Index[2] < 0 ? 0 : Index[2]);
-	return Result;
-}
+	template<class T>
+	TVector<T, 3> TUniformGrid<T, 3>::ClampMinusHalf(const TVector<T, 3>& X) const
+	{
+		TVector<T, 3> Result;
+		TVector<T, 3> Max = MMaxCorner - MDx * 0.5;
+		TVector<T, 3> Min = MMinCorner + MDx * 0.5;
+		Result[0] = X[0] > Max[0] ? Max[0] : (X[0] < Min[0] ? Min[0] : X[0]);
+		Result[1] = X[1] > Max[1] ? Max[1] : (X[1] < Min[1] ? Min[1] : X[1]);
+		Result[2] = X[2] > Max[2] ? Max[2] : (X[2] < Min[2] ? Min[2] : X[2]);
+		return Result;
+	}
 
-template<class T>
-TVector<T, 3> TUniformGrid<T, 3>::Clamp(const TVector<T, 3>& X) const
-{
-	TVector<T, 3> Result;
-	Result[0] = X[0] > MMaxCorner[0] ? MMaxCorner[0] : (X[0] < MMinCorner[0] ? MMinCorner[0] : X[0]);
-	Result[1] = X[1] > MMaxCorner[1] ? MMaxCorner[1] : (X[1] < MMinCorner[1] ? MMinCorner[1] : X[1]);
-	Result[2] = X[2] > MMaxCorner[2] ? MMaxCorner[2] : (X[2] < MMinCorner[2] ? MMinCorner[2] : X[2]);
-	return Result;
-}
+	template<class T>
+	void TMPMGrid<T>::BaseNodeIndex(const TVector<T, 3>& X, TVector<int32, 3>& Index, TVector<T, 3>& Weights) const
+	{
+		for(uint32 i = 0; i < 3; ++i)
+			Index[i] = int32(FMath::Floor((X[i] - MMinCorner[i]) / MDx[i]));
+		for(uint32 i = 0; i < 3; ++i)
+			Weights[i] = (X[i] - MMinCorner[i]) / MDx[i] - T(Index[i]);
+	}
 
-template<class T>
-TVector<T, 3> TUniformGrid<T, 3>::ClampMinusHalf(const TVector<T, 3>& X) const
-{
-	TVector<T, 3> Result;
-	TVector<T, 3> Max = MMaxCorner - MDx * 0.5;
-	TVector<T, 3> Min = MMinCorner + MDx * 0.5;
-	Result[0] = X[0] > Max[0] ? Max[0] : (X[0] < Min[0] ? Min[0] : X[0]);
-	Result[1] = X[1] > Max[1] ? Max[1] : (X[1] < Min[1] ? Min[1] : X[1]);
-	Result[2] = X[2] > Max[2] ? Max[2] : (X[2] < Min[2] ? Min[2] : X[2]);
-	return Result;
-}
+	template<class T>
+	int32 TMPMGrid<T>::FlatIndex(const TVector<int32, 3>& Index) const
+	{
+		/*
+		This assumes data will be stored at each index (i.e. numbered by flatIndex(index)). For linear interpolation over
+		the grid we assume that the index[i] \in [0,gridN[i]-1]. For quadratic interpolation, we assume that the
+		index[i] \in [-1,gridN[i]-1].
 
-template<class T>
-void TMPMGrid<T>::BaseNodeIndex(const TVector<T, 3>& X, TVector<int32, 3>& Index, TVector<T, 3>& Weights) const
-{
-	for (uint32 i = 0; i < 3; ++i)
-		Index[i] = int32(FMath::Floor((X[i] - MMinCorner[i]) / MDx[i]));
-	for (uint32 i = 0; i < 3; ++i)
-		Weights[i] = (X[i] - MMinCorner[i]) / MDx[i] - T(Index[i]);
-}
+		For linear interpolation, data is stored at the nodes of the grid, e.g. xi = i dx (i=0,1,...,gridN-1) and
+		for quadratic interpolation, data is stored at the cell centers of the grid, e.g. xi = i dx + dx/2 (i=-1,0,1,...,gridN-1)
+		*/
 
-template<class T>
-int32 TMPMGrid<T>::FlatIndex(const TVector<int32, 3>& Index) const 
-{
-	/*
-	This assumes data will be stored at each index (i.e. numbered by flatIndex(index)). For linear interpolation over
-	the grid we assume that the index[i] \in [0,gridN[i]-1]. For quadratic interpolation, we assume that the
-	index[i] \in [-1,gridN[i]-1].
+		if(interp == linear)
+		{
+			return Index[0] * MCells[1] * MCells[2] + Index[1] * MCells[2] + Index[2];
+		}
+		else
+		{
+			return (MCells[1] + 1) * (MCells[2] + 1) * (Index[0] + 1) + (MCells[2] + 1) * (Index[1] + 1) + Index[2] + 1;
+		}
 
-	For linear interpolation, data is stored at the nodes of the grid, e.g. xi = i dx (i=0,1,...,gridN-1) and
-	for quadratic interpolation, data is stored at the cell centers of the grid, e.g. xi = i dx + dx/2 (i=-1,0,1,...,gridN-1)
-	*/
 
-	if (interp == linear) {
-		return Index[0] * MCells[1] * MCells[2] + Index[1] * MCells[2] + Index[2];
-	}
-	else {
-		return (MCells[1] + 1) * (MCells[2] + 1) * (Index[0] + 1) + (MCells[2] + 1) * (Index[1] + 1) + Index[2] + 1;
+		return -1;
 	}
 
+	template<class T>
+	TVector<int32, 3> TMPMGrid<T>::Lin2MultiIndex(const int32 IndexIn) const
+	{
+		/*
+		This version returns an Index instead of taking one in by reference.
+		*/
+		TVector<int32, 3> MultiIndex;
 
-	return -1;
-}
+		if(interp == linear)
+		{
+			MultiIndex[0] = IndexIn / (MCells[1] * MCells[2]);
+			MultiIndex[1] = (IndexIn / MCells[2]) % MCells[1];
+			MultiIndex[2] = IndexIn % MCells[2];
+		}
+		else
+		{
+			MultiIndex[0] = IndexIn / ((MCells[1] + 1) * (MCells[2] + 1)) - 1;
+			MultiIndex[1] = (IndexIn / (MCells[2] + 1)) % (MCells[1] + 1) - 1;
+			MultiIndex[2] = IndexIn % (MCells[2] + 1) - 1;
+		}
 
-template<class T>
-TVector<int32, 3> TMPMGrid<T>::Lin2MultiIndex(const int32 IndexIn) const 
-{
-	/*
-	This version returns an Index instead of taking one in by reference.
-	*/
-	TVector<int32, 3> MultiIndex;
-
-	if (interp == linear) {
-		MultiIndex[0] = IndexIn / (MCells[1] * MCells[2]);
-		MultiIndex[1] = (IndexIn / MCells[2]) % MCells[1];
-		MultiIndex[2] = IndexIn % MCells[2];
-	}
-	else {
-		MultiIndex[0] = IndexIn / ((MCells[1] + 1) * (MCells[2] + 1)) - 1;
-		MultiIndex[1] = (IndexIn / (MCells[2] + 1)) % (MCells[1] + 1) - 1;
-		MultiIndex[2] = IndexIn % (MCells[2] + 1) - 1;
-	}
 
+		return MultiIndex;
+	}
 
-	return MultiIndex;
-}
 
 
+	template<class T>
+	TVector<T, 3> TMPMGrid<T>::Node(int32 FlatIndexIn) const
+	{
+		TVector<T, 3> result;
+		TVector<int32, 3> index = Lin2MultiIndex(FlatIndexIn);
 
-template<class T>
-TVector<T, 3> TMPMGrid<T>::Node(int32 FlatIndexIn) const {
-	TVector<T, 3> result;
-	TVector<int32, 3> index = Lin2MultiIndex(FlatIndexIn);
-	
-	if (interp == linear) {
-		for (int32 i = 0; i < 3; ++i) {
-			result[i] = T(index[i]) * MDx[i] + MMinCorner[i];
+		if(interp == linear)
+		{
+			for(int32 i = 0; i < 3; ++i)
+			{
+				result[i] = T(index[i]) * MDx[i] + MMinCorner[i];
+			}
 		}
-	}
-	else {
-		for (size_t i = 0; i < 3; ++i) {
-			result[i] = (T(index[i]) + T(0.5)) * MDx[i] + MMinCorner[i];
+		else
+		{
+			for(size_t i = 0; i < 3; ++i)
+			{
+				result[i] = (T(index[i]) + T(0.5)) * MDx[i] + MMinCorner[i];
+			}
 		}
-	}
 
-	return result;
-}
+		return result;
+	}
 
-template<class T>
-int32 TMPMGrid<T>::Loc2GlobIndex(const int32 IndexIn, const TVector<int32, 3>& LocalIndexIn) const
-{
+	template<class T>
+	int32 TMPMGrid<T>::Loc2GlobIndex(const int32 IndexIn, const TVector<int32, 3>& LocalIndexIn) const
+	{
 
-	if (interp == linear) {
-		return IndexIn + MCells[1] * MCells[2] * LocalIndexIn[0] + MCells[2] * LocalIndexIn[1] + LocalIndexIn[2];
-	}
-	else {
-		return IndexIn + (MCells[1] + 1) * (MCells[2] + 1) * (LocalIndexIn[0] - 1) + (MCells[2] + 1) * (LocalIndexIn[1] - 1) + LocalIndexIn[2] - 1;
-	}
+		if(interp == linear)
+		{
+			return IndexIn + MCells[1] * MCells[2] * LocalIndexIn[0] + MCells[2] * LocalIndexIn[1] + LocalIndexIn[2];
+		}
+		else
+		{
+			return IndexIn + (MCells[1] + 1) * (MCells[2] + 1) * (LocalIndexIn[0] - 1) + (MCells[2] + 1) * (LocalIndexIn[1] - 1) + LocalIndexIn[2] - 1;
+		}
 
-	return -1;
+		return -1;
 
-}
+	}
 
-template<class T>
-void TMPMGrid<T>::UpdateGridFromPositions(const Chaos::TDynamicParticles<T, 3>& InParticles)
-{
-	TVector<T, 3> ParticleMin, ParticleMax;
-	if (InParticles.Size() > 0)
+	template<class T>
+	void TMPMGrid<T>::UpdateGridFromPositions(const Chaos::TDynamicParticles<T, 3>& InParticles)
 	{
-		ParticleMin = InParticles.X()[0];
-		ParticleMax = InParticles.X()[0];
-		for (uint32 p = 0; p < InParticles.Size(); p++) {
-			for (int32 alpha = 0; alpha < 3; alpha++) {
-				if (ParticleMin[alpha] > InParticles.X()[p][alpha]) {
-					ParticleMin[alpha] = InParticles.X()[p][alpha];
-				}
-				if (ParticleMax[alpha] < InParticles.X()[p][alpha]) {
-					ParticleMax[alpha] = InParticles.X()[p][alpha];
+		TVector<T, 3> ParticleMin, ParticleMax;
+		if(InParticles.Size() > 0)
+		{
+			ParticleMin = InParticles.X()[0];
+			ParticleMax = InParticles.X()[0];
+			for(uint32 p = 0; p < InParticles.Size(); p++)
+			{
+				for(int32 alpha = 0; alpha < 3; alpha++)
+				{
+					if(ParticleMin[alpha] > InParticles.X()[p][alpha])
+					{
+						ParticleMin[alpha] = InParticles.X()[p][alpha];
+					}
+					if(ParticleMax[alpha] < InParticles.X()[p][alpha])
+					{
+						ParticleMax[alpha] = InParticles.X()[p][alpha];
+					}
 				}
 			}
-		}
-		TVector<T, 3> dims = ParticleMax - ParticleMin;
+			TVector<T, 3> dims = ParticleMax - ParticleMin;
 
-		for (uint32 c = 0; c < 3; c++) {
-			MMinCorner[c] = MDx[c] * T(FMath::CeilToInt((ParticleMin[c] / MDx[c]))) - T(2) * MDx[c];
-			MCells[c] = int32(FMath::CeilToInt((dims[c] / MDx[c]))) + 4;
-			MMaxCorner[c] = MMinCorner[c] + T(MCells[c]) * MDx[c];
+			for(uint32 c = 0; c < 3; c++)
+			{
+				MMinCorner[c] = MDx[c] * T(FMath::CeilToInt((ParticleMin[c] / MDx[c]))) - T(2) * MDx[c];
+				MCells[c] = int32(FMath::CeilToInt((dims[c] / MDx[c]))) + 4;
+				MMaxCorner[c] = MMinCorner[c] + T(MCells[c]) * MDx[c];
+			}
 		}
 	}
-}
 
-template<class T>
-void TMPMGrid<T>::SetInterp(InterpType InterpIn) {
-	interp = InterpIn;
-	if (interp == linear)
-		NPerDir = 2;
-	else
-		NPerDir = 3;
-}
+	template<class T>
+	void TMPMGrid<T>::SetInterp(InterpType InterpIn)
+	{
+		interp = InterpIn;
+		if(interp == linear)
+			NPerDir = 2;
+		else
+			NPerDir = 3;
+	}
 
 
 
-template<class T>
-bool Chaos::TUniformGrid<T, 3>::IsValid(const TVector<int32, 3>& X) const
-{
-	return X == ClampIndex(X);
+	template<class T>
+	bool Chaos::TUniformGrid<T, 3>::IsValid(const TVector<int32, 3>& X) const
+	{
+		return X == ClampIndex(X);
+	}
 }
 
 #if PLATFORM_MAC || PLATFORM_LINUX
@@ -400,5 +418,5 @@ template class Chaos::TMPMGrid<Chaos::FRealSingle>;
 template class Chaos::TMPMGrid<Chaos::FReal>;
 #endif
 
-template FVec3 Chaos::TUniformGridBase<Chaos::FReal, 3>::LinearlyInterpolate<FVec3>(const TArrayND<FVec3, 3>&, const FVec3&) const;
-template CHAOS_API Chaos::FReal Chaos::TUniformGridBase<Chaos::FReal, 3>::LinearlyInterpolate<Chaos::FReal>(const TArrayND<Chaos::FReal, 3>&, const FVec3&) const;
\ No newline at end of file
+template Chaos::FVec3 Chaos::TUniformGridBase<Chaos::FReal, 3>::LinearlyInterpolate<Chaos::FVec3>(const Chaos::TArrayND<Chaos::FVec3, 3>&, const Chaos::FVec3&) const;
+template CHAOS_API Chaos::FReal Chaos::TUniformGridBase<Chaos::FReal, 3>::LinearlyInterpolate<Chaos::FReal>(const TArrayND<Chaos::FReal, 3>&, const Chaos::FVec3&) const;
\ No newline at end of file
diff --git a/Engine/Source/Runtime/Experimental/Chaos/Private/ChaosVisualDebugger/ChaosVisualDebuggerTrace.cpp b/Engine/Source/Runtime/Experimental/Chaos/Private/ChaosVisualDebugger/ChaosVisualDebuggerTrace.cpp
--- a/Engine/Source/Runtime/Experimental/Chaos/Private/ChaosVisualDebugger/ChaosVisualDebuggerTrace.cpp
+++ b/Engine/Source/Runtime/Experimental/Chaos/Private/ChaosVisualDebugger/ChaosVisualDebuggerTrace.cpp
@@ -10,10 +10,8 @@ UE_TRACE_EVENT_FIELD(Chaos::FReal, PositionY)
 UE_TRACE_EVENT_FIELD(Chaos::FReal, PositionZ)
 UE_TRACE_EVENT_END()
 
-using namespace Chaos;
-
 //A very minimal test debug log function
-void ChaosVisualDebugger::ParticlePositionLog(const FVec3& Position)
+void ChaosVisualDebugger::ParticlePositionLog(const Chaos::FVec3& Position)
 {
 	UE_TRACE_LOG(Physics, ParticlePosition, PhysicsChannel)
 					<< ParticlePosition.PositionX(Position.X)
diff --git a/Engine/Source/Runtime/PhysicsCore/Private/ChaosEngineInterface.cpp b/Engine/Source/Runtime/PhysicsCore/Private/ChaosEngineInterface.cpp
--- a/Engine/Source/Runtime/PhysicsCore/Private/ChaosEngineInterface.cpp
+++ b/Engine/Source/Runtime/PhysicsCore/Private/ChaosEngineInterface.cpp
@@ -47,7 +47,7 @@ FPhysicsGeometryCollection_Chaos::FPhysicsGeometryCollection_Chaos(FPhysicsGeome
 
 ECollisionShapeType FPhysicsGeometryCollection_Chaos::GetType() const
 {
-	return GetImplicitType(Geom);
+	return ChaosInterface::GetImplicitType(Geom);
 }
 
 const Chaos::FImplicitObject& FPhysicsGeometryCollection_Chaos::GetGeometry() const
@@ -1625,7 +1625,7 @@ bool FChaosEngineInterface::IsQueryShape(const FPhysicsShapeHandle& InShape)
 
 ECollisionShapeType FChaosEngineInterface::GetShapeType(const FPhysicsShapeReference_Chaos& InShapeRef)
 {
-	return GetImplicitType(*InShapeRef.Shape->GetGeometry());
+	return ChaosInterface::GetImplicitType(*InShapeRef.Shape->GetGeometry());
 }
 
 FTransform FChaosEngineInterface::GetLocalTransform(const FPhysicsShapeReference_Chaos& InShapeRef)
diff --git a/Engine/Source/Runtime/PhysicsCore/Private/PhysTestSerializer.cpp b/Engine/Source/Runtime/PhysicsCore/Private/PhysTestSerializer.cpp
--- a/Engine/Source/Runtime/PhysicsCore/Private/PhysTestSerializer.cpp
+++ b/Engine/Source/Runtime/PhysicsCore/Private/PhysTestSerializer.cpp
@@ -11,8 +11,6 @@
 #include "Chaos/Sphere.h"
 #include "Chaos/Capsule.h"
 
-using namespace Chaos;
-
 #include "PhysicsPublicCore.h"
 #include "PhysicsCore.h"
 #include "HAL/FileManager.h"
@@ -41,7 +39,7 @@ void FPhysTestSerializer::Serialize(const TCHAR* FilePrefix)
 	TUniquePtr<FArchive> File(IFileManager::Get().CreateFileWriter(*UseFileName));
 	if (File)
 	{
-		FChaosArchive Ar(*File);
+		Chaos::FChaosArchive Ar(*File);
 		UE_LOG(LogPhysicsCore, Log, TEXT("PhysTestSerialize File: %s"), *UseFileName);
 		Serialize(Ar);
 	}
@@ -62,7 +60,7 @@ void FPhysTestSerializer::Serialize(Chaos::FChaosArchive& Ar)
 	static const FName TestSerializerName = TEXT("PhysTestSerializer");
 
 	{
-		FChaosArchiveScopedMemory ScopedMemory(Ar, TestSerializerName, false);
+		Chaos::FChaosArchiveScopedMemory ScopedMemory(Ar, TestSerializerName, false);
 		int Version = 1;
 		Ar << Version;
 		Ar << bDiskDataIsChaos;
@@ -95,7 +93,7 @@ void FPhysTestSerializer::Serialize(Chaos::FChaosArchive& Ar)
 
 	bool bHasSQCapture = !!SQCapture;
 	{
-		FChaosArchiveScopedMemory ScopedMemory(Ar, TestSerializerName, false);
+		Chaos::FChaosArchiveScopedMemory ScopedMemory(Ar, TestSerializerName, false);
 		Ar << bHasSQCapture;
 	}
 	if(bHasSQCapture)
@@ -114,7 +112,7 @@ void FPhysTestSerializer::SetPhysicsData(Chaos::FPBDRigidsEvolution& Evolution)
 	bDiskDataIsChaos = true;
 	Data.Empty();
 	FMemoryWriter Ar(Data);
-	FChaosArchive ChaosAr(Ar);
+	Chaos::FChaosArchive ChaosAr(Ar);
 	Evolution.Serialize(ChaosAr);
 	ChaosContext = ChaosAr.StealContext();
 	ArchiveVersion = Ar.GetCustomVersions();
@@ -234,14 +232,14 @@ void FPhysTestSerializer::CreateChaosData()
 			++Idx;
 		}
 
-		ChaosEvolution = MakeUnique<FPBDRigidsEvolution>(Particles, PhysicalMaterials);
+		ChaosEvolution = MakeUnique<Chaos::FPBDRigidsEvolution>(Particles, PhysicalMaterials);
 	}
 	else
 	{
-		ChaosEvolution = MakeUnique<FPBDRigidsEvolution>(Particles, PhysicalMaterials);
+		ChaosEvolution = MakeUnique<Chaos::FPBDRigidsEvolution>(Particles, PhysicalMaterials);
 
 		FMemoryReader Ar(Data);
-		FChaosArchive ChaosAr(Ar);
+		Chaos::FChaosArchive ChaosAr(Ar);
 
 		Ar.SetCustomVersions(ArchiveVersion);
 
diff --git a/Engine/Source/Runtime/PhysicsCore/Private/SQAccelerator.cpp b/Engine/Source/Runtime/PhysicsCore/Private/SQAccelerator.cpp
--- a/Engine/Source/Runtime/PhysicsCore/Private/SQAccelerator.cpp
+++ b/Engine/Source/Runtime/PhysicsCore/Private/SQAccelerator.cpp
@@ -19,25 +19,25 @@ int32 ChaosSQDrawDebugVisitorQueries = 0;
 FAutoConsoleVariableRef CVarChaosSQDrawDebugQueries(TEXT("p.Chaos.SQ.DrawDebugVisitorQueries"), ChaosSQDrawDebugVisitorQueries, TEXT("Draw bounds of objects visited by visitors in scene queries."));
 #endif
 
-void FSQAcceleratorUnion::Raycast(const FVector& Start, const FVector& Dir, const float DeltaMagnitude, FPhysicsHitCallback<FHitRaycast>& HitBuffer, EHitFlags OutputFlags, const FQueryFilterData& QueryFilterData, ICollisionQueryFilterCallbackBase& QueryCallback) const
+void FSQAcceleratorUnion::Raycast(const FVector& Start, const FVector& Dir, const float DeltaMagnitude, FPhysicsHitCallback<FHitRaycast>& HitBuffer, EHitFlags OutputFlags, const ChaosInterface::FQueryFilterData& QueryFilterData, ICollisionQueryFilterCallbackBase& QueryCallback) const
 {
-	for (const ISQAccelerator* Accelerator : Accelerators)
+	for(const ISQAccelerator* Accelerator : Accelerators)
 	{
 		Accelerator->Raycast(Start, Dir, DeltaMagnitude, HitBuffer, OutputFlags, QueryFilterData, QueryCallback);
 	}
 }
 
-void FSQAcceleratorUnion::Sweep(const FPhysicsGeometry& QueryGeom, const FTransform& StartTM, const FVector& Dir, const float DeltaMagnitude, FPhysicsHitCallback<FHitSweep>& HitBuffer, EHitFlags OutputFlags, const FQueryFilterData& QueryFilterData, ICollisionQueryFilterCallbackBase& QueryCallback) const
+void FSQAcceleratorUnion::Sweep(const FPhysicsGeometry& QueryGeom, const FTransform& StartTM, const FVector& Dir, const float DeltaMagnitude, FPhysicsHitCallback<FHitSweep>& HitBuffer, EHitFlags OutputFlags, const ChaosInterface::FQueryFilterData& QueryFilterData, ICollisionQueryFilterCallbackBase& QueryCallback) const
 {
-	for (const ISQAccelerator* Accelerator : Accelerators)
+	for(const ISQAccelerator* Accelerator : Accelerators)
 	{
 		Accelerator->Sweep(QueryGeom, StartTM, Dir, DeltaMagnitude, HitBuffer, OutputFlags, QueryFilterData, QueryCallback);
 	}
 }
 
-void FSQAcceleratorUnion::Overlap(const FPhysicsGeometry& QueryGeom, const FTransform& GeomPose, FPhysicsHitCallback<FHitOverlap>& HitBuffer, const FQueryFilterData& QueryFilterData, ICollisionQueryFilterCallbackBase& QueryCallback) const
+void FSQAcceleratorUnion::Overlap(const FPhysicsGeometry& QueryGeom, const FTransform& GeomPose, FPhysicsHitCallback<FHitOverlap>& HitBuffer, const ChaosInterface::FQueryFilterData& QueryFilterData, ICollisionQueryFilterCallbackBase& QueryCallback) const
 {
-	for (const ISQAccelerator* Accelerator : Accelerators)
+	for(const ISQAccelerator* Accelerator : Accelerators)
 	{
 		Accelerator->Overlap(QueryGeom, GeomPose, HitBuffer, QueryFilterData, QueryCallback);
 	}
@@ -84,19 +84,17 @@ void FillHitHelper(ChaosInterface::FPTLocationHit& Hit, const Chaos::FReal Dista
 }
 
 void FillHitHelper(ChaosInterface::FOverlapHit& Hit, const Chaos::FReal Distance, const FVector& WorldPosition, const FVector& WorldNormal, int32 FaceIdx, bool bComputeMTD)
-{
-}
+{}
 
 void FillHitHelper(ChaosInterface::FPTOverlapHit& Hit, const Chaos::FReal Distance, const FVector& WorldPosition, const FVector& WorldNormal, int32 FaceIdx, bool bComputeMTD)
-{
-}
+{}
 
 template <typename QueryGeometryType, typename TPayload, typename THitType, bool bGTData = true>
 struct TSQVisitor : public Chaos::ISpatialVisitor<TPayload, Chaos::FReal>
 {
 	using TGeometryType = typename TChooseClass<bGTData, Chaos::FGeometryParticle, Chaos::FGeometryParticleHandle>::Result;
 	TSQVisitor(const FVector& InStartPoint, const FVector& InDir, ChaosInterface::FSQHitBuffer<THitType>& InHitBuffer, EHitFlags InOutputFlags,
-		const FQueryFilterData& InQueryFilterData, ICollisionQueryFilterCallbackBase& InQueryCallback, const FQueryDebugParams& InDebugParams)
+			   const ChaosInterface::FQueryFilterData& InQueryFilterData, ICollisionQueryFilterCallbackBase& InQueryCallback, const ChaosInterface::FQueryDebugParams& InDebugParams)
 		: StartPoint(InStartPoint)
 		, Dir(InDir)
 		, HalfExtents(0)
@@ -112,7 +110,7 @@ struct TSQVisitor : public Chaos::ISpatialVisitor<TPayload, Chaos::FReal>
 	}
 
 	TSQVisitor(const FTransform& InStartTM, const FVector& InDir, ChaosInterface::FSQHitBuffer<THitType>& InHitBuffer, EHitFlags InOutputFlags,
-		const FQueryFilterData& InQueryFilterData, ICollisionQueryFilterCallbackBase& InQueryCallback, const QueryGeometryType& InQueryGeom, const FQueryDebugParams& InDebugParams)
+			   const ChaosInterface::FQueryFilterData& InQueryFilterData, ICollisionQueryFilterCallbackBase& InQueryCallback, const QueryGeometryType& InQueryGeom, const ChaosInterface::FQueryDebugParams& InDebugParams)
 		: Dir(InDir)
 		, HalfExtents(InQueryGeom.BoundingBox().Extents() * 0.5)
 		, OutputFlags(InOutputFlags)
@@ -131,7 +129,7 @@ struct TSQVisitor : public Chaos::ISpatialVisitor<TPayload, Chaos::FReal>
 	}
 
 	TSQVisitor(const FTransform& InWorldTM, ChaosInterface::FSQHitBuffer<THitType>& InHitBuffer,
-		const FQueryFilterData& InQueryFilterData, ICollisionQueryFilterCallbackBase& InQueryCallback, const QueryGeometryType& InQueryGeom, const FQueryDebugParams& InDebugParams)
+			   const ChaosInterface::FQueryFilterData& InQueryFilterData, ICollisionQueryFilterCallbackBase& InQueryCallback, const QueryGeometryType& InQueryGeom, const ChaosInterface::FQueryDebugParams& InDebugParams)
 		: HalfExtents(InQueryGeom.BoundingBox().Extents() * 0.5)
 		, bAnyHit(false)
 		, DebugParams(InDebugParams)
@@ -179,7 +177,7 @@ struct TSQVisitor : public Chaos::ISpatialVisitor<TPayload, Chaos::FReal>
 	}
 
 	virtual bool HasBlockingHit() const override
-	{ 
+	{
 		return HitBuffer.HasBlockingHit();
 	}
 
@@ -225,7 +223,7 @@ private:
 
 		const FShapesArray& Shapes = GeometryParticle->ShapesArray();
 
-		const bool bTestShapeBounds =  Shapes.Num() > 1;
+		const bool bTestShapeBounds = Shapes.Num() > 1;
 		bool bContinue = true;
 
 		const FRigidTransform3 ActorTM(GeometryParticle->X(), GeometryParticle->R());
@@ -236,14 +234,14 @@ private:
 		bool bHitBufferIncreased = false;
 #endif
 
-		for (const auto& Shape : Shapes)
+		for(const auto& Shape : Shapes)
 		{
 			const FImplicitObject* Geom = Shape->GetGeometry().Get();
 
-			if (bTestShapeBounds)
+			if(bTestShapeBounds)
 			{
 				FAABB3 InflatedWorldBounds;
-				if (SQ == ESQType::Raycast)
+				if(SQ == ESQType::Raycast)
 				{
 					InflatedWorldBounds = Shape->GetWorldSpaceInflatedShapeBounds();
 				}
@@ -254,12 +252,12 @@ private:
 
 					InflatedWorldBounds = FAABB3(Shape->GetWorldSpaceInflatedShapeBounds().Min() - WorldHalfExtent, Shape->GetWorldSpaceInflatedShapeBounds().Max() + WorldHalfExtent);
 				}
-				if (SQ != ESQType::Overlap)
+				if(SQ != ESQType::Overlap)
 				{
 					//todo: use fast raycast
 					Chaos::FReal TmpTime, TmpExitTime;
 					const FVec3 InflatedBoundsTraceStart = SQ == ESQType::Raycast ? StartPoint : QueryGeomWorldBounds.Center();
-					if (!InflatedWorldBounds.RaycastFast(InflatedBoundsTraceStart, CurData->Dir, CurData->InvDir, CurData->bParallel, CurData->CurrentLength, CurData->InvCurrentLength, TmpTime, TmpExitTime))
+					if(!InflatedWorldBounds.RaycastFast(InflatedBoundsTraceStart, CurData->Dir, CurData->InvDir, CurData->bParallel, CurData->CurrentLength, CurData->InvCurrentLength, TmpTime, TmpExitTime))
 					{
 						continue;
 					}
@@ -267,7 +265,7 @@ private:
 				else
 				{
 					const FVec3 QueryCenter = QueryGeom ? QueryGeomWorldBounds.Center() : StartTM.GetLocation();
-					if (!InflatedWorldBounds.Contains(QueryCenter))
+					if(!InflatedWorldBounds.Contains(QueryCenter))
 					{
 						continue;
 					}
@@ -277,7 +275,7 @@ private:
 			//TODO: use gt particles directly
 			ECollisionQueryHitType HitType = QueryFilterData.flags & FChaosQueryFlag::ePREFILTER ? QueryCallback.PreFilter(QueryFilterDataConcrete, *Shape, *GeometryParticle) : ECollisionQueryHitType::Block;
 
-			if (HitType != ECollisionQueryHitType::None)
+			if(HitType != ECollisionQueryHitType::None)
 			{
 #if CHAOS_DEBUG_DRAW
 				bAllShapesIgnoredInPrefilter = false;
@@ -296,7 +294,7 @@ private:
 				FVec3 FaceNormal = FVec3::ZeroVector;
 				const bool bComputeMTD = !!((uint16)(OutputFlags.HitFlags & EHitFlags::MTD));
 
-				if (SQ == ESQType::Raycast)
+				if(SQ == ESQType::Raycast)
 				{
 					FVec3 LocalNormal;
 					FVec3 LocalPosition;
@@ -304,7 +302,7 @@ private:
 					const FVec3 DirLocal = ActorTM.InverseTransformVectorNoScale(Dir);
 					const FVec3 StartLocal = ActorTM.InverseTransformPositionNoScale(StartPoint);
 					bHit = Geom->Raycast(StartLocal, DirLocal, CurData->CurrentLength, /*Thickness=*/0.0, Distance, LocalPosition, LocalNormal, FaceIdx);
-					if (bHit)
+					if(bHit)
 					{
 						WorldPosition = ActorTM.TransformPositionNoScale(LocalPosition);
 						WorldNormal = ActorTM.TransformVectorNoScale(LocalNormal);
@@ -314,13 +312,13 @@ private:
 				{
 					bHit = SweepQuery(*Geom, ActorTM, *QueryGeom, StartTM, CurData->Dir, CurData->CurrentLength, Distance, WorldPosition, WorldNormal, FaceIdx, FaceNormal, 0.f, bComputeMTD);
 				}
-				else if ((SQ == ESQType::Overlap || (SQ == ESQType::Sweep && CurData->CurrentLength == 0)) && ensure(QueryGeom))
+				else if((SQ == ESQType::Overlap || (SQ == ESQType::Sweep && CurData->CurrentLength == 0)) && ensure(QueryGeom))
 				{
-					if (bComputeMTD)
+					if(bComputeMTD)
 					{
 						FMTDInfo MTDInfo;
 						bHit = OverlapQuery(*Geom, ActorTM, *QueryGeom, StartTM, /*Thickness=*/0, &MTDInfo);
-						if (bHit)
+						if(bHit)
 						{
 							WorldNormal = MTDInfo.Normal * MTDInfo.Penetration;
 						}
@@ -334,9 +332,9 @@ private:
 				if(bHit)
 				{
 					//QUICK_SCOPE_CYCLE_COUNTER(SQNarrowHit);
-					
+
 					bool bAcceptHit = true;
-					if constexpr(SQ == ESQType::Sweep && std::is_same_v<THitType, FSweepHit>)
+					if constexpr(SQ == ESQType::Sweep && std::is_same_v<THitType, ChaosInterface::FSweepHit>)
 					{
 						const THitType* CurrentHit = HitBuffer.GetCurrentHit();
 						if(FaceIdx != INDEX_NONE)
@@ -365,7 +363,7 @@ private:
 									}
 								}
 							}
-							
+
 							if(bAcceptHit)
 							{
 								// Record the new face normal
@@ -412,7 +410,7 @@ private:
 		}
 
 #if CHAOS_DEBUG_DRAW
-		if (DebugParams.IsDebugQuery() && ChaosSQDrawDebugVisitorQueries)
+		if(DebugParams.IsDebugQuery() && ChaosSQDrawDebugVisitorQueries)
 		{
 			DebugDraw<SQ>(Instance, CurData, bAllShapesIgnoredInPrefilter, bHitBufferIncreased);
 		}
@@ -423,29 +421,36 @@ private:
 
 #if CHAOS_DEBUG_DRAW
 
-	void DebugDrawPayloadImpl(const TPayload& Payload, const bool bExternal, const bool bHit, decltype(&TPayload::DebugDraw)) { Payload.DebugDraw(bExternal, bHit); }
-	void DebugDrawPayloadImpl(const TPayload& Payload, const bool bExternal, const bool bHit, ...) { }
-	void DebugDrawPayload(const TPayload& Payload, const bool bExternal, const bool bHit) { DebugDrawPayloadImpl(Payload, bExternal, bHit, 0); }
+	void DebugDrawPayloadImpl(const TPayload& Payload, const bool bExternal, const bool bHit, decltype(&TPayload::DebugDraw))
+	{
+		Payload.DebugDraw(bExternal, bHit);
+	}
+	void DebugDrawPayloadImpl(const TPayload& Payload, const bool bExternal, const bool bHit, ...)
+	{}
+	void DebugDrawPayload(const TPayload& Payload, const bool bExternal, const bool bHit)
+	{
+		DebugDrawPayloadImpl(Payload, bExternal, bHit, 0);
+	}
 
 	template <ESQType SQ>
 	void DebugDraw(const Chaos::TSpatialVisitorData<TPayload>& Instance, const Chaos::FQueryFastData* CurData, const bool bPrefiltered, const bool bHit)
 	{
-		if (SQ == ESQType::Raycast)
+		if(SQ == ESQType::Raycast)
 		{
 			const FVector EndPoint = StartPoint + (Dir * CurData->CurrentLength);
 			Chaos::FDebugDrawQueue::GetInstance().DrawDebugDirectionalArrow(StartPoint, EndPoint, 5.f, bHit ? FColor::Red : FColor::Green);
 		}
-		else if (SQ == ESQType::Overlap)
+		else if(SQ == ESQType::Overlap)
 		{
 			Chaos::DebugDraw::DrawShape(StartTM, QueryGeom, Chaos::FShapeOrShapesArray(), bHit ? FColor::Red : FColor::Green);
 		}
 
-		if (Instance.bHasBounds)
+		if(Instance.bHasBounds)
 		{
 			Chaos::FDebugDrawQueue::GetInstance().DrawDebugBox(Instance.Bounds.Center(), Instance.Bounds.Extents(), FQuat::Identity, bHit ? FColor(100, 50, 50) : FColor(50, 100, 50), false, -1.f, 0, 0.f);
 		}
 
-		if (!bPrefiltered)
+		if(!bPrefiltered)
 		{
 			DebugDrawPayload(Instance.Payload, DebugParams.bExternalQuery, bHit);
 		}
@@ -457,10 +462,10 @@ private:
 	const FVector HalfExtents;
 	FHitFlags OutputFlags;
 	bool bAnyHit;
-	const FQueryDebugParams DebugParams;
+	const ChaosInterface::FQueryDebugParams DebugParams;
 	Chaos::FVec3 HitFaceNormal;
 	ChaosInterface::FSQHitBuffer<THitType>& HitBuffer;
-	const FQueryFilterData& QueryFilterData;
+	const ChaosInterface::FQueryFilterData& QueryFilterData;
 	const FCollisionFilterData QueryFilterDataConcrete;
 	const QueryGeometryType* QueryGeom = nullptr;
 	ICollisionQueryFilterCallbackBase& QueryCallback;
@@ -473,7 +478,7 @@ struct TBPVisitor : public Chaos::ISpatialVisitor<TPayload, Chaos::FReal>
 	using TGeometryType = typename TChooseClass<bGTData, Chaos::FGeometryParticle, Chaos::FGeometryParticleHandle>::Result;
 
 	TBPVisitor(const FTransform& InWorldTM, ChaosInterface::FSQHitBuffer<THitType>& InHitBuffer,
-		const FQueryFilterData& InQueryFilterData, ICollisionQueryFilterCallbackBase& InQueryCallback, const QueryGeometryType& InQueryGeom, const FQueryDebugParams& InDebugParams)
+			   const ChaosInterface::FQueryFilterData& InQueryFilterData, ICollisionQueryFilterCallbackBase& InQueryCallback, const QueryGeometryType& InQueryGeom, const ChaosInterface::FQueryDebugParams& InDebugParams)
 		: HalfExtents(InQueryGeom.BoundingBox().Extents() * 0.5)
 		, bAnyHit(false)
 		, DebugParams(InDebugParams)
@@ -512,7 +517,7 @@ struct TBPVisitor : public Chaos::ISpatialVisitor<TPayload, Chaos::FReal>
 	{
 		return nullptr;
 	}
-	
+
 	/** Return a pointer to the payload on which we are querying the acceleration structure */
 	virtual const void* GetQueryPayload() const override
 	{
@@ -530,7 +535,7 @@ private:
 
 	auto GetPayloadForThread(const TPayload& Payload)
 	{
-		if constexpr (bGTData)
+		if constexpr(bGTData)
 		{
 			return Payload.GetExternalGeometryParticle_ExternalThread();
 		}
@@ -548,7 +553,7 @@ private:
 		//todo: add a check to ensure hitbuffer matches SQ type
 		using namespace Chaos;
 		TGeometryType* GeometryParticle = GetPayloadForThread(Payload);
-		if (!GeometryParticle)
+		if(!GeometryParticle)
 		{
 			// This case handles particles created by the physics simulation without the main thread
 			// being made aware of their creation. We have a PT particle but no external particle
@@ -558,15 +563,15 @@ private:
 		const FShapesArray& Shapes = GeometryParticle->ShapesArray();
 		THitType Hit;
 		Hit.Actor = GeometryParticle;
-		for (const auto& Shape : Shapes)
+		for(const auto& Shape : Shapes)
 		{
 			ECollisionQueryHitType HitType = QueryFilterData.flags & FPhysicsQueryFlag::ePREFILTER ? QueryCallback.PreFilter(QueryFilterDataConcrete, *Shape, *GeometryParticle) : ECollisionQueryHitType::Block;
-			if (HitType != ECollisionQueryHitType::None)
+			if(HitType != ECollisionQueryHitType::None)
 			{
 				const bool bBlocker = (HitType == ECollisionQueryHitType::Block || bAnyHit || HitBuffer.WantsSingleResult());
 				Hit.Shape = Shape.Get();
 				HitBuffer.InsertHit(Hit, bBlocker);
-				if (bAnyHit)
+				if(bAnyHit)
 				{
 					bContinue = false;
 				}
@@ -581,9 +586,9 @@ private:
 	const FVector HalfExtents;
 	FHitFlags OutputFlags;
 	bool bAnyHit;
-	const FQueryDebugParams DebugParams;
+	const ChaosInterface::FQueryDebugParams DebugParams;
 	ChaosInterface::FSQHitBuffer<THitType>& HitBuffer;
-	const FQueryFilterData& QueryFilterData;
+	const ChaosInterface::FQueryFilterData& QueryFilterData;
 	const FCollisionFilterData QueryFilterDataConcrete;
 	const QueryGeometryType* QueryGeom;
 	ICollisionQueryFilterCallbackBase& QueryCallback;
@@ -591,36 +596,36 @@ private:
 };
 
 template <typename TRaycastHit>
-void FChaosSQAccelerator::RaycastImp(const FVector& Start, const FVector& Dir, const float DeltaMagnitude, ChaosInterface::FSQHitBuffer<TRaycastHit>& HitBuffer, EHitFlags OutputFlags, const FQueryFilterData& QueryFilterData, ICollisionQueryFilterCallbackBase& QueryCallback, const FQueryDebugParams& DebugParams) const
+void FChaosSQAccelerator::RaycastImp(const FVector& Start, const FVector& Dir, const float DeltaMagnitude, ChaosInterface::FSQHitBuffer<TRaycastHit>& HitBuffer, EHitFlags OutputFlags, const ChaosInterface::FQueryFilterData& QueryFilterData, ICollisionQueryFilterCallbackBase& QueryCallback, const ChaosInterface::FQueryDebugParams& DebugParams) const
 {
 	using namespace Chaos;
 	using namespace ChaosInterface;
 
-	TSQVisitor<TSphere<FReal, 3>, FAccelerationStructureHandle, TRaycastHit, std::is_same<TRaycastHit,FRaycastHit>::value> RaycastVisitor(Start, Dir, HitBuffer, OutputFlags, QueryFilterData, QueryCallback, DebugParams);
+	TSQVisitor<TSphere<FReal, 3>, FAccelerationStructureHandle, TRaycastHit, std::is_same<TRaycastHit, FRaycastHit>::value> RaycastVisitor(Start, Dir, HitBuffer, OutputFlags, QueryFilterData, QueryCallback, DebugParams);
 	HitBuffer.IncFlushCount();
 	SpatialAcceleration.Raycast(Start, Dir, DeltaMagnitude, RaycastVisitor);
 	HitBuffer.DecFlushCount();
 }
 
-void FChaosSQAccelerator::Raycast(const FVector& Start, const FVector& Dir, const float DeltaMagnitude, ChaosInterface::FSQHitBuffer<ChaosInterface::FRaycastHit>& HitBuffer, EHitFlags OutputFlags, const FQueryFilterData& QueryFilterData, ICollisionQueryFilterCallbackBase& QueryCallback, const FQueryDebugParams& DebugParams) const
+void FChaosSQAccelerator::Raycast(const FVector& Start, const FVector& Dir, const float DeltaMagnitude, ChaosInterface::FSQHitBuffer<ChaosInterface::FRaycastHit>& HitBuffer, EHitFlags OutputFlags, const ChaosInterface::FQueryFilterData& QueryFilterData, ICollisionQueryFilterCallbackBase& QueryCallback, const ChaosInterface::FQueryDebugParams& DebugParams) const
 {
 	RaycastImp(Start, Dir, DeltaMagnitude, HitBuffer, OutputFlags, QueryFilterData, QueryCallback, DebugParams);
 }
 
-void FChaosSQAccelerator::Raycast(const FVector& Start, const FVector& Dir, const float DeltaMagnitude, ChaosInterface::FSQHitBuffer<ChaosInterface::FPTRaycastHit>& HitBuffer, EHitFlags OutputFlags, const FQueryFilterData& QueryFilterData, ICollisionQueryFilterCallbackBase& QueryCallback, const FQueryDebugParams& DebugParams) const
+void FChaosSQAccelerator::Raycast(const FVector& Start, const FVector& Dir, const float DeltaMagnitude, ChaosInterface::FSQHitBuffer<ChaosInterface::FPTRaycastHit>& HitBuffer, EHitFlags OutputFlags, const ChaosInterface::FQueryFilterData& QueryFilterData, ICollisionQueryFilterCallbackBase& QueryCallback, const ChaosInterface::FQueryDebugParams& DebugParams) const
 {
 	RaycastImp(Start, Dir, DeltaMagnitude, HitBuffer, OutputFlags, QueryFilterData, QueryCallback, DebugParams);
 }
 
 template <typename QueryGeomType, typename TSweepHit>
-void SweepHelper(const QueryGeomType& QueryGeom,const Chaos::ISpatialAcceleration<Chaos::FAccelerationStructureHandle, Chaos::FReal,3>& SpatialAcceleration,const FTransform& StartTM,const FVector& Dir,const float DeltaMagnitude,ChaosInterface::FSQHitBuffer<TSweepHit>& HitBuffer,EHitFlags OutputFlags,const FQueryFilterData& QueryFilterData,ICollisionQueryFilterCallbackBase& QueryCallback,const FQueryDebugParams& DebugParams)
+void SweepHelper(const QueryGeomType& QueryGeom, const Chaos::ISpatialAcceleration<Chaos::FAccelerationStructureHandle, Chaos::FReal, 3>& SpatialAcceleration, const FTransform& StartTM, const FVector& Dir, const float DeltaMagnitude, ChaosInterface::FSQHitBuffer<TSweepHit>& HitBuffer, EHitFlags OutputFlags, const ChaosInterface::FQueryFilterData& QueryFilterData, ICollisionQueryFilterCallbackBase& QueryCallback, const ChaosInterface::FQueryDebugParams& DebugParams)
 {
 	using namespace Chaos;
 	using namespace ChaosInterface;
 
 	const FAABB3 Bounds = QueryGeom.BoundingBox().TransformedAABB(StartTM);
 	const bool bSweepAsOverlap = DeltaMagnitude == 0;	//question: do we care about tiny sweeps?
-	TSQVisitor<QueryGeomType,FAccelerationStructureHandle,TSweepHit, std::is_same<TSweepHit, FSweepHit>::value> SweepVisitor(StartTM,Dir,HitBuffer,OutputFlags,QueryFilterData,QueryCallback,QueryGeom,DebugParams);
+	TSQVisitor<QueryGeomType, FAccelerationStructureHandle, TSweepHit, std::is_same<TSweepHit, FSweepHit>::value> SweepVisitor(StartTM, Dir, HitBuffer, OutputFlags, QueryFilterData, QueryCallback, QueryGeom, DebugParams);
 
 	HitBuffer.IncFlushCount();
 
@@ -628,27 +633,34 @@ void SweepHelper(const QueryGeomType& QueryGeom,const Chaos::ISpatialAcceleratio
 	{
 		//fallback to overlap
 		SpatialAcceleration.Overlap(Bounds, SweepVisitor);
-	} else
+	}
+	else
 	{
 		const FVector HalfExtents = Bounds.Extents() * 0.5f;
-		SpatialAcceleration.Sweep(Bounds.GetCenter(),Dir,DeltaMagnitude,HalfExtents,SweepVisitor);
+		SpatialAcceleration.Sweep(Bounds.GetCenter(), Dir, DeltaMagnitude, HalfExtents, SweepVisitor);
 	}
 
 	HitBuffer.DecFlushCount();
 }
 
-void FChaosSQAccelerator::Sweep(const Chaos::FImplicitObject& QueryGeom, const FTransform& StartTM, const FVector& Dir, const float DeltaMagnitude, ChaosInterface::FSQHitBuffer<ChaosInterface::FSweepHit>& HitBuffer, EHitFlags OutputFlags, const FQueryFilterData& QueryFilterData, ICollisionQueryFilterCallbackBase& QueryCallback, const FQueryDebugParams& DebugParams) const
+void FChaosSQAccelerator::Sweep(const Chaos::FImplicitObject& QueryGeom, const FTransform& StartTM, const FVector& Dir, const float DeltaMagnitude, ChaosInterface::FSQHitBuffer<ChaosInterface::FSweepHit>& HitBuffer, EHitFlags OutputFlags, const ChaosInterface::FQueryFilterData& QueryFilterData, ICollisionQueryFilterCallbackBase& QueryCallback, const ChaosInterface::FQueryDebugParams& DebugParams) const
 {
-	return Chaos::Utilities::CastHelper(QueryGeom, StartTM, [&](const auto& Downcast, const FTransform& StartFullTM) { return SweepHelper(Downcast, SpatialAcceleration, StartFullTM, Dir, DeltaMagnitude, HitBuffer, OutputFlags, QueryFilterData, QueryCallback, DebugParams); });
+	return Chaos::Utilities::CastHelper(QueryGeom, StartTM, [&](const auto& Downcast, const FTransform& StartFullTM)
+	{
+		return SweepHelper(Downcast, SpatialAcceleration, StartFullTM, Dir, DeltaMagnitude, HitBuffer, OutputFlags, QueryFilterData, QueryCallback, DebugParams);
+	});
 }
 
-void FChaosSQAccelerator::Sweep(const Chaos::FImplicitObject& QueryGeom, const FTransform& StartTM, const FVector& Dir, const float DeltaMagnitude, ChaosInterface::FSQHitBuffer<ChaosInterface::FPTSweepHit>& HitBuffer, EHitFlags OutputFlags, const FQueryFilterData& QueryFilterData, ICollisionQueryFilterCallbackBase& QueryCallback, const FQueryDebugParams& DebugParams) const
+void FChaosSQAccelerator::Sweep(const Chaos::FImplicitObject& QueryGeom, const FTransform& StartTM, const FVector& Dir, const float DeltaMagnitude, ChaosInterface::FSQHitBuffer<ChaosInterface::FPTSweepHit>& HitBuffer, EHitFlags OutputFlags, const ChaosInterface::FQueryFilterData& QueryFilterData, ICollisionQueryFilterCallbackBase& QueryCallback, const ChaosInterface::FQueryDebugParams& DebugParams) const
 {
-	return Chaos::Utilities::CastHelper(QueryGeom, StartTM, [&](const auto& Downcast, const FTransform& StartFullTM) { return SweepHelper(Downcast, SpatialAcceleration, StartFullTM, Dir, DeltaMagnitude, HitBuffer, OutputFlags, QueryFilterData, QueryCallback, DebugParams); });
+	return Chaos::Utilities::CastHelper(QueryGeom, StartTM, [&](const auto& Downcast, const FTransform& StartFullTM)
+	{
+		return SweepHelper(Downcast, SpatialAcceleration, StartFullTM, Dir, DeltaMagnitude, HitBuffer, OutputFlags, QueryFilterData, QueryCallback, DebugParams);
+	});
 }
 
 template <typename QueryGeomType, typename TOverlapHit>
-void OverlapHelper(const QueryGeomType& QueryGeom, const Chaos::ISpatialAcceleration<Chaos::FAccelerationStructureHandle, Chaos::FReal, 3>& SpatialAcceleration, const FTransform& GeomPose, ChaosInterface::FSQHitBuffer<TOverlapHit>& HitBuffer, const FQueryFilterData& QueryFilterData, ICollisionQueryFilterCallbackBase& QueryCallback, const FQueryDebugParams& DebugParams)
+void OverlapHelper(const QueryGeomType& QueryGeom, const Chaos::ISpatialAcceleration<Chaos::FAccelerationStructureHandle, Chaos::FReal, 3>& SpatialAcceleration, const FTransform& GeomPose, ChaosInterface::FSQHitBuffer<TOverlapHit>& HitBuffer, const ChaosInterface::FQueryFilterData& QueryFilterData, ICollisionQueryFilterCallbackBase& QueryCallback, const ChaosInterface::FQueryDebugParams& DebugParams)
 {
 	using namespace Chaos;
 	using namespace ChaosInterface;
@@ -660,7 +672,7 @@ void OverlapHelper(const QueryGeomType& QueryGeom, const Chaos::ISpatialAccelera
 	bool bSkipNarrowPhase = QueryFilterData.flags & FPhysicsQueryFlag::eSKIPNARROWPHASE;
 
 	constexpr bool bGTData = std::is_same<TOverlapHit, FOverlapHit>::value;
-	if (bSkipNarrowPhase)
+	if(bSkipNarrowPhase)
 	{
 		TBPVisitor<QueryGeomType, FAccelerationStructureHandle, TOverlapHit, bGTData> OverlapVisitor(GeomPose, HitBuffer, QueryFilterData, QueryCallback, QueryGeom, DebugParams);
 		SpatialAcceleration.Overlap(Bounds, OverlapVisitor);
@@ -673,12 +685,15 @@ void OverlapHelper(const QueryGeomType& QueryGeom, const Chaos::ISpatialAccelera
 	HitBuffer.DecFlushCount();
 }
 
-void FChaosSQAccelerator::Overlap(const Chaos::FImplicitObject& QueryGeom, const FTransform& GeomPose, ChaosInterface::FSQHitBuffer<ChaosInterface::FOverlapHit>& HitBuffer, const FQueryFilterData& QueryFilterData, ICollisionQueryFilterCallbackBase& QueryCallback, const FQueryDebugParams& DebugParams) const
+void FChaosSQAccelerator::Overlap(const Chaos::FImplicitObject& QueryGeom, const FTransform& GeomPose, ChaosInterface::FSQHitBuffer<ChaosInterface::FOverlapHit>& HitBuffer, const ChaosInterface::FQueryFilterData& QueryFilterData, ICollisionQueryFilterCallbackBase& QueryCallback, const ChaosInterface::FQueryDebugParams& DebugParams) const
 {
-	return Chaos::Utilities::CastHelper(QueryGeom, GeomPose, [&](const auto& Downcast, const FTransform& GeomFullPose) { return OverlapHelper(Downcast, SpatialAcceleration, GeomFullPose, HitBuffer, QueryFilterData, QueryCallback, DebugParams); });
+	return Chaos::Utilities::CastHelper(QueryGeom, GeomPose, [&](const auto& Downcast, const FTransform& GeomFullPose)
+	{
+		return OverlapHelper(Downcast, SpatialAcceleration, GeomFullPose, HitBuffer, QueryFilterData, QueryCallback, DebugParams);
+	});
 }
 
-void FChaosSQAccelerator::Overlap(const Chaos::FImplicitObject& QueryGeom, const FTransform& GeomPose, ChaosInterface::FSQHitBuffer<ChaosInterface::FPTOverlapHit>& HitBuffer, const FQueryFilterData& QueryFilterData, ICollisionQueryFilterCallbackBase& QueryCallback, const FQueryDebugParams& DebugParams) const
+void FChaosSQAccelerator::Overlap(const Chaos::FImplicitObject& QueryGeom, const FTransform& GeomPose, ChaosInterface::FSQHitBuffer<ChaosInterface::FPTOverlapHit>& HitBuffer, const ChaosInterface::FQueryFilterData& QueryFilterData, ICollisionQueryFilterCallbackBase& QueryCallback, const ChaosInterface::FQueryDebugParams& DebugParams) const
 {
 	return Chaos::Utilities::CastHelper(QueryGeom, GeomPose, [&](const auto& Downcast, const FTransform& GeomFullPose) { return OverlapHelper(Downcast, SpatialAcceleration, GeomFullPose, HitBuffer, QueryFilterData, QueryCallback, DebugParams); });
 }
diff --git a/Engine/Source/Runtime/PhysicsCore/Private/SQCapture.cpp b/Engine/Source/Runtime/PhysicsCore/Private/SQCapture.cpp
--- a/Engine/Source/Runtime/PhysicsCore/Private/SQCapture.cpp
+++ b/Engine/Source/Runtime/PhysicsCore/Private/SQCapture.cpp
@@ -39,7 +39,7 @@ FSQCapture::~FSQCapture()
 {
 }
 
-void SerializeQueryFilterData(FArchive& Ar, FQueryFilterData& QueryFilterData)
+void SerializeQueryFilterData(FArchive& Ar, ChaosInterface::FQueryFilterData& QueryFilterData)
 {
 	Ar << QueryFilterData.data.word0;
 	Ar << QueryFilterData.data.word1;
@@ -188,7 +188,7 @@ void FSQCapture::Serialize(Chaos::FChaosArchive& Ar)
 	}
 }
 
-void FSQCapture::StartCaptureChaosSweep(const Chaos::FPBDRigidsEvolution& Evolution, const Chaos::FImplicitObject& InQueryGeom, const FTransform& InStartTM, const FVector& InDir, float InDeltaMag, FHitFlags InOutputFlags, const FQueryFilterData& QueryFilter, const FCollisionFilterData& FilterData, ICollisionQueryFilterCallbackBase& Callback)
+void FSQCapture::StartCaptureChaosSweep(const Chaos::FPBDRigidsEvolution& Evolution, const Chaos::FImplicitObject& InQueryGeom, const FTransform& InStartTM, const FVector& InDir, float InDeltaMag, FHitFlags InOutputFlags, const ChaosInterface::FQueryFilterData& QueryFilter, const FCollisionFilterData& FilterData, ICollisionQueryFilterCallbackBase& Callback)
 {
 	if (IsInGameThread())
 	{
@@ -216,7 +216,7 @@ void FSQCapture::EndCaptureChaosSweep(const ChaosInterface::FSQHitBuffer<ChaosIn
 	}
 }
 
-void FSQCapture::StartCaptureChaosRaycast(const Chaos::FPBDRigidsEvolution& Evolution, const FVector& InStartPoint, const FVector& InDir, float InDeltaMag, FHitFlags InOutputFlags, const FQueryFilterData& QueryFilter, const FCollisionFilterData& FilterData, ICollisionQueryFilterCallbackBase& Callback)
+void FSQCapture::StartCaptureChaosRaycast(const Chaos::FPBDRigidsEvolution& Evolution, const FVector& InStartPoint, const FVector& InDir, float InDeltaMag, FHitFlags InOutputFlags, const ChaosInterface::FQueryFilterData& QueryFilter, const FCollisionFilterData& FilterData, ICollisionQueryFilterCallbackBase& Callback)
 {
 	if (IsInGameThread())
 	{
@@ -242,7 +242,7 @@ void FSQCapture::EndCaptureChaosRaycast(const ChaosInterface::FSQHitBuffer<Chaos
 	}
 }
 
-void FSQCapture::StartCaptureChaosOverlap(const Chaos::FPBDRigidsEvolution& Evolution, const Chaos::FImplicitObject& InQueryGeom, const FTransform& InStartTM, const FQueryFilterData& QueryFilter, const FCollisionFilterData& FilterData, ICollisionQueryFilterCallbackBase& Callback)
+void FSQCapture::StartCaptureChaosOverlap(const Chaos::FPBDRigidsEvolution& Evolution, const Chaos::FImplicitObject& InQueryGeom, const FTransform& InStartTM, const ChaosInterface::FQueryFilterData& QueryFilter, const FCollisionFilterData& FilterData, ICollisionQueryFilterCallbackBase& Callback)
 {
 	if (IsInGameThread())
 	{
diff --git a/Engine/Source/Runtime/PhysicsCore/Public/ChaosInterfaceWrapperCore.h b/Engine/Source/Runtime/PhysicsCore/Public/ChaosInterfaceWrapperCore.h
--- a/Engine/Source/Runtime/PhysicsCore/Public/ChaosInterfaceWrapperCore.h
+++ b/Engine/Source/Runtime/PhysicsCore/Public/ChaosInterfaceWrapperCore.h
@@ -262,5 +262,3 @@ bool GetHasBlock(const FSQHitBuffer<HitType>& Callback)
 }
 
 } // namespace ChaosInterface
-
-using namespace ChaosInterface;
diff --git a/Engine/Source/Runtime/PhysicsCore/Public/SQAccelerator.h b/Engine/Source/Runtime/PhysicsCore/Public/SQAccelerator.h
--- a/Engine/Source/Runtime/PhysicsCore/Public/SQAccelerator.h
+++ b/Engine/Source/Runtime/PhysicsCore/Public/SQAccelerator.h
@@ -41,20 +41,20 @@ public:
 	FChaosSQAccelerator(const Chaos::ISpatialAcceleration<Chaos::FAccelerationStructureHandle,Chaos::FReal, 3>& InSpatialAcceleration);
 	virtual ~FChaosSQAccelerator() {};
 
-	void Raycast(const FVector& Start, const FVector& Dir, const float DeltaMagnitude, ChaosInterface::FSQHitBuffer<ChaosInterface::FRaycastHit>& HitBuffer, EHitFlags OutputFlags, const FQueryFilterData& QueryFilterData, ICollisionQueryFilterCallbackBase& QueryCallback, const FQueryDebugParams& DebugParams = FQueryDebugParams()) const;
-	void Raycast(const FVector& Start, const FVector& Dir, const float DeltaMagnitude, ChaosInterface::FSQHitBuffer<ChaosInterface::FPTRaycastHit>& HitBuffer, EHitFlags OutputFlags, const FQueryFilterData& QueryFilterData, ICollisionQueryFilterCallbackBase& QueryCallback, const FQueryDebugParams& DebugParams = FQueryDebugParams()) const;
+	void Raycast(const FVector& Start, const FVector& Dir, const float DeltaMagnitude, ChaosInterface::FSQHitBuffer<ChaosInterface::FRaycastHit>& HitBuffer, EHitFlags OutputFlags, const ChaosInterface::FQueryFilterData& QueryFilterData, ICollisionQueryFilterCallbackBase& QueryCallback, const ChaosInterface::FQueryDebugParams& DebugParams = {}) const;
+	void Raycast(const FVector& Start, const FVector& Dir, const float DeltaMagnitude, ChaosInterface::FSQHitBuffer<ChaosInterface::FPTRaycastHit>& HitBuffer, EHitFlags OutputFlags, const ChaosInterface::FQueryFilterData& QueryFilterData, ICollisionQueryFilterCallbackBase& QueryCallback, const ChaosInterface::FQueryDebugParams& DebugParams = {}) const;
 
-	void Sweep(const Chaos::FImplicitObject& QueryGeom, const FTransform& StartTM, const FVector& Dir, const float DeltaMagnitude, ChaosInterface::FSQHitBuffer<ChaosInterface::FSweepHit>& HitBuffer, EHitFlags OutputFlags, const FQueryFilterData& QueryFilterData, ICollisionQueryFilterCallbackBase& QueryCallback, const FQueryDebugParams& DebugParams = FQueryDebugParams()) const;
-	void Sweep(const Chaos::FImplicitObject& QueryGeom, const FTransform& StartTM, const FVector& Dir, const float DeltaMagnitude, ChaosInterface::FSQHitBuffer<ChaosInterface::FPTSweepHit>& HitBuffer, EHitFlags OutputFlags, const FQueryFilterData& QueryFilterData, ICollisionQueryFilterCallbackBase& QueryCallback, const FQueryDebugParams& DebugParams = FQueryDebugParams()) const;
+	void Sweep(const Chaos::FImplicitObject& QueryGeom, const FTransform& StartTM, const FVector& Dir, const float DeltaMagnitude, ChaosInterface::FSQHitBuffer<ChaosInterface::FSweepHit>& HitBuffer, EHitFlags OutputFlags, const ChaosInterface::FQueryFilterData& QueryFilterData, ICollisionQueryFilterCallbackBase& QueryCallback, const ChaosInterface::FQueryDebugParams& DebugParams = {}) const;
+	void Sweep(const Chaos::FImplicitObject& QueryGeom, const FTransform& StartTM, const FVector& Dir, const float DeltaMagnitude, ChaosInterface::FSQHitBuffer<ChaosInterface::FPTSweepHit>& HitBuffer, EHitFlags OutputFlags, const ChaosInterface::FQueryFilterData& QueryFilterData, ICollisionQueryFilterCallbackBase& QueryCallback, const ChaosInterface::FQueryDebugParams& DebugParams = {}) const;
 
-	void Overlap(const Chaos::FImplicitObject& QueryGeom, const FTransform& GeomPose, ChaosInterface::FSQHitBuffer<ChaosInterface::FOverlapHit>& HitBuffer, const FQueryFilterData& QueryFilterData, ICollisionQueryFilterCallbackBase& QueryCallback, const FQueryDebugParams& DebugParams = FQueryDebugParams()) const;
-	void Overlap(const Chaos::FImplicitObject& QueryGeom, const FTransform& GeomPose, ChaosInterface::FSQHitBuffer<ChaosInterface::FPTOverlapHit>& HitBuffer, const FQueryFilterData& QueryFilterData, ICollisionQueryFilterCallbackBase& QueryCallback, const FQueryDebugParams& DebugParams = FQueryDebugParams()) const;
+	void Overlap(const Chaos::FImplicitObject& QueryGeom, const FTransform& GeomPose, ChaosInterface::FSQHitBuffer<ChaosInterface::FOverlapHit>& HitBuffer, const ChaosInterface::FQueryFilterData& QueryFilterData, ICollisionQueryFilterCallbackBase& QueryCallback, const ChaosInterface::FQueryDebugParams& DebugParams = {}) const;
+	void Overlap(const Chaos::FImplicitObject& QueryGeom, const FTransform& GeomPose, ChaosInterface::FSQHitBuffer<ChaosInterface::FPTOverlapHit>& HitBuffer, const ChaosInterface::FQueryFilterData& QueryFilterData, ICollisionQueryFilterCallbackBase& QueryCallback, const ChaosInterface::FQueryDebugParams& DebugParams = {}) const;
 
 private:
 	const Chaos::ISpatialAcceleration<Chaos::FAccelerationStructureHandle, Chaos::FReal, 3>& SpatialAcceleration;
 
 	template <typename TRaycastHit>
-	void RaycastImp(const FVector& Start, const FVector& Dir, const float DeltaMagnitude, ChaosInterface::FSQHitBuffer<TRaycastHit>& HitBuffer, EHitFlags OutputFlags, const FQueryFilterData& QueryFilterData, ICollisionQueryFilterCallbackBase& QueryCallback, const FQueryDebugParams& DebugParams) const;
+	void RaycastImp(const FVector& Start, const FVector& Dir, const float DeltaMagnitude, ChaosInterface::FSQHitBuffer<TRaycastHit>& HitBuffer, EHitFlags OutputFlags, const ChaosInterface::FQueryFilterData& QueryFilterData, ICollisionQueryFilterCallbackBase& QueryCallback, const ChaosInterface::FQueryDebugParams& DebugParams) const;
 
 };
 
@@ -66,18 +66,18 @@ class PHYSICSCORE_API ISQAccelerator
 {
 public:
 	virtual ~ISQAccelerator() {};
-	virtual void Raycast(const FVector& Start, const FVector& Dir, const float DeltaMagnitude, FPhysicsHitCallback<FHitRaycast>& HitBuffer, EHitFlags OutputFlags, const FQueryFilterData& QueryFilterData, ICollisionQueryFilterCallbackBase& QueryCallback) const = 0;
-	virtual void Sweep(const FPhysicsGeometry& QueryGeom, const FTransform& StartTM, const FVector& Dir, const float DeltaMagnitude, FPhysicsHitCallback<FHitSweep>& HitBuffer, EHitFlags OutputFlags, const FQueryFilterData& QueryFilterData, ICollisionQueryFilterCallbackBase& QueryCallback) const = 0;
-	virtual void Overlap(const FPhysicsGeometry& QueryGeom, const FTransform& GeomPose, FPhysicsHitCallback<FHitOverlap>& HitBuffer, const FQueryFilterData& QueryFilterData, ICollisionQueryFilterCallbackBase& QueryCallback) const = 0;
+	virtual void Raycast(const FVector& Start, const FVector& Dir, const float DeltaMagnitude, FPhysicsHitCallback<FHitRaycast>& HitBuffer, EHitFlags OutputFlags, const ChaosInterface::FQueryFilterData& QueryFilterData, ICollisionQueryFilterCallbackBase& QueryCallback) const = 0;
+	virtual void Sweep(const FPhysicsGeometry& QueryGeom, const FTransform& StartTM, const FVector& Dir, const float DeltaMagnitude, FPhysicsHitCallback<FHitSweep>& HitBuffer, EHitFlags OutputFlags, const ChaosInterface::FQueryFilterData& QueryFilterData, ICollisionQueryFilterCallbackBase& QueryCallback) const = 0;
+	virtual void Overlap(const FPhysicsGeometry& QueryGeom, const FTransform& GeomPose, FPhysicsHitCallback<FHitOverlap>& HitBuffer, const ChaosInterface::FQueryFilterData& QueryFilterData, ICollisionQueryFilterCallbackBase& QueryCallback) const = 0;
 };
 
 class PHYSICSCORE_API FSQAcceleratorUnion : public ISQAccelerator
 {
 public:
 
-	virtual void Raycast(const FVector& Start, const FVector& Dir, const float DeltaMagnitude, FPhysicsHitCallback<FHitRaycast>& HitBuffer, EHitFlags OutputFlags, const FQueryFilterData& QueryFilterData, ICollisionQueryFilterCallbackBase& QueryCallback) const override;
-	virtual void Sweep(const FPhysicsGeometry& QueryGeom, const FTransform& StartTM, const FVector& Dir, const float DeltaMagnitude, FPhysicsHitCallback<FHitSweep>& HitBuffer, EHitFlags OutputFlags, const FQueryFilterData& QueryFilterData, ICollisionQueryFilterCallbackBase& QueryCallback) const override;
-	virtual void Overlap(const FPhysicsGeometry& QueryGeom, const FTransform& GeomPose, FPhysicsHitCallback<FHitOverlap>& HitBuffer, const FQueryFilterData& QueryFilterData, ICollisionQueryFilterCallbackBase& QueryCallback) const override;
+	virtual void Raycast(const FVector& Start, const FVector& Dir, const float DeltaMagnitude, FPhysicsHitCallback<FHitRaycast>& HitBuffer, EHitFlags OutputFlags, const ChaosInterface::FQueryFilterData& QueryFilterData, ICollisionQueryFilterCallbackBase& QueryCallback) const override;
+	virtual void Sweep(const FPhysicsGeometry& QueryGeom, const FTransform& StartTM, const FVector& Dir, const float DeltaMagnitude, FPhysicsHitCallback<FHitSweep>& HitBuffer, EHitFlags OutputFlags, const ChaosInterface::FQueryFilterData& QueryFilterData, ICollisionQueryFilterCallbackBase& QueryCallback) const override;
+	virtual void Overlap(const FPhysicsGeometry& QueryGeom, const FTransform& GeomPose, FPhysicsHitCallback<FHitOverlap>& HitBuffer, const ChaosInterface::FQueryFilterData& QueryFilterData, ICollisionQueryFilterCallbackBase& QueryCallback) const override;
 
 	void AddSQAccelerator(ISQAccelerator* InAccelerator);
 	void RemoveSQAccelerator(ISQAccelerator* AcceleratorToRemove);
diff --git a/Engine/Source/Runtime/PhysicsCore/Public/SQCapture.h b/Engine/Source/Runtime/PhysicsCore/Public/SQCapture.h
--- a/Engine/Source/Runtime/PhysicsCore/Public/SQCapture.h
+++ b/Engine/Source/Runtime/PhysicsCore/Public/SQCapture.h
@@ -50,13 +50,13 @@ struct PHYSICSCORE_API FSQCapture
 		Overlap
 	} SQType;
 
-	void StartCaptureChaosSweep(const Chaos::FPBDRigidsEvolution& Evolution, const Chaos::FImplicitObject& InQueryGeom, const FTransform& InStartTM, const FVector& InDir, float InDeltaMag, FHitFlags InOutputFlags, const FQueryFilterData& QueryFilter, const FCollisionFilterData& FilterData, ICollisionQueryFilterCallbackBase& Callback);
+	void StartCaptureChaosSweep(const Chaos::FPBDRigidsEvolution& Evolution, const Chaos::FImplicitObject& InQueryGeom, const FTransform& InStartTM, const FVector& InDir, float InDeltaMag, FHitFlags InOutputFlags, const ChaosInterface::FQueryFilterData& QueryFilter, const FCollisionFilterData& FilterData, ICollisionQueryFilterCallbackBase& Callback);
 	void EndCaptureChaosSweep(const ChaosInterface::FSQHitBuffer<ChaosInterface::FSweepHit>& Results);
 
-	void StartCaptureChaosRaycast(const Chaos::FPBDRigidsEvolution& Evolution, const FVector& InStartPoint, const FVector& InDir, float InDeltaMag, FHitFlags InOutputFlags, const FQueryFilterData& QueryFilter, const FCollisionFilterData& FilterData, ICollisionQueryFilterCallbackBase& Callback);
+	void StartCaptureChaosRaycast(const Chaos::FPBDRigidsEvolution& Evolution, const FVector& InStartPoint, const FVector& InDir, float InDeltaMag, FHitFlags InOutputFlags, const ChaosInterface::FQueryFilterData& QueryFilter, const FCollisionFilterData& FilterData, ICollisionQueryFilterCallbackBase& Callback);
 	void EndCaptureChaosRaycast(const ChaosInterface::FSQHitBuffer<ChaosInterface::FRaycastHit>& Results);
 
-	void StartCaptureChaosOverlap(const Chaos::FPBDRigidsEvolution& Evolution, const Chaos::FImplicitObject& InQueryGeom, const FTransform& InStartTM, const FQueryFilterData& QueryFilter, const FCollisionFilterData& FilterData, ICollisionQueryFilterCallbackBase& Callback);
+	void StartCaptureChaosOverlap(const Chaos::FPBDRigidsEvolution& Evolution, const Chaos::FImplicitObject& InQueryGeom, const FTransform& InStartTM, const ChaosInterface::FQueryFilterData& QueryFilter, const FCollisionFilterData& FilterData, ICollisionQueryFilterCallbackBase& Callback);
 	void EndCaptureChaosOverlap(const ChaosInterface::FSQHitBuffer<ChaosInterface::FOverlapHit>& Results);
 
 	ECollisionQueryHitType GetFilterResult(const Chaos::FPerShapeData* Shape, const Chaos::FGeometryParticle* Actor) const;
@@ -67,7 +67,7 @@ struct PHYSICSCORE_API FSQCapture
 
 	float DeltaMag;
 	FHitFlags OutputFlags;
-	FQueryFilterData QueryFilterData;
+	ChaosInterface::FQueryFilterData QueryFilterData;
 	TUniquePtr<ICollisionQueryFilterCallbackBase> FilterCallback;
 
 	TUniquePtr<Chaos::FImplicitObject> ChaosOwnerObject;	//should be private, do not access directly
