make Chaos MassProperties.cpp's CalculateVolumeAndCenterOfMassImpl and CalculateInertiaAndRotationOfMassImpl consistent: Both now compute their result in double precision, and handle surfaces with more than three vertices by implicitly fan-triangulating them

diff --git a/Engine/Source/Runtime/Experimental/Chaos/Private/Chaos/MassProperties.cpp b/Engine/Source/Runtime/Experimental/Chaos/Private/Chaos/MassProperties.cpp
--- a/Engine/Source/Runtime/Experimental/Chaos/Private/Chaos/MassProperties.cpp
+++ b/Engine/Source/Runtime/Experimental/Chaos/Private/Chaos/MassProperties.cpp
@@ -144,20 +144,20 @@ namespace Chaos
 			return;
 		}
 		
-		T Volume = 0;
-		TVec VolumeTimesSum(0);
-		TVec Center = Vertices[Surfaces[0][0]];
+		double Volume = 0;
+		FVec3 VolumeTimesSum(0);
+		FVec3 Center = Vertices[Surfaces[0][0]];
 		for (const auto& Element : Surfaces)
 		{
 			// implicitly triangulate with a fan
 			for (int32 FanIdx = 0; FanIdx + 2 < Element.Num(); ++FanIdx)
 			{
 				int32 FanSubInds[]{ 0, FanIdx + 1, FanIdx + 2 };
-				PMatrix<T, 3, 3> DeltaMatrix;
-				TVec PerElementSize;
+				FMatrix33 DeltaMatrix;
+				FVec3 PerElementSize;
 				for (int32 i = 0; i < 3; ++i)
 				{
-					TVec DeltaVector = Vertices[Element[FanSubInds[i]]] - Center;
+					FVec3 DeltaVector = FVec3(Vertices[Element[FanSubInds[i]]]) - Center;
 					DeltaMatrix.M[0][i] = DeltaVector[0];
 					DeltaMatrix.M[1][i] = DeltaVector[1];
 					DeltaMatrix.M[2][i] = DeltaVector[2];
@@ -167,7 +167,7 @@ namespace Chaos
 				PerElementSize[1] = DeltaMatrix.M[1][0] + DeltaMatrix.M[1][1] + DeltaMatrix.M[1][2];
 				PerElementSize[2] = DeltaMatrix.M[2][0] + DeltaMatrix.M[2][1] + DeltaMatrix.M[2][2];
 
-				T Det = DeltaMatrix.M[0][0] * (DeltaMatrix.M[1][1] * DeltaMatrix.M[2][2] - DeltaMatrix.M[1][2] * DeltaMatrix.M[2][1]) -
+				double Det = DeltaMatrix.M[0][0] * (DeltaMatrix.M[1][1] * DeltaMatrix.M[2][2] - DeltaMatrix.M[1][2] * DeltaMatrix.M[2][1]) -
 					DeltaMatrix.M[0][1] * (DeltaMatrix.M[1][0] * DeltaMatrix.M[2][2] - DeltaMatrix.M[1][2] * DeltaMatrix.M[2][0]) +
 					DeltaMatrix.M[0][2] * (DeltaMatrix.M[1][0] * DeltaMatrix.M[2][1] - DeltaMatrix.M[1][1] * DeltaMatrix.M[2][0]);
 
@@ -175,14 +175,14 @@ namespace Chaos
 				VolumeTimesSum += Det * PerElementSize;
 			}
 		}
-		// @todo(mlentine): Should add suppoert for thin shell mass properties
-		if (Volume < UE_KINDA_SMALL_NUMBER)	//handle negative volume using fallback for now. Need to investigate cases where this happens
+		// @todo(mlentine): Should add support for thin shell mass properties
+		if (Volume < UE_DOUBLE_KINDA_SMALL_NUMBER)	//handle negative volume using fallback for now. Need to investigate cases where this happens
 			{
 			OutVolume = 0;
 			return;
 			}
-		OutCenterOfMass = Center + VolumeTimesSum / (4 * Volume);
-		OutVolume = Volume / 6;
+		OutCenterOfMass = TVec(Center + VolumeTimesSum / (4 * Volume));
+		OutVolume = T(Volume / 6);
 	}
 
 	template<typename T, typename TSurfaces>
@@ -214,19 +214,24 @@ namespace Chaos
 		FMatrix33 Covariance(0);
 		for (const auto& Element : Surfaces)
 		{
-			FMatrix33 DeltaMatrix(0);
-			for (int32 i = 0; i < Element.Num(); ++i)
+			// implicitly triangulate with a fan
+			for (int32 FanIdx = 0; FanIdx + 2 < Element.Num(); ++FanIdx)
 			{
-				FVec3 DeltaVector = FVec3(Vertices[Element[i]] - CenterOfMass);
-				DeltaMatrix.M[0][i] = DeltaVector[0];
-				DeltaMatrix.M[1][i] = DeltaVector[1];
-				DeltaMatrix.M[2][i] = DeltaVector[2];
+				int32 FanSubInds[]{ 0, FanIdx + 1, FanIdx + 2 };
+				FMatrix33 DeltaMatrix(0);
+				for (int32 i = 0; i < 3; ++i)
+				{
+					FVec3 DeltaVector = FVec3(Vertices[Element[FanSubInds[i]]] - CenterOfMass);
+					DeltaMatrix.M[0][i] = DeltaVector[0];
+					DeltaMatrix.M[1][i] = DeltaVector[1];
+					DeltaMatrix.M[2][i] = DeltaVector[2];
+				}
+				const FReal Det = DeltaMatrix.M[0][0] * (DeltaMatrix.M[1][1] * DeltaMatrix.M[2][2] - DeltaMatrix.M[1][2] * DeltaMatrix.M[2][1]) -
+					DeltaMatrix.M[0][1] * (DeltaMatrix.M[1][0] * DeltaMatrix.M[2][2] - DeltaMatrix.M[1][2] * DeltaMatrix.M[2][0]) +
+					DeltaMatrix.M[0][2] * (DeltaMatrix.M[1][0] * DeltaMatrix.M[2][1] - DeltaMatrix.M[1][1] * DeltaMatrix.M[2][0]);
+				const FMatrix33 ScaledStandard = Standard * Det;
+				Covariance += DeltaMatrix * ScaledStandard * DeltaMatrix.GetTransposed();
 			}
-			const FReal Det = DeltaMatrix.M[0][0] * (DeltaMatrix.M[1][1] * DeltaMatrix.M[2][2] - DeltaMatrix.M[1][2] * DeltaMatrix.M[2][1]) -
-				DeltaMatrix.M[0][1] * (DeltaMatrix.M[1][0] * DeltaMatrix.M[2][2] - DeltaMatrix.M[1][2] * DeltaMatrix.M[2][0]) +
-				DeltaMatrix.M[0][2] * (DeltaMatrix.M[1][0] * DeltaMatrix.M[2][1] - DeltaMatrix.M[1][1] * DeltaMatrix.M[2][0]);
-			const FMatrix33 ScaledStandard = Standard * Det;
-			Covariance += DeltaMatrix * ScaledStandard * DeltaMatrix.GetTransposed();
 		}
 		const FReal Trace = Covariance.M[0][0] + Covariance.M[1][1] + Covariance.M[2][2];
 		const FMatrix33 TraceMat(Trace, Trace, Trace);
