Chaos - convex-triangle collision manifold
- SAT-based collision generation, no margins
- Fix negative scaled heightfield collision

diff --git a/Engine/Source/Runtime/Experimental/Chaos/Private/Chaos/Collision/ConvexTriangleContactPoint.cpp b/Engine/Source/Runtime/Experimental/Chaos/Private/Chaos/Collision/ConvexTriangleContactPoint.cpp
new file mode 100644
--- /dev/null
+++ b/Engine/Source/Runtime/Experimental/Chaos/Private/Chaos/Collision/ConvexTriangleContactPoint.cpp
@@ -0,0 +1,532 @@
+// Copyright Epic Games, Inc. All Rights Reserved.
+#include "Chaos/Collision/ConvexTriangleContactPoint.h"
+#include "Chaos/Collision/ContactPoint.h"
+#include "Chaos/Collision/ConvexContactPointUtilities.h"
+#include "Chaos/CollisionOneShotManifolds.h"
+#include "Chaos/Convex.h"
+#include "Chaos/DebugDrawQueue.h"
+#include "Chaos/SAT.h"
+#include "Chaos/Triangle.h"
+
+//PRAGMA_DISABLE_OPTIMIZATION
+
+namespace Chaos
+{
+	// Check whether the convex-triangle edge pair form part of the Minkowski Sum. Only edge pairs
+	// that contribute to the Minkowki Sum surface need to be checked for separation. The inputs
+	// are the convex normals for the two faces that share the convex edge, and the normal and
+	// edge vector of the triangle.
+	// 
+	// This is a custom version of IsMinkowskiSum for triangles where the two normals are directly
+	// opposing and therefore the regular edge vector calculation returns zero.
+	// 
+	// @param A ConvexNormalA
+	// @param B ConvexNormalB
+	// @param BA ConvexEdge
+	// @param C TriNormal (negated)
+	// @param DC TriEdge
+	bool IsMinkowskiSumConvexTriangle(const FVec3& A, const FVec3& B, const FVec3& BA, const FVec3& C, const FVec3& DC)
+	{
+		const FReal CBA = FVec3::DotProduct(C, BA);		// TriNormal | ConvexEdge
+		const FReal ADC = FVec3::DotProduct(A, DC);		// ConvexNormalA | TriEdge
+		const FReal BDC = FVec3::DotProduct(B, DC);		// ConvexNormalB | TriEdge
+
+		const FReal Tolerance = 1.e-2f;
+		return ((ADC * BDC) < -Tolerance) && ((CBA * BDC) > Tolerance);
+	}
+
+	// Clip the vertices of a triangle to a face of a convex, using some arbitrary vector as the clipping axis
+	// (the axis is assumed to not be parallel to the convex face surface).
+	template <typename ConvexImplicitType>
+	TArrayView<FVec3> ClipTriangleToConvex(
+		const FTriangle& Triangle,
+		const ConvexImplicitType& Convex,
+		const int32 ConvexPlaneIndex,
+		const FVec3& Axis,
+		TArrayView<FVec3> VertexBuffer1,
+		TArrayView<FVec3> VertexBuffer2)
+	{
+		check(VertexBuffer1.Num() == VertexBuffer2.Num());
+		check(VertexBuffer1.Num() >= 3);
+
+		// Start with the triangle vertices
+		int32 ContactPointCount = 0;
+		VertexBuffer1[0] = Triangle.GetVertex(0);
+		VertexBuffer1[1] = Triangle.GetVertex(1);
+		VertexBuffer1[2] = Triangle.GetVertex(2);
+		ContactPointCount = 3;
+
+		// Now clip against all planes that belong to the convex plane's, edges
+		// Note winding order matters here, and we have to handle negative scales
+		const FReal WindingOrder = Convex.GetWindingOrder();
+		const int32 ConvexFaceVerticesNum = Convex.NumPlaneVertices(ConvexPlaneIndex);
+		int32 ClippingPlaneCount = ConvexFaceVerticesNum;
+		FVec3 PrevPoint = Convex.GetVertex(Convex.GetPlaneVertex(ConvexPlaneIndex, ClippingPlaneCount - 1));
+		for (int32 ClippingPlaneIndex = 0; (ClippingPlaneIndex < ClippingPlaneCount) && (ContactPointCount > 1); ++ClippingPlaneIndex)
+		{
+			const FVec3 CurrentPoint = Convex.GetVertex(Convex.GetPlaneVertex(ConvexPlaneIndex, ClippingPlaneIndex));
+
+			// Convex edge clipping plane
+			// NOTE: Plane is not normalized, but the length cancels out in all clip operations
+			const FVec3 ClippingPlaneNormal = WindingOrder * FVec3::CrossProduct(Axis, PrevPoint - CurrentPoint);
+			if (ClippingPlaneNormal.SizeSquared() > UE_SMALL_NUMBER)
+			{
+				ContactPointCount = Collisions::ClipVerticesAgainstPlane(VertexBuffer1.GetData(), VertexBuffer2.GetData(), ContactPointCount, VertexBuffer1.Num(), ClippingPlaneNormal, FVec3::DotProduct(CurrentPoint, ClippingPlaneNormal));
+				Swap(VertexBuffer1, VertexBuffer2); // VertexBuffer1 will now point to the latest
+			}
+
+			PrevPoint = CurrentPoint;
+		}
+
+		return TArrayView<FVec3>(VertexBuffer1.GetData(), ContactPointCount);
+	}
+
+	// Clip the vertices of a convex to a face of a triangle
+	template <typename ConvexImplicitType>
+	TArrayView<FVec3> ClipConvexToTriangle(
+		const ConvexImplicitType& Convex,
+		const int32 ConvexPlaneIndex,
+		const FTriangle& Triangle,
+		const FVec3& TriangleN,
+		TArrayView<FVec3> VertexBuffer1,
+		TArrayView<FVec3> VertexBuffer2)
+	{
+		check(VertexBuffer1.Num() == VertexBuffer2.Num());
+
+		// Populate the clipped vertices with the convex face vertices
+		int32 ContactPointCount = 0;
+		const int32 ConvexFaceVerticesNum = Convex.NumPlaneVertices(ConvexPlaneIndex);
+		ContactPointCount = FMath::Min(ConvexFaceVerticesNum, VertexBuffer1.Num()); // Number of face vertices
+		for (int32 VertexIndex = 0; VertexIndex < ContactPointCount; ++VertexIndex)
+		{
+			// Todo Check for Grey code
+			VertexBuffer1[VertexIndex] = Convex.GetVertex(Convex.GetPlaneVertex(ConvexPlaneIndex, VertexIndex));
+		}
+
+		// Now clip against all planes that belong to the reference plane's, edges
+		// Note winding order matters here, and we have to handle negative scales
+		int32 ClippingPlaneCount = 3;
+		FVec3 PrevPoint = Triangle.GetVertex(2);
+		for (int32 ClippingPlaneIndex = 0; (ClippingPlaneIndex < ClippingPlaneCount) && (ContactPointCount > 1); ++ClippingPlaneIndex)
+		{
+			const FVec3 CurrentPoint = Triangle.GetVertex(ClippingPlaneIndex);
+
+			// Triangle edge clipping plane
+			// NOTE: Plane is not normalized, but the length cancels out in all clip operations
+			const FVec3 ClippingPlaneNormal = FVec3::CrossProduct(TriangleN, PrevPoint - CurrentPoint);
+			if (ClippingPlaneNormal.SizeSquared() > UE_SMALL_NUMBER)
+			{
+				ContactPointCount = Collisions::ClipVerticesAgainstPlane(VertexBuffer1.GetData(), VertexBuffer2.GetData(), ContactPointCount, VertexBuffer1.Num(), ClippingPlaneNormal, FVec3::DotProduct(CurrentPoint, ClippingPlaneNormal));
+				Swap(VertexBuffer1, VertexBuffer2); // VertexBuffer1 will now point to the latest
+			}
+
+			PrevPoint = CurrentPoint;
+		}
+
+		return TArrayView<FVec3>(VertexBuffer1.GetData(), ContactPointCount);
+	}
+
+	// Triangle is in (possibly scaled) Convex space
+	template <typename ConvexType>
+	void ConstructConvexTriangleOneShotManifold2(const ConvexType& Convex, const FTriangle& Triangle, const FReal CullDistance, TCArray<FContactPoint, 4>& OutContactPoints)
+	{
+		const FReal NormalTolerance = FReal(1.e-8);
+		const FReal NormalToleranceSq = NormalTolerance * NormalTolerance;
+
+		// Triangle (same space as convex)
+		const FVec3 TriN = Triangle.GetNormal();
+		const FVec3 TriC = Triangle.GetCentroid();
+
+		//
+		// SAT: Triangle face vs Convex verts
+		//
+		
+		// We store the convex vertex distances to the triangle face for use in the edge-edge culling
+		// @todo(chaos): scratch or stack allocator
+		TArray<FReal> ConvexVertexDs;
+		ConvexVertexDs.SetNum(Convex.NumVertices());
+		TArrayView<FReal> ConvexVertexDsView = MakeArrayView(ConvexVertexDs);
+
+		// Project the convex onto the triangle normal, with distances relative to the triangle plane
+		FReal TriPlaneDMin, TriPlaneDMax;
+		int32 ConvexVertexIndexMin, ConvexVertexIndexMax;
+		ProjectOntoAxis(Convex, TriN, TriC, TriPlaneDMin, TriPlaneDMax, ConvexVertexIndexMin, ConvexVertexIndexMax, &ConvexVertexDsView);
+
+		// Distance culling
+		if (TriPlaneDMin > CullDistance)
+		{
+			// Outside triangle face and separated by more than CullDistance
+			return;
+		}
+		if (TriPlaneDMax < FReal(0))
+		{
+			// Inside triangle face (single-sided collision)
+			return;
+		}
+
+		//
+		// SAT: Convex faces vs triangle verts
+		//
+
+		// For each convex plane, project the Triangle onto the convex plane normal
+		// and reject if the separation is more than cull distance.
+		FVec3 ConvexPlaneN = FVec3(0);
+		FVec3 ConvexPlaneX = FVec3(0);
+		FReal ConvexPlaneDMin = std::numeric_limits<FReal>::lowest();
+		int32 ConvexPlaneIndexMin = INDEX_NONE;
+		for (int32 PlaneIndex = 0; PlaneIndex < Convex.NumPlanes(); ++PlaneIndex)
+		{
+			FVec3 ConN, ConX;
+			Convex.GetPlaneNX(PlaneIndex, ConN, ConX);
+
+			FReal DMin, DMax;
+			int32 IndexMin, IndexMax;
+			ProjectOntoAxis(Triangle, ConN, ConX, DMin, DMax, IndexMin, IndexMax);
+
+			// Distance culling
+			// @todo(chaos): Cull against the projected convex hull, not just the outward face 
+			// (we can store the most-distance vertex for each face with the convex to avoid actually having to project)
+			if (DMin > CullDistance)
+			{
+				// Separated by more than CullDistance
+				return;
+			}
+
+			if (DMin > ConvexPlaneDMin)
+			{
+				ConvexPlaneN = ConN;
+				ConvexPlaneX = ConX;
+				ConvexPlaneDMin = DMin;
+				ConvexPlaneIndexMin = PlaneIndex;
+			}
+		}
+
+		//
+		// SAT: Convex edges vs triangle edges
+		//
+		
+		// Calculate the distance of triangle each edge to the convex separating plane
+		FReal TriVertexConvexDMin0 = FVec3::DotProduct(Triangle.GetVertex(0) - ConvexPlaneX, ConvexPlaneN);
+		FReal TriVertexConvexDMin1 = FVec3::DotProduct(Triangle.GetVertex(1) - ConvexPlaneX, ConvexPlaneN);
+		FReal TriVertexConvexDMin2 = FVec3::DotProduct(Triangle.GetVertex(2) - ConvexPlaneX, ConvexPlaneN);
+		FReal TriEdgeConvexDMin[3] =
+		{
+			FMath::Max(TriVertexConvexDMin2, TriVertexConvexDMin0),
+			FMath::Max(TriVertexConvexDMin0, TriVertexConvexDMin1),
+			FMath::Max(TriVertexConvexDMin1, TriVertexConvexDMin2),
+		};
+
+		FVec3 EdgeEdgeN = FVec3(0);
+		FReal EdgeEdgeDMin = std::numeric_limits<FReal>::lowest();
+		int32 ConvexEdgeIndexMin = INDEX_NONE;
+		int32 TriEdgeIndexMin = INDEX_NONE;
+		for (int32 ConvexEdgeIndex = 0; ConvexEdgeIndex < Convex.NumEdges(); ++ConvexEdgeIndex)
+		{
+			// Skip convex edges beyond CullDistance of the triangle face
+			const int32 ConvexEdgeVertexIndex0 = Convex.GetEdgeVertex(ConvexEdgeIndex, 0);
+			const int32 ConvexEdgeVertexIndex1 = Convex.GetEdgeVertex(ConvexEdgeIndex, 1);
+			const FReal FaceConvexD0 = ConvexVertexDs[ConvexEdgeVertexIndex0];
+			const FReal FaceConvexD1 = ConvexVertexDs[ConvexEdgeVertexIndex1];
+			if ((FaceConvexD0 > CullDistance) && (FaceConvexD1 > CullDistance))
+			{
+				continue;
+			}
+
+			// Convex edge vertices
+			const FVec3 ConvexEdgeV0 = Convex.GetVertex(ConvexEdgeVertexIndex0);
+			const FVec3 ConvexEdgeV1 = Convex.GetVertex(ConvexEdgeVertexIndex1);
+
+			// Convex planes that form the edge
+			const int32 ConvexEdgePlaneIndexA = Convex.GetEdgePlane(ConvexEdgeIndex, 0);
+			const int32 ConvexEdgePlaneIndexB = Convex.GetEdgePlane(ConvexEdgeIndex, 1);
+			const FVec3 ConvexEdgePlaneNormalA = Convex.GetPlane(ConvexEdgePlaneIndexA).Normal();
+			const FVec3 ConvexEdgePlaneNormalB = Convex.GetPlane(ConvexEdgePlaneIndexB).Normal();
+
+			for (int32 TriEdgeIndex = 0; TriEdgeIndex < 3; ++TriEdgeIndex)
+			{
+				// Skip triangle edges beyond cull distance of the convex separating face
+				if (TriEdgeConvexDMin[TriEdgeIndex] > CullDistance)
+				{
+					continue;
+				}
+
+				// Triangle edge vertices
+				const FVec3& TriEdgeV0 = (TriEdgeIndex == 0) ? Triangle.GetVertex(2) : Triangle.GetVertex(TriEdgeIndex - 1);
+				const FVec3& TriEdgeV1 = Triangle.GetVertex(TriEdgeIndex);
+
+				// Skip edge pairs that do not contribute to the Minkowski Sum surface
+				// NOTE: This relies on the ordering of the edge planes from above. 
+				// I.e., we require Sign(ConvexEdgePlaneNormalA x ConvexEdgePlaneNormalB) == Sign(ConvexEdgeV1 - ConvexEdgeV0)
+				// Also note that we must pass the negated triangle normal in
+				if (!IsMinkowskiSumConvexTriangle(ConvexEdgePlaneNormalA, ConvexEdgePlaneNormalB, ConvexEdgeV1 - ConvexEdgeV0, -TriN, TriEdgeV1 - TriEdgeV0))
+				{
+					continue;
+				}
+
+				// Separating axis
+				// NOTE: Not normalized at this stage. We perform the projection against the
+				// non-normalized axis and defer the square root until we know we need it
+				FVec3 Axis = FVec3::CrossProduct(ConvexEdgeV1 - ConvexEdgeV0, TriEdgeV1 - TriEdgeV0);
+				const FReal AxisLenSq = Axis.SizeSquared();
+
+				// Pick consistent axis direction: away from triangle (we want a signed distance)
+				const FReal Sign = FVec3::DotProduct(TriEdgeV0 - TriC, Axis);
+				if (Sign < FReal(0))
+				{
+					Axis = -Axis;
+				}
+
+				const FReal ScaledSeparation = FVec3::DotProduct(ConvexEdgeV0 - TriEdgeV0, Axis);
+
+				// Check cull distance on projected segments
+				// Comparing square distances scaled by axis length to defer square root (keep the sign)
+				const FReal ScaledSeparationSq = ScaledSeparation * FMath::Abs(ScaledSeparation);
+				const FReal ScaledCullDistanceSq = FMath::Square(CullDistance) * AxisLenSq;
+				if (ScaledSeparationSq > ScaledCullDistanceSq)
+				{
+					return;
+				}
+
+				// Comparing square distances scaled by axis length to defer square root (keep the sign)
+				const FReal ScaledEdgeEdgeDMinSq = EdgeEdgeDMin * FMath::Abs(EdgeEdgeDMin) * AxisLenSq;
+				if (ScaledSeparationSq > ScaledEdgeEdgeDMinSq)
+				{
+					// Now we need to know the actual separation and axis
+					const FReal AxisInvLen = FMath::InvSqrt(AxisLenSq);
+					EdgeEdgeDMin = ScaledSeparation * AxisInvLen;
+					EdgeEdgeN = Axis * AxisInvLen;
+					ConvexEdgeIndexMin = ConvexEdgeIndex;
+					TriEdgeIndexMin = TriEdgeIndex;
+				}
+			}
+		}
+
+		// Determine the best features to use for this collision
+		FVec3 SeparatingAxis, ClipAxis;
+		EContactPointType ContactType;
+		bool bClipConvexToTriangle;
+		bool bClipToFaceNormal;
+		const FReal TriFaceBias = FReal(1.e-2);	// Prevent flip=flip on near parallel cases
+		if ((TriPlaneDMin + TriFaceBias > ConvexPlaneDMin) && (TriPlaneDMin + TriFaceBias > EdgeEdgeDMin))
+		{
+			// Triangle face contact - clip the convex vertices to the triangle
+			ContactType = EContactPointType::VertexPlane;
+			bClipConvexToTriangle = true;
+			bClipToFaceNormal = true;
+
+			// The triangle normal is the separating axis
+			SeparatingAxis = TriN;
+			if (FVec3::DotProduct(SeparatingAxis, Convex.GetCenterOfMass() - TriC) < FReal(0))
+			{
+				SeparatingAxis = -SeparatingAxis;
+			}
+
+			ClipAxis = TriN;
+
+			// Find the convex face most opposing the separating axis			
+			ConvexPlaneIndexMin = Convex.GetMostOpposingPlane(TriN);		// @todo(chaos): should use the known vertex index
+			Convex.GetPlaneNX(ConvexPlaneIndexMin, ConvexPlaneN, ConvexPlaneX);
+
+		}
+		else if (ConvexPlaneDMin > EdgeEdgeDMin)
+		{
+			// Convex face contact - clip the triangle to the convex face
+			ContactType = EContactPointType::PlaneVertex;
+			bClipConvexToTriangle = false;
+			bClipToFaceNormal = true;
+
+			// The convex face is the separating axis, but it must point from the triangle to the convex
+			SeparatingAxis = ConvexPlaneN;
+			if (FVec3::DotProduct(SeparatingAxis, Convex.GetCenterOfMass() - TriC) < FReal(0))
+			{
+				SeparatingAxis = -SeparatingAxis;
+			}
+
+			ClipAxis = ConvexPlaneN;
+		}
+		else
+		{
+			// Edge-edge contact - clip triangle vs convex or vice-versa based on most opposing normals
+			ContactType = EContactPointType::EdgeEdge;
+
+			// The separating axis must point from the triangle to the convex
+			SeparatingAxis = EdgeEdgeN;
+			if (FVec3::DotProduct(SeparatingAxis, Convex.GetCenterOfMass() - TriC) < FReal(0))
+			{
+				SeparatingAxis = -SeparatingAxis;
+			}
+
+			// Find the convex face most opposing the separating axis			
+			ConvexPlaneIndexMin = Convex.GetMostOpposingPlane(SeparatingAxis);	// @todo(chaos): should use the known edge index
+			Convex.GetPlaneNX(ConvexPlaneIndexMin, ConvexPlaneN, ConvexPlaneX);
+
+			// Decide whether to clip against the triangle or the convex
+			ClipAxis = SeparatingAxis;
+			const FReal TriNDotAxis = FVec3::DotProduct(TriN, SeparatingAxis);
+			const FReal ConvexNDotAxis = FVec3::DotProduct(ConvexPlaneN, SeparatingAxis);
+			if (FMath::Abs(TriNDotAxis) > FMath::Abs(ConvexNDotAxis))
+			{
+				bClipConvexToTriangle = true;
+				if (FVec3::DotProduct(ClipAxis, TriN) < 0)
+				{
+					ClipAxis = -ClipAxis;
+				}
+			}
+			else
+			{
+				bClipConvexToTriangle = false;
+				if (FVec3::DotProduct(ClipAxis, ConvexPlaneN) < 0)
+				{
+					ClipAxis = -ClipAxis;
+				}
+			}
+			bClipToFaceNormal = false;
+		}
+
+		// @todo(chaos): scratch or stack allocation of clipped vertex buffers
+		const int32 MaxContactPointCount = 32;
+		TCArray<FVec3, MaxContactPointCount> ClippedVertexBuffers0 = TCArray<FVec3, MaxContactPointCount>::MakeFull();
+		TCArray<FVec3, MaxContactPointCount> ClippedVertexBuffers1 = TCArray<FVec3, MaxContactPointCount>::MakeFull();
+
+		// @todo(chaos): 2D Clip
+		TArrayView<FVec3> ClippedVertices;
+		if (bClipConvexToTriangle)
+		{
+			ClippedVertices = ClipConvexToTriangle(Convex, ConvexPlaneIndexMin, Triangle, ClipAxis, MakeArrayView(ClippedVertexBuffers0), MakeArrayView(ClippedVertexBuffers1));
+		}
+		else
+		{
+			ClippedVertices = ClipTriangleToConvex(Triangle, Convex, ConvexPlaneIndexMin, ClipAxis, MakeArrayView(ClippedVertexBuffers0), MakeArrayView(ClippedVertexBuffers1));
+		}
+
+		// Reduce number of contacts to the maximum allowed
+		if (ClippedVertices.Num() > 4)
+		{
+			const FRotation3 RotateSeperationToZ = FRotation3::FromRotatedVector(SeparatingAxis, FVec3(0.0f, 0.0f, 1.0f));
+			for (int32 ContactPointIndex = 0; ContactPointIndex < ClippedVertices.Num(); ++ContactPointIndex)
+			{
+				ClippedVertices[ContactPointIndex] = RotateSeperationToZ * ClippedVertices[ContactPointIndex];
+			}
+
+			const int32 ReducedContactPointCount = Collisions::ReduceManifoldContactPoints(ClippedVertices.GetData(), ClippedVertices.Num());
+			ClippedVertices = MakeArrayView(ClippedVertices.GetData(), ReducedContactPointCount);
+
+			for (int32 ContactPointIndex = 0; ContactPointIndex < ClippedVertices.Num(); ++ContactPointIndex)
+			{
+				ClippedVertices[ContactPointIndex] = RotateSeperationToZ.Inverse() * ClippedVertices[ContactPointIndex];
+			}
+		}
+
+		// Add the clipped points to the contact list
+		const auto& AddContact = [&CullDistance, &ContactType, &SeparatingAxis, &OutContactPoints](const FVec3& ConvexX, const FVec3 TriX, const FReal Distance)
+		{
+			if (Distance < CullDistance)
+			{
+				FContactPoint& ContactPoint = OutContactPoints[OutContactPoints.Add()];
+				ContactPoint.ShapeContactPoints[0] = ConvexX;
+				ContactPoint.ShapeContactPoints[1] = TriX;
+				ContactPoint.ShapeContactNormal = SeparatingAxis;
+				ContactPoint.Phi = Distance;
+				ContactPoint.ContactType = ContactType;
+				ContactPoint.FaceIndex = INDEX_NONE;
+			}
+		};
+
+		check(ClippedVertices.Num() <= OutContactPoints.Max());
+		if (bClipConvexToTriangle && bClipToFaceNormal)
+		{
+			// Clipped points are on the convex, and we clipped to the triangle face along its normal
+			for (int32 ContactIndex = 0; ContactIndex < ClippedVertices.Num(); ++ContactIndex)
+			{
+				const FVec3& ConvexX = ClippedVertices[ContactIndex];
+				const FReal Distance = FVec3::DotProduct(ConvexX - TriC, SeparatingAxis);
+				const FVec3 TriX = ConvexX - Distance * SeparatingAxis;
+
+				AddContact(ConvexX, TriX, Distance);
+			}
+		}
+		else if (bClipConvexToTriangle && !bClipToFaceNormal)
+		{
+			// Clipped points are on the convex, and we clipped to the triangle face, but not along its normal
+			for (int32 ContactIndex = 0; ContactIndex < ClippedVertices.Num(); ++ContactIndex)
+			{
+				const FVec3& ConvexX = ClippedVertices[ContactIndex];
+				const FReal IntersectDenom = FVec3::DotProduct(SeparatingAxis, TriN);
+				check(FMath::Abs(IntersectDenom) > UE_SMALL_NUMBER);	// Guaranteed based on axis selection above
+				const FReal Distance = FVec3::DotProduct(ConvexX - TriC, TriN) / IntersectDenom;
+				const FVec3 TriX = ConvexX - Distance * SeparatingAxis;
+
+				AddContact(ConvexX, TriX, Distance);
+			}
+		}
+		else if (!bClipConvexToTriangle && bClipToFaceNormal)
+		{
+			// Clipped points are on the triangle, and we clipped to the convex face along its normal
+			for (int32 ContactIndex = 0; ContactIndex < ClippedVertices.Num(); ++ContactIndex)
+			{
+				const FVec3& TriX = ClippedVertices[ContactIndex];
+				const FReal Distance = FVec3::DotProduct(ConvexPlaneX - TriX, SeparatingAxis);
+				const FVec3 ConvexX = TriX + Distance * SeparatingAxis;
+
+				AddContact(ConvexX, TriX, Distance);
+			}
+		}
+		else if (!bClipConvexToTriangle && !bClipToFaceNormal)
+		{
+			// Clipped points are on the triangle, and we clipped to the convex face, but not along its normal
+			for (int32 ContactIndex = 0; ContactIndex < ClippedVertices.Num(); ++ContactIndex)
+			{
+				const FVec3& TriX = ClippedVertices[ContactIndex];
+				const FReal IntersectDenom = FVec3::DotProduct(SeparatingAxis, ConvexPlaneN);
+				check(FMath::Abs(IntersectDenom) > UE_SMALL_NUMBER);	// Guaranteed based on axis selection above
+				const FReal Distance = FVec3::DotProduct(ConvexPlaneX - TriX, ConvexPlaneN) / IntersectDenom;
+				const FVec3 ConvexX = TriX + Distance * SeparatingAxis;
+
+				AddContact(ConvexX, TriX, Distance);
+			}
+		}
+	}
+
+
+	template
+	void ConstructConvexTriangleOneShotManifold2(
+		const FImplicitConvex3& Convex, 
+		const FTriangle& Triangle,
+		const FReal CullDistance, 
+		TCArray<FContactPoint, 4>& OutContactPoints);
+
+	template
+	void ConstructConvexTriangleOneShotManifold2(
+		const TImplicitObjectInstanced<FImplicitConvex3>& Convex,
+		const FTriangle& Triangle,
+		const FReal CullDistance,
+		TCArray<FContactPoint, 4>& OutContactPoints);
+
+	template
+	void ConstructConvexTriangleOneShotManifold2(
+		const TImplicitObjectScaled<FImplicitConvex3>& Convex,
+		const FTriangle& Triangle,
+		const FReal CullDistance,
+		TCArray<FContactPoint, 4>& OutContactPoints);
+
+	template
+	void ConstructConvexTriangleOneShotManifold2(
+		const FImplicitBox3& Convex,
+		const FTriangle& Triangle,
+		const FReal CullDistance,
+		TCArray<FContactPoint, 4>& OutContactPoints);
+
+	template
+	void ConstructConvexTriangleOneShotManifold2(
+		const TImplicitObjectScaled<FImplicitBox3>& Convex,
+		const FTriangle& Triangle,
+		const FReal CullDistance,
+		TCArray<FContactPoint, 4>& OutContactPoints);
+
+	template
+	void ConstructConvexTriangleOneShotManifold2(
+		const TImplicitObjectInstanced<FImplicitBox3>& Convex,
+		const FTriangle& Triangle,
+		const FReal CullDistance,
+		TCArray<FContactPoint, 4>& OutContactPoints);
+}
\ No newline at end of file
diff --git a/Engine/Source/Runtime/Experimental/Chaos/Private/Chaos/CollisionOneShotManifolds.cpp b/Engine/Source/Runtime/Experimental/Chaos/Private/Chaos/CollisionOneShotManifolds.cpp
--- a/Engine/Source/Runtime/Experimental/Chaos/Private/Chaos/CollisionOneShotManifolds.cpp
+++ b/Engine/Source/Runtime/Experimental/Chaos/Private/Chaos/CollisionOneShotManifolds.cpp
@@ -92,6 +92,9 @@ namespace Chaos
 	bool bChaos_Collision_UseCapsuleTriMesh2 = true;
 	FAutoConsoleVariableRef CVarChaos_Collision_UseCapsuleTriMesh2(TEXT("p.Chaos.Collision.UseCapsuleTriMesh2"), bChaos_Collision_UseCapsuleTriMesh2, TEXT(""));
 
+	bool bChaos_Collision_UseConvexTriMesh2 = true;
+	FAutoConsoleVariableRef CVarChaos_Collision_UseConvexTriMesh2(TEXT("p.Chaos.Collision.UseConvexTriMesh2"), bChaos_Collision_UseConvexTriMesh2, TEXT(""));
+
 	namespace Collisions
 	{
 		// Forward delarations we need from CollisionRestitution.cpp
@@ -283,7 +286,7 @@ namespace Chaos
 		// more vertices may be added to outputVertexBuffer by this function
 		// This is the core of the Sutherland-Hodgman algorithm
 		// Plane Normals face outwards 
-		inline uint32 ClipVerticesAgainstPlane(const FVec3* InputVertexBuffer, FVec3* OutputVertexBuffer, const uint32 ClipPointCount, const uint32 MaxNumberOfOutputPoints, const FVec3 ClippingPlaneNormal, const FReal PlaneDistance)
+		uint32 ClipVerticesAgainstPlane(const FVec3* InputVertexBuffer, FVec3* OutputVertexBuffer, const uint32 ClipPointCount, const uint32 MaxNumberOfOutputPoints, const FVec3 ClippingPlaneNormal, const FReal PlaneDistance)
 		{
 			check(ClipPointCount > 0);
 
diff --git a/Engine/Source/Runtime/Experimental/Chaos/Private/Chaos/CollisionOneShotManifolds.h b/Engine/Source/Runtime/Experimental/Chaos/Private/Chaos/CollisionOneShotManifolds.h
--- a/Engine/Source/Runtime/Experimental/Chaos/Private/Chaos/CollisionOneShotManifolds.h
+++ b/Engine/Source/Runtime/Experimental/Chaos/Private/Chaos/CollisionOneShotManifolds.h
@@ -14,6 +14,7 @@ namespace Chaos
 	{
 		uint32 BoxBoxClipVerticesAgainstPlane(const FVec3* InputVertexBuffer, FVec3* outputVertexBuffer, uint32 ClipPointCount, int32 ClippingAxis, FReal Distance);
 		uint32 ReduceManifoldContactPoints(FVec3* Points, uint32 PointCount);
+		uint32 ClipVerticesAgainstPlane(const FVec3* InputVertexBuffer, FVec3* OutputVertexBuffer, const uint32 ClipPointCount, const uint32 MaxNumberOfOutputPoints, const FVec3 ClippingPlaneNormal, const FReal PlaneDistance);
 
 		void ConstructBoxBoxOneShotManifold(
 			const FImplicitBox3& Box1,
diff --git a/Engine/Source/Runtime/Experimental/Chaos/Private/Chaos/CollisionOneShotManifoldsMiscShapes.cpp b/Engine/Source/Runtime/Experimental/Chaos/Private/Chaos/CollisionOneShotManifoldsMiscShapes.cpp
--- a/Engine/Source/Runtime/Experimental/Chaos/Private/Chaos/CollisionOneShotManifoldsMiscShapes.cpp
+++ b/Engine/Source/Runtime/Experimental/Chaos/Private/Chaos/CollisionOneShotManifoldsMiscShapes.cpp
@@ -8,6 +8,7 @@
 #include "Chaos/Collision/CapsuleTriangleContactPoint.h"
 #include "Chaos/Collision/ContactPointsMiscShapes.h"
 #include "Chaos/Collision/ContactTriangles.h"
+#include "Chaos/Collision/ConvexTriangleContactPoint.h"
 #include "Chaos/Collision/SphereConvexContactPoint.h"
 #include "Chaos/Collision/PBDCollisionConstraint.h"
 #include "Chaos/Convex.h"
@@ -41,6 +42,7 @@ namespace Chaos
 	extern bool bChaos_Collision_OneSidedHeightField;
 
 	extern bool bChaos_Collision_UseCapsuleTriMesh2;
+	extern bool bChaos_Collision_UseConvexTriMesh2;
 
 	namespace Collisions
 	{
@@ -621,7 +623,14 @@ namespace Chaos
 		template<typename ConvexType>
 		void GenerateConvexTriangleOneShotManifold(const ConvexType& Convex, const FTriangle& Triangle, const FReal CullDistance, TCArray<FContactPoint, 4>& OutContactPoints)
 		{
-			ConstructPlanarConvexTriangleOneShotManifold(Convex, Triangle, CullDistance, OutContactPoints);
+			if (bChaos_Collision_UseConvexTriMesh2)
+			{
+				ConstructConvexTriangleOneShotManifold2(Convex, Triangle, CullDistance, OutContactPoints);
+			}
+			else
+			{
+				ConstructPlanarConvexTriangleOneShotManifold(Convex, Triangle, CullDistance, OutContactPoints);
+			}
 		}
 
 		template<>
diff --git a/Engine/Source/Runtime/Experimental/Chaos/Private/Chaos/Convex.cpp b/Engine/Source/Runtime/Experimental/Chaos/Private/Chaos/Convex.cpp
--- a/Engine/Source/Runtime/Experimental/Chaos/Private/Chaos/Convex.cpp
+++ b/Engine/Source/Runtime/Experimental/Chaos/Private/Chaos/Convex.cpp
@@ -302,51 +302,6 @@ namespace Chaos
 		return 0;
 	}
 
-	int32 FConvex::NumPlaneVertices(int32 PlaneIndex) const
-	{
-		if (StructureData.IsValid())
-		{
-			return StructureData.NumPlaneVertices(PlaneIndex);
-		}
-		return 0;
-	}
-
-	int32 FConvex::GetPlaneVertex(int32 PlaneIndex, int32 PlaneVertexIndex) const
-	{
-		if (StructureData.IsValid())
-		{
-			return StructureData.GetPlaneVertex(PlaneIndex, PlaneVertexIndex);
-		}
-		return INDEX_NONE;
-	}
-
-	int32 FConvex::GetEdgeVertex(int32 EdgeIndex, int32 EdgeVertexIndex) const
-	{
-		if (StructureData.IsValid())
-		{
-			return StructureData.GetEdgeVertex(EdgeIndex, EdgeVertexIndex);
-		}
-		return INDEX_NONE;
-	}
-
-	int32 FConvex::GetEdgePlane(int32 EdgeIndex, int32 EdgePlaneIndex) const
-	{
-		if (StructureData.IsValid())
-		{
-			return StructureData.GetEdgePlane(EdgeIndex, EdgePlaneIndex);
-		}
-		return INDEX_NONE;
-	}
-
-	int32 FConvex::NumEdges() const
-	{
-		if (StructureData.IsValid())
-		{
-			return StructureData.NumEdges();
-		}
-		return 0;
-	}
-
 	// Store the structure data with the convex. This is used by manifold generation, for example
 	void FConvex::CreateStructureData(TArray<TArray<int32>>&& PlaneVertexIndices)
 	{
diff --git a/Engine/Source/Runtime/Experimental/Chaos/Private/Chaos/HeightField.cpp b/Engine/Source/Runtime/Experimental/Chaos/Private/Chaos/HeightField.cpp
--- a/Engine/Source/Runtime/Experimental/Chaos/Private/Chaos/HeightField.cpp
+++ b/Engine/Source/Runtime/Experimental/Chaos/Private/Chaos/HeightField.cpp
@@ -1807,6 +1807,21 @@ namespace Chaos
 
 		GetGridIntersections(FlatQueryBounds, Intersections);
 
+		// Forward winding point indices for each triangle
+		int32 Indices[2][3] =
+		{
+			{0, 1, 3},
+			{0, 3, 2}
+		};
+
+		// Swap indices if reverse winding
+		const bool bStandardWinding = ((GeomData.Scale.X * GeomData.Scale.Y * GeomData.Scale.Z) >= FReal(0));
+		if (!bStandardWinding)
+		{
+			Swap(Indices[0][1], Indices[0][2]);
+			Swap(Indices[1][1], Indices[1][2]);
+		}
+
 		FReal LocalContactPhi = FLT_MAX;
 		FVec3 LocalContactLocation, LocalContactNormal;
 		for (const TVec2<int32>& Cell : Intersections)
@@ -1823,7 +1838,7 @@ namespace Chaos
 			// The triangle is solid so proceed to test it
 			GeomData.GetPointsScaled(SingleIndex, Points);
 
-			if (OverlapTriangle(Points[0], Points[1], Points[3], LocalContactLocation, LocalContactNormal, LocalContactPhi))
+			if (OverlapTriangle(Points[Indices[0][0]], Points[Indices[0][1]], Points[Indices[0][2]], LocalContactLocation, LocalContactNormal, LocalContactPhi))
 			{
 				if (LocalContactPhi < ContactPhi)
 				{
@@ -1834,7 +1849,7 @@ namespace Chaos
 				}
 			}
 
-			if (OverlapTriangle(Points[0], Points[3], Points[2], LocalContactLocation, LocalContactNormal, LocalContactPhi))
+			if (OverlapTriangle(Points[Indices[1][0]], Points[Indices[1][1]], Points[Indices[1][2]], LocalContactLocation, LocalContactNormal, LocalContactPhi))
 			{
 				if (LocalContactPhi < ContactPhi)
 				{
diff --git a/Engine/Source/Runtime/Experimental/Chaos/Private/Chaos/TriangleMeshImplicitObject.cpp b/Engine/Source/Runtime/Experimental/Chaos/Private/Chaos/TriangleMeshImplicitObject.cpp
--- a/Engine/Source/Runtime/Experimental/Chaos/Private/Chaos/TriangleMeshImplicitObject.cpp
+++ b/Engine/Source/Runtime/Experimental/Chaos/Private/Chaos/TriangleMeshImplicitObject.cpp
@@ -567,6 +567,8 @@ bool FTriangleMeshImplicitObject::GJKContactPointImp(const QueryGeomType& QueryG
 	TRigidTransform<FReal, 3> WorldScaleQueryTM;
 	ScaleTransformHelper(TriMeshScale, QueryTM, WorldScaleQueryTM);
 
+	const bool bStandardWinding = ((TriMeshScale.X * TriMeshScale.Y * TriMeshScale.Z) >= FReal(0));
+
 	auto CalculateTriangleContact = [&](const FVec3& A, const FVec3& B, const FVec3& C,
 		FVec3& LocalContactLocation, FVec3& LocalContactNormal, FReal& LocalContactPhi) -> bool
 	{
@@ -580,6 +582,7 @@ bool FTriangleMeshImplicitObject::GJKContactPointImp(const QueryGeomType& QueryG
 		bool GJKValidResult = GJKPenetration<true>(TriangleConvex, WorldScaleGeom, WorldScaleQueryTM, LambdaPenetration, ClosestA, ClosestB, LambdaNormal, ClosestVertexIndexA, ClosestVertexIndexB, (FReal)0);
 		if (GJKValidResult)
 		{
+			// @todo(chaos): single-sided collision
 			LocalContactLocation = ClosestB;
 			LocalContactNormal = LambdaNormal;
 			LocalContactPhi = -LambdaPenetration;
@@ -600,6 +603,11 @@ bool FTriangleMeshImplicitObject::GJKContactPointImp(const QueryGeomType& QueryG
 			FVec3 A, B, C;
 			TriangleMeshTransformVertsHelper(TriMeshScale, TriIdx, MParticles, Elements, A, B, C);
 
+			if (!bStandardWinding)
+			{
+				Swap(B, C);
+			}
+
 			if (CalculateTriangleContact(A, B, C, LocalContactLocation, LocalContactNormal, LocalContactPhi))
 			{
 				if (LocalContactPhi < OutContactPhi)
diff --git a/Engine/Source/Runtime/Experimental/Chaos/Public/Chaos/Collision/ConvexContactPointUtilities.h b/Engine/Source/Runtime/Experimental/Chaos/Public/Chaos/Collision/ConvexContactPointUtilities.h
new file mode 100644
--- /dev/null
+++ b/Engine/Source/Runtime/Experimental/Chaos/Public/Chaos/Collision/ConvexContactPointUtilities.h
@@ -0,0 +1,78 @@
+// Copyright Epic Games, Inc. All Rights Reserved.
+#pragma once
+
+#include "Chaos/Core.h"
+#include "Chaos/Capsule.h"
+#include "Chaos/Triangle.h"
+
+namespace Chaos
+{
+	// Project a convex onto an axis and return the projected range as well as the vertex indices that bound the range
+	template <typename ConvexType>
+	inline void ProjectOntoAxis(const ConvexType& Convex, const FVec3& AxisN, const FVec3& AxisX, FReal& PMin, FReal& PMax, int32& MinVertexIndex, int32& MaxVertexIndex, TArrayView<FReal>* VertexDs)
+	{
+		PMin = std::numeric_limits<FReal>::max();
+		PMax = std::numeric_limits<FReal>::lowest();
+		for (int32 VertexIndex = 0; VertexIndex < Convex.NumVertices(); ++VertexIndex)
+		{
+			const FVec3 V = Convex.GetVertex(VertexIndex);
+			const FReal D = FVec3::DotProduct(V - AxisX, AxisN);
+			if (D < PMin)
+			{
+				PMin = D;
+				MinVertexIndex = VertexIndex;
+			}
+			if (D > PMax)
+			{
+				PMax = D;
+				MaxVertexIndex = VertexIndex;
+			}
+			if (VertexDs != nullptr)
+			{
+				(*VertexDs)[VertexIndex] = D;
+			}
+		}
+	}
+
+	// Project a triangle onto an axis and return the projected range as well as the vertex indices that bound the range
+	inline void ProjectOntoAxis(const FTriangle& Triangle, const FVec3& AxisN, const FVec3& AxisX, FReal& PMin, FReal& PMax, int32& MinVertexIndex, int32& MaxVertexIndex)
+	{
+		const FVec3& V0 = Triangle.GetVertex(0);
+		const FVec3& V1 = Triangle.GetVertex(1);
+		const FVec3& V2 = Triangle.GetVertex(2);
+		const FReal Ds[3] =
+		{
+			FVec3::DotProduct(V0 - AxisX, AxisN),
+			FVec3::DotProduct(V1 - AxisX, AxisN),
+			FVec3::DotProduct(V2 - AxisX, AxisN)
+		};
+		MinVertexIndex = FMath::Min3Index(Ds[0], Ds[1], Ds[2]);
+		MaxVertexIndex = FMath::Max3Index(Ds[0], Ds[1], Ds[2]);
+		PMin = Ds[MinVertexIndex];
+		PMax = Ds[MaxVertexIndex];
+	}
+
+	// Project a capsule segment onto an axis and return the projected range as well as the vertex indices that bound the range
+	inline void ProjectOntoAxis(const FCapsule& Capsule, const FVec3& AxisN, const FVec3& AxisX, FReal& PMin, FReal& PMax, int32& MinVertexIndex, int32& MaxVertexIndex)
+	{
+		const FVec3 V0 = Capsule.GetX1();
+		const FVec3 V1 = Capsule.GetX2();
+		const FReal D0 = FVec3::DotProduct(V0 - AxisX, AxisN);
+		const FReal D1 = FVec3::DotProduct(V1 - AxisX, AxisN);
+		if (D0 < D1)
+		{
+			MinVertexIndex = 0;
+			MaxVertexIndex = 1;
+			PMin = D0;
+			PMax = D1;
+		}
+		else
+		{
+			MinVertexIndex = 1;
+			MaxVertexIndex = 01;
+			PMin = D1;
+			PMax = D0;
+		}
+	}
+
+}
\ No newline at end of file
diff --git a/Engine/Source/Runtime/Experimental/Chaos/Public/Chaos/Collision/ConvexTriangleContactPoint.h b/Engine/Source/Runtime/Experimental/Chaos/Public/Chaos/Collision/ConvexTriangleContactPoint.h
new file mode 100644
--- /dev/null
+++ b/Engine/Source/Runtime/Experimental/Chaos/Public/Chaos/Collision/ConvexTriangleContactPoint.h
@@ -0,0 +1,12 @@
+// Copyright Epic Games, Inc. All Rights Reserved.
+#pragma once
+#include "Chaos/Core.h"
+#include "Chaos/CollisionResolutionTypes.h"
+#include "Chaos/Framework/UncheckedArray.h"
+#include "Chaos/ImplicitFwd.h"
+
+namespace Chaos
+{
+	template <typename ConvexType>
+	extern CHAOS_API void ConstructConvexTriangleOneShotManifold2(const ConvexType& Convex, const FTriangle& Triangle, const FReal CullDistance, TCArray<FContactPoint, 4>& OutContactPoints);
+}
\ No newline at end of file
diff --git a/Engine/Source/Runtime/Experimental/Chaos/Public/Chaos/Convex.h b/Engine/Source/Runtime/Experimental/Chaos/Public/Chaos/Convex.h
--- a/Engine/Source/Runtime/Experimental/Chaos/Public/Chaos/Convex.h
+++ b/Engine/Source/Runtime/Experimental/Chaos/Public/Chaos/Convex.h
@@ -407,47 +407,84 @@ namespace Chaos
 		int32 GetVertexPlanes3(int32 VertexIndex, int32& PlaneIndex0, int32& PlaneIndex1, int32& PlaneIndex2) const;
 
 		// The number of vertices that make up the corners of the specified face
-		int32 NumPlaneVertices(int32 PlaneIndex) const;
+		inline int32 NumPlaneVertices(int32 PlaneIndex) const
+		{
+			if (StructureData.IsValid())
+			{
+				return StructureData.NumPlaneVertices(PlaneIndex);
+			}
+			return 0;
+		}
 
 		// Get the vertex index of one of the vertices making up the corners of the specified face
-		int32 GetPlaneVertex(int32 PlaneIndex, int32 PlaneVertexIndex) const;
+		inline int32 GetPlaneVertex(int32 PlaneIndex, int32 PlaneVertexIndex) const
+		{
+			if (StructureData.IsValid())
+			{
+				return StructureData.GetPlaneVertex(PlaneIndex, PlaneVertexIndex);
+			}
+			return INDEX_NONE;
+		}
 
-		int32 GetEdgeVertex(int32 EdgeIndex, int32 EdgeVertexIndex) const;
+		// Get the vertex index of one of the two vertices in an edge
+		inline int32 GetEdgeVertex(int32 EdgeIndex, int32 EdgeVertexIndex) const
+		{
+			if (StructureData.IsValid())
+			{
+				return StructureData.GetEdgeVertex(EdgeIndex, EdgeVertexIndex);
+			}
+			return INDEX_NONE;
+		}
 
-		int32 GetEdgePlane(int32 EdgeIndex, int32 EdgePlaneIndex) const;
+		// Get the plane index of one of the two planes using an edge
+		inline int32 GetEdgePlane(int32 EdgeIndex, int32 EdgePlaneIndex) const
+		{
+			if (StructureData.IsValid())
+			{
+				return StructureData.GetEdgePlane(EdgeIndex, EdgePlaneIndex);
+			}
+			return INDEX_NONE;
+		}
 
-		int32 NumPlanes() const
+		inline int32 NumPlanes() const
 		{
 			return Planes.Num();
 		}
 
-		int32 NumEdges() const;
+		inline int32 NumEdges() const
+		{
+			if (StructureData.IsValid())
+			{
+				return StructureData.NumEdges();
+			}
+			return 0;
+		}
 
-		int32 NumVertices() const
+		inline int32 NumVertices() const
 		{
 			return (int32)Vertices.Num();
 		}
 
 		// Get the plane at the specified index (e.g., indices from FindVertexPlanes)
-		const TPlaneConcrete<FReal, 3> GetPlane(int32 FaceIndex) const
+		inline const TPlaneConcrete<FReal, 3> GetPlane(int32 FaceIndex) const
 		{
 			// @todo(chaos) this is needed because this API is shared with BOx implicit - we shoudl eventually only need local space planes and be able to use single precision
 			return TPlaneConcrete<FReal, 3>::MakeFrom(Planes[FaceIndex]);
 		}
 
-		void GetPlaneNX(const int32 FaceIndex, FVec3& OutN, FVec3& OutX) const
+		inline void GetPlaneNX(const int32 FaceIndex, FVec3& OutN, FVec3& OutX) const
 		{
 			OutN = FVec3(Planes[FaceIndex].Normal());
 			OutX = FVec3(Planes[FaceIndex].X());
 		}
 
-		const FPlaneType& GetPlaneRaw(int32 FaceIndex) const
+		inline const FPlaneType& GetPlaneRaw(int32 FaceIndex) const
 		{
 			return Planes[FaceIndex];
 		}
 
 		// Get the vertex at the specified index (e.g., indices from GetPlaneVertexs)
-		const FVec3Type& GetVertex(int32 VertexIndex) const
+		inline const FVec3Type& GetVertex(int32 VertexIndex) const
 		{
 			return Vertices[VertexIndex];
 		}
diff --git a/Engine/Source/Runtime/Experimental/Chaos/Public/Chaos/Framework/UncheckedArray.h b/Engine/Source/Runtime/Experimental/Chaos/Public/Chaos/Framework/UncheckedArray.h
--- a/Engine/Source/Runtime/Experimental/Chaos/Public/Chaos/Framework/UncheckedArray.h
+++ b/Engine/Source/Runtime/Experimental/Chaos/Public/Chaos/Framework/UncheckedArray.h
@@ -33,6 +33,16 @@ namespace Chaos
 		static const int32 MaxElements = N;
 		using ElementType = T;
 
+		static TCArray<T, N> MakeFull()
+		{
+			return TCArray<T, N>(N);
+		}
+
+		static TCArray<T, N> MakeEmpty()
+		{
+			return TCArray<T, N>(0);
+		}
+
 		inline TCArray() : NumElements(0) {}
 
 		inline int32 Num() const { return NumElements; }
@@ -145,6 +155,16 @@ namespace Chaos
 			--NumElements;
 		}
 
+		inline ElementType* GetData()
+		{
+			return &Elements[0];
+		}
+
+		inline const ElementType* GetData() const
+		{
+			return &Elements[0];
+		}
+
 		inline ElementType* begin()
 		{
 			return &Elements[0];
@@ -166,6 +186,12 @@ namespace Chaos
 		}
 
 	private:
+		explicit TCArray(int32 InNumElements)
+		: NumElements(InNumElements)
+		{
+			check(Num() <= Max());
+		}
+
 		// Element count before array for better cache behaviour
 		int32 NumElements;
 		ElementType Elements[MaxElements];
@@ -217,4 +243,9 @@ struct TAllocatorTraits<Chaos::FUncheckedHeapAllocator> : TAllocatorTraits<FHeap
 {
 };
 
+template <typename T, int N>
+struct TIsContiguousContainer<Chaos::TCArray<T, N>>
+{
+	enum { Value = true };
+};
 
diff --git a/Engine/Source/Runtime/Experimental/Chaos/Public/Chaos/HeightField.h b/Engine/Source/Runtime/Experimental/Chaos/Public/Chaos/HeightField.h
--- a/Engine/Source/Runtime/Experimental/Chaos/Public/Chaos/HeightField.h
+++ b/Engine/Source/Runtime/Experimental/Chaos/Public/Chaos/HeightField.h
@@ -98,6 +98,8 @@ namespace Chaos
 		virtual FVec3 FindGeometryOpposingNormal(const FVec3& DenormDir, int32 FaceIndex, const FVec3& OriginalNormal) const override;
 
 		/**
+		 * Visit all triangles transformed into the specified space. Triangles will have standard winding, regardless of heightfield scale.
+		 * 
 		 * @param QueryBounds Bounding box in which we want to produce triangles, in HeightField space
 		 * @param QueryTransform Transforms from HeightField space to query space (usually the other object's space)
 		 * @param CullDistance The distance at which we can ignore triangles
@@ -114,7 +116,10 @@ namespace Chaos
 			TArray<TVec2<int32>> Intersections;
 			GetGridIntersections(GridQueryBounds, Intersections);
 
+			const bool bStandardWinding = ((GeomData.Scale.X * GeomData.Scale.Y * GeomData.Scale.Z) >= FReal(0));
+
 			FVec3 Points[4];
+			FTriangle Triangles[2];
 
 			for (const TVec2<int32>& Cell : Intersections)
 			{
@@ -149,8 +154,17 @@ namespace Chaos
 					// Generate contacts if overlapping
 					const int32 FaceIndex0 = CellIndex * 2 + 0;
 					const int32 FaceIndex1 = CellIndex * 2 + 1;
-					Visitor(FTriangle(Points[0], Points[1], Points[3]), FaceIndex0, VertexIndex0, VertexIndex1, VertexIndex3);
-					Visitor(FTriangle(Points[0], Points[3], Points[2]), FaceIndex1, VertexIndex0, VertexIndex3, VertexIndex2);
+
+					Triangles[0] = FTriangle(Points[0], Points[1], Points[3]);
+					Triangles[1] = FTriangle(Points[0], Points[3], Points[2]);
+					if (!bStandardWinding)
+					{
+						Triangles[0].ReverseWinding();
+						Triangles[1].ReverseWinding();
+					}
+
+					Visitor(Triangles[0], FaceIndex0, VertexIndex0, VertexIndex1, VertexIndex3);
+					Visitor(Triangles[1], FaceIndex1, VertexIndex0, VertexIndex3, VertexIndex2);
 				}
 			}
 		}
diff --git a/Engine/Source/Runtime/Experimental/Chaos/Public/Chaos/SAT.h b/Engine/Source/Runtime/Experimental/Chaos/Public/Chaos/SAT.h
--- a/Engine/Source/Runtime/Experimental/Chaos/Public/Chaos/SAT.h
+++ b/Engine/Source/Runtime/Experimental/Chaos/Public/Chaos/SAT.h
@@ -1,11 +1,9 @@
 // Copyright Epic Games, Inc. All Rights Reserved.
 #pragma once
 
-#include "CoreMinimal.h"
 #include "Chaos/Core.h"
 #include "Chaos/Plane.h"
-#include "Chaos/Vector.h"
-#include "Chaos/Transform.h"
+#include "Chaos/Utilities.h"
 
 namespace Chaos
 {
@@ -71,8 +69,8 @@ namespace Chaos
 
 	// Check whether the two edges of two convex shapes contribute to the Minkowski sum.
 	// A and B are the face normals for the faces of the edge convex 1
-	// C and D are the face normals for the faces of the edge convex 2
-	bool IsMinkowskiSumFace(const FVec3& A, const FVec3& B, const FVec3& C, const FVec3& D)
+	// C and D are the negated face normals for the faces of the edge convex 2
+	inline bool IsMinkowskiSumFace(const FVec3& A, const FVec3& B, const FVec3& C, const FVec3& D)
 	{
 		const FVec3 BA = FVec3::CrossProduct(B, A);
 		const FVec3 DC = FVec3::CrossProduct(D, C);
@@ -198,6 +196,7 @@ namespace Chaos
 				}
 
 				// Separating normal (always points away from convex 2)
+				// @todo(chaos): we can perform the distance culling with a non-normalized axis and defer the sqrt
 				FVec3 EdgeNormal = FVec3::CrossProduct(EdgeVertex1B - EdgeVertex1A, EdgeVertex2BIn1 - EdgeVertex2AIn1);
 				if (!Utilities::NormalizeSafe(EdgeNormal))
 				{
diff --git a/Engine/Source/Runtime/Experimental/Chaos/Public/Chaos/Triangle.h b/Engine/Source/Runtime/Experimental/Chaos/Public/Chaos/Triangle.h
--- a/Engine/Source/Runtime/Experimental/Chaos/Public/Chaos/Triangle.h
+++ b/Engine/Source/Runtime/Experimental/Chaos/Public/Chaos/Triangle.h
@@ -192,7 +192,14 @@ namespace Chaos
 			return 3;
 		}
 
-		// Returns a winding order multiplier used in the manifold clipping and required when we have negative scales (See ImplicitObjectScaled)
+		// Change the winding order of this triangle by flipping the last two vertices
+		FTriangle& ReverseWinding()
+		{
+			Swap(ABC[1], ABC[2]);
+			return *this;
+		}
+
+		// Triangle winding order is always 1.0. When we reverse the winding of a triangle we flip the verts. @see ReverseWinding
 		// Used for manifold generation
 		FORCEINLINE T GetWindingOrder() const
 		{
diff --git a/Engine/Source/Runtime/Experimental/Chaos/Public/Chaos/TriangleMeshImplicitObject.h b/Engine/Source/Runtime/Experimental/Chaos/Public/Chaos/TriangleMeshImplicitObject.h
--- a/Engine/Source/Runtime/Experimental/Chaos/Public/Chaos/TriangleMeshImplicitObject.h
+++ b/Engine/Source/Runtime/Experimental/Chaos/Public/Chaos/TriangleMeshImplicitObject.h
@@ -701,6 +701,7 @@ namespace Chaos
 
 		/**
 		 * @brief Generate the triangle at the specified index with the specified transform (including scale)
+		 * @note does not correct winding for negative scales
 		*/
 		void GetTransformedTriangle(const int32 TriangleIndex, const FRigidTransform3& Transform, FTriangle& OutTriangle, int32& OutVertexIndex0, int32& OutVertexIndex1, int32& OutVertexIndex2) const
 		{
@@ -735,12 +736,21 @@ namespace Chaos
 			TArray<int32> OverlapIndices;
 			FindOverlappingTriangles(QueryBounds, OverlapIndices);
 
+			const bool bStandardWinding = ((QueryTransform.GetScale3D().X * QueryTransform.GetScale3D().Y * QueryTransform.GetScale3D().Z) >= FReal(0));
+
 			for (int32 OverlapIndex = 0; OverlapIndex < OverlapIndices.Num(); ++OverlapIndex)
 			{
 				const int32 TriangleIndex = OverlapIndices[OverlapIndex];
 				FTriangle Triangle;
 				int32 VertexIndex0, VertexIndex1, VertexIndex2;
 				GetTransformedTriangle(TriangleIndex, QueryTransform, Triangle, VertexIndex0, VertexIndex1, VertexIndex2);
+
+				if (!bStandardWinding)
+				{
+					Triangle.ReverseWinding();
+					Swap(VertexIndex1, VertexIndex2);
+				}
+
 				Visitor(Triangle, TriangleIndex, VertexIndex0, VertexIndex1, VertexIndex2);
 			}
 		}
