Geometrty collection - add multiple blueprint functions
- GetLocalBoudns
- GetDebugInfo : to get detailed info about the various collection and their arrays ( including estimated memory usage )
- SetAnchoredByIndex / SetAnchoredBybox now has a bAnchored boolean

diff --git a/Engine/Source/Runtime/Experimental/Chaos/Private/GeometryCollection/ManagedArrayCollection.cpp b/Engine/Source/Runtime/Experimental/Chaos/Private/GeometryCollection/ManagedArrayCollection.cpp
--- a/Engine/Source/Runtime/Experimental/Chaos/Private/GeometryCollection/ManagedArrayCollection.cpp
+++ b/Engine/Source/Runtime/Experimental/Chaos/Private/GeometryCollection/ManagedArrayCollection.cpp
@@ -535,7 +535,11 @@ bool FManagedArrayCollection::IsConnected(FName StartingNode, FName TargetNode)
 #include <string>
 FString FManagedArrayCollection::ToString() const
 {
-	FString Buffer("");
+	FString Buffer("Group : Attribute [Ptr] [AllocatedSize]\n");
+
+	const TArray<FStringFormatArg> CollectionInfos = { FString::FormatAsNumber((int32)GetAllocatedSize()) };
+	Buffer += FString::Format(TEXT("*:* [n/a] [{0}]\n"), CollectionInfos);
+
 	for (FName GroupName : GroupNames())
 	{
 		Buffer += GroupName.ToString() + "\n";
@@ -544,16 +548,29 @@ FString FManagedArrayCollection::ToString() const
 			FKeyType Key = FManagedArrayCollection::MakeMapKey(AttributeName, GroupName);
 			const FValueType& Value = Map[Key];
 
-			const void* PointerAddress = static_cast<const void*>(Value.Value);
-			std::stringstream AddressStream;
-			AddressStream << PointerAddress;
+			const SIZE_T AttributeAllocatedSize = Value.Value? Value.Value->GetAllocatedSize() : 0;
+			const FString AttributeAllocatedSizeStr = FString::FormatAsNumber((int32)AttributeAllocatedSize);
 
-			Buffer += GroupName.ToString() + ":" + AttributeName.ToString() + " [" + FString(AddressStream.str().c_str()) + "]\n";
+			const TArray<FStringFormatArg> AttributeInfos = { GroupName.ToString(), AttributeName.ToString(), (uint64)Value.Value, AttributeAllocatedSizeStr };
+			Buffer += FString::Format(TEXT("{0}:{1} [{2}] [{3}]\n"), AttributeInfos);
 		}
 	}
 	return Buffer;
 }
 
+SIZE_T FManagedArrayCollection::GetAllocatedSize() const
+{
+	SIZE_T AllocatedSize = Map.GetAllocatedSize();
+	for (const TTuple<FKeyType, FValueType>& Entry : Map)
+	{
+		if (Entry.Value.Value)
+		{
+			AllocatedSize += Entry.Value.Value->GetAllocatedSize();
+		}
+	}
+	return AllocatedSize;
+}
+
 static const FName GuidName("GUID");
 
 // this is a reference wrapper to avoid copying attributes as some may not support copy (unique ptr ones) 
diff --git a/Engine/Source/Runtime/Experimental/Chaos/Private/PhysicsProxy/GeometryCollectionPhysicsProxy.cpp b/Engine/Source/Runtime/Experimental/Chaos/Private/PhysicsProxy/GeometryCollectionPhysicsProxy.cpp
--- a/Engine/Source/Runtime/Experimental/Chaos/Private/PhysicsProxy/GeometryCollectionPhysicsProxy.cpp
+++ b/Engine/Source/Runtime/Experimental/Chaos/Private/PhysicsProxy/GeometryCollectionPhysicsProxy.cpp
@@ -1727,48 +1727,58 @@ void FGeometryCollectionPhysicsProxy::BreakActiveClusters_External()
 	}
 }
 
-void FGeometryCollectionPhysicsProxy::SetAnchored_External(int32 Index)
+static void SetParticleAnchored(Chaos::FPBDRigidsEvolution* Evolution, Chaos::FPBDRigidClusteredParticleHandle* ParticleHandle, bool bAnchored)
+{
+	if (ParticleHandle)
+	{
+		ParticleHandle->SetIsAnchored(bAnchored);
+		if (!bAnchored)
+		{
+			if (Evolution && !ParticleHandle->IsDynamic())
+			{
+				Chaos::FKinematicTarget NoKinematicTarget;
+				Evolution->SetParticleObjectState(ParticleHandle, Chaos::EObjectStateType::Dynamic);
+				Evolution->SetParticleKinematicTarget(ParticleHandle, NoKinematicTarget);
+			}
+		}
+	}
+}
+
+void FGeometryCollectionPhysicsProxy::SetAnchoredByIndex_External(int32 Index, bool bAnchored)
 {
 	check(IsInGameThread());
 	if (Chaos::FPhysicsSolver* RBDSolver = GetSolver<Chaos::FPhysicsSolver>())
 	{
-		RBDSolver->EnqueueCommandImmediate([this, Index, RBDSolver]()
+		RBDSolver->EnqueueCommandImmediate([this, Index, bAnchored, RBDSolver]()
 			{
 				Chaos::FPBDRigidsEvolution* Evolution = RBDSolver->GetEvolution();
 				if (SolverParticleHandles.IsValidIndex(Index))
 				{
-					if (FClusterHandle* ParticleHandle = SolverParticleHandles[Index])
-					{
-						ParticleHandle->SetIsAnchored(true);
-					}
+					SetParticleAnchored(Evolution, SolverParticleHandles[Index], bAnchored);
 				}
 			});
 	}
 }
 
-void FGeometryCollectionPhysicsProxy::SetAnchored_External(const FBox& WorldSpaceBox)
+void FGeometryCollectionPhysicsProxy::SetAnchoredByBox_External(const FBox& WorldSpaceBox, bool bAnchored)
 {
 	check(IsInGameThread());
 	if (Chaos::FPhysicsSolver* RBDSolver = GetSolver<Chaos::FPhysicsSolver>())
 	{
 		Chaos::FAABB3 BoundsToCheck(WorldSpaceBox.Min, WorldSpaceBox.Max);
-		RBDSolver->EnqueueCommandImmediate([this, BoundsToCheck, RBDSolver]()
+		RBDSolver->EnqueueCommandImmediate([this, BoundsToCheck, bAnchored, RBDSolver]()
 			{
 				Chaos::FPBDRigidsEvolution* Evolution = RBDSolver->GetEvolution();
 				for (FClusterHandle* ParticleHandle : GetSolverParticleHandles())
 				{
-					if (ParticleHandle)
+					if (ParticleHandle && BoundsToCheck.Contains(ParticleHandle->X()))
 					{
-						if (BoundsToCheck.Contains(ParticleHandle->X()))
-						{
-							ParticleHandle->SetIsAnchored(true);
-						}
+						SetParticleAnchored(Evolution, ParticleHandle, bAnchored);
 					}
 				}
 			});
 	}
 }
-
 void FGeometryCollectionPhysicsProxy::RemoveAllAnchors_External()
 {
 	check(IsInGameThread());
@@ -1778,20 +1788,9 @@ void FGeometryCollectionPhysicsProxy::RemoveAllAnchors_External()
 		RBDSolver->EnqueueCommandImmediate([this, RBDSolver]()
 			{
 				Chaos::FPBDRigidsEvolution* Evolution = RBDSolver->GetEvolution();
-
-				// disable anchoring where relevant and collect the nodes to update
-				Chaos::FKinematicTarget NoKinematicTarget;
 				for (FClusterHandle* ParticleHandle : GetSolverParticleHandles())
 				{
-					if (ParticleHandle)
-					{
-						ParticleHandle->SetIsAnchored(false);
-						if (!ParticleHandle->IsDynamic())
-						{
-							Evolution->SetParticleObjectState(ParticleHandle, Chaos::EObjectStateType::Dynamic);
-							Evolution->SetParticleKinematicTarget(ParticleHandle, NoKinematicTarget);
-						}
-					}
+					SetParticleAnchored(Evolution, ParticleHandle, false);
 				}
 			});
 	}
diff --git a/Engine/Source/Runtime/Experimental/Chaos/Public/Chaos/Particles.h b/Engine/Source/Runtime/Experimental/Chaos/Public/Chaos/Particles.h
--- a/Engine/Source/Runtime/Experimental/Chaos/Public/Chaos/Particles.h
+++ b/Engine/Source/Runtime/Experimental/Chaos/Public/Chaos/Particles.h
@@ -168,6 +168,11 @@ namespace Chaos
 			return OutHash;
 		}
 
+		SIZE_T GetAllocatedSize() const
+		{
+			return MX.GetAllocatedSize();
+		}
+
 #if PARTICLE_ITERATOR_RANGED_FOR_CHECK
 		int32 DirtyValidationCount() const { return MDirtyValidationCount; }
 #endif
diff --git a/Engine/Source/Runtime/Experimental/Chaos/Public/GeometryCollection/ManagedArray.h b/Engine/Source/Runtime/Experimental/Chaos/Public/GeometryCollection/ManagedArray.h
--- a/Engine/Source/Runtime/Experimental/Chaos/Public/GeometryCollection/ManagedArray.h
+++ b/Engine/Source/Runtime/Experimental/Chaos/Public/GeometryCollection/ManagedArray.h
@@ -77,6 +77,54 @@ inline void TryBulkSerializeManagedArray(Chaos::FChaosArchive& Ar, TArray<FIntVe
 	Array.BulkSerialize(Ar);
 }
 
+
+// --------------------------------------------------------------------------
+// utility functions to estimate the allocated size of managed arrays
+// --------------------------------------------------------------------------
+namespace ManagedArrayTypeSize
+{
+	template<typename T>
+	inline SIZE_T GetAllocatedSize(const T& Value)
+	{
+		return sizeof(T);
+	}
+
+	template<typename T>
+	inline SIZE_T GetAllocatedSize(const TArray<T>& Array)
+	{
+		return Array.GetAllocatedSize();
+	}
+
+	template<typename T>
+	inline SIZE_T GetAllocatedSize(const TSet<T>& Set)
+	{
+		return Set.GetAllocatedSize();
+	}
+
+	template<typename T>
+	inline SIZE_T GetAllocatedSize(const TUniquePtr<T>& Ptr)
+	{
+		return Ptr ? ManagedArrayTypeSize::GetAllocatedSize(*Ptr) : 0;
+	}
+
+	template<typename T, ESPMode Mode>
+	inline SIZE_T GetAllocatedSize(const TSharedPtr<T, Mode>& Ptr)
+	{
+		return Ptr ? ManagedArrayTypeSize::GetAllocatedSize(*Ptr) : 0;
+	}
+
+	inline SIZE_T GetAllocatedSize(const Chaos::FImplicitObject3* ImplicitObjectPtr)
+	{
+		return ImplicitObjectPtr ? sizeof(Chaos::FImplicitObject3) : 0;
+	}
+
+	inline SIZE_T GetAllocatedSize(const Chaos::FBVHParticlesFloat3& BVHParticles)
+	{
+		return BVHParticles.GetAllocatedSize();
+	}
+}
+
+
 /***
 *  Managed Array Base
 *
@@ -118,12 +166,16 @@ protected:
 	*/
 	virtual void Init(const FManagedArrayBase& ) {};
 
+	/**
+	* Get allocated memory 
+	*/
+	virtual SIZE_T GetAllocatedSize() const { return 0; }
+
 	/**
 	* Copy a range of values from the ConstArray into this
 	*/
 	virtual void CopyRange(const FManagedArrayBase& ConstArray, int32 Start, int32 Stop, int32 Offset = 0) {};
 
-
 public:
 	FManagedArrayBase()
 	{
@@ -320,6 +372,11 @@ public:
 		InitHelper(Array, NewTypedArray, Size);
 	}
 
+	virtual SIZE_T GetAllocatedSize() const override
+	{
+		return ManagedArrayTypeSize::GetAllocatedSize(Array);
+	}
+
 	/**
 	* Copy from a predefined Array of matching type
 	*/
diff --git a/Engine/Source/Runtime/Experimental/Chaos/Public/GeometryCollection/ManagedArrayCollection.h b/Engine/Source/Runtime/Experimental/Chaos/Public/GeometryCollection/ManagedArrayCollection.h
--- a/Engine/Source/Runtime/Experimental/Chaos/Public/GeometryCollection/ManagedArrayCollection.h
+++ b/Engine/Source/Runtime/Experimental/Chaos/Public/GeometryCollection/ManagedArrayCollection.h
@@ -533,6 +533,11 @@ public:
 	*/
 	FString ToString() const;
 
+	/**
+	* Get allocated memory size 
+	*/
+	SIZE_T GetAllocatedSize() const;
+
 private:
 
 	/****
diff --git a/Engine/Source/Runtime/Experimental/Chaos/Public/PhysicsProxy/GeometryCollectionPhysicsProxy.h b/Engine/Source/Runtime/Experimental/Chaos/Public/PhysicsProxy/GeometryCollectionPhysicsProxy.h
--- a/Engine/Source/Runtime/Experimental/Chaos/Public/PhysicsProxy/GeometryCollectionPhysicsProxy.h
+++ b/Engine/Source/Runtime/Experimental/Chaos/Public/PhysicsProxy/GeometryCollectionPhysicsProxy.h
@@ -376,8 +376,8 @@ public:
 	void ApplyImpulseAt_External(FVector Force, FVector WorldLocation);
 	void BreakClusters_External(TArray<FGeometryCollectionItemIndex>&& ItemIndices);
 	void BreakActiveClusters_External();
-	void SetAnchored_External(int32 Index);
-	void SetAnchored_External(const FBox& WorldSpaceBox);
+	void SetAnchoredByIndex_External(int32 Index, bool bAnchored);
+	void SetAnchoredByBox_External(const FBox& WorldSpaceBox, bool bAnchored);
 	void RemoveAllAnchors_External();
 	void ApplyExternalStrain_External(FGeometryCollectionItemIndex ItemIndex, const FVector& WorldLocation, float Radius, int32 PropagationDepth, float PropagationFactor, float StrainValue);
 	void ApplyInternalStrain_External(FGeometryCollectionItemIndex ItemIndex, const FVector& WorldLocation, float Radius, int32 PropagationDepth, float PropagationFactor, float StrainValue);
diff --git a/Engine/Source/Runtime/Experimental/GeometryCollectionEngine/Private/GeometryCollection/GeometryCollectionComponent.cpp b/Engine/Source/Runtime/Experimental/GeometryCollectionEngine/Private/GeometryCollection/GeometryCollectionComponent.cpp
--- a/Engine/Source/Runtime/Experimental/GeometryCollectionEngine/Private/GeometryCollection/GeometryCollectionComponent.cpp
+++ b/Engine/Source/Runtime/Experimental/GeometryCollectionEngine/Private/GeometryCollection/GeometryCollectionComponent.cpp
@@ -2722,6 +2722,25 @@ void UGeometryCollectionComponent::SetRestCollection(const UGeometryCollection*
 	}
 }
 
+FString UGeometryCollectionComponent::GetDebugInfo()
+{
+	// print the game thread side of things
+	FString DebugInfo;
+	DebugInfo += FString("RestCollection - ") + (RestCollection? RestCollection->GetName() : FString("None"));
+	DebugInfo += "\n";
+	if (RestCollection && RestCollection->GetGeometryCollection())
+	{
+		DebugInfo += RestCollection->GetGeometryCollection()->ToString();
+	}
+	DebugInfo += FString("DynamicCollection - ") + FString(DynamicCollection ? "Yes": "No");
+	DebugInfo += "\n";
+	if (DynamicCollection)
+	{
+		DebugInfo += DynamicCollection->ToString();
+	}
+	return DebugInfo;
+}
+
 FGeometryCollectionEdit::FGeometryCollectionEdit(UGeometryCollectionComponent* InComponent, GeometryCollection::EEditUpdate InEditUpdate, bool bShapeIsUnchanged)
 	: Component(InComponent)
 	, EditUpdate(InEditUpdate)
@@ -4232,19 +4251,19 @@ void UGeometryCollectionComponent::CrumbleActiveClusters()
 	}
 }
 
-void UGeometryCollectionComponent::SetAnchoredByIndex(int32 Index)
+void UGeometryCollectionComponent::SetAnchoredByIndex(int32 Index, bool bAnchored)
 {
 	if (PhysicsProxy)
 	{
-		PhysicsProxy->SetAnchored_External(Index);
+		PhysicsProxy->SetAnchoredByIndex_External(Index, bAnchored);
 	}
 }
 
-void UGeometryCollectionComponent::SetAnchoredByBox(FBox WorldSpaceBox)
+void UGeometryCollectionComponent::SetAnchoredByBox(FBox WorldSpaceBox, bool bAnchored)
 {
 	if (PhysicsProxy)
 	{
-		PhysicsProxy->SetAnchored_External(WorldSpaceBox);
+		PhysicsProxy->SetAnchoredByBox_External(WorldSpaceBox, bAnchored);
 	}
 }
 
diff --git a/Engine/Source/Runtime/Experimental/GeometryCollectionEngine/Public/GeometryCollection/GeometryCollectionComponent.h b/Engine/Source/Runtime/Experimental/GeometryCollectionEngine/Public/GeometryCollection/GeometryCollectionComponent.h
--- a/Engine/Source/Runtime/Experimental/GeometryCollectionEngine/Public/GeometryCollection/GeometryCollectionComponent.h
+++ b/Engine/Source/Runtime/Experimental/GeometryCollectionEngine/Public/GeometryCollection/GeometryCollectionComponent.h
@@ -514,6 +514,12 @@ public:
 	UPROPERTY(EditAnywhere, Category = "ChaosPhysics", meta = (DisplayName = "Chaos Solver"))
 	TObjectPtr<AChaosSolverActor> ChaosSolverActor;
 
+	/**
+	* Get local bounds of the geometry collection
+	*/
+	UFUNCTION(BlueprintCallable, Category = "ChaosPhysics")
+	FBox GetLocalBounds() const { return LocalBounds; }
+
 	/**
 	 * Apply an external strain to specific piece of the geometry collection
 	 * @param ItemIndex item index ( from HitResult) of the piece to apply strain on
@@ -553,19 +559,19 @@ public:
 	void CrumbleActiveClusters();
 
 	/**
-	* Set a fragment or cluster to be anchored
+	* Set a piece or cluster to be anchored or not 
 	*/
 	UFUNCTION(BlueprintCallable, Category = "ChaosPhysics")
-	void SetAnchoredByIndex(int32 Index);
+	void SetAnchoredByIndex(int32 Index, bool bAnchored);
 
 	/**
-	* Set all fragments within a world bounding box to be anchored
+	* Set all pieces within a world bounding space box to be anchored or not 
 	*/
 	UFUNCTION(BlueprintCallable, Category = "ChaosPhysics")
-	void SetAnchoredByBox(FBox WorldSpaceBox);
+	void SetAnchoredByBox(FBox WorldSpaceBox, bool bAnchored);
 
 	/**
-	* this will remove anchors on all the pieces ( inlcuding the static and kinematic initial states ones ) of the geometry colection
+	* this will remove anchors on all the pieces ( including the static and kinematic initial states ones ) of the geometry colection
 	*/
 	UFUNCTION(BlueprintCallable, Category = "ChaosPhysics")
 	void RemoveAllAnchors();
@@ -629,6 +635,9 @@ public:
 	UFUNCTION(BlueprintCallable, Category = "ChaosPhysics")
 	void SetRestCollection(const UGeometryCollection * RestCollectionIn);
 
+	/** RestCollection */
+	UFUNCTION(BlueprintCallable, Category = "ChaosPhysics")
+	FString GetDebugInfo();
 
 	FORCEINLINE const UGeometryCollection* GetRestCollection() const { return RestCollection; }
 
