Static Analysis Fix

diff --git a/Engine/Source/Runtime/Experimental/Iris/Core/Private/Iris/ReplicationState/ReplicationStateDescriptorBuilder.cpp b/Engine/Source/Runtime/Experimental/Iris/Core/Private/Iris/ReplicationState/ReplicationStateDescriptorBuilder.cpp
--- a/Engine/Source/Runtime/Experimental/Iris/Core/Private/Iris/ReplicationState/ReplicationStateDescriptorBuilder.cpp
+++ b/Engine/Source/Runtime/Experimental/Iris/Core/Private/Iris/ReplicationState/ReplicationStateDescriptorBuilder.cpp
@@ -1917,7 +1917,7 @@ FRepTag FPropertyReplicationStateDescriptorBuilder::GetRepTagFromProperty(const
 			return PropertyName == NAME_Role ? RepTag_NetRole : RepTag_NetRemoteRole;
 		}
 	}
-	else if (PropertyName == ReplicationStateDescriptorBuilder_NAME_NetCullDistanceSquared)
+	else if (Property && PropertyName == ReplicationStateDescriptorBuilder_NAME_NetCullDistanceSquared)
 	{
 		const FFieldClass* Class = Property->GetClass();
 		if (Class->GetFName() == NAME_FloatProperty)
