Chaos : Integration CL22107450 from /UE5/Main.

diff --git a/Engine/Source/Runtime/Experimental/Chaos/Public/GeometryCollection/ManagedArray.h b/Engine/Source/Runtime/Experimental/Chaos/Public/GeometryCollection/ManagedArray.h
--- a/Engine/Source/Runtime/Experimental/Chaos/Public/GeometryCollection/ManagedArray.h
+++ b/Engine/Source/Runtime/Experimental/Chaos/Public/GeometryCollection/ManagedArray.h
@@ -166,6 +166,12 @@ protected:
 	*/
 	virtual void Init(const FManagedArrayBase& ) {};
 
+	/**
+	* Copy a range of values from the ConstArray into this
+	*/
+	virtual void CopyRange(const FManagedArrayBase& ConstArray, int32 Start, int32 Stop, int32 Offset = 0) {};
+
+
 	/**
 	* Get allocated memory 
 	*/
