Choas : Fix CIS Non-unity build.

diff --git a/Engine/Source/Runtime/Experimental/Chaos/Public/Chaos/Deformable/Utilities.h b/Engine/Source/Runtime/Experimental/Chaos/Public/Chaos/Deformable/Utilities.h
--- a/Engine/Source/Runtime/Experimental/Chaos/Public/Chaos/Deformable/Utilities.h
+++ b/Engine/Source/Runtime/Experimental/Chaos/Public/Chaos/Deformable/Utilities.h
@@ -118,7 +118,7 @@ namespace Chaos
 						Tet[1] != Tet[3] &&
 						Tet[2] != Tet[3]))
 					{
-						UE_LOG(LogChaos, Display, TEXT("Skipping degenerate tet %d of %d."), Idx, Tets.Num());
+						//UE_LOG(LogChaosFlesh, Display, TEXT("Skipping degenerate tet %d of %d."), Idx, Tets.Num());
 						continue;
 					}
 
@@ -173,8 +173,8 @@ namespace Chaos
 						}
 					}
 				}
-				if (nonManifold)
-					UE_LOG(LogChaos, Display, TEXT("WARNING: Encountered %d non-manifold tetrahedral mesh faces."), nonManifold);
+				//if (nonManifold)
+				//	UE_LOG(LogChaosFlesh, Display, TEXT("WARNING: Encountered %d non-manifold tetrahedral mesh faces."), nonManifold);
 			}
 		}
 	
