Chaos - Jacobi pair solver for collision (Disabled)

Using a Jacobi solver for the contacts between a shape pair, but a Gauss Seidel solver main loop gives more  stable stacking and friction at the cost of some sponginess at low FPS and an extra parameter to tune.

diff --git a/Engine/Source/Runtime/Experimental/Chaos/Private/Chaos/Collision/PBDCollisionSolver.cpp b/Engine/Source/Runtime/Experimental/Chaos/Private/Chaos/Collision/PBDCollisionSolver.cpp
--- a/Engine/Source/Runtime/Experimental/Chaos/Private/Chaos/Collision/PBDCollisionSolver.cpp
+++ b/Engine/Source/Runtime/Experimental/Chaos/Private/Chaos/Collision/PBDCollisionSolver.cpp
@@ -227,68 +227,4 @@ namespace Chaos
 			}
 		}
 	}
-
-	bool FPBDCollisionSolver::SolveVelocity(const FSolverReal Dt, const bool bApplyDynamicFriction)
-	{
-		// Apply restitution at the average contact point
-		// This means we don't need to run as many iterations to get stable bouncing
-		// It also helps with zero restitution to counter any velocioty added by the PBD solve
-		const bool bSolveAverageContact = (State.NumManifoldPoints > 1) && bChaos_PBDCollisionSolver_Velocity_AveragePointEnabled;
-		if (bSolveAverageContact)
-		{
-			SolveVelocityAverage(Dt);
-		}
-
-		FConstraintSolverBody& Body0 = SolverBody0();
-		FConstraintSolverBody& Body1 = SolverBody1();
-
-		// NOTE: this dynamic friction implementation is iteration-count sensitive
-		// @todo(chaos): fix iteration count dependence of dynamic friction
-		const FSolverReal DynamicFriction = (bApplyDynamicFriction && (Dt > 0) && bChaos_PBDCollisionSolver_Velocity_FrictionEnabled) ? State.VelocityFriction : FSolverReal(0);
-
-		for (int32 PointIndex = 0; PointIndex < State.NumManifoldPoints; ++PointIndex)
-		{
-			FPBDCollisionSolverManifoldPoint& SolverManifoldPoint = State.ManifoldPoints[PointIndex];
-			if (SolverManifoldPoint.ShouldSolveVelocity())
-			{
-				const FSolverReal MinImpulseNormal = FMath::Min(FSolverReal(0), -SolverManifoldPoint.NetPushOutNormal / Dt);
-
-				if (DynamicFriction > 0)
-				{
-					FSolverReal ContactVelocityDeltaNormal, ContactVelocityDeltaTangentU, ContactVelocityDeltaTangentV;
-					SolverManifoldPoint.CalculateContactVelocityError(Body0, Body1, DynamicFriction, Dt, ContactVelocityDeltaNormal, ContactVelocityDeltaTangentU, ContactVelocityDeltaTangentV);
-
-					ApplyVelocityCorrection(
-						State.Stiffness,
-						Dt,
-						DynamicFriction,
-						ContactVelocityDeltaNormal,
-						ContactVelocityDeltaTangentU,
-						ContactVelocityDeltaTangentV,
-						MinImpulseNormal,
-						SolverManifoldPoint,
-						Body0,
-						Body1);
-				}
-				else
-				{
-					FSolverReal ContactVelocityDeltaNormal;
-					SolverManifoldPoint.CalculateContactVelocityErrorNormal(Body0, Body1, ContactVelocityDeltaNormal);
-
-					ApplyVelocityCorrectionNormal(
-						State.Stiffness,
-						ContactVelocityDeltaNormal,
-						MinImpulseNormal,
-						SolverManifoldPoint,
-						Body0,
-						Body1);
-				}
-			}
-		}
-
-		// Early-out support for the velocity solve is not currently very important because we
-		// only run one iteration in the velocity solve phase.
-		// @todo(chaos): support early-out in velocity solve if necessary
-		return true;
-	}
 }
\ No newline at end of file
diff --git a/Engine/Source/Runtime/Experimental/Chaos/Private/Chaos/Collision/PBDCollisionSolver.h b/Engine/Source/Runtime/Experimental/Chaos/Private/Chaos/Collision/PBDCollisionSolver.h
--- a/Engine/Source/Runtime/Experimental/Chaos/Private/Chaos/Collision/PBDCollisionSolver.h
+++ b/Engine/Source/Runtime/Experimental/Chaos/Private/Chaos/Collision/PBDCollisionSolver.h
@@ -21,7 +21,12 @@ namespace Chaos
 
 	namespace CVars
 	{
+		extern bool bChaos_PBDCollisionSolver_Velocity_AveragePointEnabled;
+		extern bool bChaos_PBDCollisionSolver_Velocity_FrictionEnabled;
 		extern float Chaos_PBDCollisionSolver_Position_StaticFrictionStiffness;
+		extern float Chaos_PBDCollisionSolver_JacobiStiffness;
+		extern float Chaos_PBDCollisionSolver_JacobiPositionTolerance;
+		extern float Chaos_PBDCollisionSolver_JacobiRotationTolerance;
 	}
 
 
@@ -259,13 +264,16 @@ namespace Chaos
 		 * @return true if we need to run more iterations, false if we did not apply any correction
 		*/
 		bool SolvePositionWithFriction(const FSolverReal Dt, const FSolverReal MaxPushOut);
+		bool SolvePositionWithFrictionJacobi(const FSolverReal Dt, const FSolverReal MaxPushOut);
 		bool SolvePositionNoFriction(const FSolverReal Dt, const FSolverReal MaxPushOut);
+		bool SolvePositionNoFrictionJacobi(const FSolverReal Dt, const FSolverReal MaxPushOut);
 
 		/**
 		 * @brief Calculate and apply the velocity correction for this iteration
 		 * @return true if we need to run more iterations, false if we did not apply any correction
 		*/
 		bool SolveVelocity(const FSolverReal Dt, const bool bApplyDynamicFriction);
+		bool SolveVelocityJacobi(const FSolverReal Dt, const bool bApplyDynamicFriction);
 
 	private:
 		/**
@@ -337,6 +345,19 @@ namespace Chaos
 	//////////////////////////////////////////////////////////////////////////////////////////////////
 	//////////////////////////////////////////////////////////////////////////////////////////////////
 
+	// If the vector size squared is below the threshold, zero the vector and return false, otherwise leave vector as-is and return true.
+	// This is useful for maintaining perfect restitution bounces - we remove any rounding errors in the impulses before applying to the rotation
+	// @todo(chaos): consider using a test of absolute element values instead
+	// @todo(chaos): disabled for now - too expensive :(
+	FORCEINLINE_DEBUGGABLE bool AboveThresholdElseZero(FSolverVec3& InOutV, const FSolverReal ThresholdSq)
+	{
+		//if (InOutV.SizeSquared() < ThresholdSq)
+		//{
+		//	InOutV = FSolverVec3(0);
+		//	return false;
+		//}
+		return true;
+	}
 
 	FORCEINLINE_DEBUGGABLE void CalculatePositionCorrectionNormal(
 		const FSolverReal Stiffness,
@@ -369,10 +390,8 @@ namespace Chaos
 		const FSolverReal NetPushOutTangent,
 		FSolverReal& OutPushOutTangent)
 	{
-		const FSolverReal FrictionStiffness = CVars::Chaos_PBDCollisionSolver_Position_StaticFrictionStiffness;
-
 		// Bilateral constraint - negative values allowed (unlike the normal correction)
-		OutPushOutTangent = -Stiffness * FrictionStiffness * ContactMassTangent * ContactDeltaTangent;
+		OutPushOutTangent = -Stiffness * ContactMassTangent * ContactDeltaTangent;
 	}
 
 
@@ -466,7 +485,7 @@ namespace Chaos
 			ManifoldPoint.StaticFrictionRatio);					// Out
 
 		// Update the particle state based on the pushout
-		const FVec3 PushOut = PushOutTangentU * ManifoldPoint.WorldContactTangentU + PushOutTangentV * ManifoldPoint.WorldContactTangentV;
+		const FSolverVec3 PushOut = PushOutTangentU * ManifoldPoint.WorldContactTangentU + PushOutTangentV * ManifoldPoint.WorldContactTangentV;
 		if (Body0.IsDynamic())
 		{
 			const FSolverVec3 DX0 = Body0.InvM() * PushOut;
@@ -932,6 +951,7 @@ namespace Chaos
 			// differences in contacts from tick to tick
 			// @todo(chaos): try a decaying maximum per contact point rather than an average (again - we had that once!)
 			const FSolverReal FrictionMaxPushOut = TotalPushOutNormal / FSolverReal(NumFrictionContacts);
+			const FSolverReal FrictionStiffness = State.Stiffness * CVars::Chaos_PBDCollisionSolver_Position_StaticFrictionStiffness;
 
 			for (int32 PointIndex = 0; PointIndex < State.NumManifoldPoints; ++PointIndex)
 			{
@@ -943,7 +963,7 @@ namespace Chaos
 					SolverManifoldPoint.CalculateContactPositionErrorTangential(Body0.SolverBody(), Body1.SolverBody(), ContactDeltaTangentU, ContactDeltaTangentV);
 
 					ApplyPositionCorrectionTangential(
-						State.Stiffness,
+						FrictionStiffness,
 						State.Regularization,
 						State.StaticFriction,
 						State.DynamicFriction,
@@ -990,5 +1010,373 @@ namespace Chaos
 		return false;
 	}
 
+	FORCEINLINE_DEBUGGABLE bool FPBDCollisionSolver::SolvePositionNoFrictionJacobi(const FSolverReal Dt, const FSolverReal MaxPushOut)
+	{
+		// SolverBody decorator used to add mass scaling
+		FConstraintSolverBody& Body0 = SolverBody0();
+		FConstraintSolverBody& Body1 = SolverBody1();
+
+		// Apply the position correction so that all contacts have zero separation
+		FSolverVec3 DX0 = FSolverVec3(0);
+		FSolverVec3 DX1 = FSolverVec3(0);
+		FSolverVec3 DR0 = FSolverVec3(0);
+		FSolverVec3 DR1 = FSolverVec3(0);
+
+		const FSolverReal Stiffness = CVars::Chaos_PBDCollisionSolver_JacobiStiffness * State.Stiffness;
+
+		for (int32 PointIndex = 0; PointIndex < State.NumManifoldPoints; ++PointIndex)
+		{
+			FPBDCollisionSolverManifoldPoint& SolverManifoldPoint = State.ManifoldPoints[PointIndex];
+
+			FSolverReal ContactDeltaNormal;
+			SolverManifoldPoint.CalculateContactPositionErrorNormal(Body0.SolverBody(), Body1.SolverBody(), MaxPushOut, ContactDeltaNormal);
+
+			const bool bProcessManifoldPoint = (ContactDeltaNormal < FSolverReal(0)) || (SolverManifoldPoint.NetPushOutNormal > FSolverReal(UE_SMALL_NUMBER));
+			if (bProcessManifoldPoint)
+			{
+				FSolverReal PushOutNormal = FSolverReal(0);
+				CalculatePositionCorrectionNormal(
+					Stiffness,
+					State.Regularization,
+					ContactDeltaNormal,
+					SolverManifoldPoint.ContactMassNormal,
+					SolverManifoldPoint.NetPushOutNormal,
+					PushOutNormal);						// Out
+			
+				SolverManifoldPoint.NetPushOutNormal += PushOutNormal;
+
+				// Update the particle state based on the pushout
+				if (Body0.IsDynamic())
+				{
+					DX0 += (Body0.InvM() * PushOutNormal) * SolverManifoldPoint.WorldContactNormal;
+					DR0 += SolverManifoldPoint.WorldContactNormalAngular0 * PushOutNormal;
+				}
+				if (Body1.IsDynamic())
+				{
+					DX1 -= (Body1.InvM() * PushOutNormal) * SolverManifoldPoint.WorldContactNormal;
+					DR1 -= SolverManifoldPoint.WorldContactNormalAngular1 * PushOutNormal;
+				}
+			}
+		}
+
+		if (Body0.IsDynamic())
+		{
+			if (AboveThresholdElseZero(DX0, CVars::Chaos_PBDCollisionSolver_JacobiPositionTolerance))
+			{
+				Body0.ApplyPositionDelta(DX0);
+			}
+			if (AboveThresholdElseZero(DR0, CVars::Chaos_PBDCollisionSolver_JacobiRotationTolerance))
+			{
+				Body0.ApplyRotationDelta(DR0);
+			}
+		}
+		if (Body1.IsDynamic())
+		{
+			if (AboveThresholdElseZero(DX1, CVars::Chaos_PBDCollisionSolver_JacobiPositionTolerance))
+			{
+				Body1.ApplyPositionDelta(DX1);
+			}
+			if (AboveThresholdElseZero(DR1, CVars::Chaos_PBDCollisionSolver_JacobiRotationTolerance))
+			{
+				Body1.ApplyRotationDelta(DR1);
+			}
+		}
+
+		return false;
+	}
+
+	FORCEINLINE_DEBUGGABLE bool FPBDCollisionSolver::SolvePositionWithFrictionJacobi(const FSolverReal Dt, const FSolverReal MaxPushOut)
+	{
+		// SolverBody decorator used to add mass scaling
+		FConstraintSolverBody& Body0 = SolverBody0();
+		FConstraintSolverBody& Body1 = SolverBody1();
+
+		// Accumulate net pushout for friction limits below
+		bool bApplyFriction[MaxPointsPerConstraint] = { false, };
+		int32 NumFrictionContacts = 0;
+		FSolverReal TotalPushOutNormal = FSolverReal(0);
+
+		FSolverVec3 DX0 = FSolverVec3(0);
+		FSolverVec3 DX1 = FSolverVec3(0);
+		FSolverVec3 DR0 = FSolverVec3(0);
+		FSolverVec3 DR1 = FSolverVec3(0);
+
+		const FSolverReal Stiffness = CVars::Chaos_PBDCollisionSolver_JacobiStiffness * State.Stiffness;
+
+		// Apply the position correction along the normal and determine if we want to run friction on each point
+		for (int32 PointIndex = 0; PointIndex < State.NumManifoldPoints; ++PointIndex)
+		{
+			FPBDCollisionSolverManifoldPoint& SolverManifoldPoint = State.ManifoldPoints[PointIndex];
+
+			FSolverReal ContactDeltaNormal;
+			SolverManifoldPoint.CalculateContactPositionErrorNormal(Body0.SolverBody(), Body1.SolverBody(), MaxPushOut, ContactDeltaNormal);
+
+			// Apply a normal correction if we still have penetration or if we are now separated but have previously applied a correction that we may want to undo
+			const bool bProcessManifoldPoint = (ContactDeltaNormal < FSolverReal(0)) || (SolverManifoldPoint.NetPushOutNormal > FSolverReal(UE_SMALL_NUMBER));
+			if (bProcessManifoldPoint)
+			{
+				FSolverReal PushOutNormal = FSolverReal(0);
+				CalculatePositionCorrectionNormal(
+					Stiffness,
+					State.Regularization,
+					ContactDeltaNormal,
+					SolverManifoldPoint.ContactMassNormal,
+					SolverManifoldPoint.NetPushOutNormal,
+					PushOutNormal);						// Out
+
+				SolverManifoldPoint.NetPushOutNormal += PushOutNormal;
+
+				// Update the particle state based on the pushout
+				if (Body0.IsDynamic())
+				{
+					DX0 += (Body0.InvM() * PushOutNormal) * SolverManifoldPoint.WorldContactNormal;
+					DR0 += SolverManifoldPoint.WorldContactNormalAngular0 * PushOutNormal;
+				}
+				if (Body1.IsDynamic())
+				{
+					DX1 -= (Body1.InvM() * PushOutNormal) * SolverManifoldPoint.WorldContactNormal;
+					DR1 -= SolverManifoldPoint.WorldContactNormalAngular1 * PushOutNormal;
+				}
+
+				TotalPushOutNormal += SolverManifoldPoint.NetPushOutNormal;
+			}
+
+			// Friction gets updated for any point with a net normal correction or where we have previously had a normal correction and 
+			// already applied friction (in which case we may need to zero it)
+			if ((SolverManifoldPoint.NetPushOutNormal != 0) || (SolverManifoldPoint.NetPushOutTangentU != 0) || (SolverManifoldPoint.NetPushOutTangentV != 0))
+			{
+				bApplyFriction[PointIndex] = true;
+				++NumFrictionContacts;
+			}
+		}
+
+		// Apply the tangential position correction if required
+		if (NumFrictionContacts > 0)
+		{
+			// We clip the tangential correction at each contact to the friction cone, but we use to average impulse
+			// among all contacts as the clipping limit. This is not really correct but it is much more stable to 
+			// differences in contacts from tick to tick
+			const FSolverReal FrictionMaxPushOut = TotalPushOutNormal / FSolverReal(NumFrictionContacts);
+			const FSolverReal FrictionStiffness = Stiffness * CVars::Chaos_PBDCollisionSolver_Position_StaticFrictionStiffness;
+
+			for (int32 PointIndex = 0; PointIndex < State.NumManifoldPoints; ++PointIndex)
+			{
+				if (bApplyFriction[PointIndex])
+				{
+					FPBDCollisionSolverManifoldPoint& SolverManifoldPoint = State.ManifoldPoints[PointIndex];
+
+					FSolverReal ContactDeltaTangentU, ContactDeltaTangentV;
+					SolverManifoldPoint.CalculateContactPositionErrorTangential(Body0.SolverBody(), Body1.SolverBody(), ContactDeltaTangentU, ContactDeltaTangentV);
+
+					FSolverReal PushOutTangentU = FSolverReal(0);
+					FSolverReal PushOutTangentV = FSolverReal(0);
+
+					CalculatePositionCorrectionTangent(
+						FrictionStiffness,
+						State.Regularization,
+						ContactDeltaTangentU,
+						SolverManifoldPoint.ContactMassTangentU,
+						SolverManifoldPoint.NetPushOutTangentU,
+						PushOutTangentU);					// Out
+
+					CalculatePositionCorrectionTangent(
+						FrictionStiffness,
+						State.Regularization,
+						ContactDeltaTangentV,
+						SolverManifoldPoint.ContactMassTangentV,
+						SolverManifoldPoint.NetPushOutTangentV,
+						PushOutTangentV);					// Out
+
+					SolverManifoldPoint.NetPushOutTangentU += PushOutTangentU;
+					SolverManifoldPoint.NetPushOutTangentV += PushOutTangentV;
+
+					ApplyFrictionCone(
+						State.StaticFriction,
+						State.DynamicFriction,
+						FrictionMaxPushOut,
+						PushOutTangentU,									// InOut
+						PushOutTangentV,									// InOut
+						SolverManifoldPoint.NetPushOutTangentU,				// InOut
+						SolverManifoldPoint.NetPushOutTangentV,				// InOut
+						SolverManifoldPoint.StaticFrictionRatio);			// Out
+
+					// Update the particle state based on the pushout
+					const FSolverVec3 PushOut = PushOutTangentU * SolverManifoldPoint.WorldContactTangentU + PushOutTangentV * SolverManifoldPoint.WorldContactTangentV;
+					if (Body0.IsDynamic())
+					{
+						DX0 += Body0.InvM() * PushOut;
+						DR0 += SolverManifoldPoint.WorldContactTangentUAngular0 * PushOutTangentU + SolverManifoldPoint.WorldContactTangentVAngular0 * PushOutTangentV;
+					}
+					if (Body1.IsDynamic())
+					{
+						DX1 -= Body1.InvM() * PushOut;
+						DR1 -= SolverManifoldPoint.WorldContactTangentUAngular1 * PushOutTangentU + SolverManifoldPoint.WorldContactTangentVAngular1 * PushOutTangentV;
+					}
+				}
+			}
+
+			if (Body0.IsDynamic())
+			{
+				if (AboveThresholdElseZero(DX0, CVars::Chaos_PBDCollisionSolver_JacobiPositionTolerance))
+				{
+					Body0.ApplyPositionDelta(DX0);
+				}
+				if (AboveThresholdElseZero(DR0, CVars::Chaos_PBDCollisionSolver_JacobiRotationTolerance))
+				{
+					Body0.ApplyRotationDelta(DR0);
+				}
+			}
+			if (Body1.IsDynamic())
+			{
+				if (AboveThresholdElseZero(DX1, CVars::Chaos_PBDCollisionSolver_JacobiPositionTolerance))
+				{
+					Body1.ApplyPositionDelta(DX1);
+				}
+				if (AboveThresholdElseZero(DR1, CVars::Chaos_PBDCollisionSolver_JacobiRotationTolerance))
+				{
+					Body1.ApplyRotationDelta(DR1);
+				}
+			}
+		}
+
+		return false;
+	}
+
+	FORCEINLINE_DEBUGGABLE bool FPBDCollisionSolver::SolveVelocity(const FSolverReal Dt, const bool bApplyDynamicFriction)
+	{
+		// Apply restitution at the average contact point
+		// This means we don't need to run as many iterations to get stable bouncing
+		// It also helps with zero restitution to counter any velocioty added by the PBD solve
+		const bool bSolveAverageContact = (State.NumManifoldPoints > 1) && CVars::bChaos_PBDCollisionSolver_Velocity_AveragePointEnabled;
+		if (bSolveAverageContact)
+		{
+			SolveVelocityAverage(Dt);
+		}
+
+		FConstraintSolverBody& Body0 = SolverBody0();
+		FConstraintSolverBody& Body1 = SolverBody1();
+
+		// NOTE: this dynamic friction implementation is iteration-count sensitive
+		// @todo(chaos): fix iteration count dependence of dynamic friction
+		const FSolverReal DynamicFriction = (bApplyDynamicFriction && (Dt > 0) && CVars::bChaos_PBDCollisionSolver_Velocity_FrictionEnabled) ? State.VelocityFriction : FSolverReal(0);
+
+		for (int32 PointIndex = 0; PointIndex < State.NumManifoldPoints; ++PointIndex)
+		{
+			FPBDCollisionSolverManifoldPoint& SolverManifoldPoint = State.ManifoldPoints[PointIndex];
+			if (SolverManifoldPoint.ShouldSolveVelocity())
+			{
+				const FSolverReal MinImpulseNormal = FMath::Min(FSolverReal(0), -SolverManifoldPoint.NetPushOutNormal / Dt);
+
+				if (DynamicFriction > 0)
+				{
+					FSolverReal ContactVelocityDeltaNormal, ContactVelocityDeltaTangentU, ContactVelocityDeltaTangentV;
+					SolverManifoldPoint.CalculateContactVelocityError(Body0, Body1, DynamicFriction, Dt, ContactVelocityDeltaNormal, ContactVelocityDeltaTangentU, ContactVelocityDeltaTangentV);
+
+					ApplyVelocityCorrection(
+						State.Stiffness,
+						Dt,
+						DynamicFriction,
+						ContactVelocityDeltaNormal,
+						ContactVelocityDeltaTangentU,
+						ContactVelocityDeltaTangentV,
+						MinImpulseNormal,
+						SolverManifoldPoint,
+						Body0,
+						Body1);
+				}
+				else
+				{
+					FSolverReal ContactVelocityDeltaNormal;
+					SolverManifoldPoint.CalculateContactVelocityErrorNormal(Body0, Body1, ContactVelocityDeltaNormal);
+
+					ApplyVelocityCorrectionNormal(
+						State.Stiffness,
+						ContactVelocityDeltaNormal,
+						MinImpulseNormal,
+						SolverManifoldPoint,
+						Body0,
+						Body1);
+				}
+			}
+		}
+
+		// Early-out support for the velocity solve is not currently very important because we
+		// only run one iteration in the velocity solve phase.
+		// @todo(chaos): support early-out in velocity solve if necessary
+		return true;
+	}
+
+
+	FORCEINLINE_DEBUGGABLE bool FPBDCollisionSolver::SolveVelocityJacobi(const FSolverReal Dt, const bool bApplyDynamicFriction)
+	{
+		FConstraintSolverBody& Body0 = SolverBody0();
+		FConstraintSolverBody& Body1 = SolverBody1();
+
+		FSolverVec3 DV0 = FSolverVec3(0);
+		FSolverVec3 DV1 = FSolverVec3(0);
+		FSolverVec3 DW0 = FSolverVec3(0);
+		FSolverVec3 DW1 = FSolverVec3(0);
+
+		const FSolverReal Stiffness = CVars::Chaos_PBDCollisionSolver_JacobiStiffness * State.Stiffness;
+
+		for (int32 PointIndex = 0; PointIndex < State.NumManifoldPoints; ++PointIndex)
+		{
+			FPBDCollisionSolverManifoldPoint& SolverManifoldPoint = State.ManifoldPoints[PointIndex];
+			if (SolverManifoldPoint.ShouldSolveVelocity())
+			{
+				const FSolverReal MinImpulseNormal = FMath::Min(FSolverReal(0), -SolverManifoldPoint.NetPushOutNormal / Dt);
+
+				FSolverReal ContactVelocityDeltaNormal;
+				SolverManifoldPoint.CalculateContactVelocityErrorNormal(Body0, Body1, ContactVelocityDeltaNormal);
+
+				FSolverReal ImpulseNormal = -(Stiffness * SolverManifoldPoint.ContactMassNormal) * ContactVelocityDeltaNormal;
+
+				// See comments in ApplyVelocityCorrection
+				const FSolverReal NetImpulseNormal = SolverManifoldPoint.NetImpulseNormal + ImpulseNormal;
+				if (NetImpulseNormal < MinImpulseNormal)
+				{
+					ImpulseNormal = ImpulseNormal - NetImpulseNormal + MinImpulseNormal;
+				}
+
+				SolverManifoldPoint.NetImpulseNormal += ImpulseNormal;
+
+				FSolverVec3 Impulse = ImpulseNormal * SolverManifoldPoint.WorldContactNormal;
+				if (Body0.IsDynamic())
+				{
+					DV0 += Body0.InvM() * Impulse;
+					DW0 += SolverManifoldPoint.WorldContactNormalAngular0 * ImpulseNormal;
+				}
+				if (Body1.IsDynamic())
+				{
+					DV1 -= Body1.InvM() * Impulse;
+					DW1 -= SolverManifoldPoint.WorldContactNormalAngular1 * ImpulseNormal;
+				}
+			}
+		}
+
+		const FSolverReal PositionFloor = (CVars::Chaos_PBDCollisionSolver_JacobiPositionTolerance / FMath::Square(Dt));
+		const FSolverReal RotationFloor = (CVars::Chaos_PBDCollisionSolver_JacobiRotationTolerance / FMath::Square(Dt));
+
+		if (Body0.IsDynamic())
+		{
+			AboveThresholdElseZero(DV0, PositionFloor);
+			AboveThresholdElseZero(DW0, RotationFloor);
+
+			Body0.ApplyVelocityDelta(DV0, DW0);
+		}
+		if (Body1.IsDynamic())
+		{
+			AboveThresholdElseZero(DV1, PositionFloor);
+			AboveThresholdElseZero(DW1, RotationFloor);
+
+			Body1.ApplyVelocityDelta(DV1, DW1);
+		}
+
+		// Early-out support for the velocity solve is not currently very important because we
+		// only run one iteration in the velocity solve phase.
+		// @todo(chaos): support early-out in velocity solve if necessary
+		return true;
+	}
+
 }
 
diff --git a/Engine/Source/Runtime/Experimental/Chaos/Private/Chaos/Collision/SolverCollisionContainer.cpp b/Engine/Source/Runtime/Experimental/Chaos/Private/Chaos/Collision/SolverCollisionContainer.cpp
--- a/Engine/Source/Runtime/Experimental/Chaos/Private/Chaos/Collision/SolverCollisionContainer.cpp
+++ b/Engine/Source/Runtime/Experimental/Chaos/Private/Chaos/Collision/SolverCollisionContainer.cpp
@@ -37,6 +37,25 @@ namespace Chaos
 		// Collision solver regularization (experimental, disabled by default)
 		FRealSingle Chaos_PBDCollisionSolver_Regularization = 0;
 		FAutoConsoleVariableRef CVarChaosPBDCollisionSolverRegularization(TEXT("p.Chaos.PBDCollisionSolver.Regularization"), Chaos_PBDCollisionSolver_Regularization, TEXT(""));
+
+		// Whether to use the Jacobi collision pair solver
+		bool bChaos_PBDCollisionSolver_UseJacobiPairSolver = false;
+		FAutoConsoleVariableRef CVarChaosPBDCollisionSolverUseJacobiPairSolver(TEXT("p.Chaos.PBDCollisionSolver.UseJacobi"), bChaos_PBDCollisionSolver_UseJacobiPairSolver, TEXT(""));
+
+		// Jacobi solver stiffness
+		// @todo(chaos): to be tuned
+		FRealSingle Chaos_PBDCollisionSolver_JacobiStiffness = 0.5f;
+		FAutoConsoleVariableRef CVarChaosPBDCollisionSolverJacobiStiffness(TEXT("p.Chaos.PBDCollisionSolver.JacobiStiffness"), Chaos_PBDCollisionSolver_JacobiStiffness, TEXT(""));
+
+		// Jacobi position tolerance. Position corrections below this are zeroed.
+		// @todo(chaos): to be tuned
+		FRealSingle Chaos_PBDCollisionSolver_JacobiPositionTolerance = 1.e-6f;
+		FAutoConsoleVariableRef CVarChaosPBDCollisionSolverJacobiPositionTolerance(TEXT("p.Chaos.PBDCollisionSolver.JacobiPositionTolerance"), Chaos_PBDCollisionSolver_JacobiPositionTolerance, TEXT(""));
+		
+		// Jacobi rotation tolerance. Rotation corrections below this are zeroed.
+		// @todo(chaos): to be tuned
+		FRealSingle Chaos_PBDCollisionSolver_JacobiRotationTolerance = 1.e-8f;
+		FAutoConsoleVariableRef CVarChaosPBDCollisionSolverJacobiRotationTolerance(TEXT("p.Chaos.PBDCollisionSolver.JacobiRotationTolerance"), Chaos_PBDCollisionSolver_JacobiRotationTolerance, TEXT(""));
 	}
 	using namespace CVars;
 
@@ -619,9 +638,19 @@ namespace Chaos
 		const FSolverReal Dt = FSolverReal(InDt);
 		const FSolverReal MaxPushOut = FSolverReal(InMaxPushOut);
 
-		for (int32 SolverIndex = BeginIndex; SolverIndex < EndIndex; ++SolverIndex)
+		if (!CVars::bChaos_PBDCollisionSolver_UseJacobiPairSolver)
+		{
+			for (int32 SolverIndex = BeginIndex; SolverIndex < EndIndex; ++SolverIndex)
+			{
+				CollisionSolvers[SolverIndex].GetSolver().SolvePositionWithFriction(Dt, MaxPushOut);
+			}
+		}
+		else
 		{
-			CollisionSolvers[SolverIndex].GetSolver().SolvePositionWithFriction(Dt, MaxPushOut);
+			for (int32 SolverIndex = BeginIndex; SolverIndex < EndIndex; ++SolverIndex)
+			{
+				CollisionSolvers[SolverIndex].GetSolver().SolvePositionWithFrictionJacobi(Dt, MaxPushOut);
+			}
 		}
 
 		return true;
@@ -637,9 +666,19 @@ namespace Chaos
 		const FSolverReal Dt = FSolverReal(InDt);
 		const FSolverReal MaxPushOut = FSolverReal(InMaxPushOut);
 
-		for (int32 SolverIndex = BeginIndex; SolverIndex < EndIndex; ++SolverIndex)
+		if (!CVars::bChaos_PBDCollisionSolver_UseJacobiPairSolver)
+		{
+			for (int32 SolverIndex = BeginIndex; SolverIndex < EndIndex; ++SolverIndex)
+			{
+				CollisionSolvers[SolverIndex].GetSolver().SolvePositionNoFriction(Dt, MaxPushOut);
+			}
+		}
+		else
 		{
-			CollisionSolvers[SolverIndex].GetSolver().SolvePositionNoFriction(Dt, MaxPushOut);
+			for (int32 SolverIndex = BeginIndex; SolverIndex < EndIndex; ++SolverIndex)
+			{
+				CollisionSolvers[SolverIndex].GetSolver().SolvePositionNoFrictionJacobi(Dt, MaxPushOut);
+			}
 		}
 
 		return true;
@@ -661,11 +700,19 @@ namespace Chaos
 		// Apply the velocity correction
 		// @todo(chaos): parallel version of SolveVelocity
 		bool bNeedsAnotherIteration = false;
-		for (int32 SolverIndex = BeginIndex; SolverIndex < EndIndex; ++SolverIndex)
+		if (!CVars::bChaos_PBDCollisionSolver_UseJacobiPairSolver)
 		{
-			FPBDCollisionSolverAdapter& CollisionSolver = CollisionSolvers[SolverIndex];
-
-			bNeedsAnotherIteration |= CollisionSolver.GetSolver().SolveVelocity(Dt, bApplyDynamicFriction);
+			for (int32 SolverIndex = BeginIndex; SolverIndex < EndIndex; ++SolverIndex)
+			{
+				bNeedsAnotherIteration |= CollisionSolvers[SolverIndex].GetSolver().SolveVelocity(Dt, bApplyDynamicFriction);
+			}
+		}
+		else
+		{
+			for (int32 SolverIndex = BeginIndex; SolverIndex < EndIndex; ++SolverIndex)
+			{
+				bNeedsAnotherIteration |= CollisionSolvers[SolverIndex].GetSolver().SolveVelocityJacobi(Dt, bApplyDynamicFriction);
+			}
 		}
 
 		return bNeedsAnotherIteration;
