GeometryCollection : Fix for CIS

diff --git a/Engine/Source/Runtime/Experimental/Chaos/Public/GeometryCollection/ManagedArray.h b/Engine/Source/Runtime/Experimental/Chaos/Public/GeometryCollection/ManagedArray.h
--- a/Engine/Source/Runtime/Experimental/Chaos/Public/GeometryCollection/ManagedArray.h
+++ b/Engine/Source/Runtime/Experimental/Chaos/Public/GeometryCollection/ManagedArray.h
@@ -177,11 +177,6 @@ protected:
 	*/
 	virtual SIZE_T GetAllocatedSize() const { return 0; }
 
-	/**
-	* Copy a range of values from the ConstArray into this
-	*/
-	virtual void CopyRange(const FManagedArrayBase& ConstArray, int32 Start, int32 Stop, int32 Offset = 0) {};
-
 public:
 	FManagedArrayBase()
 	{
