Chaos - clean up ContactPoint array decls

diff --git a/Engine/Source/Programs/HeadlessChaos/Private/HeadlessChaosTestCapsuleTriangleCollision.cpp b/Engine/Source/Programs/HeadlessChaos/Private/HeadlessChaosTestCapsuleTriangleCollision.cpp
--- a/Engine/Source/Programs/HeadlessChaos/Private/HeadlessChaosTestCapsuleTriangleCollision.cpp
+++ b/Engine/Source/Programs/HeadlessChaos/Private/HeadlessChaosTestCapsuleTriangleCollision.cpp
@@ -9,7 +9,7 @@ namespace ChaosTest
 {
 	using namespace Chaos;
 
-	bool ManifoldContainsPoint(const TCArray<FContactPoint, 4>& ContactPoints, const int32 BodyIndex, const FVec3& Point, const FReal Tolerance)
+	bool ManifoldContainsPoint(const FContactPointManifold& ContactPoints, const int32 BodyIndex, const FVec3& Point, const FReal Tolerance)
 	{
 		for (FContactPoint ContactPoint : ContactPoints)
 		{
@@ -53,7 +53,7 @@ namespace ChaosTest
 		FTriangle Triangle(FVec3(0,0,0), FVec3(100, -100, 0), FVec3(100, 0, 0));
 		FImplicitCapsule3 Capsule(FVec3(30, -10, 10), FVec3(90, -10, 10), FReal(20));
 
-		TCArray<FContactPoint, 4> ContactPoints;
+		FContactPointManifold ContactPoints;
 		ConstructCapsuleTriangleOneShotManifold2(Capsule, Triangle, CullDistance, ContactPoints);
 
 		EXPECT_EQ(ContactPoints.Num(), 2);
@@ -75,7 +75,7 @@ namespace ChaosTest
 		FTriangle Triangle(FVec3(0, 0, 0), FVec3(100, -100, 0), FVec3(100, 0, 0));
 		FImplicitCapsule3 Capsule(FVec3(30, -10, -1), FVec3(90, -10, -1), FReal(20));
 
-		TCArray<FContactPoint, 4> ContactPoints;
+		FContactPointManifold ContactPoints;
 		ConstructCapsuleTriangleOneShotManifold2(Capsule, Triangle, CullDistance, ContactPoints);
 
 		EXPECT_EQ(ContactPoints.Num(), 0);
@@ -92,7 +92,7 @@ namespace ChaosTest
 		FTriangle Triangle(FVec3(0, 0, 0), FVec3(100, -100, 0), FVec3(100, 0, 0));
 		FImplicitCapsule3 Capsule(FVec3(30, -10, 0), FVec3(90, -10, 0), FReal(20));
 
-		TCArray<FContactPoint, 4> ContactPoints;
+		FContactPointManifold ContactPoints;
 		ConstructCapsuleTriangleOneShotManifold2(Capsule, Triangle, CullDistance, ContactPoints);
 
 		EXPECT_EQ(ContactPoints.Num(), 2);
@@ -125,7 +125,7 @@ namespace ChaosTest
 		FTriangle Triangle(FVec3(0, 0, 0), FVec3(100, -100, 0), FVec3(100, 0, 0));
 		FImplicitCapsule3 Capsule(FVec3(50, -10, 10), FVec3(120, -10, 10), FReal(20));
 
-		TCArray<FContactPoint, 4> ContactPoints;
+		FContactPointManifold ContactPoints;
 		ConstructCapsuleTriangleOneShotManifold2(Capsule, Triangle, CullDistance, ContactPoints);
 
 		EXPECT_EQ(ContactPoints.Num(), 2);
@@ -161,7 +161,7 @@ namespace ChaosTest
 		const FTriangle Triangle(FVec3(0, 0, 0), FVec3(100, -100, 0), FVec3(100, 0, 0));
 		const FImplicitCapsule3 Capsule(FVec3(50, -10, 10 + 50 * Sin10), FVec3(150, -10, 10 - 50 * Sin10), FReal(20));
 
-		TCArray<FContactPoint, 4> ContactPoints;
+		FContactPointManifold ContactPoints;
 		ConstructCapsuleTriangleOneShotManifold2(Capsule, Triangle, CullDistance, ContactPoints);
 
 		EXPECT_EQ(ContactPoints.Num(), 2);
@@ -197,7 +197,7 @@ namespace ChaosTest
 		FTriangle Triangle(FVec3(0, 0, 0), FVec3(100, -100, 0), FVec3(100, 0, 0));
 		FImplicitCapsule3 Capsule(FVec3(50, -10, 50), FVec3(150, -10, -50), FReal(10));
 
-		TCArray<FContactPoint, 4> ContactPoints;
+		FContactPointManifold ContactPoints;
 		ConstructCapsuleTriangleOneShotManifold2(Capsule, Triangle, CullDistance, ContactPoints);
 
 		EXPECT_EQ(ContactPoints.Num(), 1);
@@ -210,7 +210,7 @@ namespace ChaosTest
 
 	// Capsule at a 45 degree angle to the triangle, and the capsule axis passes right through an edge.
 	// This time we have an end cap within the edge planes, so we should get a face contact at that location.
-	// The edge contact will be ignored because it would generate an inward dacing normal, and we are
+	// The edge contact will be ignored because it would generate an inward-facing normal, and we are
 	// above the face angle threshold.
 	// 
 	//         / /
@@ -228,7 +228,7 @@ namespace ChaosTest
 		FTriangle Triangle(FVec3(0, 0, 0), FVec3(100, -100, 0), FVec3(100, 0, 0));
 		FImplicitCapsule3 Capsule(FVec3(50, -10, -50), FVec3(150, -10, 50), FReal(10));
 
-		TCArray<FContactPoint, 4> ContactPoints;
+		FContactPointManifold ContactPoints;
 		ConstructCapsuleTriangleOneShotManifold2(Capsule, Triangle, CullDistance, ContactPoints);
 
 		EXPECT_EQ(ContactPoints.Num(), 1);
diff --git a/Engine/Source/Runtime/Experimental/Chaos/Private/Chaos/Collision/CapsuleTriangleContactPoint.cpp b/Engine/Source/Runtime/Experimental/Chaos/Private/Chaos/Collision/CapsuleTriangleContactPoint.cpp
--- a/Engine/Source/Runtime/Experimental/Chaos/Private/Chaos/Collision/CapsuleTriangleContactPoint.cpp
+++ b/Engine/Source/Runtime/Experimental/Chaos/Private/Chaos/Collision/CapsuleTriangleContactPoint.cpp
@@ -14,7 +14,7 @@ namespace Chaos
 		FAutoConsoleVariableRef CVarUseGenericSweptConvexConstraints(TEXT("p.Chaos.Collision.CapsuleTriMeshSATCull"), bChaosCollisionCapsuleTriMeshSATCull, TEXT("Enable the SAT cull check in capsule-triangle [default: false]."));
 	}
 
-	void AddCapsuleTriangleParallelEdgeManifoldContacts(const FVec3& P0, const FVec3& P1, const FVec3& EdgeP0, const FVec3& EdgeP1, const FReal R, const FReal RejectDistanceSq, const FReal NormalToleranceSq, TCArray<FContactPoint, 4>& OutContactPoints);
+	void AddCapsuleTriangleParallelEdgeManifoldContacts(const FVec3& P0, const FVec3& P1, const FVec3& EdgeP0, const FVec3& EdgeP1, const FReal R, const FReal RejectDistanceSq, const FReal NormalToleranceSq, FContactPointManifold& OutContactPoints);
 
 	// Return true if the value V is in the closed/inclusive range [RangeMin, RangeMax]
 	// @todo(chaos): move to Utilities.h
@@ -50,7 +50,7 @@ namespace Chaos
 	//		- Add edge contact where the segment is closest to the edge
 	//		- When the angle between segment and face is small, create a face contact instead
 	//
-	void ConstructCapsuleTriangleOneShotManifold2(const FImplicitCapsule3& Capsule, const FTriangle& Triangle, const FReal CullDistance, TCArray<FContactPoint, 4>& OutContactPoints)
+	void ConstructCapsuleTriangleOneShotManifold2(const FImplicitCapsule3& Capsule, const FTriangle& Triangle, const FReal CullDistance, FContactPointManifold& OutContactPoints)
 	{
 		//
 		// @todo(chaos): this function is a good SIMD candidate
@@ -456,10 +456,10 @@ namespace Chaos
 	// Utility for ConstructCapsuleTriangleOneShotManifold2 to handle the rare case when a capsule is perfectly aligned with an edge of a triangle.
 	// In this case we project the segment onto the edge and, depending on whether there is any overlap, use either the closest point or the
 	// two clipped points as contacts.
-	void AddCapsuleTriangleParallelEdgeManifoldContacts(const FVec3& P0, const FVec3& P1, const FVec3& EdgeP0, const FVec3& EdgeP1, const FReal R, const FReal RejectDistanceSq, const FReal NormalToleranceSq, TCArray<FContactPoint, 4>& OutContactPoints)
+	void AddCapsuleTriangleParallelEdgeManifoldContacts(const FVec3& P0, const FVec3& P1, const FVec3& EdgeP0, const FVec3& EdgeP1, const FReal R, const FReal RejectDistanceSq, const FReal NormalToleranceSq, FContactPointManifold& OutContactPoints)
 	{
 		// Utility to add a contact to the array if it is within cull distance
-		const auto& AddContact = [](const FVec3& SegmentEdgeC, const FVec3& SegmentEdgeDelta, const FReal R, const FReal RejectDistanceSq, const FReal NormalToleranceSq, TCArray<FContactPoint, 4>& OutContactPoints) -> void
+		const auto& AddContact = [](const FVec3& SegmentEdgeC, const FVec3& SegmentEdgeDelta, const FReal R, const FReal RejectDistanceSq, const FReal NormalToleranceSq, FContactPointManifold& OutContactPoints) -> void
 		{
 			const FReal SegmentEdgeDistSq = SegmentEdgeDelta.SizeSquared();
 			if ((SegmentEdgeDistSq < RejectDistanceSq) && (SegmentEdgeDistSq > NormalToleranceSq))
diff --git a/Engine/Source/Runtime/Experimental/Chaos/Private/Chaos/Collision/ConvexTriangleContactPoint.cpp b/Engine/Source/Runtime/Experimental/Chaos/Private/Chaos/Collision/ConvexTriangleContactPoint.cpp
--- a/Engine/Source/Runtime/Experimental/Chaos/Private/Chaos/Collision/ConvexTriangleContactPoint.cpp
+++ b/Engine/Source/Runtime/Experimental/Chaos/Private/Chaos/Collision/ConvexTriangleContactPoint.cpp
@@ -128,7 +128,7 @@ namespace Chaos
 
 	// Triangle is in (possibly scaled) Convex space
 	template <typename ConvexType>
-	void ConstructConvexTriangleOneShotManifold2(const ConvexType& Convex, const FTriangle& Triangle, const FReal CullDistance, TCArray<FContactPoint, 4>& OutContactPoints)
+	void ConstructConvexTriangleOneShotManifold2(const ConvexType& Convex, const FTriangle& Triangle, const FReal CullDistance, FContactPointManifold& OutContactPoints)
 	{
 		const FReal NormalTolerance = FReal(1.e-8);
 		const FReal NormalToleranceSq = NormalTolerance * NormalTolerance;
@@ -493,40 +493,40 @@ namespace Chaos
 		const FImplicitConvex3& Convex, 
 		const FTriangle& Triangle,
 		const FReal CullDistance, 
-		TCArray<FContactPoint, 4>& OutContactPoints);
+		FContactPointManifold& OutContactPoints);
 
 	template
 	void ConstructConvexTriangleOneShotManifold2(
 		const TImplicitObjectInstanced<FImplicitConvex3>& Convex,
 		const FTriangle& Triangle,
 		const FReal CullDistance,
-		TCArray<FContactPoint, 4>& OutContactPoints);
+		FContactPointManifold& OutContactPoints);
 
 	template
 	void ConstructConvexTriangleOneShotManifold2(
 		const TImplicitObjectScaled<FImplicitConvex3>& Convex,
 		const FTriangle& Triangle,
 		const FReal CullDistance,
-		TCArray<FContactPoint, 4>& OutContactPoints);
+		FContactPointManifold& OutContactPoints);
 
 	template
 	void ConstructConvexTriangleOneShotManifold2(
 		const FImplicitBox3& Convex,
 		const FTriangle& Triangle,
 		const FReal CullDistance,
-		TCArray<FContactPoint, 4>& OutContactPoints);
+		FContactPointManifold& OutContactPoints);
 
 	template
 	void ConstructConvexTriangleOneShotManifold2(
 		const TImplicitObjectScaled<FImplicitBox3>& Convex,
 		const FTriangle& Triangle,
 		const FReal CullDistance,
-		TCArray<FContactPoint, 4>& OutContactPoints);
+		FContactPointManifold& OutContactPoints);
 
 	template
 	void ConstructConvexTriangleOneShotManifold2(
 		const TImplicitObjectInstanced<FImplicitBox3>& Convex,
 		const FTriangle& Triangle,
 		const FReal CullDistance,
-		TCArray<FContactPoint, 4>& OutContactPoints);
+		FContactPointManifold& OutContactPoints);
 }
\ No newline at end of file
diff --git a/Engine/Source/Runtime/Experimental/Chaos/Private/Chaos/CollisionOneShotManifolds.cpp b/Engine/Source/Runtime/Experimental/Chaos/Private/Chaos/CollisionOneShotManifolds.cpp
--- a/Engine/Source/Runtime/Experimental/Chaos/Private/Chaos/CollisionOneShotManifolds.cpp
+++ b/Engine/Source/Runtime/Experimental/Chaos/Private/Chaos/CollisionOneShotManifolds.cpp
@@ -1004,7 +1004,7 @@ namespace Chaos
 			const ConvexType& Convex,
 			const FRigidTransform3& ConvexTransform,
 			const FReal CullDistance,
-			TCArray<FContactPoint, 4>& OutContactPoints)
+			FContactPointManifold& OutContactPoints)
 		{
 			SCOPE_CYCLE_COUNTER_MANIFOLD();
 			check(OutContactPoints.Num() == 0);
@@ -1192,7 +1192,7 @@ namespace Chaos
 			}
 		}
 
-		void ConstructCapsuleTriangleOneShotManifold(const FImplicitCapsule3& Capsule, const FTriangle& Triangle, const FReal CullDistance, TCArray<FContactPoint, 4>& OutContactPoints)
+		void ConstructCapsuleTriangleOneShotManifold(const FImplicitCapsule3& Capsule, const FTriangle& Triangle, const FReal CullDistance, FContactPointManifold& OutContactPoints)
 		{
 			// @todo(chaos): make custom capsule-triangle manifold function.
 			// NOTE: ConstructCapsuleConvexOneShotManifold could be used but it has issues when we collide a triangle edge with the capsule cylinder - that needs fixing...
@@ -1207,7 +1207,7 @@ namespace Chaos
 
 		// @todo(chaos): don't use GJK/EPA for triangle-convex so we can avoid using margins and errors related to misidentifying edge contacts
 		template <typename ConvexType>
-		void ConstructPlanarConvexTriangleOneShotManifold(const ConvexType& Convex, const FTriangle& Triangle, const FReal CullDistance, TCArray<FContactPoint, 4>& OutContactPoints)
+		void ConstructPlanarConvexTriangleOneShotManifold(const ConvexType& Convex, const FTriangle& Triangle, const FReal CullDistance, FContactPointManifold& OutContactPoints)
 		{
 			SCOPE_CYCLE_COUNTER_MANIFOLD();
 			check(OutContactPoints.Num() == 0);
@@ -1588,7 +1588,7 @@ namespace Chaos
 			const FImplicitConvex3& Convex,
 			const FRigidTransform3& ConvexTransform,
 			const FReal CullDistance,
-			TCArray<FContactPoint, 4>& OutContactPoints);
+			FContactPointManifold& OutContactPoints);
 
 		template
 		void ConstructCapsuleConvexOneShotManifold(
@@ -1597,7 +1597,7 @@ namespace Chaos
 			const TImplicitObjectInstanced<FImplicitConvex3>& Convex,
 			const FRigidTransform3& ConvexTransform,
 			const FReal CullDistance,
-			TCArray<FContactPoint, 4>& OutContactPoints);
+			FContactPointManifold& OutContactPoints);
 
 		template
 		void ConstructCapsuleConvexOneShotManifold(
@@ -1606,7 +1606,7 @@ namespace Chaos
 			const TImplicitObjectScaled<FImplicitConvex3>& Convex,
 			const FRigidTransform3& ConvexTransform,
 			const FReal CullDistance,
-			TCArray<FContactPoint, 4>& OutContactPoints);
+			FContactPointManifold& OutContactPoints);
 
 		template
 		void ConstructCapsuleConvexOneShotManifold(
@@ -1615,7 +1615,7 @@ namespace Chaos
 			const FImplicitBox3& Convex,
 			const FRigidTransform3& ConvexTransform,
 			const FReal CullDistance,
-			TCArray<FContactPoint, 4>& OutContactPoints);
+			FContactPointManifold& OutContactPoints);
 
 
 		template
@@ -1623,42 +1623,42 @@ namespace Chaos
 			const FImplicitConvex3& Convex, 
 			const FTriangle& Triangle,
 			const FReal CullDistance, 
-			TCArray<FContactPoint, 4>& OutContactPoints);
+			FContactPointManifold& OutContactPoints);
 
 		template
 		void ConstructPlanarConvexTriangleOneShotManifold(
 			const TImplicitObjectInstanced<FImplicitConvex3>& Convex,
 			const FTriangle& Triangle,
 			const FReal CullDistance,
-			TCArray<FContactPoint, 4>& OutContactPoints);
+			FContactPointManifold& OutContactPoints);
 
 		template
 		void ConstructPlanarConvexTriangleOneShotManifold(
 			const TImplicitObjectScaled<FImplicitConvex3>& Convex,
 			const FTriangle& Triangle,
 			const FReal CullDistance,
-			TCArray<FContactPoint, 4>& OutContactPoints);
+			FContactPointManifold& OutContactPoints);
 
 		template
 		void ConstructPlanarConvexTriangleOneShotManifold(
 			const FImplicitBox3& Convex,
 			const FTriangle& Triangle,
 			const FReal CullDistance,
-			TCArray<FContactPoint, 4>& OutContactPoints);
+			FContactPointManifold& OutContactPoints);
 
 		template
 		void ConstructPlanarConvexTriangleOneShotManifold(
 			const TImplicitObjectScaled<FImplicitBox3>& Convex,
 			const FTriangle& Triangle,
 			const FReal CullDistance,
-			TCArray<FContactPoint, 4>& OutContactPoints);
+			FContactPointManifold& OutContactPoints);
 
 		template
 		void ConstructPlanarConvexTriangleOneShotManifold(
 			const TImplicitObjectInstanced<FImplicitBox3>& Convex,
 			const FTriangle& Triangle,
 			const FReal CullDistance,
-			TCArray<FContactPoint, 4>& OutContactPoints);
+			FContactPointManifold& OutContactPoints);
 
 
 		// @todo(chaos): all ConstructConvexConvexOneShotManifold taking a triangle are still used in FHeightField::ContactManifoldImp and the 
diff --git a/Engine/Source/Runtime/Experimental/Chaos/Private/Chaos/CollisionOneShotManifolds.h b/Engine/Source/Runtime/Experimental/Chaos/Private/Chaos/CollisionOneShotManifolds.h
--- a/Engine/Source/Runtime/Experimental/Chaos/Private/Chaos/CollisionOneShotManifolds.h
+++ b/Engine/Source/Runtime/Experimental/Chaos/Private/Chaos/CollisionOneShotManifolds.h
@@ -40,20 +40,20 @@ namespace Chaos
 			const ConvexType& Convex,
 			const FRigidTransform3& ConvexTransform,
 			const FReal CullDistance,
-			TCArray<FContactPoint, 4>& OutContactPoints);
+			FContactPointManifold& OutContactPoints);
 
 		void ConstructCapsuleTriangleOneShotManifold(
 			const FImplicitCapsule3& Capsule,
 			const FTriangle& Triangle,
 			const FReal CullDistance,
-			TCArray<FContactPoint, 4>& OutContactPoints);
+			FContactPointManifold& OutContactPoints);
 
 		template <typename ConvexType>
 		void ConstructPlanarConvexTriangleOneShotManifold(
 			const ConvexType& Convex, 
 			const FTriangle& Triangle, 
 			const FReal CullDistance, 
-			TCArray<FContactPoint, 4>& OutContactPoints);
+			FContactPointManifold& OutContactPoints);
 	}
 }
 
diff --git a/Engine/Source/Runtime/Experimental/Chaos/Private/Chaos/CollisionOneShotManifoldsMiscShapes.cpp b/Engine/Source/Runtime/Experimental/Chaos/Private/Chaos/CollisionOneShotManifoldsMiscShapes.cpp
--- a/Engine/Source/Runtime/Experimental/Chaos/Private/Chaos/CollisionOneShotManifoldsMiscShapes.cpp
+++ b/Engine/Source/Runtime/Experimental/Chaos/Private/Chaos/CollisionOneShotManifoldsMiscShapes.cpp
@@ -247,7 +247,7 @@ namespace Chaos
 		}
 
 		template<typename ConvexType>
-		void ConstructSphereConvexManifoldImpl(const FImplicitSphere3& Sphere, const ConvexType& Convex, const FRigidTransform3& SphereToConvexTransform, const FReal CullDistance, TCArray<FContactPoint, 4>& ContactPoints)
+		void ConstructSphereConvexManifoldImpl(const FImplicitSphere3& Sphere, const ConvexType& Convex, const FRigidTransform3& SphereToConvexTransform, const FReal CullDistance, FContactPointManifold& ContactPoints)
 		{
 			FContactPoint ClosestContactPoint = SphereConvexContactPoint(Sphere, Convex, SphereToConvexTransform);
 
@@ -314,7 +314,7 @@ namespace Chaos
 
 			const FRigidTransform3 SphereToConvexTransform = SphereTransform.GetRelativeTransformNoScale(ConvexTransform);
 
-			TCArray<FContactPoint, 4> ContactPoints;
+			FContactPointManifold ContactPoints;
 			if (const FImplicitBox3* RawBox = Convex.template GetObject<FImplicitBox3>())
 			{
 				ConstructSphereConvexManifoldImpl(Sphere, *RawBox, SphereToConvexTransform, Constraint.GetCullDistance(), ContactPoints);
@@ -621,7 +621,7 @@ namespace Chaos
 		 * Templated so we can specialize for some shape types
 		*/
 		template<typename ConvexType>
-		void GenerateConvexTriangleOneShotManifold(const ConvexType& Convex, const FTriangle& Triangle, const FReal CullDistance, TCArray<FContactPoint, 4>& OutContactPoints)
+		void GenerateConvexTriangleOneShotManifold(const ConvexType& Convex, const FTriangle& Triangle, const FReal CullDistance, FContactPointManifold& OutContactPoints)
 		{
 			if (bChaos_Collision_UseConvexTriMesh2)
 			{
@@ -634,7 +634,7 @@ namespace Chaos
 		}
 
 		template<>
-		void GenerateConvexTriangleOneShotManifold<FCapsule>(const FCapsule& Capsule, const FTriangle& Triangle, const FReal CullDistance, TCArray<FContactPoint, 4>& OutContactPoints)
+		void GenerateConvexTriangleOneShotManifold<FCapsule>(const FCapsule& Capsule, const FTriangle& Triangle, const FReal CullDistance, FContactPointManifold& OutContactPoints)
 		{
 			if (bChaos_Collision_UseCapsuleTriMesh2)
 			{
@@ -659,7 +659,7 @@ namespace Chaos
 		template<typename ConvexType, typename MeshType>
 		void GenerateConvexMeshContactPoints(const ConvexType& Convex, const MeshType& Mesh, const FAABB3& MeshQueryBounds, const FRigidTransform3& MeshToConvexTransform, const FReal CullDistance, FContactTriangleCollector& MeshContacts)
 		{
-			TCArray<FContactPoint, 4> TriangleManifoldPoints;
+			FContactPointManifold TriangleManifoldPoints;
 
 			// Loop over all the triangles, build a manifold and add the points to the total manifold
 			// NOTE: contact points will be in the space of the convex until the end of the function when we convert into shape local space
@@ -673,7 +673,7 @@ namespace Chaos
 				{
 					// Add the points into the main contact array
 					// NOTE: The Contacts' FaceIndices will be an index into the ContactTriangles not the original tri mesh (this will get mapped back to the mesh index below)
-					MeshContacts.AddTriangleContacts(MakeArrayView(&TriangleManifoldPoints[0], TriangleManifoldPoints.Num()), Triangle, TriangleIndex, VertexIndex0, VertexIndex1, VertexIndex2, CullDistance);
+					MeshContacts.AddTriangleContacts(MakeArrayView(TriangleManifoldPoints.begin(), TriangleManifoldPoints.Num()), Triangle, TriangleIndex, VertexIndex0, VertexIndex1, VertexIndex2, CullDistance);
 				}
 			});
 
diff --git a/Engine/Source/Runtime/Experimental/Chaos/Private/Chaos/CollisionResolution.cpp b/Engine/Source/Runtime/Experimental/Chaos/Private/Chaos/CollisionResolution.cpp
--- a/Engine/Source/Runtime/Experimental/Chaos/Private/Chaos/CollisionResolution.cpp
+++ b/Engine/Source/Runtime/Experimental/Chaos/Private/Chaos/CollisionResolution.cpp
@@ -967,7 +967,7 @@ namespace Chaos
 		{
 			if (Constraint.GetUseManifold())
 			{
-				TCArray<FContactPoint, 4> ContactPoints;
+				FContactPointManifold ContactPoints;
 				ConstructCapsuleConvexOneShotManifold<FImplicitBox3>(A, ATransform, B, BTransform, Constraint.GetCullDistance(), ContactPoints);
 
 				Constraint.ResetActiveManifoldContacts();
@@ -1027,7 +1027,7 @@ namespace Chaos
 			{
 				if (Constraint.GetUseManifold())
 				{
-					TCArray<FContactPoint, 4> ContactPoints;
+					FContactPointManifold ContactPoints;
 					ConstructCapsuleConvexOneShotManifold<ConvexType>(A, ATransform, B, BTransform, Constraint.GetCullDistance(), ContactPoints);
 
 					Constraint.ResetActiveManifoldContacts();
diff --git a/Engine/Source/Runtime/Experimental/Chaos/Public/Chaos/Collision/CapsuleTriangleContactPoint.h b/Engine/Source/Runtime/Experimental/Chaos/Public/Chaos/Collision/CapsuleTriangleContactPoint.h
--- a/Engine/Source/Runtime/Experimental/Chaos/Public/Chaos/Collision/CapsuleTriangleContactPoint.h
+++ b/Engine/Source/Runtime/Experimental/Chaos/Public/Chaos/Collision/CapsuleTriangleContactPoint.h
@@ -7,5 +7,5 @@
 
 namespace Chaos
 {
-	extern CHAOS_API void ConstructCapsuleTriangleOneShotManifold2(const FImplicitCapsule3& Capsule, const FTriangle& Triangle, const FReal CullDistance, TCArray<FContactPoint, 4>& OutContactPoints);
+	extern CHAOS_API void ConstructCapsuleTriangleOneShotManifold2(const FImplicitCapsule3& Capsule, const FTriangle& Triangle, const FReal CullDistance, FContactPointManifold& OutContactPoints);
 }
\ No newline at end of file
diff --git a/Engine/Source/Runtime/Experimental/Chaos/Public/Chaos/Collision/ContactPoint.h b/Engine/Source/Runtime/Experimental/Chaos/Public/Chaos/Collision/ContactPoint.h
--- a/Engine/Source/Runtime/Experimental/Chaos/Public/Chaos/Collision/ContactPoint.h
+++ b/Engine/Source/Runtime/Experimental/Chaos/Public/Chaos/Collision/ContactPoint.h
@@ -80,20 +80,5 @@ namespace Chaos
 		{ 
 			return (Phi != InvalidPhi<FRealType>());
 		}
-
-		// Switch the shape indices. For use when calling a collision detection method which takes shape types in the opposite order to what you want.
-		// WARNING: this function is not "general purporse" and only works when ShapeContactPoints are both in the smae space, because
-		// the normal is in the space of the second body and negating it does not otherwise transform it to be relative to the other body. 
-		// However, it is useful in the cases where we perform collision detection in the space of one of the bodies before transforming into 
-		// shape space, whcih is actually most of the time.
-		//TContactPoint<FRealType>& SwapShapes()
-		//{
-		//	if (IsSet())
-		//	{
-		//		Swap(ShapeContactPoints[0], ShapeContactPoints[1]);
-		//		ShapeContactNormal = -ShapeContactNormal;
-		//	}
-		//	return *this;
-		//}
 	};
 }
\ No newline at end of file
diff --git a/Engine/Source/Runtime/Experimental/Chaos/Public/Chaos/Collision/ConvexTriangleContactPoint.h b/Engine/Source/Runtime/Experimental/Chaos/Public/Chaos/Collision/ConvexTriangleContactPoint.h
--- a/Engine/Source/Runtime/Experimental/Chaos/Public/Chaos/Collision/ConvexTriangleContactPoint.h
+++ b/Engine/Source/Runtime/Experimental/Chaos/Public/Chaos/Collision/ConvexTriangleContactPoint.h
@@ -8,5 +8,5 @@
 namespace Chaos
 {
 	template <typename ConvexType>
-	extern CHAOS_API void ConstructConvexTriangleOneShotManifold2(const ConvexType& Convex, const FTriangle& Triangle, const FReal CullDistance, TCArray<FContactPoint, 4>& OutContactPoints);
+	extern CHAOS_API void ConstructConvexTriangleOneShotManifold2(const ConvexType& Convex, const FTriangle& Triangle, const FReal CullDistance, FContactPointManifold& OutContactPoints);
 }
\ No newline at end of file
diff --git a/Engine/Source/Runtime/Experimental/Chaos/Public/Chaos/CollisionResolutionTypes.h b/Engine/Source/Runtime/Experimental/Chaos/Public/Chaos/CollisionResolutionTypes.h
--- a/Engine/Source/Runtime/Experimental/Chaos/Public/Chaos/CollisionResolutionTypes.h
+++ b/Engine/Source/Runtime/Experimental/Chaos/Public/Chaos/CollisionResolutionTypes.h
@@ -3,14 +3,24 @@
 
 #include "Chaos/Core.h"
 #include "Chaos/ParticleHandleFwd.h"
-#include "Chaos/Vector.h"
 
 namespace Chaos
 {
+	// Float and Double versions of ContactPoint
 	template<typename T> class TContactPoint;
 	using FContactPoint = TContactPoint<FReal>;
 	using FContactPointf = TContactPoint<FRealSingle>;
 
+	// C-Style array type
+	template<typename T, int32 N> class TCArray;
+
+	// A set of contact points for a manifold (standard size is up to 4 contacts)
+	using FContactPointManifold = TCArray<FContactPoint, 4>;
+
+	// A set of contact points for a manifold that may be more than 4 points
+	using FContactPointLargeManifold = TArray<FContactPoint>;
+
+
 	/** Specifies the type of work we should do*/
 	enum class ECollisionUpdateType
 	{
