Dataflow : Proper Context type identification.

diff --git a/Engine/Plugins/Experimental/ChaosFlesh/Source/ChaosFleshEditor/Private/ChaosFlesh/Editor/FleshEditorToolkit.cpp b/Engine/Plugins/Experimental/ChaosFlesh/Source/ChaosFleshEditor/Private/ChaosFlesh/Editor/FleshEditorToolkit.cpp
--- a/Engine/Plugins/Experimental/ChaosFlesh/Source/ChaosFleshEditor/Private/ChaosFlesh/Editor/FleshEditorToolkit.cpp
+++ b/Engine/Plugins/Experimental/ChaosFlesh/Source/ChaosFleshEditor/Private/ChaosFlesh/Editor/FleshEditorToolkit.cpp
@@ -64,7 +64,7 @@ void FFleshEditorToolkit::InitFleshAssetEditor(const EToolkitMode::Type Mode, co
 		GraphEditor = CreateGraphEditorWidget(FleshAsset->Dataflow, NodeDetailsEditor);
 		SkeletalEditor = CreateSkeletalEditorWidget(FleshAsset->SkeletalMesh);
 
-		Context = TSharedPtr< Dataflow::FEngineContext>(new Dataflow::FEngineContext(FleshAsset, Dataflow, FPlatformTime::Cycles64(), FString("UFleshAsset")));
+		Context = TSharedPtr< Dataflow::FEngineContext>(new Dataflow::FFleshContext(FleshAsset, Dataflow, FPlatformTime::Cycles64()));
 		LastNodeTimestamp = Context->GetTimestamp();
 
 		const TSharedRef<FTabManager::FLayout> StandaloneDefaultLayout = FTabManager::NewLayout("FleshAsset_Layout.V1")
@@ -129,7 +129,7 @@ void FFleshEditorToolkit::Tick(float DeltaTime)
 		{
 			if (!Context)
 			{
-				Context = TSharedPtr< Dataflow::FEngineContext>(new Dataflow::FEngineContext(FleshAsset, Dataflow, Dataflow::FTimestamp::Invalid, FString("UFleshAsset")));
+				Context = TSharedPtr< Dataflow::FEngineContext>(new Dataflow::FFleshContext(FleshAsset, Dataflow, Dataflow::FTimestamp::Invalid));
 				LastNodeTimestamp = Dataflow::FTimestamp::Invalid;
 			}
 
@@ -155,7 +155,7 @@ TSharedRef<SGraphEditor> FFleshEditorToolkit::CreateGraphEditorWidget(UDataflow*
 		{
 			if (!Context)
 			{
-				Context = TSharedPtr< Dataflow::FEngineContext>(new Dataflow::FEngineContext(FleshAsset, Dataflow, Dataflow::FTimestamp::Invalid, FString("UFleshAsset")));
+				Context = TSharedPtr< Dataflow::FEngineContext>(new Dataflow::FFleshContext(FleshAsset, Dataflow, Dataflow::FTimestamp::Invalid));
 			}
 			LastNodeTimestamp = Dataflow::FTimestamp::Invalid;
 
diff --git a/Engine/Plugins/Experimental/ChaosFlesh/Source/ChaosFleshEditor/Private/ChaosFlesh/Editor/FleshEditorToolkit.h b/Engine/Plugins/Experimental/ChaosFlesh/Source/ChaosFleshEditor/Private/ChaosFlesh/Editor/FleshEditorToolkit.h
--- a/Engine/Plugins/Experimental/ChaosFlesh/Source/ChaosFleshEditor/Private/ChaosFlesh/Editor/FleshEditorToolkit.h
+++ b/Engine/Plugins/Experimental/ChaosFlesh/Source/ChaosFleshEditor/Private/ChaosFlesh/Editor/FleshEditorToolkit.h
@@ -19,6 +19,19 @@ class UFleshAsset;
 class UDataflow;
 class USkeletalMesh;
 
+namespace Dataflow
+{
+	class FFleshContext : public TEngineContext<FContextSingle>
+	{
+	public:
+		DATAFLOW_CONTEXT_INTERNAL(TEngineContext<FContextSingle>, FFleshContext);
+
+		FFleshContext(UObject* InOwner, UDataflow* InGraph, FTimestamp InTimestamp)
+			: Super(InOwner, InGraph, InTimestamp)
+		{}
+	};
+}
+
 class FFleshEditorToolkit : public FAssetEditorToolkit, public FTickableEditorObject, public FNotifyHook, public FGCObject
 {
 public:
diff --git a/Engine/Plugins/Experimental/ChaosFlesh/Source/ChaosFleshEngine/Public/ChaosFlesh/FleshAsset.h b/Engine/Plugins/Experimental/ChaosFlesh/Source/ChaosFleshEngine/Public/ChaosFlesh/FleshAsset.h
--- a/Engine/Plugins/Experimental/ChaosFlesh/Source/ChaosFleshEngine/Public/ChaosFlesh/FleshAsset.h
+++ b/Engine/Plugins/Experimental/ChaosFlesh/Source/ChaosFleshEngine/Public/ChaosFlesh/FleshAsset.h
@@ -14,21 +14,6 @@ class UFleshAsset;
 class UDataflow;
 class USkeletalMesh;
 
-namespace Dataflow
-{
-	class CHAOSFLESHENGINE_API FFleshContext : public FContext
-	{
-	public:
-		UFleshAsset* Asset;
-
-		FFleshContext(UFleshAsset* InAsset, FTimestamp InTimestamp)
-			: FContext( InTimestamp )
-			, Asset(InAsset)
-		{}
-
-	};
-}
-
 
 /**
 *	FFleshAssetEdit
diff --git a/Engine/Plugins/Experimental/GeometryCollectionPlugin/Source/GeometryCollectionEditor/Private/GeometryCollection/GeometryCollectionEditorToolkit.cpp b/Engine/Plugins/Experimental/GeometryCollectionPlugin/Source/GeometryCollectionEditor/Private/GeometryCollection/GeometryCollectionEditorToolkit.cpp
--- a/Engine/Plugins/Experimental/GeometryCollectionPlugin/Source/GeometryCollectionEditor/Private/GeometryCollection/GeometryCollectionEditorToolkit.cpp
+++ b/Engine/Plugins/Experimental/GeometryCollectionPlugin/Source/GeometryCollectionEditor/Private/GeometryCollection/GeometryCollectionEditorToolkit.cpp
@@ -54,7 +54,7 @@ void FGeometryCollectionEditorToolkit::InitGeometryCollectionAssetEditor(const E
 		AssetDetailsEditor = CreateAssetDetailsEditorWidget(GeometryCollection);
 		GraphEditor = CreateGraphEditorWidget(Dataflow, NodeDetailsEditor);
 
-		Context = TSharedPtr< Dataflow::FEngineContext>(new Dataflow::FEngineContext(GeometryCollection, Dataflow, FPlatformTime::Cycles64(), FString("UFleshAsset")));
+		Context = TSharedPtr< Dataflow::FEngineContext>(new Dataflow::FCollectionContext(GeometryCollection, Dataflow, FPlatformTime::Cycles64()));
 		LastNodeTimestamp = Context->GetTimestamp();
 
 		const TSharedRef<FTabManager::FLayout> StandaloneDefaultLayout = FTabManager::NewLayout("GeometryCollectionDataflowEditor_Layout.V1")
@@ -108,7 +108,7 @@ void FGeometryCollectionEditorToolkit::Tick(float DeltaTime)
 		{
 			if (!Context)
 			{
-				Context = TSharedPtr< Dataflow::FEngineContext>(new Dataflow::FEngineContext(GeometryCollection, Dataflow, Dataflow::FTimestamp::Invalid, FString("UGeometryCollection")));
+				Context = TSharedPtr< Dataflow::FEngineContext>(new Dataflow::FCollectionContext(GeometryCollection, Dataflow, Dataflow::FTimestamp::Invalid));
 				LastNodeTimestamp = Dataflow::FTimestamp::Invalid;
 			}
 
@@ -134,7 +134,7 @@ TSharedRef<SGraphEditor> FGeometryCollectionEditorToolkit::CreateGraphEditorWidg
 		{
 			if (!Context)
 			{
-				Context = TSharedPtr< Dataflow::FEngineContext>(new FEngineContext(GeometryCollection, Dataflow, Dataflow::FTimestamp::Invalid, FString("UGeometryCollection")));
+				Context = TSharedPtr< Dataflow::FEngineContext>(new FCollectionContext(GeometryCollection, Dataflow, Dataflow::FTimestamp::Invalid));
 				LastNodeTimestamp = Dataflow::FTimestamp::Invalid;
 			}
 
diff --git a/Engine/Plugins/Experimental/GeometryCollectionPlugin/Source/GeometryCollectionEditor/Public/GeometryCollection/GeometryCollectionEditorToolkit.h b/Engine/Plugins/Experimental/GeometryCollectionPlugin/Source/GeometryCollectionEditor/Public/GeometryCollection/GeometryCollectionEditorToolkit.h
--- a/Engine/Plugins/Experimental/GeometryCollectionPlugin/Source/GeometryCollectionEditor/Public/GeometryCollection/GeometryCollectionEditorToolkit.h
+++ b/Engine/Plugins/Experimental/GeometryCollectionPlugin/Source/GeometryCollectionEditor/Public/GeometryCollection/GeometryCollectionEditorToolkit.h
@@ -5,6 +5,7 @@
 #include "CoreMinimal.h"
 #include "Dataflow/DataflowNodeParameters.h"
 #include "Dataflow/DataflowGraphEditor.h"
+#include "Dataflow/DataflowObjectInterface.h"
 #include "GeometryCollection/GeometryCollectionObject.h"
 #include "GraphEditor.h"
 #include "Misc/NotifyHook.h"
@@ -18,6 +19,18 @@ class FTabManager;
 class IToolkitHost;
 class UDataflow;
 
+namespace Dataflow
+{
+	class FCollectionContext : public TEngineContext<FContextSingle>
+	{
+	public:
+		DATAFLOW_CONTEXT_INTERNAL(TEngineContext<FContextSingle>, FCollectionContext);
+
+		FCollectionContext(UObject* InOwner, UDataflow* InGraph, FTimestamp InTimestamp)
+			: Super(InOwner, InGraph, InTimestamp)
+		{}
+	};
+}
 
 class FGeometryCollectionEditorToolkit : public FAssetEditorToolkit, public FTickableEditorObject, public FNotifyHook, public FGCObject
 {
diff --git a/Engine/Source/Runtime/Experimental/Dataflow/Core/Public/Dataflow/DataflowNodeParameters.h b/Engine/Source/Runtime/Experimental/Dataflow/Core/Public/Dataflow/DataflowNodeParameters.h
--- a/Engine/Source/Runtime/Experimental/Dataflow/Core/Public/Dataflow/DataflowNodeParameters.h
+++ b/Engine/Source/Runtime/Experimental/Dataflow/Core/Public/Dataflow/DataflowNodeParameters.h
@@ -94,21 +94,24 @@ namespace Dataflow
 
 
 	public:
-		FContext(FTimestamp InTimestamp, FString InType = FString(""))
+		FContext(FTimestamp InTimestamp)
 			: Timestamp(InTimestamp)
-			, Type(StaticType().Append(InType))
 		{}
 
 		virtual ~FContext() {}
 		
 		FTimestamp Timestamp = FTimestamp::Invalid;
-		FString Type;
+
 		static FString StaticType() { return "FContext"; }
 
+		virtual bool IsA(FString InType) const { return InType.Equals(StaticType()); }
+
+		virtual FString GetType() const { return FContext::StaticType(); }
+
 		template<class T>
 		const T* AsType() const
 		{
-			if (Type.Contains(T::StaticType()))
+			if (IsA(T::StaticType()))
 			{
 				return (T*)this;
 			}
@@ -170,14 +173,21 @@ namespace Dataflow
 		virtual bool Evaluate(const FDataflowOutput& Connection) = 0;
 	};
 
+#define DATAFLOW_CONTEXT_INTERNAL(PARENTTYPE, TYPENAME)															\
+	typedef PARENTTYPE Super;																					\
+	static FString StaticType() { return #TYPENAME; }															\
+	virtual bool IsA(FString InType) const override { return InType.Equals(StaticType()) || Super::IsA(InType); }	\
+	virtual FString GetType() const override { return StaticType(); }
+
 	class DATAFLOWCORE_API FContextSingle : public FContext
 	{
 		FContextCache DataStore;
 
 	public:
+		DATAFLOW_CONTEXT_INTERNAL(FContext, FContextSingle);
 
-		FContextSingle(FTimestamp InTime, FString InType = FString(""))
-			: FContext(InTime, InType)
+		FContextSingle(FTimestamp InTime)
+			: FContext(InTime)
 		{}
 
 		virtual void SetDataImpl(int64 Key, TUniquePtr<FContextCacheElementBase>&& DataStoreEntry) override
@@ -210,9 +220,11 @@ namespace Dataflow
 		TSharedPtr<FCriticalSection> CacheLock;
 
 	public:
+		DATAFLOW_CONTEXT_INTERNAL(FContext, FContextThreaded);
+
 
-		FContextThreaded(FTimestamp InTime, FString InType = FString(""))
-			: FContext(InTime, InType)
+		FContextThreaded(FTimestamp InTime)
+			: FContext(InTime)
 		{
 			CacheLock = MakeShared<FCriticalSection>();
 		}
diff --git a/Engine/Source/Runtime/Experimental/Dataflow/Engine/Public/Dataflow/DataflowObjectInterface.h b/Engine/Source/Runtime/Experimental/Dataflow/Engine/Public/Dataflow/DataflowObjectInterface.h
--- a/Engine/Source/Runtime/Experimental/Dataflow/Engine/Public/Dataflow/DataflowObjectInterface.h
+++ b/Engine/Source/Runtime/Experimental/Dataflow/Engine/Public/Dataflow/DataflowObjectInterface.h
@@ -16,17 +16,14 @@ namespace Dataflow
 	class DATAFLOWENGINE_API TEngineContext : public Base
 	{
 	public:
-		TEngineContext(UObject* InOwner, UDataflow* InGraph, FTimestamp InTimestamp, FString InType)
-				: Base(InTimestamp, StaticType().Append(InType))
+		DATAFLOW_CONTEXT_INTERNAL(Base, TEngineContext);
+
+		TEngineContext(UObject* InOwner, UDataflow* InGraph, FTimestamp InTimestamp)
+				: Base(InTimestamp)
 				, Owner(InOwner)
 				, Graph(InGraph)
 		{}
 	
-		static FString StaticType()
-		{
-			return "TEngineContext";
-		}
-	
 		UObject* Owner = nullptr;
 		UDataflow* Graph = nullptr;
 	};
