Dataflow : Node type casting

diff --git a/Engine/Source/Runtime/Experimental/Dataflow/Core/Public/Dataflow/DataflowNode.h b/Engine/Source/Runtime/Experimental/Dataflow/Core/Public/Dataflow/DataflowNode.h
--- a/Engine/Source/Runtime/Experimental/Dataflow/Core/Public/Dataflow/DataflowNode.h
+++ b/Engine/Source/Runtime/Experimental/Dataflow/Core/Public/Dataflow/DataflowNode.h
@@ -70,7 +70,8 @@ struct DATAFLOWCORE_API FDataflowNode
 	void SetName(FName InName) { Name = InName; }
 	Dataflow::FTimestamp GetTimestamp() const { return LastModifiedTimestamp;  }
 
-	virtual FName GetType() const { check(true); return FName("invalid"); }
+	static FName StaticType() { return FName("FDataflowNode"); }
+	virtual FName GetType() const { return StaticType(); }
 	virtual FName GetDisplayName() const { return ""; }
 	virtual FName GetCategory() const { return ""; }
 	virtual FString GetTags() const { return ""; }
@@ -182,6 +183,18 @@ struct DATAFLOWCORE_API FDataflowNode
 
 	bool IsValid() const { return bValid; }
 
+	virtual bool IsA(FName InType) const { return InType==StaticType(); }
+
+	template<class T>
+	const T* AsType() const
+	{
+		if (IsA(T::StaticType()))
+		{
+			return (T*)this;
+		}
+		return nullptr;
+	}
+
 
 private:
 	virtual TArray<Dataflow::FRenderingParameter> GetRenderParametersImpl() const { return TArray<Dataflow::FRenderingParameter>(); }
@@ -219,6 +232,8 @@ public:																				\
 	static FString StaticTags() {return TAGS;}										\
 	static FString StaticToolTip() {return FString("Create a dataflow node.");}		\
 	virtual FName GetType() const { return #TYPE; }									\
+	virtual bool IsA(FName InType) const override									\
+		{ return InType==StaticType() || Super::IsA(InType); }						\
 	virtual FStructOnScope* NewStructOnScope() override {							\
 	   return new FStructOnScope(TYPE::StaticStruct(), (uint8*)this);}				\
 	virtual void SerializeInternal(FArchive& Ar) override {							\
diff --git a/Engine/Source/Runtime/Experimental/Dataflow/Core/Public/Dataflow/DataflowNodeParameters.h b/Engine/Source/Runtime/Experimental/Dataflow/Core/Public/Dataflow/DataflowNodeParameters.h
--- a/Engine/Source/Runtime/Experimental/Dataflow/Core/Public/Dataflow/DataflowNodeParameters.h
+++ b/Engine/Source/Runtime/Experimental/Dataflow/Core/Public/Dataflow/DataflowNodeParameters.h
@@ -102,11 +102,11 @@ namespace Dataflow
 		
 		FTimestamp Timestamp = FTimestamp::Invalid;
 
-		static FString StaticType() { return "FContext"; }
+		static FName StaticType() { return FName("FContext"); }
 
-		virtual bool IsA(FString InType) const { return InType.Equals(StaticType()); }
+		virtual bool IsA(FName InType) const { return InType==StaticType(); }
 
-		virtual FString GetType() const { return FContext::StaticType(); }
+		virtual FName GetType() const { return FContext::StaticType(); }
 
 		template<class T>
 		const T* AsType() const
@@ -173,11 +173,11 @@ namespace Dataflow
 		virtual bool Evaluate(const FDataflowOutput& Connection) = 0;
 	};
 
-#define DATAFLOW_CONTEXT_INTERNAL(PARENTTYPE, TYPENAME)															\
-	typedef PARENTTYPE Super;																					\
-	static FString StaticType() { return #TYPENAME; }															\
-	virtual bool IsA(FString InType) const override { return InType.Equals(StaticType()) || Super::IsA(InType); }	\
-	virtual FString GetType() const override { return StaticType(); }
+#define DATAFLOW_CONTEXT_INTERNAL(PARENTTYPE, TYPENAME)														\
+	typedef PARENTTYPE Super;																				\
+	static FName StaticType() { return FName(#TYPENAME); }													\
+	virtual bool IsA(FName InType) const override { return InType==StaticType() || Super::IsA(InType); }	\
+	virtual FName GetType() const override { return StaticType(); }
 
 	class DATAFLOWCORE_API FContextSingle : public FContext
 	{
