Dataflow : Removed virtual call in constructor

diff --git a/Engine/Source/Runtime/Experimental/Dataflow/Core/Private/Dataflow/DataflowTerminalNode.cpp b/Engine/Source/Runtime/Experimental/Dataflow/Core/Private/Dataflow/DataflowTerminalNode.cpp
new file mode 100644
--- /dev/null
+++ b/Engine/Source/Runtime/Experimental/Dataflow/Core/Private/Dataflow/DataflowTerminalNode.cpp
@@ -0,0 +1,18 @@
+// Copyright Epic Games, Inc. All Rights Reserved.
+
+#include "Dataflow/DataflowTerminalNode.h"
+
+#include "ChaosLog.h"
+#include "Dataflow/DataflowInputOutput.h"
+#include "Dataflow/DataflowArchive.h"
+
+#include UE_INLINE_GENERATED_CPP_BY_NAME(DataflowTerminalNode)
+
+bool FDataflowTerminalNode::ValidateConnections()
+{
+	if (ensureMsgf(NumOutputs() == 0, TEXT("Error: Terminal nodes can not have outputs.")))
+	{
+		return Super::ValidateConnections();
+	}
+	return false;
+}
diff --git a/Engine/Source/Runtime/Experimental/Dataflow/Core/Public/Dataflow/DataflowNode.h b/Engine/Source/Runtime/Experimental/Dataflow/Core/Public/Dataflow/DataflowNode.h
--- a/Engine/Source/Runtime/Experimental/Dataflow/Core/Public/Dataflow/DataflowNode.h
+++ b/Engine/Source/Runtime/Experimental/Dataflow/Core/Public/Dataflow/DataflowNode.h
@@ -115,8 +115,8 @@ struct DATAFLOWCORE_API FDataflowNode
 	virtual const UScriptStruct* TypedScriptStruct() const { return nullptr; }
 
 	/** Register the Input and Outputs after the creation in the factory */
-	virtual void RegisterInputConnection(const void* Property);
-	virtual void RegisterOutputConnection(const void* Property, const void* Passthrough = nullptr);
+	void RegisterInputConnection(const void* Property);
+	void RegisterOutputConnection(const void* Property, const void* Passthrough = nullptr);
 
 	//
 	// Evaluation
@@ -179,7 +179,7 @@ struct DATAFLOWCORE_API FDataflowNode
 
 	void Invalidate();
 
-	bool ValidateConnections();
+	virtual bool ValidateConnections();
 
 	bool IsValid() const { return bValid; }
 
diff --git a/Engine/Source/Runtime/Experimental/Dataflow/Core/Public/Dataflow/DataflowTerminalNode.h b/Engine/Source/Runtime/Experimental/Dataflow/Core/Public/Dataflow/DataflowTerminalNode.h
--- a/Engine/Source/Runtime/Experimental/Dataflow/Core/Public/Dataflow/DataflowTerminalNode.h
+++ b/Engine/Source/Runtime/Experimental/Dataflow/Core/Public/Dataflow/DataflowTerminalNode.h
@@ -44,11 +44,13 @@ struct DATAFLOWCORE_API FDataflowTerminalNode : public FDataflowNode
 	virtual void AddOutput(FDataflowOutput* InPtr) override { ensure(false); }
 
 	//
-	//  Struct Support
+	// Error Checking
 	//
+	virtual bool ValidateConnections() override ;
 
-	virtual void RegisterOutputConnection(const void* Property, const void* Passthrough = nullptr) { ensure(false); }
-
+	//
+	// Evaluate
+	//
 
 	virtual void Evaluate(Dataflow::FContext& Context) const { ensure(false); }
 
