Dataflow : fix compilation issue

diff --git a/Engine/Source/Runtime/Experimental/Dataflow/Engine/Private/Dataflow/DataflowObject.cpp b/Engine/Source/Runtime/Experimental/Dataflow/Engine/Private/Dataflow/DataflowObject.cpp
--- a/Engine/Source/Runtime/Experimental/Dataflow/Engine/Private/Dataflow/DataflowObject.cpp
+++ b/Engine/Source/Runtime/Experimental/Dataflow/Engine/Private/Dataflow/DataflowObject.cpp
@@ -45,9 +45,7 @@ void UDataflow::EvaluateTerminalNodeByName(FName NodeName, UObject* Asset)
 		TSharedPtr<FDataflowNode> TerminalNode = Dataflow->FindTerminalNode(NodeName);
 		if (TerminalNode)
 		{
-			const float EvalTime = FGameTime::GetTimeSinceAppStart().GetRealTimeSeconds();
-			Dataflow::FEngineContext Context(Asset, this, EvalTime);
-			// Todo , we should be able to cast to a final node type and call EvaluateFinal
+			Dataflow::FEngineContext Context(Asset, this, FPlatformTime::Cycles64());
 			TerminalNode->Evaluate(Context, nullptr);
 		}
 		else
