Enable strip on cook to work without nanite enabled

diff --git a/Engine/Source/Runtime/Experimental/Chaos/Public/GeometryCollection/ManagedArrayCollection.h b/Engine/Source/Runtime/Experimental/Chaos/Public/GeometryCollection/ManagedArrayCollection.h
--- a/Engine/Source/Runtime/Experimental/Chaos/Public/GeometryCollection/ManagedArrayCollection.h
+++ b/Engine/Source/Runtime/Experimental/Chaos/Public/GeometryCollection/ManagedArrayCollection.h
@@ -158,24 +158,29 @@ public:
 		return Collection;
 	}
 
-	void CopyTo(FManagedArrayCollection* Collection) const
+	void CopyTo(FManagedArrayCollection* Collection, const TArray<FName>& GroupsToSkip = TArray<FName>()) const
 	{
 		if (!Map.IsEmpty())
 		{
 			for (const TTuple<FKeyType, FValueType>& Entry : Map)
 			{
-				if (!Collection->HasGroup(Entry.Key.Get<1>()))
+				const FName& AttributeName = Entry.Key.Get<0>();
+				const FName& GroupName = Entry.Key.Get<1>();
+				if (!GroupsToSkip.Contains(GroupName))
 				{
-					Collection->AddGroup(Entry.Key.Get<1>());
+					if (!Collection->HasGroup(GroupName))
+					{
+						Collection->AddGroup(GroupName);
+					}
+
+					if (NumElements(GroupName) != Collection->NumElements(GroupName))
+					{
+						ensure(!Collection->NumElements(GroupName));
+						Collection->AddElements(NumElements(GroupName), GroupName);
+					}
+
+					Collection->CopyAttribute(*this, AttributeName, GroupName);
 				}
-
-				if (NumElements(Entry.Key.Get<1>()) != Collection->NumElements(Entry.Key.Get<1>()))
-				{
-					ensure(!Collection->NumElements(Entry.Key.Get<1>()));
-					Collection->AddElements(NumElements(Entry.Key.Get<1>()), Entry.Key.Get<1>());
-				}
-
-				Collection->CopyAttribute(*this, Entry.Key.Get<0>(), Entry.Key.Get<1>());
 			}
 		}
 	}
diff --git a/Engine/Source/Runtime/Experimental/GeometryCollectionEngine/Private/GeometryCollection/GeometryCollectionObject.cpp b/Engine/Source/Runtime/Experimental/GeometryCollectionEngine/Private/GeometryCollection/GeometryCollectionObject.cpp
--- a/Engine/Source/Runtime/Experimental/GeometryCollectionEngine/Private/GeometryCollection/GeometryCollectionObject.cpp
+++ b/Engine/Source/Runtime/Experimental/GeometryCollectionEngine/Private/GeometryCollection/GeometryCollectionObject.cpp
@@ -18,7 +18,7 @@
 #include "EngineUtils.h"
 #include "Engine/StaticMesh.h"
 #include "PhysicsEngine/PhysicsSettings.h"
-
+#include "EditorFramework/AssetImportData.h"
 
 #if WITH_EDITOR
 #include "GeometryCollection/DerivedDataGeometryCollectionCooker.h"
@@ -306,6 +306,16 @@ void UGeometryCollection::UpdateConvexGeometry()
 #endif
 }
 
+void UGeometryCollection::PostInitProperties()
+{
+#if WITH_EDITORONLY_DATA
+	if (!HasAnyFlags(RF_ClassDefaultObject))
+	{
+		AssetImportData = NewObject<UAssetImportData>(this, TEXT("AssetImportData"));
+	}
+#endif
+	Super::PostInitProperties();
+}
 
 float KgCm3ToKgM3(float Density)
 {
@@ -678,10 +688,18 @@ void UGeometryCollection::Serialize(FArchive& Ar)
 			Materials[SelectedMaterialIndex] = Materials[0];
 		}
 
-		if (bStripOnCook && EnableNanite && NaniteData)
+		if (bStripOnCook)
 		{
-			// If this is a cooked archive, we strip unnecessary data from the Geometry Collection to keep the memory footprint as small as possible.
-			ArchiveGeometryCollection = GenerateMinimalGeometryCollection();
+			if (EnableNanite && NaniteData)
+			{
+				// If this is a cooked archive, we strip unnecessary data from the Geometry Collection to keep the memory footprint as small as possible.
+				ArchiveGeometryCollection = GenerateMinimalGeometryCollection();
+			}
+			else
+			{
+				// non-nanite path where it may be necessary to remove geometry if the geometry collection is rendered using ISMPool or an external rendering system 
+				ArchiveGeometryCollection = CopyCollectionAndRemoveGeometry(GeometryCollection);
+			}
 		}
 #endif
 	}
@@ -1115,6 +1133,16 @@ TSharedPtr<FGeometryCollection, ESPMode::ThreadSafe> UGeometryCollection::Genera
 	return DuplicateGeometryCollection;
 }
 
+TSharedPtr<FGeometryCollection, ESPMode::ThreadSafe> UGeometryCollection::CopyCollectionAndRemoveGeometry(const TSharedPtr<const FGeometryCollection, ESPMode::ThreadSafe>& CollectionToCopy)
+{
+	TSharedPtr<FGeometryCollection, ESPMode::ThreadSafe> GeometryCollectionToReturn(new FGeometryCollection());
+
+	const TArray<FName> GroupsToSkip{ FGeometryCollection::GeometryGroup, FGeometryCollection::VerticesGroup, FGeometryCollection::FacesGroup };
+	CollectionToCopy->CopyTo(GeometryCollectionToReturn.Get(), GroupsToSkip);
+
+	return GeometryCollectionToReturn;
+}
+
 #endif
 
 void UGeometryCollection::InitResources()
diff --git a/Engine/Source/Runtime/Experimental/GeometryCollectionEngine/Public/GeometryCollection/GeometryCollectionObject.h b/Engine/Source/Runtime/Experimental/GeometryCollectionEngine/Public/GeometryCollection/GeometryCollectionObject.h
--- a/Engine/Source/Runtime/Experimental/GeometryCollectionEngine/Public/GeometryCollection/GeometryCollectionObject.h
+++ b/Engine/Source/Runtime/Experimental/GeometryCollectionEngine/Public/GeometryCollection/GeometryCollectionObject.h
@@ -354,6 +354,7 @@ public:
 	virtual void PostEditChangeProperty(struct FPropertyChangedEvent& PropertyChangedEvent) override;
 	virtual bool Modify(bool bAlwaysMarkDirty = true) override;
 #endif
+	virtual void PostInitProperties() override;
 	virtual void PostLoad() override;
 	virtual void BeginDestroy() override;
 	/** End UObject Interface */
@@ -445,6 +446,9 @@ public:
 	/** Produce a deep copy of GeometryCollection member, stripped of data unecessary for gameplay. */
 	TSharedPtr<FGeometryCollection, ESPMode::ThreadSafe> GenerateMinimalGeometryCollection() const;
 
+	/** copy a collection and remove geometry from it */
+	static TSharedPtr<FGeometryCollection, ESPMode::ThreadSafe> CopyCollectionAndRemoveGeometry(const TSharedPtr<const FGeometryCollection, ESPMode::ThreadSafe>& CollectionToCopy);
+
 #if WITH_EDITOR
 	/** If this flag is set, we only regenerate simulation data when requested via CreateSimulationData() */
 	bool bManualDataCreate;
@@ -674,6 +678,10 @@ public:
 	static const TCHAR* GetSelectedMaterialPath();
 
 #if WITH_EDITORONLY_DATA
+	/** Importing data and options used for this geometry collection */
+	UPROPERTY(EditAnywhere, Instanced, Category = ImportSettings)
+	TObjectPtr<class UAssetImportData> AssetImportData;
+
 	/** Information for thumbnail rendering */
 	UPROPERTY(VisibleAnywhere, Instanced, AdvancedDisplay, Category = GeometryCollection)
 	TObjectPtr<class UThumbnailInfo> ThumbnailInfo;
@@ -693,7 +701,7 @@ public:
 	//
 	// Dataflow
 	//
-	UPROPERTY(EditAnywhere, Category = "Dataflow")
+	UPROPERTY(EditAnywhere, BlueprintReadWrite, Category = "Dataflow")
 	TObjectPtr<UDataflow> Dataflow;
 
 	UPROPERTY(EditAnywhere, Category = "Dataflow")
