Chaos: imporve connection graph debug draw by showing the nodes and their states using colors ( blue : kinematic/anchored, yellow: dynamic , red static )

diff --git a/Engine/Source/Runtime/Experimental/Chaos/Private/Chaos/ChaosDebugDraw.cpp b/Engine/Source/Runtime/Experimental/Chaos/Private/Chaos/ChaosDebugDraw.cpp
--- a/Engine/Source/Runtime/Experimental/Chaos/Private/Chaos/ChaosDebugDraw.cpp
+++ b/Engine/Source/Runtime/Experimental/Chaos/Private/Chaos/ChaosDebugDraw.cpp
@@ -1445,7 +1445,7 @@ namespace Chaos
 
 					FColor Color = FColor::Green; // FMath::Lerp(FColor::Green, FColor::Red, Connection.Strain / MaxStrain);
 
-					FDebugDrawQueue::GetInstance().DrawDebugLine(Pos0, Pos1, Color, false, UE_KINDA_SMALL_NUMBER, Settings.DrawPriority, Settings.LineThickness * 3);
+					FDebugDrawQueue::GetInstance().DrawDebugLine(Pos0, Pos1, Color, false, UE_KINDA_SMALL_NUMBER, Settings.DrawPriority, Settings.LineThickness * 2);
 				}
 
 			};
@@ -1484,6 +1484,15 @@ namespace Chaos
 							{
 								if (const FRigidClustering::FClusterHandle& Child = Rigid->CastToClustered())
 								{
+
+									FColor PointColor = Settings.ShapesColorsPerState.GetColorFromState(Child->ObjectState());
+									if (Child->IsAnchored())
+									{ 
+										PointColor = Settings.ShapesColorsPerState.KinematicColor;
+									}
+									FDebugDrawQueue::GetInstance().DrawDebugPoint(Child->X(), PointColor, false, UE_KINDA_SMALL_NUMBER, Settings.DrawPriority, Settings.LineThickness * 8);
+
+
 									const FConnectivityEdgeArray& Edges = Child->ConnectivityEdges();
 									for (const TConnectivityEdge<FReal>& Edge : Edges)
 									{
