GeometryCollection : Material index out of bounds crash.

diff --git a/Engine/Source/Runtime/Experimental/GeometryCollectionEngine/Private/GeometryCollection/GeometryCollectionSceneProxy.cpp b/Engine/Source/Runtime/Experimental/GeometryCollectionEngine/Private/GeometryCollection/GeometryCollectionSceneProxy.cpp
--- a/Engine/Source/Runtime/Experimental/GeometryCollectionEngine/Private/GeometryCollection/GeometryCollectionSceneProxy.cpp
+++ b/Engine/Source/Runtime/Experimental/GeometryCollectionEngine/Private/GeometryCollection/GeometryCollectionSceneProxy.cpp
@@ -818,7 +818,7 @@ FMaterialRenderProxy* FGeometryCollectionSceneProxy::GetMaterial(FMeshElementCol
 		Collector.RegisterOneFrameMaterialProxy(VertexColorVisualizationMaterialInstance);
 		MaterialProxy = VertexColorVisualizationMaterialInstance;
 	}
-	else
+	else if(Materials.IsValidIndex(MaterialIndex))
 	{
 		MaterialProxy = Materials[MaterialIndex]->GetRenderProxy();
 	}
