[Chaos]
* Removed cmath include in public api

diff --git a/Engine/Source/Runtime/Experimental/Chaos/Public/Chaos/AABB.h b/Engine/Source/Runtime/Experimental/Chaos/Public/Chaos/AABB.h
--- a/Engine/Source/Runtime/Experimental/Chaos/Public/Chaos/AABB.h
+++ b/Engine/Source/Runtime/Experimental/Chaos/Public/Chaos/AABB.h
@@ -333,7 +333,7 @@ namespace Chaos
 				{
 					return ChosenPt;
 				}
-				const TVector<T, d> Normalized = Direction / sqrt(SizeSqr);
+				const TVector<T, d> Normalized = Direction / FMath::Sqrt(SizeSqr);
 
 				const TVector<T, d> InflatedPt = ChosenPt + Normalized.GetSafeNormal() * Thickness;
 				return InflatedPt;
diff --git a/Engine/Source/Runtime/Experimental/Chaos/Public/Chaos/PBDSphericalConstraint.h b/Engine/Source/Runtime/Experimental/Chaos/Public/Chaos/PBDSphericalConstraint.h
--- a/Engine/Source/Runtime/Experimental/Chaos/Public/Chaos/PBDSphericalConstraint.h
+++ b/Engine/Source/Runtime/Experimental/Chaos/Public/Chaos/PBDSphericalConstraint.h
@@ -85,7 +85,7 @@ private:
 			static const FSolverReal DeadZoneSquareRadius = UE_SMALL_NUMBER; // We will not push the particle away in the dead zone
 			if (DistanceSquared > FMath::Square(Radius) + DeadZoneSquareRadius)
 			{
-				const FSolverReal Distance = sqrt(DistanceSquared);
+				const FSolverReal Distance = FMath::Sqrt(DistanceSquared);
 				const FSolverVec3 PositionOnSphere = (Radius / Distance) * CenterToParticle;
 				Particles.P(ParticleIndex) = Center + PositionOnSphere;
 			}
diff --git a/Engine/Source/Runtime/Experimental/ChaosCore/Public/Chaos/Matrix.h b/Engine/Source/Runtime/Experimental/ChaosCore/Public/Chaos/Matrix.h
--- a/Engine/Source/Runtime/Experimental/ChaosCore/Public/Chaos/Matrix.h
+++ b/Engine/Source/Runtime/Experimental/ChaosCore/Public/Chaos/Matrix.h
@@ -576,15 +576,15 @@ namespace Chaos
 			FRealDouble c2 = m20 + m21 + M[2][2] * M[2][2];
 			if (c0 > c1 && c0 > c2)
 			{
-				return TVector<FRealDouble, 3>(M[0][0], M[1][0], M[2][0]) / sqrt(c0);
+				return TVector<FRealDouble, 3>(M[0][0], M[1][0], M[2][0]) / FMath::Sqrt(c0);
 			}
 			if (c1 > c2)
 			{
-				return TVector<FRealDouble, 3>(M[1][0], M[1][1], M[2][1]) / sqrt(c1);
+				return TVector<FRealDouble, 3>(M[1][0], M[1][1], M[2][1]) / FMath::Sqrt(c1);
 			}
 			if (c2 > 0)
 			{
-				return TVector<FRealDouble, 3>(M[2][0], M[2][1], M[2][2]) / sqrt(c2);
+				return TVector<FRealDouble, 3>(M[2][0], M[2][1], M[2][2]) / FMath::Sqrt(c2);
 			}
 			return TVector<FRealDouble, 3>(1, 0, 0);
 		}
@@ -1003,15 +1003,15 @@ namespace Chaos
 			FRealSingle c2 = m20 + m21 + M[2][2] * M[2][2];
 			if (c0 > c1 && c0 > c2)
 			{
-				return TVector<FRealSingle, 3>(M[0][0], M[1][0], M[2][0]) / sqrt(c0);
+				return TVector<FRealSingle, 3>(M[0][0], M[1][0], M[2][0]) / FMath::Sqrt(c0);
 			}
 			if (c1 > c2)
 			{
-				return TVector<FRealSingle, 3>(M[1][0], M[1][1], M[2][1]) / sqrt(c1);
+				return TVector<FRealSingle, 3>(M[1][0], M[1][1], M[2][1]) / FMath::Sqrt(c1);
 			}
 			if (c2 > 0)
 			{
-				return TVector<FRealSingle, 3>(M[2][0], M[2][1], M[2][2]) / sqrt(c2);
+				return TVector<FRealSingle, 3>(M[2][0], M[2][1], M[2][2]) / FMath::Sqrt(c2);
 			}
 			return TVector<FRealSingle, 3>(1, 0, 0);
 		}
diff --git a/Engine/Source/Runtime/Experimental/ChaosCore/Public/Chaos/Vector.h b/Engine/Source/Runtime/Experimental/ChaosCore/Public/Chaos/Vector.h
--- a/Engine/Source/Runtime/Experimental/ChaosCore/Public/Chaos/Vector.h
+++ b/Engine/Source/Runtime/Experimental/ChaosCore/Public/Chaos/Vector.h
@@ -5,11 +5,18 @@
 #include "Math/Vector.h"
 #include "Math/Vector2D.h"
 #include "Math/Vector4.h"
-#include <cmath>
 #else
+#include <cmath>
 #include <iostream>
 #include <utility>
 #include <limits>
+namespace FMath
+{
+	float Sqrt(float v) { return sqrt(v); }
+	double Sqrt(double v) { return sqrt(v); }
+	float Atan2(float x, float y) { return atan2(x, y); }
+	double Atan2(double x, double y) { return atan2(x, y); }
+}
 #endif
 
 #include "Chaos/Real.h"
@@ -195,7 +202,7 @@ namespace Chaos
 //			{
 //				SquaredSum += ((*this)[i] * (*this)[i]);
 //			}
-//			return sqrt(SquaredSum);
+//			return FMath::Sqrt(SquaredSum);
 //		}
 //		T Product() const
 //		{
@@ -229,7 +236,7 @@ namespace Chaos
 //			T SizeSqr = SizeSquared();
 //			if (SizeSqr <= TNumericLimits<T>::Min())
 //				return AxisVector(0);
-//			return (*this) / sqrt(SizeSqr);
+//			return (*this) / FMath::Sqrt(SizeSqr);
 //		}
 //		T SafeNormalize()
 //		{
@@ -239,7 +246,7 @@ namespace Chaos
 //				*this = AxisVector(0);
 //				return (T)0.;
 //			}
-//			Size = sqrt(Size);
+//			Size = FMath::Sqrt(Size);
 //			*this = (*this) / Size;
 //			return Size;
 //		}
@@ -557,7 +564,7 @@ namespace Chaos
 				*this = AxisVector(0);
 				return 0.f;
 			}
-			Size = sqrt(Size);
+			Size = FMath::Sqrt(Size);
 			*this = (*this) / Size;
 			return Size;
 		}
@@ -581,7 +588,7 @@ namespace Chaos
 		{
 			FRealSingle s = CrossProduct(V1, V2).Size();
 			FRealSingle c = DotProduct(V1, V2);
-			return atan2(s, c);
+			return FMath::Atan2(s, c);
 		}
 		/** Calculate the velocity to move from P0 to P1 in time Dt. Exists just for symmetry with TRotation::CalculateAngularVelocity! */
 		static TVector<FRealSingle, 3> CalculateVelocity(const TVector<FRealSingle, 3>& P0, const TVector<FRealSingle, 3>& P1, const FRealSingle Dt)
@@ -762,7 +769,7 @@ namespace Chaos
 				*this = AxisVector(0);
 				return 0.f;
 			}
-			Size = sqrt(Size);
+			Size = FMath::Sqrt(Size);
 			*this = (*this) / Size;
 			return Size;
 		}
@@ -786,7 +793,7 @@ namespace Chaos
 		{
 			FRealDouble s = CrossProduct(V1, V2).Size();
 			FRealDouble c = DotProduct(V1, V2);
-			return atan2(s, c);
+			return FMath::Atan2(s, c);
 		}
 		/** Calculate the velocity to move from P0 to P1 in time Dt. Exists just for symmetry with TRotation::CalculateAngularVelocity! */
 		static TVector<FRealDouble, 3> CalculateVelocity(const TVector<FRealDouble, 3>& P0, const TVector<FRealDouble, 3>& P1, const FRealDouble Dt)
@@ -1030,7 +1037,7 @@ namespace Chaos
 		FORCEINLINE T Size() const
 		{
 			const T SquaredSum = X * X + Y * Y + Z * Z;
-			return sqrt(SquaredSum);
+			return FMath::Sqrt(SquaredSum);
 		}
 		FORCEINLINE T Product() const { return X * Y * Z; }
 		FORCEINLINE static TVector<T, 3> AxisVector(const int32 Axis)
@@ -1061,7 +1068,7 @@ namespace Chaos
 			T SizeSqr = SizeSquared();
 			if (SizeSqr <= TNumericLimits<T>::Min())
 				return AxisVector(0);
-			return (*this) / sqrt(SizeSqr);
+			return (*this) / FMath::Sqrt(SizeSqr);
 		}
 		FORCEINLINE T SafeNormalize()
 		{
@@ -1071,7 +1078,7 @@ namespace Chaos
 				*this = AxisVector(0);
 				return (T)0.;
 			}
-			Size = sqrt(Size);
+			Size = FMath::Sqrt(Size);
 			*this = (*this) / Size;
 			return Size;
 		}
