Include removals
* GeometryCollectionComponent.h - Most files are related to the removal of includes in this file. Removed 7 headers
   - Moved GeometryCollectionDamagePropagationData in to its own file
   - Changed COPY_ON_WRITE_ATTRIBUTE to be set in a way so implementations can be placed in cpp file.
   - Moved some function implementations to cpp file
* Removed include in FieldSystemActor.h
* Removed include in FieldSystem.h
* Removed include in CommonUISettings.h

diff --git a/Engine/Plugins/Experimental/ChaosCaching/Source/ChaosCaching/Private/Chaos/Adapters/GeometryCollectionComponentCacheAdapter.cpp b/Engine/Plugins/Experimental/ChaosCaching/Source/ChaosCaching/Private/Chaos/Adapters/GeometryCollectionComponentCacheAdapter.cpp
--- a/Engine/Plugins/Experimental/ChaosCaching/Source/ChaosCaching/Private/Chaos/Adapters/GeometryCollectionComponentCacheAdapter.cpp
+++ b/Engine/Plugins/Experimental/ChaosCaching/Source/ChaosCaching/Private/Chaos/Adapters/GeometryCollectionComponentCacheAdapter.cpp
@@ -5,6 +5,7 @@
 #include "Chaos/ParticleHandle.h"
 #include "Engine/World.h"
 #include "GeometryCollection/GeometryCollectionComponent.h"
+#include "GeometryCollection/GeometryCollectionObject.h"
 #include "PBDRigidsSolver.h"
 #include "PhysicsProxy/GeometryCollectionPhysicsProxy.h"
 #include "PhysicsProxy/SingleParticlePhysicsProxy.h"
diff --git a/Engine/Plugins/Experimental/ChaosCaching/Source/ChaosCaching/Private/Chaos/ChaosCache.cpp b/Engine/Plugins/Experimental/ChaosCaching/Source/ChaosCaching/Private/Chaos/ChaosCache.cpp
--- a/Engine/Plugins/Experimental/ChaosCaching/Source/ChaosCaching/Private/Chaos/ChaosCache.cpp
+++ b/Engine/Plugins/Experimental/ChaosCaching/Source/ChaosCaching/Private/Chaos/ChaosCache.cpp
@@ -3,6 +3,7 @@
 #include "Chaos/ChaosCache.h"
 #include "Chaos/ChaosCachingPlugin.h"
 #include "Components/PrimitiveComponent.h"
+#include "GeometryCollection/GeometryCollection.h"
 #include "UObject/Package.h"
 #include "UObject/UE5MainStreamObjectVersion.h"
 #include "GeometryCollection/GeometryCollectionComponent.h"
diff --git a/Engine/Plugins/Experimental/ChaosCaching/Source/ChaosCachingEditor/Private/Chaos/ActorFactoryCacheManager.cpp b/Engine/Plugins/Experimental/ChaosCaching/Source/ChaosCachingEditor/Private/Chaos/ActorFactoryCacheManager.cpp
--- a/Engine/Plugins/Experimental/ChaosCaching/Source/ChaosCachingEditor/Private/Chaos/ActorFactoryCacheManager.cpp
+++ b/Engine/Plugins/Experimental/ChaosCaching/Source/ChaosCachingEditor/Private/Chaos/ActorFactoryCacheManager.cpp
@@ -7,6 +7,7 @@
 #include "Chaos/ChaosCache.h"
 #include "Components/PrimitiveComponent.h"
 #include "GeometryCollection/GeometryCollectionComponent.h"
+#include "GeometryCollection/GeometryCollectionSimulationTypes.h"
 #include "Kismet2/ComponentEditorUtils.h"
 
 #define LOCTEXT_NAMESPACE "CacheManagerActorFactory"
diff --git a/Engine/Plugins/Experimental/ChaosEditor/Source/FractureEditor/Private/FractureEditorModeToolkit.cpp b/Engine/Plugins/Experimental/ChaosEditor/Source/FractureEditor/Private/FractureEditorModeToolkit.cpp
--- a/Engine/Plugins/Experimental/ChaosEditor/Source/FractureEditor/Private/FractureEditorModeToolkit.cpp
+++ b/Engine/Plugins/Experimental/ChaosEditor/Source/FractureEditor/Private/FractureEditorModeToolkit.cpp
@@ -10,6 +10,7 @@
 #include "FractureEditorMode.h"
 #include "Framework/MultiBox/MultiBoxBuilder.h"
 #include "Framework/MultiBox/MultiBoxDefs.h"
+#include "GeometryCollection/GeometryCollectionObject.h"
 #include "IDetailsView.h"
 #include "IDetailRootObjectCustomization.h"
 #include "PropertyEditorModule.h"
diff --git a/Engine/Plugins/Experimental/ChaosEditor/Source/FractureEditor/Private/FractureSelectionTools.cpp b/Engine/Plugins/Experimental/ChaosEditor/Source/FractureEditor/Private/FractureSelectionTools.cpp
--- a/Engine/Plugins/Experimental/ChaosEditor/Source/FractureEditor/Private/FractureSelectionTools.cpp
+++ b/Engine/Plugins/Experimental/ChaosEditor/Source/FractureEditor/Private/FractureSelectionTools.cpp
@@ -4,6 +4,7 @@
 #include "GeometryCollection/GeometryCollectionComponent.h"
 #include "GeometryCollection/GeometryCollectionClusteringUtility.h"
 #include "EditorSupportDelegates.h"
+#include "GeometryCollection/GeometryCollectionObject.h"
 
 void FFractureSelectionTools::ToggleSelectedBones(UGeometryCollectionComponent* GeometryCollectionComponent, TArray<int32>& BoneIndices, bool bClearCurrentSelection, bool bAdd, bool bSnapToLevel)
 {
diff --git a/Engine/Plugins/Experimental/ChaosEditor/Source/FractureEditor/Private/FractureToolProperties.cpp b/Engine/Plugins/Experimental/ChaosEditor/Source/FractureEditor/Private/FractureToolProperties.cpp
--- a/Engine/Plugins/Experimental/ChaosEditor/Source/FractureEditor/Private/FractureToolProperties.cpp
+++ b/Engine/Plugins/Experimental/ChaosEditor/Source/FractureEditor/Private/FractureToolProperties.cpp
@@ -2,9 +2,11 @@
 
 #include "FractureToolProperties.h"
 
+#include "GeometryCollection/GeometryCollection.h"
 #include "GeometryCollection/GeometryCollectionComponent.h"
 #include "GeometryCollection/GeometryCollectionEngineRemoval.h"
 #include "GeometryCollection/Facades/CollectionAnchoringFacade.h"
+#include "GeometryCollection/GeometryCollectionObject.h"
 
 #include UE_INLINE_GENERATED_CPP_BY_NAME(FractureToolProperties)
 
diff --git a/Engine/Plugins/Experimental/ChaosEditor/Source/FractureEditor/Private/SGeometryCollectionHistogram.cpp b/Engine/Plugins/Experimental/ChaosEditor/Source/FractureEditor/Private/SGeometryCollectionHistogram.cpp
--- a/Engine/Plugins/Experimental/ChaosEditor/Source/FractureEditor/Private/SGeometryCollectionHistogram.cpp
+++ b/Engine/Plugins/Experimental/ChaosEditor/Source/FractureEditor/Private/SGeometryCollectionHistogram.cpp
@@ -5,6 +5,7 @@
 #include "FractureToolProperties.h"
 
 #include "GeometryCollection/GeometryCollectionComponent.h"
+#include "GeometryCollection/GeometryCollectionObject.h"
 #include "GeometryCollection/GeometryCollectionUtility.h"
 
 #include "Widgets/Colors/SColorBlock.h"
diff --git a/Engine/Plugins/Experimental/ChaosEditor/Source/FractureEditor/Private/SGeometryCollectionOutliner.cpp b/Engine/Plugins/Experimental/ChaosEditor/Source/FractureEditor/Private/SGeometryCollectionOutliner.cpp
--- a/Engine/Plugins/Experimental/ChaosEditor/Source/FractureEditor/Private/SGeometryCollectionOutliner.cpp
+++ b/Engine/Plugins/Experimental/ChaosEditor/Source/FractureEditor/Private/SGeometryCollectionOutliner.cpp
@@ -2,6 +2,7 @@
 
 #include "SGeometryCollectionOutliner.h"
 
+#include "GeometryCollection/GeometryCollectionObject.h"
 #include "Widgets/Layout/SBorder.h"
 #include "Widgets/Text/STextBlock.h"
 #include "Widgets/Views/STreeView.h"
diff --git a/Engine/Plugins/Experimental/ChaosNiagara/Source/ChaosNiagara/Private/NiagaraDataInterfaceChaosDestruction.cpp b/Engine/Plugins/Experimental/ChaosNiagara/Source/ChaosNiagara/Private/NiagaraDataInterfaceChaosDestruction.cpp
--- a/Engine/Plugins/Experimental/ChaosNiagara/Source/ChaosNiagara/Private/NiagaraDataInterfaceChaosDestruction.cpp
+++ b/Engine/Plugins/Experimental/ChaosNiagara/Source/ChaosNiagara/Private/NiagaraDataInterfaceChaosDestruction.cpp
@@ -4,6 +4,7 @@
 #include "Chaos/CollisionResolutionTypes.h"
 #include "Chaos/PBDCollisionConstraintsUtil.h"
 #include "Engine/World.h"
+#include "GeometryCollection/GeometryCollectionObject.h"
 #include "Misc/FileHelper.h"
 #include "NiagaraComponent.h"
 #include "NiagaraGpuComputeDispatchInterface.h"
diff --git a/Engine/Plugins/Experimental/ChaosNiagara/Source/ChaosNiagara/Private/NiagaraDataInterfaceGeometryCollection.cpp b/Engine/Plugins/Experimental/ChaosNiagara/Source/ChaosNiagara/Private/NiagaraDataInterfaceGeometryCollection.cpp
--- a/Engine/Plugins/Experimental/ChaosNiagara/Source/ChaosNiagara/Private/NiagaraDataInterfaceGeometryCollection.cpp
+++ b/Engine/Plugins/Experimental/ChaosNiagara/Source/ChaosNiagara/Private/NiagaraDataInterfaceGeometryCollection.cpp
@@ -1,6 +1,8 @@
 // Copyright Epic Games, Inc. All Rights Reserved.
 
 #include "NiagaraDataInterfaceGeometryCollection.h"
+#include "GeometryCollection/GeometryCollection.h"
+#include "GeometryCollection/GeometryCollectionObject.h"
 #include "NiagaraRenderer.h"
 #include "NiagaraSimStageData.h"
 #include "NiagaraShaderParametersBuilder.h"
diff --git a/Engine/Plugins/Experimental/CommonUI/Source/CommonUI/Public/CommonUISettings.h b/Engine/Plugins/Experimental/CommonUI/Source/CommonUI/Public/CommonUISettings.h
--- a/Engine/Plugins/Experimental/CommonUI/Source/CommonUI/Public/CommonUISettings.h
+++ b/Engine/Plugins/Experimental/CommonUI/Source/CommonUI/Public/CommonUISettings.h
@@ -10,9 +10,13 @@
 #include "ICommonUIModule.h"
 #include "CommonUITypes.h"
 #include "CommonUIRichTextData.h"
-#include "CommonTextBlock.h"
 #include "GameplayTagContainer.h"
 #include "NativeGameplayTags.h"
+
+#if UE_ENABLE_INCLUDE_ORDER_DEPRECATED_IN_5_2
+#include "CommonTextBlock.h"
+#endif
+
 #include "CommonUISettings.generated.h"
 
 class UMaterial;
diff --git a/Engine/Plugins/Experimental/GeometryCollectionPlugin/Source/GeometryCollectionEditor/Private/GeometryCollection/GeometryCollectionAssetBroker.cpp b/Engine/Plugins/Experimental/GeometryCollectionPlugin/Source/GeometryCollectionEditor/Private/GeometryCollection/GeometryCollectionAssetBroker.cpp
--- a/Engine/Plugins/Experimental/GeometryCollectionPlugin/Source/GeometryCollectionEditor/Private/GeometryCollection/GeometryCollectionAssetBroker.cpp
+++ b/Engine/Plugins/Experimental/GeometryCollectionPlugin/Source/GeometryCollectionEditor/Private/GeometryCollection/GeometryCollectionAssetBroker.cpp
@@ -2,6 +2,7 @@
 
 #include "GeometryCollection/GeometryCollectionAssetBroker.h"
 #include "GeometryCollection/GeometryCollectionComponent.h"
+#include "GeometryCollection/GeometryCollectionObject.h"
 
 
 UClass* FGeometryCollectionAssetBroker::GetSupportedAssetClass()
diff --git a/Engine/Plugins/Experimental/GeometryCollectionPlugin/Source/GeometryCollectionEditor/Private/GeometryCollection/GeometryCollectionSelectRigidBodyEdMode.cpp b/Engine/Plugins/Experimental/GeometryCollectionPlugin/Source/GeometryCollectionEditor/Private/GeometryCollection/GeometryCollectionSelectRigidBodyEdMode.cpp
--- a/Engine/Plugins/Experimental/GeometryCollectionPlugin/Source/GeometryCollectionEditor/Private/GeometryCollection/GeometryCollectionSelectRigidBodyEdMode.cpp
+++ b/Engine/Plugins/Experimental/GeometryCollectionPlugin/Source/GeometryCollectionEditor/Private/GeometryCollection/GeometryCollectionSelectRigidBodyEdMode.cpp
@@ -4,7 +4,10 @@
 
 #include "GeometryCollection/GeometryCollectionSelectRigidBodyEdMode.h"
 
+#include "GeometryCollection/GeometryCollection.h"
 #include "GeometryCollection/GeometryCollectionActor.h"
+#include "GeometryCollection/GeometryCollectionBoneNode.h"
+#include "GeometryCollection/GeometryCollectionComponent.h"
 #include "GeometryCollection/GeometryCollectionDebugDrawActor.h"
 #include "GeometryCollection/GeometryCollectionHitProxy.h"
 #include "Editor.h"
@@ -12,7 +15,6 @@
 #include "EditorViewportClient.h"
 #include "EngineUtils.h"
 #include "PropertyHandle.h"
-#include "GeometryCollection/GeometryCollectionComponent.h"
 
 DEFINE_LOG_CATEGORY_STATIC(LogGeometryCollectionSelectRigidBodyEdMode, Log, All);
 
diff --git a/Engine/Plugins/Experimental/GeometryCollectionPlugin/Source/GeometryCollectionEditor/Private/GeometryCollection/GeometryCollectionThumbnailScene.cpp b/Engine/Plugins/Experimental/GeometryCollectionPlugin/Source/GeometryCollectionEditor/Private/GeometryCollection/GeometryCollectionThumbnailScene.cpp
--- a/Engine/Plugins/Experimental/GeometryCollectionPlugin/Source/GeometryCollectionEditor/Private/GeometryCollection/GeometryCollectionThumbnailScene.cpp
+++ b/Engine/Plugins/Experimental/GeometryCollectionPlugin/Source/GeometryCollectionEditor/Private/GeometryCollection/GeometryCollectionThumbnailScene.cpp
@@ -1,6 +1,7 @@
 // Copyright Epic Games, Inc. All Rights Reserved.
 
 #include "GeometryCollection/GeometryCollectionThumbnailScene.h"
+#include "GeometryCollection/GeometryCollectionObject.h"
 #include "ThumbnailRendering/SceneThumbnailInfo.h"
 #include "GeometryCollection/GeometryCollection.h"
 #include "GeometryCollection/GeometryCollectionComponent.h"
diff --git a/Engine/Plugins/Experimental/GeometryCollectionPlugin/Source/GeometryCollectionTracks/Private/MovieSceneGeometryCollectionTemplate.cpp b/Engine/Plugins/Experimental/GeometryCollectionPlugin/Source/GeometryCollectionTracks/Private/MovieSceneGeometryCollectionTemplate.cpp
--- a/Engine/Plugins/Experimental/GeometryCollectionPlugin/Source/GeometryCollectionTracks/Private/MovieSceneGeometryCollectionTemplate.cpp
+++ b/Engine/Plugins/Experimental/GeometryCollectionPlugin/Source/GeometryCollectionTracks/Private/MovieSceneGeometryCollectionTemplate.cpp
@@ -4,6 +4,7 @@
 #include "Compilation/MovieSceneCompilerRules.h"
 #include "GeometryCollection/GeometryCollectionComponent.h"
 #include "Evaluation/MovieSceneEvaluation.h"
+#include "GeometryCollection/GeometryCollectionSimulationTypes.h"
 #include "IMovieScenePlayer.h"
 #include "UObject/ObjectKey.h"
 
diff --git a/Engine/Plugins/MeshPainting/Source/MeshPaintingToolset/Private/MeshPaintGeometryCollectionAdapter.cpp b/Engine/Plugins/MeshPainting/Source/MeshPaintingToolset/Private/MeshPaintGeometryCollectionAdapter.cpp
--- a/Engine/Plugins/MeshPainting/Source/MeshPaintingToolset/Private/MeshPaintGeometryCollectionAdapter.cpp
+++ b/Engine/Plugins/MeshPainting/Source/MeshPaintingToolset/Private/MeshPaintGeometryCollectionAdapter.cpp
@@ -1,7 +1,9 @@
 // Copyright Epic Games, Inc. All Rights Reserved.
 
 #include "MeshPaintGeometryCollectionAdapter.h"
+#include "GeometryCollection/GeometryCollection.h"
 #include "GeometryCollection/GeometryCollectionComponent.h"
+#include "GeometryCollection/GeometryCollectionObject.h"
 #include "MeshPaintHelpers.h"
 #include "MeshPaintingToolsetTypes.h"
 #include "ComponentReregisterContext.h"
diff --git a/Engine/Source/Runtime/Engine/Private/PhysicsField/PhysicsFieldComponent.cpp b/Engine/Source/Runtime/Engine/Private/PhysicsField/PhysicsFieldComponent.cpp
--- a/Engine/Source/Runtime/Engine/Private/PhysicsField/PhysicsFieldComponent.cpp
+++ b/Engine/Source/Runtime/Engine/Private/PhysicsField/PhysicsFieldComponent.cpp
@@ -1,6 +1,7 @@
 // Copyright Epic Games, Inc. All Rights Reserved.
 
 #include "PhysicsField/PhysicsFieldComponent.h"
+#include "ChaosStats.h"
 #include "SceneManagement.h"
 #include "EngineGlobals.h"
 #include "Engine/Engine.h"
diff --git a/Engine/Source/Runtime/Experimental/Chaos/Public/Field/FieldSystem.h b/Engine/Source/Runtime/Experimental/Chaos/Public/Field/FieldSystem.h
--- a/Engine/Source/Runtime/Experimental/Chaos/Public/Field/FieldSystem.h
+++ b/Engine/Source/Runtime/Experimental/Chaos/Public/Field/FieldSystem.h
@@ -2,13 +2,16 @@
 
 #pragma once
 
+#include "Chaos/ParticleHandleFwd.h"
 #include "CoreMinimal.h"
 #include "Containers/Array.h"
 #include "Containers/ArrayView.h"
 #include "Containers/Queue.h"
 #include "Field/FieldSystemTypes.h"
 #include "Field/FieldArrayView.h"
+#if UE_ENABLE_INCLUDE_ORDER_DEPRECATED_IN_5_2
 #include "Chaos/ParticleHandle.h"
+#endif
 #include "Math/Vector.h"
 
 /**
diff --git a/Engine/Source/Runtime/Experimental/Chaos/Public/GeometryCollection/GeometryCollectionSimulationCoreTypes.h b/Engine/Source/Runtime/Experimental/Chaos/Public/GeometryCollection/GeometryCollectionSimulationCoreTypes.h
--- a/Engine/Source/Runtime/Experimental/Chaos/Public/GeometryCollection/GeometryCollectionSimulationCoreTypes.h
+++ b/Engine/Source/Runtime/Experimental/Chaos/Public/GeometryCollection/GeometryCollectionSimulationCoreTypes.h
@@ -81,7 +81,7 @@ struct FSharedSimulationSizeSpecificData
 //
 //
 //
-enum ESimulationInitializationState { Unintialized = 0, Activated, Created, Initialized };
+enum ESimulationInitializationState : uint8 { Unintialized = 0, Activated, Created, Initialized };
 
 
 /**
diff --git a/Engine/Source/Runtime/Experimental/FieldSystem/Source/FieldSystemEngine/Public/Field/FieldSystemActor.h b/Engine/Source/Runtime/Experimental/FieldSystem/Source/FieldSystemEngine/Public/Field/FieldSystemActor.h
--- a/Engine/Source/Runtime/Experimental/FieldSystem/Source/FieldSystemEngine/Public/Field/FieldSystemActor.h
+++ b/Engine/Source/Runtime/Experimental/FieldSystem/Source/FieldSystemEngine/Public/Field/FieldSystemActor.h
@@ -8,11 +8,14 @@
 #include "GameFramework/Actor.h"
 #include "Field/FieldSystem.h"
 #include "Field/FieldSystemNodes.h"
+#if UE_ENABLE_INCLUDE_ORDER_DEPRECATED_IN_5_2
 #include "Field/FieldSystemComponent.h"
+#endif
 #include "UObject/ObjectMacros.h"
 
 #include "FieldSystemActor.generated.h"
 
+class UFieldSystemComponent;
 
 UCLASS(meta=(ChildCanTick))
 class FIELDSYSTEMENGINE_API AFieldSystemActor: public AActor
diff --git a/Engine/Source/Runtime/Experimental/GeometryCollectionEngine/Private/GeometryCollection/GeometryCollectionActor.cpp b/Engine/Source/Runtime/Experimental/GeometryCollectionEngine/Private/GeometryCollection/GeometryCollectionActor.cpp
--- a/Engine/Source/Runtime/Experimental/GeometryCollectionEngine/Private/GeometryCollection/GeometryCollectionActor.cpp
+++ b/Engine/Source/Runtime/Experimental/GeometryCollectionEngine/Private/GeometryCollection/GeometryCollectionActor.cpp
@@ -6,6 +6,7 @@ GeometryCollectionActor.cpp: AGeometryCollectionActor methods.
 
 #include "GeometryCollection/GeometryCollectionActor.h"
 
+#include "Chaos/ChaosSolverActor.h"
 #include "Chaos/Utilities.h"
 #include "Chaos/Plane.h"
 #include "Chaos/Box.h"
@@ -16,8 +17,10 @@ GeometryCollectionActor.cpp: AGeometryCollectionActor methods.
 #include "Engine/World.h"
 #include "GeometryCollection/GeometryCollectionAlgo.h"
 #include "GeometryCollection/GeometryCollectionComponent.h"
+#include "GeometryCollection/GeometryCollectionObject.h"
 #include "GeometryCollection/GeometryCollectionUtility.h"
 #include "Math/Box.h"
+#include "Physics/Experimental/PhysScene_Chaos.h"
 #include "Physics/PhysicsInterfaceCore.h"
 #include "PhysicsSolver.h"
 #include "GeometryCollection/GeometryCollectionDebugDrawComponent.h"
diff --git a/Engine/Source/Runtime/Experimental/GeometryCollectionEngine/Private/GeometryCollection/GeometryCollectionComponent.cpp b/Engine/Source/Runtime/Experimental/GeometryCollectionEngine/Private/GeometryCollection/GeometryCollectionComponent.cpp
--- a/Engine/Source/Runtime/Experimental/GeometryCollectionEngine/Private/GeometryCollection/GeometryCollectionComponent.cpp
+++ b/Engine/Source/Runtime/Experimental/GeometryCollectionEngine/Private/GeometryCollection/GeometryCollectionComponent.cpp
@@ -11,7 +11,9 @@
 #include "Components/BoxComponent.h"
 #include "Engine/Engine.h"
 #include "Engine/InstancedStaticMesh.h"
+#include "Field/FieldSystemComponent.h"
 #include "GeometryCollection/Facades/CollectionHierarchyFacade.h"
+#include "GeometryCollection/GeometryCollection.h"
 #include "GeometryCollection/GeometryCollectionActor.h"
 #include "GeometryCollection/GeometryCollectionAlgo.h"
 #include "GeometryCollection/GeometryCollectionCache.h"
@@ -231,6 +233,40 @@ FGeomComponentCacheParameters::FGeomComponentCacheParameters()
 {
 }
 
+#undef COPY_ON_WRITE_ATTRIBUTE
+#define COPY_ON_WRITE_ATTRIBUTE(Type, Name, Group)										\
+const TManagedArray<Type>& UGeometryCollectionComponent::Get##Name##Array() const 		\
+{																						\
+	return Indirect##Name##Array ?														\
+		*Indirect##Name##Array : RestCollection->GetGeometryCollection()->Name;			\
+}																						\
+TManagedArray<Type>& UGeometryCollectionComponent::Get##Name##ArrayCopyOnWrite()		\
+{																						\
+	if(!Indirect##Name##Array)															\
+	{																					\
+		static FName StaticName(#Name);													\
+		DynamicCollection->AddAttribute<Type>(StaticName, Group);						\
+		DynamicCollection->CopyAttribute(												\
+			*RestCollection->GetGeometryCollection(), StaticName, Group);				\
+		Indirect##Name##Array =															\
+			&DynamicCollection->ModifyAttribute<Type>(StaticName, Group);				\
+		CopyOnWriteAttributeList.Add(													\
+			reinterpret_cast<FManagedArrayBase**>(&Indirect##Name##Array));				\
+	}																					\
+	return *Indirect##Name##Array;														\
+}																						\
+void UGeometryCollectionComponent::Reset##Name##ArrayDynamic()							\
+{																						\
+	Indirect##Name##Array = NULL;														\
+}																						\
+const TManagedArray<Type>& UGeometryCollectionComponent::Get##Name##ArrayRest() const	\
+{																						\
+	return RestCollection->GetGeometryCollection()->Name;								\
+}																						\
+
+// Define the methods
+COPY_ON_WRITE_ATTRIBUTES
+
 UGeometryCollectionComponent::UGeometryCollectionComponent(const FObjectInitializer& ObjectInitializer)
 	: Super(ObjectInitializer)
 	, ChaosSolverActor(nullptr)
@@ -537,6 +573,12 @@ FBoxSphereBounds UGeometryCollectionComponent::CalcBounds(const FTransform& Loca
 	return FBoxSphereBounds(ComputeBounds(LocalToWorldWithScale));
 }
 
+int32 UGeometryCollectionComponent::GetNumElements(FName Group) const
+{
+	int32 Size = RestCollection->NumElements(Group);	//assume rest collection has the group and is connected to dynamic.
+	return Size > 0 ? Size : DynamicCollection->NumElements(Group);	//if not, maybe dynamic has the group
+}
+
 void UGeometryCollectionComponent::UpdateCachedBounds()
 {
 	ComponentSpaceBounds = ComputeBounds(FMatrix::Identity);
@@ -1242,6 +1284,16 @@ static void DispatchGeometryCollectionCrumblingEvent(const FChaosCrumblingEvent&
 	}
 }
 
+const FGeometryDynamicCollection* UGeometryCollectionComponent::GetDynamicCollection() const
+{
+	return DynamicCollection.Get();
+}
+
+FGeometryDynamicCollection* UGeometryCollectionComponent::GetDynamicCollection()
+{
+	return DynamicCollection.Get();
+}
+
 void UGeometryCollectionComponent::DispatchChaosPhysicsCollisionBlueprintEvents(const FChaosPhysicsCollisionInfo& CollisionInfo)
 {
 	ReceivePhysicsCollision(CollisionInfo);
@@ -3519,6 +3571,16 @@ void UGeometryCollectionComponent::GetInitializationCommands(TArray<FFieldSystem
 	}
 }
 
+bool UGeometryCollectionComponent::GetSuppressSelectionMaterial() const
+{
+	return RestCollection->GetGeometryCollection()->HasAttribute("Hide", FGeometryCollection::TransformGroup);
+}
+
+const int UGeometryCollectionComponent::GetBoneSelectedMaterialID() const
+{
+	return RestCollection->GetBoneSelectedMaterialIndex();
+}
+
 FPhysScene_Chaos* UGeometryCollectionComponent::GetInnerChaosScene() const
 {
 	if (ChaosSolverActor)
diff --git a/Engine/Source/Runtime/Experimental/GeometryCollectionEngine/Private/GeometryCollection/GeometryCollectionDebugDrawActor.cpp b/Engine/Source/Runtime/Experimental/GeometryCollectionEngine/Private/GeometryCollection/GeometryCollectionDebugDrawActor.cpp
--- a/Engine/Source/Runtime/Experimental/GeometryCollectionEngine/Private/GeometryCollection/GeometryCollectionDebugDrawActor.cpp
+++ b/Engine/Source/Runtime/Experimental/GeometryCollectionEngine/Private/GeometryCollection/GeometryCollectionDebugDrawActor.cpp
@@ -2,6 +2,7 @@
 
 #include "GeometryCollection/GeometryCollectionDebugDrawActor.h"
 
+#include "Chaos/ChaosSolverActor.h"
 #include "GeometryCollection/GeometryCollection.h"
 #include "GeometryCollection/GeometryCollectionParticlesData.h"
 #include "GeometryCollection/GeometryCollectionComponent.h"
diff --git a/Engine/Source/Runtime/Experimental/GeometryCollectionEngine/Private/GeometryCollection/GeometryCollectionSceneProxy.cpp b/Engine/Source/Runtime/Experimental/GeometryCollectionEngine/Private/GeometryCollection/GeometryCollectionSceneProxy.cpp
--- a/Engine/Source/Runtime/Experimental/GeometryCollectionEngine/Private/GeometryCollection/GeometryCollectionSceneProxy.cpp
+++ b/Engine/Source/Runtime/Experimental/GeometryCollectionEngine/Private/GeometryCollection/GeometryCollectionSceneProxy.cpp
@@ -4,6 +4,8 @@
 
 #include "Async/ParallelFor.h"
 #include "Engine/Engine.h"
+#include "GeometryCollection/GeometryCollection.h"
+#include "GeometryCollection/GeometryCollectionObject.h"
 #include "MaterialShaderType.h"
 #include "Materials/Material.h"
 #include "CommonRenderResources.h"
diff --git a/Engine/Source/Runtime/Experimental/GeometryCollectionEngine/Public/GeometryCollection/GeometryCollectionComponent.h b/Engine/Source/Runtime/Experimental/GeometryCollectionEngine/Public/GeometryCollection/GeometryCollectionComponent.h
--- a/Engine/Source/Runtime/Experimental/GeometryCollectionEngine/Public/GeometryCollection/GeometryCollectionComponent.h
+++ b/Engine/Source/Runtime/Experimental/GeometryCollectionEngine/Public/GeometryCollection/GeometryCollectionComponent.h
@@ -7,15 +7,18 @@
 #include "Field/FieldSystem.h"
 #include "Field/FieldSystemActor.h"
 #include "Field/FieldSystemNodes.h"
-#include "Field/FieldSystemObjects.h"
 #include "GameFramework/Actor.h"
+#if UE_ENABLE_INCLUDE_ORDER_DEPRECATED_IN_5_2
+#include "Field/FieldSystemObjects.h"
 #include "GeometryCollection/GeometryCollection.h"
 #include "GeometryCollection/GeometryCollectionSimulationTypes.h"
 #include "GeometryCollection/GeometryCollectionSimulationCoreTypes.h"
-#include "Physics/Experimental/PhysScene_Chaos.h"
-#include "GeometryCollectionProxyData.h"
 #include "GeometryCollectionObject.h"
+#include "GeometryCollectionProxyData.h"
+#include "Physics/Experimental/PhysScene_Chaos.h"
+#endif
 #include "GeometryCollectionEditorSelection.h"
+#include "GeometryCollection/GeometryCollectionDamagePropagationData.h"
 #include "GeometryCollection/RecordedTransformTrack.h"
 #include "Templates/UniquePtr.h"
 #include "Chaos/ChaosGameplayEventDispatcher.h"
@@ -30,6 +33,7 @@
 
 struct FGeometryCollectionConstantData;
 struct FGeometryCollectionDynamicData;
+class FManagedArrayBase;
 class UGeometryCollectionComponent;
 class UBoxComponent;
 class UGeometryCollectionCache;
@@ -38,9 +42,18 @@ class AChaosSolverActor;
 struct FGeometryCollectionEmbeddedExemplar;
 class UInstancedStaticMeshComponent;
 class FGeometryCollectionDecayDynamicFacade;
+class FGeometryDynamicCollection;
 struct FGeometryCollectionDecayContext;
+struct FGeometryCollectionSection;
 struct FDamageCollector;
+class FPhysScene_Chaos;
 class AGeometryCollectionISMPoolActor;
+enum ESimulationInitializationState : uint8;
+enum class EClusterConnectionTypeEnum : uint8;
+enum class EInitialVelocityTypeEnum : uint8;
+enum class EObjectStateTypeEnum : uint8;
+namespace Chaos { enum class EObjectStateType: int8; }
+template<class InElementType> class TManagedArray;
 
 DECLARE_DYNAMIC_MULTICAST_DELEGATE_OneParam(FOnChaosBreakEvent, const FChaosBreakEvent&, BreakEvent);
 
@@ -258,37 +271,13 @@ private:
 //GetArrayRest (gives original rest value)
 //This generates pointers to arrays marked private. Macro assumes getters are public
 //todo(ocohen): may want to take in a static name
-#define COPY_ON_WRITE_ATTRIBUTE(Type, Name, Group)								\
-FORCEINLINE const TManagedArray<Type>& Get##Name##Array() const 				\
-{																				\
-	return Indirect##Name##Array ?												\
-		*Indirect##Name##Array : RestCollection->GetGeometryCollection()->Name;	\
-}																				\
-FORCEINLINE TManagedArray<Type>& Get##Name##ArrayCopyOnWrite()					\
-{																				\
-	if(!Indirect##Name##Array)													\
-	{																			\
-		static FName StaticName(#Name);											\
-		DynamicCollection->AddAttribute<Type>(StaticName, Group);				\
-		DynamicCollection->CopyAttribute(										\
-			*RestCollection->GetGeometryCollection(), StaticName, Group);		\
-		Indirect##Name##Array =													\
-			&DynamicCollection->ModifyAttribute<Type>(StaticName, Group);		\
-		CopyOnWriteAttributeList.Add(											\
-			reinterpret_cast<FManagedArrayBase**>(&Indirect##Name##Array));		\
-	}																			\
-	return *Indirect##Name##Array;												\
-}																				\
-FORCEINLINE void Reset##Name##ArrayDynamic()									\
-{																				\
-	Indirect##Name##Array = NULL;												\
-}																				\
-FORCEINLINE const TManagedArray<Type>& Get##Name##ArrayRest() const				\
-{																				\
-	return RestCollection->GetGeometryCollection()->Name;						\
-}																				\
-private:																		\
-	TManagedArray<Type>* Indirect##Name##Array;									\
+#define COPY_ON_WRITE_ATTRIBUTE(Type, Name, Group)			\
+	const TManagedArray<Type>& Get##Name##Array() const;	\
+	TManagedArray<Type>& Get##Name##ArrayCopyOnWrite();		\
+	void Reset##Name##ArrayDynamic();						\
+	const TManagedArray<Type>& Get##Name##ArrayRest() const;\
+private:													\
+	TManagedArray<Type>* Indirect##Name##Array;				\
 public:
 
 /**
@@ -650,54 +639,53 @@ public:
 #endif
 
 	/** API for getting at geometry collection data */
-	FORCEINLINE int32 GetNumElements(FName Group) const
-	{
-		int32 Size = RestCollection->NumElements(Group);	//assume rest collection has the group and is connected to dynamic.
-		return Size > 0 ? Size : DynamicCollection->NumElements(Group);	//if not, maybe dynamic has the group
-	}
+	int32 GetNumElements(FName Group) const;
 
 	// Update cached bounds; used e.g. when updating the exploded view of the geometry collection
 	void UpdateCachedBounds();
 
-	// Vertices Group
-	COPY_ON_WRITE_ATTRIBUTE(FVector3f, Vertex, FGeometryCollection::VerticesGroup) 	//GetVertexArray, GetVertexArrayCopyOnWrite, GetVertexArrayRest
-	COPY_ON_WRITE_ATTRIBUTE(TArray<FVector2f>, UVs, FGeometryCollection::VerticesGroup)		//GetUVsArray
-	COPY_ON_WRITE_ATTRIBUTE(FLinearColor, Color, FGeometryCollection::VerticesGroup)//GetColorArray
-	COPY_ON_WRITE_ATTRIBUTE(FVector3f, TangentU, FGeometryCollection::VerticesGroup)	//GetTangentUArray
-	COPY_ON_WRITE_ATTRIBUTE(FVector3f, TangentV, FGeometryCollection::VerticesGroup)	//...
-	COPY_ON_WRITE_ATTRIBUTE(FVector3f, Normal, FGeometryCollection::VerticesGroup)
-	COPY_ON_WRITE_ATTRIBUTE(int32, BoneMap, FGeometryCollection::VerticesGroup)
-
-	// Faces Group
-	COPY_ON_WRITE_ATTRIBUTE(FIntVector, Indices, FGeometryCollection::FacesGroup)
-	COPY_ON_WRITE_ATTRIBUTE(bool, Visible, FGeometryCollection::FacesGroup)
-	COPY_ON_WRITE_ATTRIBUTE(int32, MaterialIndex, FGeometryCollection::FacesGroup)
-	COPY_ON_WRITE_ATTRIBUTE(int32, MaterialID, FGeometryCollection::FacesGroup)
-
-	// Geometry Group
-	COPY_ON_WRITE_ATTRIBUTE(int32, TransformIndex, FGeometryCollection::GeometryGroup)
-	COPY_ON_WRITE_ATTRIBUTE(FBox, BoundingBox, FGeometryCollection::GeometryGroup)
-	COPY_ON_WRITE_ATTRIBUTE(float, InnerRadius, FGeometryCollection::GeometryGroup)
-	COPY_ON_WRITE_ATTRIBUTE(float, OuterRadius, FGeometryCollection::GeometryGroup)
-	COPY_ON_WRITE_ATTRIBUTE(int32, VertexStart, FGeometryCollection::GeometryGroup)
-	COPY_ON_WRITE_ATTRIBUTE(int32, VertexCount, FGeometryCollection::GeometryGroup)
-	COPY_ON_WRITE_ATTRIBUTE(int32, FaceStart, FGeometryCollection::GeometryGroup)
-	COPY_ON_WRITE_ATTRIBUTE(int32, FaceCount, FGeometryCollection::GeometryGroup)
-
-	// Material Group
-	COPY_ON_WRITE_ATTRIBUTE(FGeometryCollectionSection, Sections, FGeometryCollection::MaterialGroup)
-
-	// Transform group
-	COPY_ON_WRITE_ATTRIBUTE(FString, BoneName, FTransformCollection::TransformGroup)
-	COPY_ON_WRITE_ATTRIBUTE(FLinearColor, BoneColor, FTransformCollection::TransformGroup)
-	COPY_ON_WRITE_ATTRIBUTE(FTransform, Transform, FTransformCollection::TransformGroup)
-	COPY_ON_WRITE_ATTRIBUTE(int32, Parent, FTransformCollection::TransformGroup)
-	COPY_ON_WRITE_ATTRIBUTE(TSet<int32>, Children, FTransformCollection::TransformGroup)
-	COPY_ON_WRITE_ATTRIBUTE(int32, SimulationType, FTransformCollection::TransformGroup)
-	COPY_ON_WRITE_ATTRIBUTE(int32, TransformToGeometryIndex, FTransformCollection::TransformGroup)
-	COPY_ON_WRITE_ATTRIBUTE(int32, StatusFlags, FTransformCollection::TransformGroup)
-	COPY_ON_WRITE_ATTRIBUTE(int32, ExemplarIndex, FTransformCollection::TransformGroup)
-
+#define COPY_ON_WRITE_ATTRIBUTES \
+	/* Vertices Group */ \
+	COPY_ON_WRITE_ATTRIBUTE(FVector3f, Vertex, FGeometryCollection::VerticesGroup)		/* GetVertexArray, GetVertexArrayCopyOnWrite, GetVertexArrayRest */ \
+	COPY_ON_WRITE_ATTRIBUTE(TArray<FVector2f>, UVs, FGeometryCollection::VerticesGroup) /* GetUVsArray */		\
+	COPY_ON_WRITE_ATTRIBUTE(FLinearColor, Color, FGeometryCollection::VerticesGroup)	/* GetColorArray */		\
+	COPY_ON_WRITE_ATTRIBUTE(FVector3f, TangentU, FGeometryCollection::VerticesGroup)	/* GetTangentUArray */	\
+	COPY_ON_WRITE_ATTRIBUTE(FVector3f, TangentV, FGeometryCollection::VerticesGroup)	/* //... */		\
+	COPY_ON_WRITE_ATTRIBUTE(FVector3f, Normal, FGeometryCollection::VerticesGroup)						\
+	COPY_ON_WRITE_ATTRIBUTE(int32, BoneMap, FGeometryCollection::VerticesGroup)							\
+																										\
+	/* Faces Group */																					\
+	COPY_ON_WRITE_ATTRIBUTE(FIntVector, Indices, FGeometryCollection::FacesGroup)						\
+	COPY_ON_WRITE_ATTRIBUTE(bool, Visible, FGeometryCollection::FacesGroup)								\
+	COPY_ON_WRITE_ATTRIBUTE(int32, MaterialIndex, FGeometryCollection::FacesGroup)						\
+	COPY_ON_WRITE_ATTRIBUTE(int32, MaterialID, FGeometryCollection::FacesGroup)							\
+																										\
+	/* Geometry Group */																				\
+	COPY_ON_WRITE_ATTRIBUTE(int32, TransformIndex, FGeometryCollection::GeometryGroup)					\
+	COPY_ON_WRITE_ATTRIBUTE(FBox, BoundingBox, FGeometryCollection::GeometryGroup)						\
+	COPY_ON_WRITE_ATTRIBUTE(float, InnerRadius, FGeometryCollection::GeometryGroup)						\
+	COPY_ON_WRITE_ATTRIBUTE(float, OuterRadius, FGeometryCollection::GeometryGroup)						\
+	COPY_ON_WRITE_ATTRIBUTE(int32, VertexStart, FGeometryCollection::GeometryGroup)						\
+	COPY_ON_WRITE_ATTRIBUTE(int32, VertexCount, FGeometryCollection::GeometryGroup)						\
+	COPY_ON_WRITE_ATTRIBUTE(int32, FaceStart, FGeometryCollection::GeometryGroup)						\
+	COPY_ON_WRITE_ATTRIBUTE(int32, FaceCount, FGeometryCollection::GeometryGroup)						\
+																										\
+	/* Material Group */																				\
+	COPY_ON_WRITE_ATTRIBUTE(FGeometryCollectionSection, Sections, FGeometryCollection::MaterialGroup)	\
+																										\
+	/* Transform group */																				\
+	COPY_ON_WRITE_ATTRIBUTE(FString, BoneName, FTransformCollection::TransformGroup)					\
+	COPY_ON_WRITE_ATTRIBUTE(FLinearColor, BoneColor, FTransformCollection::TransformGroup)				\
+	COPY_ON_WRITE_ATTRIBUTE(FTransform, Transform, FTransformCollection::TransformGroup)				\
+	COPY_ON_WRITE_ATTRIBUTE(int32, Parent, FTransformCollection::TransformGroup)						\
+	COPY_ON_WRITE_ATTRIBUTE(TSet<int32>, Children, FTransformCollection::TransformGroup)				\
+	COPY_ON_WRITE_ATTRIBUTE(int32, SimulationType, FTransformCollection::TransformGroup)				\
+	COPY_ON_WRITE_ATTRIBUTE(int32, TransformToGeometryIndex, FTransformCollection::TransformGroup)		\
+	COPY_ON_WRITE_ATTRIBUTE(int32, StatusFlags, FTransformCollection::TransformGroup)					\
+	COPY_ON_WRITE_ATTRIBUTE(int32, ExemplarIndex, FTransformCollection::TransformGroup)					\
+
+	// Declare all the methods
+	COPY_ON_WRITE_ATTRIBUTES
 
 	UPROPERTY(EditAnywhere, NoClear, BlueprintReadOnly, Category = "ChaosPhysics")
 	TObjectPtr<const UGeometryCollection> RestCollection;
@@ -854,9 +842,9 @@ public:
 	bool GetShowBoneColors() const { return bShowBoneColors; }
 	bool GetEnableBoneSelection() const { return bEnableBoneSelection; }
 	
-	bool GetSuppressSelectionMaterial() const { return RestCollection->GetGeometryCollection()->HasAttribute("Hide", FGeometryCollection::TransformGroup); }
+	bool GetSuppressSelectionMaterial() const;
 	
-	FORCEINLINE const int GetBoneSelectedMaterialID() const { return RestCollection->GetBoneSelectedMaterialIndex(); }
+	const int GetBoneSelectedMaterialID() const;
 	
 #if WITH_EDITORONLY_DATA
 	FORCEINLINE const TArray<int32>& GetSelectedBones() const { return SelectedBones; }
@@ -973,8 +961,8 @@ public:
 	/** Used by Niagara DI to query global matrices rather than recalculating them again */
 	const TArray<FMatrix>& GetGlobalMatrices() { return GlobalMatrices; }
 
-	const FGeometryDynamicCollection* GetDynamicCollection() const { return DynamicCollection.Get(); }
-	FGeometryDynamicCollection* GetDynamicCollection() { return DynamicCollection.Get(); } // TEMP HACK?
+	const FGeometryDynamicCollection* GetDynamicCollection() const;
+	FGeometryDynamicCollection* GetDynamicCollection();  // TEMP HACK?
 
 public:
 	UPROPERTY(BlueprintAssignable, Category = "Collision")
diff --git a/Engine/Source/Runtime/Experimental/GeometryCollectionEngine/Public/GeometryCollection/GeometryCollectionDamagePropagationData.h b/Engine/Source/Runtime/Experimental/GeometryCollectionEngine/Public/GeometryCollection/GeometryCollectionDamagePropagationData.h
new file mode 100644
--- /dev/null
+++ b/Engine/Source/Runtime/Experimental/GeometryCollectionEngine/Public/GeometryCollection/GeometryCollectionDamagePropagationData.h
@@ -0,0 +1,24 @@
+// Copyright Epic Games, Inc. All Rights Reserved.
+
+#pragma once
+
+#include "GeometryCollectionDamagePropagationData.generated.h"
+
+USTRUCT(BlueprintType)
+struct GEOMETRYCOLLECTIONENGINE_API FGeometryCollectionDamagePropagationData
+{
+public:
+	GENERATED_BODY()
+
+	/** Whether or not damage propagation is enabled. */
+	UPROPERTY(EditAnywhere, BlueprintReadWrite, Category = "Damage Propagation")
+	bool bEnabled = true;
+
+	/** factor of the remaining strain propagated through the connection graph after a piece breaks. */
+	UPROPERTY(EditAnywhere, BlueprintReadWrite, Category = "Damage Propagation")
+	float BreakDamagePropagationFactor = 1.0f;
+
+	/** factor of the received strain propagated throug the connection graph if the piece did not break. */
+	UPROPERTY(EditAnywhere, BlueprintReadWrite, Category = "Damage Propagation")
+	float ShockDamagePropagationFactor = 0.0f;
+};
diff --git a/Engine/Source/Runtime/Experimental/GeometryCollectionEngine/Public/GeometryCollection/GeometryCollectionObject.h b/Engine/Source/Runtime/Experimental/GeometryCollectionEngine/Public/GeometryCollection/GeometryCollectionObject.h
--- a/Engine/Source/Runtime/Experimental/GeometryCollectionEngine/Public/GeometryCollection/GeometryCollectionObject.h
+++ b/Engine/Source/Runtime/Experimental/GeometryCollectionEngine/Public/GeometryCollection/GeometryCollectionObject.h
@@ -7,6 +7,7 @@
 #include "Rendering/NaniteResources.h"
 #include "InstanceUniformShaderParameters.h"
 #include "GeometryCollection/ManagedArray.h"
+#include "GeometryCollection/GeometryCollectionDamagePropagationData.h"
 #include "GeometryCollection/GeometryCollectionSimulationTypes.h"
 #include "Chaos/ChaosSolverActor.h"
 
@@ -316,25 +317,6 @@ private:
 };
 
 
-USTRUCT(BlueprintType)
-struct GEOMETRYCOLLECTIONENGINE_API FGeometryCollectionDamagePropagationData
-{
-public:
-	GENERATED_BODY()
-
-	/** Whether or not damage propagation is enabled. */
-	UPROPERTY(EditAnywhere, BlueprintReadWrite, Category = "Damage Propagation")
-	bool bEnabled = true;
-
-	/** factor of the remaining strain propagated through the connection graph after a piece breaks. */
-	UPROPERTY(EditAnywhere, BlueprintReadWrite, Category = "Damage Propagation")
-	float BreakDamagePropagationFactor = 1.0f;
-
-	/** factor of the received strain propagated throug the connection graph if the piece did not break. */
-	UPROPERTY(EditAnywhere, BlueprintReadWrite, Category = "Damage Propagation")
-	float ShockDamagePropagationFactor = 0.0f;
-};
-
 /**
 * UGeometryCollectionObject (UObject)
 *
