Voronoi: Fix disabling static analysis checkers

#rnx

diff --git a/Engine/Source/Runtime/Experimental/Voronoi/Voronoi.Build.cs b/Engine/Source/Runtime/Experimental/Voronoi/Voronoi.Build.cs
--- a/Engine/Source/Runtime/Experimental/Voronoi/Voronoi.Build.cs
+++ b/Engine/Source/Runtime/Experimental/Voronoi/Voronoi.Build.cs
@@ -17,8 +17,7 @@ namespace UnrealBuildTool.Rules
 			UnsafeTypeCastWarningLevel = WarningLevel.Warning;
 
 			// Disable all static analysis checkers for this module ( third party code ) 
-			StaticAnalyzerDisabledCheckers.Add("all");
-
+			bDisableStaticAnalysis = true;
 		}
 	}
 }
