make static mesh -> geometry collection conversion respect the static mesh build settings, so normals/tangents are computed with the same settings as used by the static mesh build, and are not recomputed if the build settings do not allow it

diff --git a/Engine/Source/Runtime/Experimental/GeometryCollectionEngine/Private/GeometryCollection/GeometryCollectionEngineConversion.cpp b/Engine/Source/Runtime/Experimental/GeometryCollectionEngine/Private/GeometryCollection/GeometryCollectionEngineConversion.cpp
--- a/Engine/Source/Runtime/Experimental/GeometryCollectionEngine/Private/GeometryCollection/GeometryCollectionEngineConversion.cpp
+++ b/Engine/Source/Runtime/Experimental/GeometryCollectionEngine/Private/GeometryCollection/GeometryCollectionEngineConversion.cpp
@@ -365,6 +365,44 @@ void FGeometryCollectionEngineConversion::AppendMeshDescription(
 
 }
 
+namespace
+{
+	// Note: This is the same as the corresponding ModelingComponents AssetUtils function, in the MeshModelingToolset plugin,
+	// but we cannot use a plugin function from here, so have our own copy
+	void InitializeAutoGeneratedAttributes(FMeshDescription& Mesh, const FMeshBuildSettings* BuildSettings)
+	{
+		check(BuildSettings);
+
+		// if the build settings don't indicate calculated normals/tangents, we will assume they are good
+		if (BuildSettings->bRecomputeNormals || BuildSettings->bRecomputeTangents)
+		{
+			// check if any are invalid
+			bool bHasInvalidNormals, bHasInvalidTangents;
+			FStaticMeshOperations::AreNormalsAndTangentsValid(Mesh, bHasInvalidNormals, bHasInvalidTangents);
+
+			// if neither are invalid we are not going to recompute
+			if (bHasInvalidNormals || bHasInvalidTangents)
+			{
+				FStaticMeshAttributes Attributes(Mesh);
+				if (!Attributes.GetTriangleNormals().IsValid() || !Attributes.GetTriangleTangents().IsValid())
+				{
+					// If these attributes don't exist, create them and compute their values for each triangle
+					FStaticMeshOperations::ComputeTriangleTangentsAndNormals(Mesh);
+				}
+
+				EComputeNTBsFlags ComputeNTBsOptions = EComputeNTBsFlags::BlendOverlappingNormals;
+				ComputeNTBsOptions |= BuildSettings->bRecomputeNormals ? EComputeNTBsFlags::Normals : EComputeNTBsFlags::None;
+				ComputeNTBsOptions |= BuildSettings->bRecomputeTangents ? EComputeNTBsFlags::Tangents : EComputeNTBsFlags::None;
+				ComputeNTBsOptions |= BuildSettings->bUseMikkTSpace ? EComputeNTBsFlags::UseMikkTSpace : EComputeNTBsFlags::None;
+				ComputeNTBsOptions |= BuildSettings->bComputeWeightedNormals ? EComputeNTBsFlags::WeightedNTBs : EComputeNTBsFlags::None;
+				ComputeNTBsOptions |= BuildSettings->bRemoveDegenerates ? EComputeNTBsFlags::IgnoreDegenerateTriangles : EComputeNTBsFlags::None;
+
+				FStaticMeshOperations::ComputeTangentsAndNormals(Mesh, ComputeNTBsOptions);
+			}
+		}
+	}
+}
+
 FMeshDescription* FGeometryCollectionEngineConversion::GetMaxResMeshDescriptionWithNormalsAndTangents(const UStaticMesh* StaticMesh)
 {
 	if (StaticMesh == nullptr)
@@ -373,19 +411,21 @@ FMeshDescription* FGeometryCollectionEngineConversion::GetMaxResMeshDescriptionW
 	}
 
 	FMeshDescription* MeshDescription = nullptr;
+	const FStaticMeshSourceModel* SourceModel = nullptr;
 #if WITH_EDITORONLY_DATA
 	// Prefer the HiRes description, although this isn't always available.
 	if (StaticMesh->IsHiResMeshDescriptionValid())
 	{
 		MeshDescription = StaticMesh->GetHiResMeshDescription();
+		SourceModel = &StaticMesh->GetHiResSourceModel();
 	}
 	else
 	{
 		MeshDescription = StaticMesh->GetMeshDescription(0);
+		SourceModel = &StaticMesh->GetSourceModel(0);
 	}
 
-	FStaticMeshOperations::ComputeTriangleTangentsAndNormals(*MeshDescription);
-	FStaticMeshOperations::RecomputeNormalsAndTangentsIfNeeded(*MeshDescription, EComputeNTBsFlags::UseMikkTSpace);
+	InitializeAutoGeneratedAttributes(*MeshDescription, &SourceModel->BuildSettings);
 #endif //WITH_EDITORONLY_DATA
 	return MeshDescription;
 }
