Chaos Cloth - Fix backstop when the component is scaled.

diff --git a/Engine/Source/Runtime/Experimental/Chaos/Private/Chaos/PBDSphericalConstraint.ispc b/Engine/Source/Runtime/Experimental/Chaos/Private/Chaos/PBDSphericalConstraint.ispc
--- a/Engine/Source/Runtime/Experimental/Chaos/Private/Chaos/PBDSphericalConstraint.ispc
+++ b/Engine/Source/Runtime/Experimental/Chaos/Private/Chaos/PBDSphericalConstraint.ispc
@@ -75,7 +75,7 @@ export void ApplyLegacySphericalBackstopConstraints(uniform FVector4f ParticlesP
 		const FVector3f AnimationPosition = VectorLoad(&AnimationPositions[extract(ParticleIndex, 0)]);
 		const FVector3f AnimationNormal = VectorLoad(&AnimationNormals[extract(ParticleIndex, 0)]);
 
-		const float SphereOffsetDistance = SphereOffsetDistances[Index];
+		const float SphereOffsetDistance = SphereOffsetDistances[Index] * SphereRadiiMultiplier;
 		const float Radius = SphereRadii[Index] * SphereRadiiMultiplier;
 
 		const FVector3f Center = AnimationPosition - SphereOffsetDistance * AnimationNormal;  // Legacy version already includes the radius within the distance
@@ -124,7 +124,7 @@ export void ApplySphericalBackstopConstraints(uniform FVector4f ParticlesPandInv
 		const FVector3f AnimationPosition = VectorLoad(&AnimationPositions[extract(ParticleIndex, 0)]);
 		const FVector3f AnimationNormal = VectorLoad(&AnimationNormals[extract(ParticleIndex, 0)]);
 
-		const float SphereOffsetDistance = SphereOffsetDistances[Index];
+		const float SphereOffsetDistance = SphereOffsetDistances[Index] * SphereRadiiMultiplier;
 		const float Radius = SphereRadii[Index] * SphereRadiiMultiplier;
 
 		const FVector3f Center = AnimationPosition - (Radius + SphereOffsetDistance) * AnimationNormal;  // Non legacy version adds radius to the distance
diff --git a/Engine/Source/Runtime/Experimental/Chaos/Public/Chaos/PBDSphericalConstraint.h b/Engine/Source/Runtime/Experimental/Chaos/Public/Chaos/PBDSphericalConstraint.h
--- a/Engine/Source/Runtime/Experimental/Chaos/Public/Chaos/PBDSphericalConstraint.h
+++ b/Engine/Source/Runtime/Experimental/Chaos/Public/Chaos/PBDSphericalConstraint.h
@@ -197,7 +197,7 @@ private:
 			const FSolverVec3& AnimationPosition = AnimationPositions[ParticleIndex];
 			const FSolverVec3& AnimationNormal = AnimationNormals[ParticleIndex];
 
-			const FSolverReal SphereOffsetDistance = SphereOffsetDistances[Index];
+			const FSolverReal SphereOffsetDistance = SphereOffsetDistances[Index] * SphereRadiiMultiplier;
 			const FSolverReal Radius = SphereRadii[Index] * SphereRadiiMultiplier;
 
 			const FSolverVec3 Center = AnimationPosition - (Radius + SphereOffsetDistance) * AnimationNormal;  // Non legacy version adds radius to the distance
@@ -234,7 +234,7 @@ private:
 			const FSolverVec3& AnimationPosition = AnimationPositions[ParticleIndex];
 			const FSolverVec3& AnimationNormal = AnimationNormals[ParticleIndex];
 
-			const FSolverReal SphereOffsetDistance = SphereOffsetDistances[Index];
+			const FSolverReal SphereOffsetDistance = SphereOffsetDistances[Index] * SphereRadiiMultiplier;
 			const FSolverReal Radius = SphereRadii[Index] * SphereRadiiMultiplier;
 
 			const FSolverVec3 Center = AnimationPosition - SphereOffsetDistance * AnimationNormal;  // Legacy version already includes the radius within the distance
