Improve set clustered function on geometry collection component
- Add a set clustered by transformed box
- Fix the set clustered functions to properly update the kinematic state of the parents
- Imporved the debug draw of connection graph to better see the node state when close

diff --git a/Engine/Source/Runtime/Experimental/Chaos/Private/Chaos/ChaosDebugDraw.cpp b/Engine/Source/Runtime/Experimental/Chaos/Private/Chaos/ChaosDebugDraw.cpp
--- a/Engine/Source/Runtime/Experimental/Chaos/Private/Chaos/ChaosDebugDraw.cpp
+++ b/Engine/Source/Runtime/Experimental/Chaos/Private/Chaos/ChaosDebugDraw.cpp
@@ -1442,10 +1442,16 @@ namespace Chaos
 					FRigidTransform3 ClusterTransform(Connection.Parent->X(), Connection.Parent->R());
 					FVec3 Pos0 = ClusterTransform.TransformPosition(Connection.Handle0->ChildToParent().GetLocation());
 					FVec3 Pos1 = ClusterTransform.TransformPosition(Connection.Handle1->ChildToParent().GetLocation());
+					const FVec3 Normal = (Pos1 - Pos0).GetSafeNormal();
+					const FReal Length = (Pos1 - Pos0).Size();
+					constexpr FReal ShrinkPercentage = 0.05;
+					const FVec3 SkrinkVector = Normal * ShrinkPercentage * Length;
+					Pos0 += SkrinkVector;
+					Pos1 -= SkrinkVector;
 
 					FColor Color = FColor::Green; // FMath::Lerp(FColor::Green, FColor::Red, Connection.Strain / MaxStrain);
 
-					FDebugDrawQueue::GetInstance().DrawDebugLine(Pos0, Pos1, Color, false, UE_KINDA_SMALL_NUMBER, Settings.DrawPriority, Settings.LineThickness * 2);
+					FDebugDrawQueue::GetInstance().DrawDebugLine(Pos0, Pos1, Color, false, UE_KINDA_SMALL_NUMBER, Settings.DrawPriority, Settings.LineThickness);
 				}
 
 			};
@@ -1490,7 +1496,7 @@ namespace Chaos
 									{ 
 										PointColor = Settings.ShapesColorsPerState.KinematicColor;
 									}
-									FDebugDrawQueue::GetInstance().DrawDebugPoint(Child->X(), PointColor, false, UE_KINDA_SMALL_NUMBER, Settings.DrawPriority, Settings.LineThickness * 8);
+									FDebugDrawQueue::GetInstance().DrawDebugPoint(Child->X(), PointColor, false, UE_KINDA_SMALL_NUMBER, Settings.DrawPriority, Settings.LineThickness * 10);
 
 
 									const FConnectivityEdgeArray& Edges = Child->ConnectivityEdges();
diff --git a/Engine/Source/Runtime/Experimental/Chaos/Private/PhysicsProxy/GeometryCollectionPhysicsProxy.cpp b/Engine/Source/Runtime/Experimental/Chaos/Private/PhysicsProxy/GeometryCollectionPhysicsProxy.cpp
--- a/Engine/Source/Runtime/Experimental/Chaos/Private/PhysicsProxy/GeometryCollectionPhysicsProxy.cpp
+++ b/Engine/Source/Runtime/Experimental/Chaos/Private/PhysicsProxy/GeometryCollectionPhysicsProxy.cpp
@@ -1727,7 +1727,7 @@ void FGeometryCollectionPhysicsProxy::BreakActiveClusters_External()
 	}
 }
 
-static void SetParticleAnchored(Chaos::FPBDRigidsEvolution* Evolution, Chaos::FPBDRigidClusteredParticleHandle* ParticleHandle, bool bAnchored)
+static void SetParticleAnchored_Internal(Chaos::FPBDRigidsEvolution* Evolution, Chaos::FPBDRigidClusteredParticleHandle* ParticleHandle, bool bAnchored)
 {
 	if (ParticleHandle)
 	{
@@ -1743,6 +1743,19 @@ static void SetParticleAnchored(Chaos::FPBDRigidsEvolution* Evolution, Chaos::FP
 		}
 	}
 }
+static Chaos::FPBDRigidClusteredParticleHandle* GetTopActiveClusteredParent_Internal(Chaos::FPBDRigidClusteredParticleHandle* ParticleHandle)
+{
+	while (ParticleHandle)
+	{
+		Chaos::FPBDRigidClusteredParticleHandle* ParentParticle = ParticleHandle->Parent();
+		if (!ParentParticle)
+		{
+			return ParticleHandle;
+		}
+		ParticleHandle = ParentParticle;
+	}
+	return nullptr;
+}
 
 void FGeometryCollectionPhysicsProxy::SetAnchoredByIndex_External(int32 Index, bool bAnchored)
 {
@@ -1754,35 +1767,52 @@ void FGeometryCollectionPhysicsProxy::SetAnchoredByIndex_External(int32 Index, b
 				Chaos::FPBDRigidsEvolution* Evolution = RBDSolver->GetEvolution();
 				if (SolverParticleHandles.IsValidIndex(Index))
 				{
-					SetParticleAnchored(Evolution, SolverParticleHandles[Index], bAnchored);
-				}
+					SetParticleAnchored_Internal(Evolution, SolverParticleHandles[Index], bAnchored);
+					if (Chaos::FPBDRigidClusteredParticleHandle* TopParentHandle = GetTopActiveClusteredParent_Internal(SolverParticleHandles[Index]))
+					{
+						Chaos::UpdateKinematicProperties(TopParentHandle, Evolution->GetRigidClustering().GetChildrenMap(), *Evolution);
+					}
+				}		
 			});
 	}
 }
 
-void FGeometryCollectionPhysicsProxy::SetAnchoredByBox_External(const FBox& WorldSpaceBox, bool bAnchored)
+void FGeometryCollectionPhysicsProxy::SetAnchoredByTransformedBox_External(const FBox& Box, const FTransform& Transform, bool bAnchored)
 {
 	check(IsInGameThread());
 	if (Chaos::FPhysicsSolver* RBDSolver = GetSolver<Chaos::FPhysicsSolver>())
 	{
-		Chaos::FAABB3 BoundsToCheck(WorldSpaceBox.Min, WorldSpaceBox.Max);
-		RBDSolver->EnqueueCommandImmediate([this, BoundsToCheck, bAnchored, RBDSolver]()
+		Chaos::FAABB3 BoundsToCheck(Box.Min, Box.Max);
+		Chaos::FRigidTransform3 BoxTransform(Transform);
+		RBDSolver->EnqueueCommandImmediate([this, BoundsToCheck, BoxTransform, bAnchored, RBDSolver]()
 			{
-				Chaos::FPBDRigidsEvolution* Evolution = RBDSolver->GetEvolution();
+				using namespace Chaos;
+				TSet<FPBDRigidClusteredParticleHandle*> TopParentHandles;
+
+				FPBDRigidsEvolution* Evolution = RBDSolver->GetEvolution();
 				for (FClusterHandle* ParticleHandle : GetSolverParticleHandles())
 				{
-					if (ParticleHandle && BoundsToCheck.Contains(ParticleHandle->X()))
+					const FVec3 PositionInBoxSpace = BoxTransform.InverseTransformPosition(ParticleHandle->X());
+					if (ParticleHandle && BoundsToCheck.Contains(PositionInBoxSpace))
 					{
-						SetParticleAnchored(Evolution, ParticleHandle, bAnchored);
+						SetParticleAnchored_Internal(Evolution, ParticleHandle, bAnchored);
+						if (Chaos::FPBDRigidClusteredParticleHandle* TopParentHandle = GetTopActiveClusteredParent_Internal(ParticleHandle))
+						{
+							TopParentHandles.Add(TopParentHandle);
+						}
 					}
 				}
+				// update the top parents state
+				for (FPBDRigidClusteredParticleHandle* TopParentHandle : TopParentHandles)
+				{
+					UpdateKinematicProperties(TopParentHandle, Evolution->GetRigidClustering().GetChildrenMap(), *Evolution);
+				}
 			});
 	}
 }
 void FGeometryCollectionPhysicsProxy::RemoveAllAnchors_External()
 {
 	check(IsInGameThread());
-
 	if (Chaos::FPhysicsSolver* RBDSolver = GetSolver<Chaos::FPhysicsSolver>())
 	{
 		RBDSolver->EnqueueCommandImmediate([this, RBDSolver]()
@@ -1790,7 +1820,9 @@ void FGeometryCollectionPhysicsProxy::RemoveAllAnchors_External()
 				Chaos::FPBDRigidsEvolution* Evolution = RBDSolver->GetEvolution();
 				for (FClusterHandle* ParticleHandle : GetSolverParticleHandles())
 				{
-					SetParticleAnchored(Evolution, ParticleHandle, false);
+					SetParticleAnchored_Internal(Evolution, ParticleHandle, false);
+					// we do not need to update the kinematic state since everything is now dynamic 
+					// todo(chaos) what about cluster groups ? 
 				}
 			});
 	}
diff --git a/Engine/Source/Runtime/Experimental/Chaos/Public/PhysicsProxy/GeometryCollectionPhysicsProxy.h b/Engine/Source/Runtime/Experimental/Chaos/Public/PhysicsProxy/GeometryCollectionPhysicsProxy.h
--- a/Engine/Source/Runtime/Experimental/Chaos/Public/PhysicsProxy/GeometryCollectionPhysicsProxy.h
+++ b/Engine/Source/Runtime/Experimental/Chaos/Public/PhysicsProxy/GeometryCollectionPhysicsProxy.h
@@ -377,7 +377,7 @@ public:
 	void BreakClusters_External(TArray<FGeometryCollectionItemIndex>&& ItemIndices);
 	void BreakActiveClusters_External();
 	void SetAnchoredByIndex_External(int32 Index, bool bAnchored);
-	void SetAnchoredByBox_External(const FBox& WorldSpaceBox, bool bAnchored);
+	void SetAnchoredByTransformedBox_External(const FBox& Box, const FTransform& Transform, bool bAnchored);
 	void RemoveAllAnchors_External();
 	void ApplyExternalStrain_External(FGeometryCollectionItemIndex ItemIndex, const FVector& WorldLocation, float Radius, int32 PropagationDepth, float PropagationFactor, float StrainValue);
 	void ApplyInternalStrain_External(FGeometryCollectionItemIndex ItemIndex, const FVector& WorldLocation, float Radius, int32 PropagationDepth, float PropagationFactor, float StrainValue);
diff --git a/Engine/Source/Runtime/Experimental/GeometryCollectionEngine/Private/GeometryCollection/GeometryCollectionComponent.cpp b/Engine/Source/Runtime/Experimental/GeometryCollectionEngine/Private/GeometryCollection/GeometryCollectionComponent.cpp
--- a/Engine/Source/Runtime/Experimental/GeometryCollectionEngine/Private/GeometryCollection/GeometryCollectionComponent.cpp
+++ b/Engine/Source/Runtime/Experimental/GeometryCollectionEngine/Private/GeometryCollection/GeometryCollectionComponent.cpp
@@ -1140,8 +1140,11 @@ void UGeometryCollectionComponent::SetEmbeddedGeometrySelectable(bool bSelectabl
 {
 	for (TObjectPtr<UInstancedStaticMeshComponent> EmbeddedGeometryComponent : EmbeddedGeometryComponents)
 	{
-		EmbeddedGeometryComponent->bSelectable = bSelectable;
-		EmbeddedGeometryComponent->bHasPerInstanceHitProxies = bSelectable;
+		if (EmbeddedGeometryComponent)
+		{
+			EmbeddedGeometryComponent->bSelectable = bSelectable;
+			EmbeddedGeometryComponent->bHasPerInstanceHitProxies = bSelectable;
+		}
 	}
 }
 
@@ -4322,13 +4325,17 @@ void UGeometryCollectionComponent::SetAnchoredByIndex(int32 Index, bool bAnchore
 }
 
 void UGeometryCollectionComponent::SetAnchoredByBox(FBox WorldSpaceBox, bool bAnchored)
+{
+	SetAnchoredByTransformedBox(WorldSpaceBox, FTransform::Identity, bAnchored);
+}
+
+void UGeometryCollectionComponent::SetAnchoredByTransformedBox(FBox Box, FTransform Transform, bool bAnchored)
 {
 	if (PhysicsProxy)
 	{
-		PhysicsProxy->SetAnchoredByBox_External(WorldSpaceBox, bAnchored);
+		PhysicsProxy->SetAnchoredByTransformedBox_External(Box, Transform, bAnchored);
 	}
 }
-
 void UGeometryCollectionComponent::RemoveAllAnchors()
 {
 	if (PhysicsProxy)
diff --git a/Engine/Source/Runtime/Experimental/GeometryCollectionEngine/Public/GeometryCollection/GeometryCollectionComponent.h b/Engine/Source/Runtime/Experimental/GeometryCollectionEngine/Public/GeometryCollection/GeometryCollectionComponent.h
--- a/Engine/Source/Runtime/Experimental/GeometryCollectionEngine/Public/GeometryCollection/GeometryCollectionComponent.h
+++ b/Engine/Source/Runtime/Experimental/GeometryCollectionEngine/Public/GeometryCollection/GeometryCollectionComponent.h
@@ -554,11 +554,17 @@ public:
 	void SetAnchoredByIndex(int32 Index, bool bAnchored);
 
 	/**
-	* Set all pieces within a world bounding space box to be anchored or not 
+	* Set all pieces within a world space bounding box to be anchored or not 
 	*/
 	UFUNCTION(BlueprintCallable, Category = "ChaosPhysics")
 	void SetAnchoredByBox(FBox WorldSpaceBox, bool bAnchored);
 
+	/**
+	* Set all pieces within a world transformed bounding box to be anchored or not
+	*/
+	UFUNCTION(BlueprintCallable, Category = "ChaosPhysics")
+	void SetAnchoredByTransformedBox(FBox Box, FTransform Transform, bool bAnchored);
+
 	/**
 	* this will remove anchors on all the pieces ( including the static and kinematic initial states ones ) of the geometry colection
 	*/
