Fix static analysis warning about recent changes to SetAnchoredByTransformedBox_External

diff --git a/Engine/Source/Runtime/Experimental/Chaos/Private/PhysicsProxy/GeometryCollectionPhysicsProxy.cpp b/Engine/Source/Runtime/Experimental/Chaos/Private/PhysicsProxy/GeometryCollectionPhysicsProxy.cpp
--- a/Engine/Source/Runtime/Experimental/Chaos/Private/PhysicsProxy/GeometryCollectionPhysicsProxy.cpp
+++ b/Engine/Source/Runtime/Experimental/Chaos/Private/PhysicsProxy/GeometryCollectionPhysicsProxy.cpp
@@ -1792,13 +1792,16 @@ void FGeometryCollectionPhysicsProxy::SetAnchoredByTransformedBox_External(const
 				FPBDRigidsEvolution* Evolution = RBDSolver->GetEvolution();
 				for (FClusterHandle* ParticleHandle : GetSolverParticleHandles())
 				{
-					const FVec3 PositionInBoxSpace = BoxTransform.InverseTransformPosition(ParticleHandle->X());
-					if (ParticleHandle && BoundsToCheck.Contains(PositionInBoxSpace))
+					if (ParticleHandle)
 					{
-						SetParticleAnchored_Internal(Evolution, ParticleHandle, bAnchored);
-						if (Chaos::FPBDRigidClusteredParticleHandle* TopParentHandle = GetTopActiveClusteredParent_Internal(ParticleHandle))
+						const FVec3 PositionInBoxSpace = BoxTransform.InverseTransformPosition(ParticleHandle->X());
+						if (BoundsToCheck.Contains(PositionInBoxSpace))
 						{
-							TopParentHandles.Add(TopParentHandle);
+							SetParticleAnchored_Internal(Evolution, ParticleHandle, bAnchored);
+							if (Chaos::FPBDRigidClusteredParticleHandle* TopParentHandle = GetTopActiveClusteredParent_Internal(ParticleHandle))
+							{
+								TopParentHandles.Add(TopParentHandle);
+							}
 						}
 					}
 				}
