Add blueprint library for dataflow and deprecated BP function on the dataflow object itself

diff --git a/Engine/Source/Runtime/Experimental/Dataflow/Engine/Private/Dataflow/DataflowBlueprintLibrary.cpp b/Engine/Source/Runtime/Experimental/Dataflow/Engine/Private/Dataflow/DataflowBlueprintLibrary.cpp
new file mode 100644
--- /dev/null
+++ b/Engine/Source/Runtime/Experimental/Dataflow/Engine/Private/Dataflow/DataflowBlueprintLibrary.cpp
@@ -0,0 +1,37 @@
+// Copyright Epic Games, Inc. All Rights Reserved.
+
+#include "Dataflow/DataflowBlueprintLibrary.h"
+
+#include "Dataflow/DataflowObjectInterface.h"
+
+#include UE_INLINE_GENERATED_CPP_BY_NAME(DataflowBlueprintLibrary)
+
+void UDataflowBlueprintLibrary::EvaluateTerminalNodeByName(UDataflow* Dataflow, FName TerminalNodeName, UObject* ResultAsset)
+{
+	if (Dataflow && Dataflow->Dataflow)
+	{
+		TSharedPtr<FDataflowNode> Node = Dataflow->Dataflow->FindTerminalNode(TerminalNodeName);
+		if (const FDataflowTerminalNode* TerminalNode = Node->AsType<const FDataflowTerminalNode>())
+		{
+			Dataflow::FEngineContext Context(ResultAsset, Dataflow, FPlatformTime::Cycles64());
+			TerminalNode->Evaluate(Context);
+			if (ResultAsset)
+			{
+				TerminalNode->SetAssetValue(ResultAsset, Context);
+			}
+		}
+	}
+}
+
+
+
+
+
+
+
+
+
+
+
+
+
diff --git a/Engine/Source/Runtime/Experimental/Dataflow/Engine/Private/Dataflow/DataflowObject.cpp b/Engine/Source/Runtime/Experimental/Dataflow/Engine/Private/Dataflow/DataflowObject.cpp
--- a/Engine/Source/Runtime/Experimental/Dataflow/Engine/Private/Dataflow/DataflowObject.cpp
+++ b/Engine/Source/Runtime/Experimental/Dataflow/Engine/Private/Dataflow/DataflowObject.cpp
@@ -8,11 +8,8 @@
 
 #include UE_INLINE_GENERATED_CPP_BY_NAME(DataflowObject)
 
-
 #define LOCTEXT_NAMESPACE "UDataflow"
 
-DEFINE_LOG_CATEGORY_STATIC(DATAFLOWOBJECT_LOG, Error, All);
-
 FDataflowAssetEdit::FDataflowAssetEdit(UDataflow* InAsset, FPostEditFunctionCallback InCallback)
 	: PostEditCallback(InCallback)
 	, Asset(InAsset)
@@ -40,23 +37,7 @@ UDataflow::UDataflow(const FObjectInitializer& ObjectInitializer)
 
 void UDataflow::EvaluateTerminalNodeByName(FName NodeName, UObject* Asset)
 {
-	if (Dataflow)
-	{
-		TSharedPtr<FDataflowNode> Node = Dataflow->FindTerminalNode(NodeName);
-		if (const FDataflowTerminalNode* TerminalNode = Node->AsType<const FDataflowTerminalNode>())
-		{
-			Dataflow::FEngineContext Context(Asset, this, FPlatformTime::Cycles64());
-			TerminalNode->Evaluate(Context);
-			if (Asset)
-			{
-				TerminalNode->SetAssetValue(Asset, Context);
-			}
-		}
-		else
-		{
-			UE_LOG(DATAFLOWOBJECT_LOG, Warning, TEXT("UDataflow::EvaluateTerminalNodeByName() : Could no find terminal node %s"), *NodeName.ToString());
-		}
-	}
+	ensureAlwaysMsgf(false, TEXT("Deprecated use the dataflow blueprint library from now on"));
 }
 
 void UDataflow::PostEditCallback()
diff --git a/Engine/Source/Runtime/Experimental/Dataflow/Engine/Public/Dataflow/DataflowBlueprintLibrary.h b/Engine/Source/Runtime/Experimental/Dataflow/Engine/Public/Dataflow/DataflowBlueprintLibrary.h
new file mode 100644
--- /dev/null
+++ b/Engine/Source/Runtime/Experimental/Dataflow/Engine/Public/Dataflow/DataflowBlueprintLibrary.h
@@ -0,0 +1,21 @@
+// Copyright Epic Games, Inc. All Rights Reserved.
+#pragma once
+
+#include "Kismet/BlueprintFunctionLibrary.h"
+#include "Engine/EngineTypes.h"
+#include "Dataflow/DataflowObject.h"
+
+#include "DataflowBlueprintLibrary.generated.h"
+
+UCLASS()
+class DATAFLOWENGINE_API UDataflowBlueprintLibrary : public UBlueprintFunctionLibrary
+{
+	GENERATED_BODY()
+
+public:
+	/**
+	* Find a specific terminal node by name evaluate it using a specific UObject
+	*/
+	UFUNCTION(BlueprintCallable, Category = "Dataflow", meta = (Keywords = "Dataflow graph"))
+	static void EvaluateTerminalNodeByName(UDataflow* Dataflow, FName TerminalNodeName, UObject* ResultAsset);
+};
\ No newline at end of file
diff --git a/Engine/Source/Runtime/Experimental/Dataflow/Engine/Public/Dataflow/DataflowObject.h b/Engine/Source/Runtime/Experimental/Dataflow/Engine/Public/Dataflow/DataflowObject.h
--- a/Engine/Source/Runtime/Experimental/Dataflow/Engine/Public/Dataflow/DataflowObject.h
+++ b/Engine/Source/Runtime/Experimental/Dataflow/Engine/Public/Dataflow/DataflowObject.h
@@ -69,7 +69,7 @@ public:
 	/** 
 	* Find all the node of a speific type and evaluate them using a specific UObject
 	*/
-	UFUNCTION(BlueprintCallable, Category = "Graph")
+	UE_DEPRECATED(5.1, "Use Blueprint library version of the function")
 	void EvaluateTerminalNodeByName(FName NodeName, UObject* Asset);
 
 public:
