Verification of Capsule and Box support functions when used with non uniform scaling

diff --git a/Engine/Source/Programs/HeadlessChaos/Private/HeadlessChaosTestGJK.cpp b/Engine/Source/Programs/HeadlessChaos/Private/HeadlessChaosTestGJK.cpp
--- a/Engine/Source/Programs/HeadlessChaos/Private/HeadlessChaosTestGJK.cpp
+++ b/Engine/Source/Programs/HeadlessChaos/Private/HeadlessChaosTestGJK.cpp
@@ -2100,4 +2100,24 @@ namespace ChaosTest
 		EXPECT_TRUE(GJKIntersection(BigBoxScaled, SmallBox, BToATM, FReal(0), Chaos::TVector<FReal, 3>(-16000, -16000, 500)));		
 
 	}
+
+
+	// Test capsule support functions when scaled
+	GTEST_TEST(GJKTests, GJK_TestSupportFunctions)
+	{
+		FVec3 PointA{ -1, 0, -100 };
+		FVec3 PointB{ 1, 0, 100 };
+		FVec3 ScaleX{ 100, 1, 1 };
+		TUniquePtr<Chaos::FCapsule> Capsule = MakeUnique<FCapsule>(PointA, PointB, 10);
+		TImplicitObjectScaled<Chaos::FCapsule> CapsuleScaled(MakeSerializable(Capsule), nullptr, ScaleX);
+
+		FVec3 SupportDir{0.1f, 0.0f, -1.0f}; // Pointing down and slightly to the right
+
+		int32 Vertex = INDEX_NONE;
+		FVec3 Support = Capsule->SupportCore(SupportDir,0, nullptr, Vertex);
+		EXPECT_EQ(Support, PointA); // Expect bottom point
+
+		Support = CapsuleScaled.SupportCore(SupportDir, 0, nullptr, Vertex);
+		EXPECT_EQ(Support, ScaleX * PointA); // Still expect bottom point (But scaled)
+	}
 }
\ No newline at end of file
diff --git a/Engine/Source/Runtime/Experimental/Chaos/Public/Chaos/AABB.h b/Engine/Source/Runtime/Experimental/Chaos/Public/Chaos/AABB.h
--- a/Engine/Source/Runtime/Experimental/Chaos/Public/Chaos/AABB.h
+++ b/Engine/Source/Runtime/Experimental/Chaos/Public/Chaos/AABB.h
@@ -382,7 +382,7 @@ namespace Chaos
 
 		FORCEINLINE_DEBUGGABLE TVector<T, d> SupportCoreScaled(const TVector<T, d>& Direction, const T InMargin, const TVector<T, d>& Scale, T* OutSupportDelta, int32& VertexIndex) const
 		{
-			const TVector<T, d> ScaledDirection = Direction * Scale;
+			const TVector<T, d> ScaledDirection = Direction * Scale; // Compensate for Negative scales, only the signs really matter here
 
 			TVector<T, d> ChosenPt;
 			FIntVector ChosenAxis;
diff --git a/Engine/Source/Runtime/Experimental/Chaos/Public/Chaos/Capsule.h b/Engine/Source/Runtime/Experimental/Chaos/Public/Chaos/Capsule.h
--- a/Engine/Source/Runtime/Experimental/Chaos/Public/Chaos/Capsule.h
+++ b/Engine/Source/Runtime/Experimental/Chaos/Public/Chaos/Capsule.h
@@ -387,6 +387,7 @@ namespace Chaos
 		FORCEINLINE FVec3 SupportCoreScaled(const FVec3& Direction, const FReal InMargin, const FVec3& Scale, FReal* OutSupportDelta, int32& VertexIndex) const
 		{
 			// NOTE: Ignores InMargin, assumes Radius
+			// Note: Scaling the direction vector like this, might not seem quite right, but works due to the commutativity of the single dot product that follows
 			return SupportCore(Scale * Direction, GetMargin(), OutSupportDelta, VertexIndex) * Scale;
 		}
 
diff --git a/Engine/Source/Runtime/Experimental/Chaos/Public/Chaos/ImplicitObjectScaled.h b/Engine/Source/Runtime/Experimental/Chaos/Public/Chaos/ImplicitObjectScaled.h
--- a/Engine/Source/Runtime/Experimental/Chaos/Public/Chaos/ImplicitObjectScaled.h
+++ b/Engine/Source/Runtime/Experimental/Chaos/Public/Chaos/ImplicitObjectScaled.h
@@ -160,24 +160,24 @@ public:
 	}
 
 	// The support position from the specified direction
-	FORCEINLINE FVec3 Support(const FVec3& Direction, const FReal Thickness, int32& VertexIndex) const
+	FORCEINLINE_DEBUGGABLE FVec3 Support(const FVec3& Direction, const FReal Thickness, int32& VertexIndex) const
 	{
 		return MObject->Support(Direction, Thickness, VertexIndex); 
 	}
 
 	// this shouldn't be called, but is required until we remove the explicit function implementations in CollisionResolution.cpp
-	FORCEINLINE FVec3 SupportScaled(const FVec3& Direction, const T Thickness, const FVec3& Scale, int32& VertexIndex) const
+	FORCEINLINE_DEBUGGABLE FVec3 SupportScaled(const FVec3& Direction, const T Thickness, const FVec3& Scale, int32& VertexIndex) const
 	{
 		return MObject->SupportScaled(Direction, Thickness, Scale, VertexIndex);
 	}
 
 	// The support position from the specified direction, if the shape is reduced by the margin
-	FORCEINLINE FVec3 SupportCore(const FVec3& Direction, const FReal InMargin, FReal* OutSupportDelta, int32& VertexIndex) const
+	FORCEINLINE_DEBUGGABLE FVec3 SupportCore(const FVec3& Direction, const FReal InMargin, FReal* OutSupportDelta, int32& VertexIndex) const
 	{
 		return MObject->SupportCore(Direction, InMargin, OutSupportDelta, VertexIndex);
 	}
 
-	FORCEINLINE VectorRegister4Float SupportCoreSimd(const VectorRegister4Float& Direction, const FReal InMargin) const
+	FORCEINLINE_DEBUGGABLE VectorRegister4Float SupportCoreSimd(const VectorRegister4Float& Direction, const FReal InMargin) const
 	{
 		FVec3 DirectionVec3;
 		VectorStoreFloat3(Direction, &DirectionVec3);
