performance : fix InitializeSharedCollisionStructures  performance for GC with a high number of triangles

diff --git a/Engine/Source/Runtime/Experimental/Chaos/Private/PhysicsProxy/GeometryCollectionPhysicsProxy.cpp b/Engine/Source/Runtime/Experimental/Chaos/Private/PhysicsProxy/GeometryCollectionPhysicsProxy.cpp
--- a/Engine/Source/Runtime/Experimental/Chaos/Private/PhysicsProxy/GeometryCollectionPhysicsProxy.cpp
+++ b/Engine/Source/Runtime/Experimental/Chaos/Private/PhysicsProxy/GeometryCollectionPhysicsProxy.cpp
@@ -3062,7 +3062,7 @@ static FGeometryDynamicCollection::FSharedImplicit CreateImplicitGeometry(
 }
 
 static void ComputeGeometryVolumeAndCenterOfMass(
-	const TArray<FVector3f> Vertices,
+	const TArray<FVector3f>& Vertices,
 	const Chaos::FTriangleMesh& TriMesh,
 	const FBox& BoundingBox,
 	const FSharedSimulationParameters& SharedParams, 
