Optimize geaometry collection creation from geometry sources
- Only call reindexMaterial once
- Make sure reindexMaterial  skips validation when calling RemoveElements

diff --git a/Engine/Plugins/Experimental/GeometryCollectionPlugin/Source/GeometryCollectionNodes/Private/Dataflow/Nodes/GeometryCollectionAssetNodes.cpp b/Engine/Plugins/Experimental/GeometryCollectionPlugin/Source/GeometryCollectionNodes/Private/Dataflow/Nodes/GeometryCollectionAssetNodes.cpp
--- a/Engine/Plugins/Experimental/GeometryCollectionPlugin/Source/GeometryCollectionNodes/Private/Dataflow/Nodes/GeometryCollectionAssetNodes.cpp
+++ b/Engine/Plugins/Experimental/GeometryCollectionPlugin/Source/GeometryCollectionNodes/Private/Dataflow/Nodes/GeometryCollectionAssetNodes.cpp
@@ -115,10 +115,15 @@ void FCreateGeometryCollectionFromSourcesDataflowNode::Evaluate(Dataflow::FConte
 
 	FGeometryCollection OutCollection;
 	TArray<TObjectPtr<UMaterial>> OutMaterials;
+	constexpr bool bReindexMaterialsInLoop = false;
 	for (const FGeometryCollectionSource& Source: InSources)
 	{
 		// todo: change AppendGeometryCollectionSource to take a FManagedArrayCollection so we could move the collection when assigning it to the output
-		FGeometryCollectionEngineConversion::AppendGeometryCollectionSource(Source, OutCollection, OutMaterials);
+		FGeometryCollectionEngineConversion::AppendGeometryCollectionSource(Source, OutCollection, OutMaterials, bReindexMaterialsInLoop);
+	}
+	if (bReindexMaterialsInLoop == false)
+	{
+		OutCollection.ReindexMaterials();
 	}
 
 	// we have to make a copy since we have generated a FGeometryCollection which is inherited from FManagedArrayCollection
diff --git a/Engine/Source/Runtime/Experimental/Chaos/Private/GeometryCollection/GeometryCollection.cpp b/Engine/Source/Runtime/Experimental/Chaos/Private/GeometryCollection/GeometryCollection.cpp
--- a/Engine/Source/Runtime/Experimental/Chaos/Private/GeometryCollection/GeometryCollection.cpp
+++ b/Engine/Source/Runtime/Experimental/Chaos/Private/GeometryCollection/GeometryCollection.cpp
@@ -427,10 +427,13 @@ void FGeometryCollection::ReindexMaterials(FManagedArrayCollection& InCollection
 	TManagedArray<FGeometryCollectionSection>& Sections = InCollection.ModifyAttribute<FGeometryCollectionSection>("Sections", FGeometryCollection::MaterialGroup);
 	TManagedArray<int32>& MaterialIndex = InCollection.ModifyAttribute<int32>("MaterialIndex", FGeometryCollection::FacesGroup);
 
+	FManagedArrayCollection::FProcessingParameters ProcessingParams;
+	ProcessingParams.bDoValidation = false; // disable validation as this can be very costly ( only in editor )
+
 	// clear all sections	
 	TArray<int32> DelSections;
 	GeometryCollectionAlgo::ContiguousArray(DelSections, InCollection.NumElements(FGeometryCollection::MaterialGroup));
-	InCollection.RemoveElements(FGeometryCollection::MaterialGroup, DelSections);
+	InCollection.RemoveElements(FGeometryCollection::MaterialGroup, DelSections, ProcessingParams);
 	DelSections.Reset(0);
 
 
@@ -481,10 +484,13 @@ void FGeometryCollection::ReindexMaterials(FManagedArrayCollection& InCollection
 	}
 
 	// remap indices so the materials appear to be grouped
+	const int32 NumSections = InCollection.NumElements(FGeometryCollection::MaterialGroup);
+	const int32 NumFaceElements = InCollection.NumElements(FGeometryCollection::FacesGroup);
+
 	int Idx = 0;
-	for (int Section=0; Section < InCollection.NumElements(FGeometryCollection::MaterialGroup); Section++)
+	for (int Section=0; Section < NumSections; Section++)
 	{
-		for (int FaceElement = 0; FaceElement < InCollection.NumElements(FGeometryCollection::FacesGroup); FaceElement++)
+		for (int FaceElement = 0; FaceElement < NumFaceElements; FaceElement++)
 		{
 			int32 ID = (MaterialID)[FaceElement];
 	
@@ -498,7 +504,7 @@ void FGeometryCollection::ReindexMaterials(FManagedArrayCollection& InCollection
 	// delete unused material sections
 	if (DelSections.Num())
 	{
-		InCollection.RemoveElements(FGeometryCollection::MaterialGroup, DelSections);
+		InCollection.RemoveElements(FGeometryCollection::MaterialGroup, DelSections, ProcessingParams);
 	}
 }
 
diff --git a/Engine/Source/Runtime/Experimental/GeometryCollectionEngine/Private/GeometryCollection/GeometryCollectionEngineConversion.cpp b/Engine/Source/Runtime/Experimental/GeometryCollectionEngine/Private/GeometryCollection/GeometryCollectionEngineConversion.cpp
--- a/Engine/Source/Runtime/Experimental/GeometryCollectionEngine/Private/GeometryCollection/GeometryCollectionEngineConversion.cpp
+++ b/Engine/Source/Runtime/Experimental/GeometryCollectionEngine/Private/GeometryCollection/GeometryCollectionEngineConversion.cpp
@@ -1305,7 +1305,7 @@ int32 FGeometryCollectionEngineConversion::AppendSkeletalMeshMaterials(const USk
 	return MaterialStart;
 }
 
-void FGeometryCollectionEngineConversion::AppendGeometryCollectionSource(const FGeometryCollectionSource& GeometryCollectionSource, FGeometryCollection& GeometryCollectionInOut, TArray<UMaterial*>& MaterialsInOut)
+void FGeometryCollectionEngineConversion::AppendGeometryCollectionSource(const FGeometryCollectionSource& GeometryCollectionSource, FGeometryCollection& GeometryCollectionInOut, TArray<UMaterial*>& MaterialsInOut, bool ReindexMaterials)
 {
 	if (const UObject* SourceObject = GeometryCollectionSource.SourceGeometryObject.ResolveObject())
 	{
@@ -1319,7 +1319,10 @@ void FGeometryCollectionEngineConversion::AppendGeometryCollectionSource(const F
 				SourceStaticMesh,
 				StartMaterialIndex,
 				GeometryCollectionSource.LocalTransform,
-				&GeometryCollectionInOut
+				&GeometryCollectionInOut,
+				ReindexMaterials,
+				GeometryCollectionSource.bAddInternalMaterials,
+				GeometryCollectionSource.bSplitComponents
 				);
 		}
 		else if (const USkeletalMesh* SourceSkeletalMesh = Cast<USkeletalMesh>(SourceObject))
@@ -1328,7 +1331,8 @@ void FGeometryCollectionEngineConversion::AppendGeometryCollectionSource(const F
 				SourceSkeletalMesh,
 				StartMaterialIndex,
 				GeometryCollectionSource.LocalTransform,
-				&GeometryCollectionInOut
+				&GeometryCollectionInOut,
+				ReindexMaterials
 				);
 		}
 		else if (const UGeometryCollection* SourceGeometryCollection = Cast<UGeometryCollection>(SourceObject))
@@ -1337,7 +1341,8 @@ void FGeometryCollectionEngineConversion::AppendGeometryCollectionSource(const F
 				SourceGeometryCollection->GetGeometryCollection().Get(),
 				StartMaterialIndex,
 				GeometryCollectionSource.LocalTransform,
-				&GeometryCollectionInOut
+				&GeometryCollectionInOut,
+				ReindexMaterials
 				);
 		}
 	}
diff --git a/Engine/Source/Runtime/Experimental/GeometryCollectionEngine/Public/GeometryCollection/GeometryCollectionEngineConversion.h b/Engine/Source/Runtime/Experimental/GeometryCollectionEngine/Public/GeometryCollection/GeometryCollectionEngineConversion.h
--- a/Engine/Source/Runtime/Experimental/GeometryCollectionEngine/Public/GeometryCollection/GeometryCollectionEngineConversion.h
+++ b/Engine/Source/Runtime/Experimental/GeometryCollectionEngine/Public/GeometryCollection/GeometryCollectionEngineConversion.h
@@ -164,5 +164,5 @@ public:
 	*  @param GeometryCollectionInOut : GeometryCollection to append to 
 	*  @param MaterialsInOut  : array of materials to append to
 	*/
-	static void AppendGeometryCollectionSource(const FGeometryCollectionSource& GeometryCollectionSource, FGeometryCollection& GeometryCollectionInOut, TArray<UMaterial*>& MaterialsInOut);
+	static void AppendGeometryCollectionSource(const FGeometryCollectionSource& GeometryCollectionSource, FGeometryCollection& GeometryCollectionInOut, TArray<UMaterial*>& MaterialsInOut, bool ReindexMaterials = true);
 };
\ No newline at end of file
