Overlap Query TriangleMesh Capsule Narrow Phase Optimization

diff --git a/Engine/Source/Runtime/Experimental/Chaos/Private/Chaos/TriangleMeshImplicitObject.cpp b/Engine/Source/Runtime/Experimental/Chaos/Private/Chaos/TriangleMeshImplicitObject.cpp
--- a/Engine/Source/Runtime/Experimental/Chaos/Private/Chaos/TriangleMeshImplicitObject.cpp
+++ b/Engine/Source/Runtime/Experimental/Chaos/Private/Chaos/TriangleMeshImplicitObject.cpp
@@ -434,6 +434,234 @@ struct FTriangleMeshOverlapVisitorNoMTD
 	const FTriangleMeshImplicitObject* TriMesh;
 	bool bFoundIntersection;
 };
+namespace
+{
+bool ComputeCapsuleTriangleCollisionSimd(const VectorRegister4Float& A, const VectorRegister4Float& B, const VectorRegister4Float& C, const VectorRegister4Float& X1, const VectorRegister4Float& X2, FRealSingle Radius)
+{
+
+	const VectorRegister4Float AB = VectorSubtract(B, A);
+	const VectorRegister4Float BC = VectorSubtract(C, B);
+
+	VectorRegister4Float Normal = VectorNormalize(VectorCross(AB, BC));
+
+
+	// Plane Triangle
+	const VectorRegister4Float AX1 = VectorSubtract(X1, A);
+	const VectorRegister4Float AX2 = VectorSubtract(X2, A);
+
+	FRealSingle AX1Dist = VectorDot3Scalar(AX1, Normal);
+	FRealSingle AX2Dist = VectorDot3Scalar(AX2, Normal);
+
+	if (FMath::Sign(AX1Dist) == FMath::Sign(AX2Dist))
+	{
+		FRealSingle ClosestDist = FMath::Min<FRealSingle>(FMath::Abs<FRealSingle>(AX1Dist), FMath::Abs<FRealSingle>(AX2Dist));
+		if (ClosestDist > Radius)
+		{
+			return false;
+		}
+	}
+
+	// Edges
+	constexpr FRealSingle ThirdFloat = 1.0f / 3.0f;
+	constexpr VectorRegister4Float Third = MakeVectorRegisterFloatConstant(ThirdFloat, ThirdFloat, ThirdFloat, ThirdFloat);
+	const VectorRegister4Float Centroid = VectorMultiply(VectorAdd(VectorAdd(A, B), C), Third);
+
+	const VectorRegister4Float X2X1 = VectorSubtract(X1, X2);
+	const VectorRegister4Float CA = VectorSubtract(A, C);
+	const VectorRegister4Float CentroidA = VectorSubtract(A, Centroid);
+	{
+		VectorRegister4Float EdgeSeparationAxis = VectorNormalize(VectorCross(X2X1, CA));
+		EdgeSeparationAxis = VectorSelect(VectorCompareGT(VectorZeroFloat(), VectorDot3(CentroidA, EdgeSeparationAxis)), VectorNegate(EdgeSeparationAxis), EdgeSeparationAxis);
+
+		const FRealSingle EdgeSeparationDist = VectorDot3Scalar(VectorSubtract(X1, C), EdgeSeparationAxis);
+		if (EdgeSeparationDist > Radius)
+		{
+			return false;
+		}
+	}
+	const VectorRegister4Float CentroidB = VectorSubtract(B, Centroid);
+	{
+		VectorRegister4Float EdgeSeparationAxis = VectorNormalize(VectorCross(X2X1, AB));
+		EdgeSeparationAxis = VectorSelect(VectorCompareGT(VectorZeroFloat(), VectorDot3(CentroidB, EdgeSeparationAxis)), VectorNegate(EdgeSeparationAxis), EdgeSeparationAxis);
+
+		const FRealSingle EdgeSeparationDist = VectorDot3Scalar(VectorSubtract(X1, A), EdgeSeparationAxis);
+		if (EdgeSeparationDist > Radius)
+		{
+			return false;
+		}
+	}
+	const VectorRegister4Float CentroidC = VectorSubtract(C, Centroid);
+	{
+		VectorRegister4Float EdgeSeparationAxis = VectorNormalize(VectorCross(X2X1, BC));
+		EdgeSeparationAxis = VectorSelect(VectorCompareGT(VectorZeroFloat(), VectorDot3(CentroidC, EdgeSeparationAxis)), VectorNegate(EdgeSeparationAxis), EdgeSeparationAxis);
+
+		const FRealSingle EdgeSeparationDist = VectorDot3Scalar(VectorSubtract(X1, B), EdgeSeparationAxis);
+		if (EdgeSeparationDist > Radius)
+		{
+			return false;
+		}
+	}
+
+	// Triangle Vertex
+	VectorRegister4Float SqrX2X1 = VectorDot3(X2X1, X2X1);
+	FRealSingle SqrRadius = Radius * Radius;
+	{
+		VectorRegister4Float TimeA = VectorClamp(VectorDivide(VectorDot3(X2X1, AX1), SqrX2X1), VectorZeroFloat(), VectorOneFloat());
+		VectorRegister4Float PA = VectorMultiplyAdd(X1, VectorSubtract(VectorOneFloat(), TimeA), VectorMultiply(X2, TimeA));
+
+		if (VectorDot3Scalar(VectorSubtract(PA, A), CentroidA) > 0.0f)
+		{
+			const VectorRegister4Float APA = VectorSubtract(PA, A);
+			if (VectorDot3Scalar(APA, APA) > SqrRadius)
+			{
+				return false;
+			}
+		}
+	}
+	{
+		const VectorRegister4Float TimeB = VectorClamp(VectorDivide(VectorDot3(X2X1, VectorSubtract(X1, B)), SqrX2X1), VectorZeroFloat(), VectorOneFloat());
+		const VectorRegister4Float PB = VectorMultiplyAdd(X1, VectorSubtract(VectorOneFloat(), TimeB), VectorMultiply(X2, TimeB));
+
+		if (VectorDot3Scalar(VectorSubtract(PB, B), CentroidB) > 0.0f)
+		{
+			const VectorRegister4Float BPB = VectorSubtract(PB, B);
+			if (VectorDot3Scalar(BPB, BPB) > SqrRadius)
+			{
+				return false;
+			}
+		}
+	}
+	{
+		const VectorRegister4Float TimeC = VectorClamp(VectorDivide(VectorDot3(X2X1, VectorSubtract(X1, C)), SqrX2X1), VectorZeroFloat(), VectorOneFloat());
+		const VectorRegister4Float PC = VectorMultiplyAdd(X1, VectorSubtract(VectorOneFloat(), TimeC), VectorMultiply(X2, TimeC));
+
+		if (VectorDot3Scalar(VectorSubtract(PC, C), CentroidC) > 0.0f)
+		{
+			const VectorRegister4Float CPC = VectorSubtract(PC, C);
+			if (VectorDot3Scalar(CPC, CPC) > SqrRadius)
+			{
+				return false;
+			}
+		}
+	}
+
+	return true;
+}
+
+template <typename IdxType, typename ParticlesType>
+inline void TriangleMeshTransformVertsHelperSimd(const VectorRegister4Float& TriMeshScale, int32 TriIdx, const ParticlesType& Particles,
+	const TArray<TVector<IdxType, 3>>& Elements, VectorRegister4Float& OutA, VectorRegister4Float& OutB, VectorRegister4Float& OutC)
+{
+	OutA = VectorLoadFloat3(&Particles.X(Elements[TriIdx][0]).X);
+	OutB = VectorLoadFloat3(&Particles.X(Elements[TriIdx][1]).X);
+	OutC = VectorLoadFloat3(&Particles.X(Elements[TriIdx][2]).X);
+	OutA = VectorMultiply(OutA, TriMeshScale);
+	OutB = VectorMultiply(OutB, TriMeshScale);
+	OutC = VectorMultiply(OutC, TriMeshScale);
+}
+
+}
+
+template <>
+struct FTriangleMeshOverlapVisitorNoMTD<FCapsule>
+{
+	FTriangleMeshOverlapVisitorNoMTD(const TRigidTransform<FReal, 3>& WorldScaleQueryTM, const FCapsule& InQueryGeom, FReal InThickness, const FVec3f& InTriMeshScale, const FTriangleMeshImplicitObject* InTriMesh)
+		: Radius(FRealSingle(InQueryGeom.GetRadius() + InThickness))
+		, TriMesh(InTriMesh)
+		, bFoundIntersection(false)
+	{
+		FVec3f X1f = WorldScaleQueryTM.TransformPosition(InQueryGeom.GetX1());
+		FVec3f X2f = WorldScaleQueryTM.TransformPosition(InQueryGeom.GetX2());
+		X1 = VectorLoadFloat3(&X1f.X);
+		X2 = VectorLoadFloat3(&X2f.X);
+		TriMeshScale = VectorLoadFloat3(&InTriMeshScale.X);
+	}
+
+	bool VisitOverlap(int32 TriIdx)
+	{
+		VectorRegister4Float A, B, C;
+		if (TriMesh->MElements.RequiresLargeIndices())
+		{
+			TriangleMeshTransformVertsHelperSimd(TriMeshScale, TriIdx, TriMesh->MParticles, TriMesh->MElements.GetLargeIndexBuffer(), A, B, C);
+		}
+		else
+		{
+			TriangleMeshTransformVertsHelperSimd(TriMeshScale, TriIdx, TriMesh->MParticles, TriMesh->MElements.GetSmallIndexBuffer(), A, B, C);
+		}
+
+		bFoundIntersection = ComputeCapsuleTriangleCollisionSimd(A, B, C, X1, X2, Radius);
+		return !bFoundIntersection;
+	}
+
+	bool VisitSweep(int32 Instance, FRealSingle& CurDataLength)
+	{
+		check(false);
+		return true;
+	}
+	bool VisitRaycast(int32 Instance, FRealSingle& CurDataLength)
+	{
+		check(false);
+		return true;
+	}
+
+	FRealSingle Radius;
+	VectorRegister4Float X1;
+	VectorRegister4Float X2;
+	VectorRegister4Float TriMeshScale;
+	const FTriangleMeshImplicitObject* TriMesh;
+	bool bFoundIntersection;
+};
+
+
+template <>
+struct FTriangleMeshOverlapVisitorNoMTD<TImplicitObjectScaled<FCapsule>>
+{
+	FTriangleMeshOverlapVisitorNoMTD(const TRigidTransform<FReal, 3>& WorldScaleQueryTM, const TImplicitObjectScaled< FCapsule >& InQueryGeom, FReal InThickness, const FVec3f& InTriMeshScale, const FTriangleMeshImplicitObject* InTriMesh)
+		: Radius(FRealSingle(InQueryGeom.GetRadius()+ InThickness))
+		, TriMesh(InTriMesh)
+		, bFoundIntersection(false)
+	{
+		FVec3f X1f = WorldScaleQueryTM.TransformPosition(InQueryGeom.GetUnscaledObject()->GetX1() * InQueryGeom.GetScale());
+		FVec3f X2f = WorldScaleQueryTM.TransformPosition(InQueryGeom.GetUnscaledObject()->GetX2() * InQueryGeom.GetScale());
+		X1 = VectorLoadFloat3(&X1f.X);
+		X2 = VectorLoadFloat3(&X2f.X);
+		TriMeshScale = VectorLoadFloat3(&InTriMeshScale.X);
+	}
+
+	bool VisitOverlap(int32 TriIdx)
+	{
+		VectorRegister4Float A, B, C;
+		if (TriMesh->MElements.RequiresLargeIndices())
+		{
+			TriangleMeshTransformVertsHelperSimd(TriMeshScale, TriIdx, TriMesh->MParticles, TriMesh->MElements.GetLargeIndexBuffer(), A, B, C);
+		}
+		else
+		{
+			TriangleMeshTransformVertsHelperSimd(TriMeshScale, TriIdx, TriMesh->MParticles, TriMesh->MElements.GetSmallIndexBuffer(), A, B, C);
+		}
+
+		bFoundIntersection = ComputeCapsuleTriangleCollisionSimd(A, B, C, X1, X2, Radius);
+		return !bFoundIntersection;
+	}
+
+	bool VisitSweep(int32 Instance, FRealSingle& CurDataLength)
+	{
+		check(false);
+		return true;
+	}
+	bool VisitRaycast(int32 Instance, FRealSingle& CurDataLength)
+	{
+		check(false);
+		return true;
+	}
+
+	FRealSingle Radius;
+	VectorRegister4Float X1;
+	VectorRegister4Float X2;
+	VectorRegister4Float TriMeshScale;
+	const FTriangleMeshImplicitObject* TriMesh;
+	bool bFoundIntersection;
+};
 
 template <typename QueryGeomType>
 bool FTrimeshBVH::FindAllIntersectionsNoMTD(const FAABB3& Intersection, const TRigidTransform<FReal, 3>& Transform, const QueryGeomType& QueryGeom, FReal Thickness, const FVec3& TriMeshScale, const FTriangleMeshImplicitObject* TriMesh) const
