Dataflow:
- Added IsConnected() to DataflowNode to check if input connected in Evaluate()

diff --git a/Engine/Source/Runtime/Experimental/Dataflow/Core/Public/Dataflow/DataflowNode.h b/Engine/Source/Runtime/Experimental/Dataflow/Core/Public/Dataflow/DataflowNode.h
--- a/Engine/Source/Runtime/Experimental/Dataflow/Core/Public/Dataflow/DataflowNode.h
+++ b/Engine/Source/Runtime/Experimental/Dataflow/Core/Public/Dataflow/DataflowNode.h
@@ -182,6 +182,19 @@ struct DATAFLOWCORE_API FDataflowNode
 		FindOutput(Reference)->template SetValue<T>(Value, Context);
 	}
 
+	/**
+	*   IsConnected(...)
+	*
+	*	Checks if Reference input is connected.
+	*
+	*   @param Reference : Pointer to a member of this node that corresponds with the input.
+	*/
+	template<class T> bool IsConnected(const T* Reference) const
+	{
+		checkSlow(FindInput(Reference));
+		return FindInput(Reference) != nullptr;
+	}
+
 	void Invalidate();
 
 	virtual bool ValidateConnections();
