Height Field Capsule Optimization

diff --git a/Engine/Source/Programs/HeadlessChaos/Private/HeadlessChaosTestHeightfield.cpp b/Engine/Source/Programs/HeadlessChaos/Private/HeadlessChaosTestHeightfield.cpp
--- a/Engine/Source/Programs/HeadlessChaos/Private/HeadlessChaosTestHeightfield.cpp
+++ b/Engine/Source/Programs/HeadlessChaos/Private/HeadlessChaosTestHeightfield.cpp
@@ -949,6 +949,7 @@ namespace ChaosTest {
 		// Long box
 		{
 			TBox<FReal, 3> Box(FVec3(-1.0, -1.0, -1.0), FVec3(1.0, 1.0, 10.0));
+			FCapsule Capsule(FVec3(0.0, 0.0, 0.0), FVec3(0.0, 0.0, 9.0), 1.14);
 			for (int32 Row = 0; Row < Rows; ++Row)
 			{
 				for (int32 Col = 0; Col < Columns; ++Col)
@@ -956,20 +957,23 @@ namespace ChaosTest {
 					const FVec3 Translation(Col, Row, 1.5);
 					FRigidTransform3 QueryTM(Translation, TRotation<FReal, 3>::Identity);
 					FVec3 Dir(1, 0, 0);
-					bool Result = Heightfield.OverlapGeom(Box, QueryTM, 0.0, nullptr);
+					bool BoxResult = Heightfield.OverlapGeom(Box, QueryTM, 0.0, nullptr);
+					bool CapsuleResult = Heightfield.OverlapGeom(Capsule, QueryTM, 0.0, nullptr);
 					// No collision on the side of the mountain
 					if ((Col < 3 || Col > 7) && (Row < 3 || Row > 7))
 					{
-						EXPECT_FALSE(Result);
+						EXPECT_FALSE(BoxResult);
+						EXPECT_FALSE(CapsuleResult);
 					}
 					// Collision with the mountain
 					else if ((Col >= 3 && Col <= 7) && (Row >= 3 && Row <= 7))
 					{
-						EXPECT_TRUE(Result);
+						EXPECT_TRUE(BoxResult);
+						EXPECT_TRUE(CapsuleResult);
 					}
 					else
 					{
-						EXPECT_FALSE(Result);
+						EXPECT_FALSE(BoxResult);
 					}
 				}
 			}
@@ -977,6 +981,7 @@ namespace ChaosTest {
 		// Box rotated in X
 		{
 			TBox<FReal, 3> Box(FVec3(-1.0, -1.0, -20.0), FVec3(1.0, 1.0, 20.0));
+			FCapsule Capsule(FVec3(0.0, 0.0, -19.0), FVec3(0.0, 0.0, 19.0), 1.14);
 			for (int32 Row = 0; Row < Rows; ++Row)
 			{
 				for (int32 Col = 0; Col < Columns; ++Col)
@@ -984,15 +989,18 @@ namespace ChaosTest {
 					const FVec3 Translation(Col, Row, 2.0);
 					FRigidTransform3 QueryTM(Translation, TRotation<FReal, 3>(UE::Math::TQuat<FReal>(TVector<FReal, 3>(1.0, 0.0, 0.0), 3.14159/2.0)));
 					FVec3 Dir(1, 0, 0);
-					bool Result = Heightfield.OverlapGeom(Box, QueryTM, 0.0, nullptr);
+					bool BoxResult = Heightfield.OverlapGeom(Box, QueryTM, 0.0, nullptr);
+					bool CapsuleResult = Heightfield.OverlapGeom(Capsule, QueryTM, 0.0, nullptr);
 					// Collision with the mountain
 					if (Col >= 3 && Col <= 7)
 					{
-						EXPECT_TRUE(Result);
+						EXPECT_TRUE(BoxResult);
+						EXPECT_TRUE(CapsuleResult);
 					}
 					else
 					{
-						EXPECT_FALSE(Result);
+						EXPECT_FALSE(BoxResult);
+						EXPECT_FALSE(CapsuleResult);
 					}
 				}
 			}
@@ -1000,6 +1008,7 @@ namespace ChaosTest {
 		// Box rotated in Y
 		{
 			TBox<FReal, 3> Box(FVec3(-1.0, -1.0, -20.0), FVec3(1.0, 1.0, 20.0));
+			FCapsule Capsule(FVec3(0.0, 0.0, -19.0), FVec3(0.0, 0.0, 19.0), 1.14);
 			for (int32 Row = 0; Row < Rows; ++Row)
 			{
 				for (int32 Col = 0; Col < Columns; ++Col)
@@ -1007,15 +1016,18 @@ namespace ChaosTest {
 					const FVec3 Translation(Col, Row, 2.0);
 					FRigidTransform3 QueryTM(Translation, TRotation<FReal, 3>(UE::Math::TQuat<FReal>(TVector<FReal, 3>(0.0, 1.0, 0.0), 3.14159 / 2.0)));
 					FVec3 Dir(1, 0, 0);
-					bool Result = Heightfield.OverlapGeom(Box, QueryTM, 0.0, nullptr);
+					bool BoxResult = Heightfield.OverlapGeom(Box, QueryTM, 0.0, nullptr);
+					bool CapsuleResult = Heightfield.OverlapGeom(Capsule, QueryTM, 0.0, nullptr);
 					// Collision with the mountain
 					if (Row >= 3 && Row <= 7)
 					{
-						EXPECT_TRUE(Result);
+						EXPECT_TRUE(BoxResult);
+						EXPECT_TRUE(CapsuleResult);
 					}
 					else
 					{
-						EXPECT_FALSE(Result);
+						EXPECT_FALSE(BoxResult);
+						EXPECT_FALSE(CapsuleResult);
 					}
 				}
 			}
@@ -1023,6 +1035,7 @@ namespace ChaosTest {
 		// Thin Box
 		{
 			TBox<FReal, 3> Box(FVec3(-10.0, -10.0, -0.001), FVec3(10.0, 10.0, 0.001));
+			FCapsule Capsule(FVec3(0.0, 0.0, -0.001), FVec3(0.0, 0.0, 0.001), 11.4);
 			for (int32 Row = 0; Row < Rows; ++Row)
 			{
 				for (int32 Col = 0; Col < Columns; ++Col)
@@ -1030,15 +1043,17 @@ namespace ChaosTest {
 					const FVec3 Translation(Col, Row, 2.0);
 					FRigidTransform3 QueryTM(Translation, TRotation<FReal, 3>(UE::Math::TQuat<FReal>(TVector<FReal, 3>(1.0, 0.0, 0.0), 0.0)));
 					FVec3 Dir(1, 0, 0);
-					bool Result = Heightfield.OverlapGeom(Box, QueryTM, 0.0, nullptr);
-					// Collision with the mountain
-					EXPECT_TRUE(Result);
+					bool BoxResult = Heightfield.OverlapGeom(Box, QueryTM, 0.0, nullptr);
+					bool CapsuleResult = Heightfield.OverlapGeom(Capsule, QueryTM, 0.0, nullptr);
+					EXPECT_TRUE(BoxResult);
+					EXPECT_TRUE(CapsuleResult);
 				}
 			}
 		}
 		// Thin Box on the top
 		{
 			TBox<FReal, 3> Box(FVec3(-10.0, -10.0, -0.001), FVec3(10.0, 10.0, 0.001));
+			FCapsule Capsule(FVec3(0.0, 0.0, -0.001), FVec3(0.0, 0.0, 0.001), 11.4);
 			for (int32 Row = 0; Row < Rows; ++Row)
 			{
 				for (int32 Col = 0; Col < Columns; ++Col)
@@ -1046,9 +1061,10 @@ namespace ChaosTest {
 					const FVec3 Translation(Col, Row, 9.5);
 					FRigidTransform3 QueryTM(Translation, TRotation<FReal, 3>(UE::Math::TQuat<FReal>(TVector<FReal, 3>(1.0, 0.0, 0.0), 0.0)));
 					FVec3 Dir(1, 0, 0);
-					bool Result = Heightfield.OverlapGeom(Box, QueryTM, 0.0, nullptr);
-					// Collision with the mountain
-					EXPECT_TRUE(Result);
+					bool BoxResult = Heightfield.OverlapGeom(Box, QueryTM, 0.0, nullptr);
+					bool CapsuleResult = Heightfield.OverlapGeom(Capsule, QueryTM, 0.0, nullptr);
+					EXPECT_TRUE(BoxResult);
+					EXPECT_TRUE(CapsuleResult);
 				}
 			}
 		}
@@ -1056,12 +1072,14 @@ namespace ChaosTest {
 		// Inclined Box
 		{
 			TBox<FReal, 3> Box(FVec3(-1.0, -1.0, -0.0001), FVec3(1.0, 1.0, 10.0));
+			FCapsule Capsule(FVec3(0.0, 0.0, 1.0-0.0001), FVec3(0.0, 0.0, 9.0), 1.0);
 			const FVec3 Translation(5.0, 0.0, 15.0);
 			FRigidTransform3 QueryTM(Translation, TRotation<FReal, 3>(UE::Math::TQuat<FReal>(TVector<FReal, 3>(1.0, 0.0, 0.0), -3*3.1415926 / 4.0)));
 			FVec3 Dir(1, 0, 0);
-			bool Result = Heightfield.OverlapGeom(Box, QueryTM, 0.0, nullptr);
-			// Collision with the mountain
-			EXPECT_TRUE(Result);
+			bool BoxResult = Heightfield.OverlapGeom(Box, QueryTM, 0.0, nullptr);
+			bool CapsuleResult = Heightfield.OverlapGeom(Capsule, QueryTM, 0.0, nullptr);
+			EXPECT_TRUE(BoxResult);
+			EXPECT_TRUE(CapsuleResult);
 		}
 	}
 	
diff --git a/Engine/Source/Runtime/Experimental/Chaos/Private/Chaos/Collision/TriangleOverlap.cpp b/Engine/Source/Runtime/Experimental/Chaos/Private/Chaos/Collision/TriangleOverlap.cpp
new file mode 100644
--- /dev/null
+++ b/Engine/Source/Runtime/Experimental/Chaos/Private/Chaos/Collision/TriangleOverlap.cpp
@@ -0,0 +1,115 @@
+// Copyright Epic Games, Inc. All Rights Reserved.
+#include "Chaos/Collision/TriangleOverlap.h"
+
+namespace Chaos
+{
+	bool ComputeCapsuleTriangleOverlapSimd(const VectorRegister4Float& A, const VectorRegister4Float& B, const VectorRegister4Float& C, const VectorRegister4Float& X1, const VectorRegister4Float& X2, FRealSingle Radius)
+	{
+		const VectorRegister4Float AB = VectorSubtract(B, A);
+		const VectorRegister4Float BC = VectorSubtract(C, B);
+
+		VectorRegister4Float Normal = VectorNormalize(VectorCross(AB, BC));
+
+		// Plane Triangle
+		const VectorRegister4Float AX1 = VectorSubtract(X1, A);
+		const VectorRegister4Float AX2 = VectorSubtract(X2, A);
+
+		FRealSingle AX1Dist = VectorDot3Scalar(AX1, Normal);
+		FRealSingle AX2Dist = VectorDot3Scalar(AX2, Normal);
+
+		if (FMath::Sign(AX1Dist) == FMath::Sign(AX2Dist))
+		{
+			FRealSingle ClosestDist = FMath::Min<FRealSingle>(FMath::Abs<FRealSingle>(AX1Dist), FMath::Abs<FRealSingle>(AX2Dist));
+			if (ClosestDist > Radius)
+			{
+				return false;
+			}
+		}
+
+		// Edges
+		constexpr FRealSingle ThirdFloat = 1.0f / 3.0f;
+		constexpr VectorRegister4Float Third = MakeVectorRegisterFloatConstant(ThirdFloat, ThirdFloat, ThirdFloat, ThirdFloat);
+		const VectorRegister4Float Centroid = VectorMultiply(VectorAdd(VectorAdd(A, B), C), Third);
+
+		const VectorRegister4Float X2X1 = VectorSubtract(X1, X2);
+		const VectorRegister4Float CA = VectorSubtract(A, C);
+		const VectorRegister4Float CentroidA = VectorSubtract(A, Centroid);
+		{
+			VectorRegister4Float EdgeSeparationAxis = VectorNormalize(VectorCross(X2X1, CA));
+			EdgeSeparationAxis = VectorSelect(VectorCompareGT(VectorZeroFloat(), VectorDot3(CentroidA, EdgeSeparationAxis)), VectorNegate(EdgeSeparationAxis), EdgeSeparationAxis);
+
+			const FRealSingle EdgeSeparationDist = VectorDot3Scalar(VectorSubtract(X1, C), EdgeSeparationAxis);
+			if (EdgeSeparationDist > Radius)
+			{
+				return false;
+			}
+		}
+		const VectorRegister4Float CentroidB = VectorSubtract(B, Centroid);
+		{
+			VectorRegister4Float EdgeSeparationAxis = VectorNormalize(VectorCross(X2X1, AB));
+			EdgeSeparationAxis = VectorSelect(VectorCompareGT(VectorZeroFloat(), VectorDot3(CentroidB, EdgeSeparationAxis)), VectorNegate(EdgeSeparationAxis), EdgeSeparationAxis);
+
+			const FRealSingle EdgeSeparationDist = VectorDot3Scalar(VectorSubtract(X1, A), EdgeSeparationAxis);
+			if (EdgeSeparationDist > Radius)
+			{
+				return false;
+			}
+		}
+		const VectorRegister4Float CentroidC = VectorSubtract(C, Centroid);
+		{
+			VectorRegister4Float EdgeSeparationAxis = VectorNormalize(VectorCross(X2X1, BC));
+			EdgeSeparationAxis = VectorSelect(VectorCompareGT(VectorZeroFloat(), VectorDot3(CentroidC, EdgeSeparationAxis)), VectorNegate(EdgeSeparationAxis), EdgeSeparationAxis);
+
+			const FRealSingle EdgeSeparationDist = VectorDot3Scalar(VectorSubtract(X1, B), EdgeSeparationAxis);
+			if (EdgeSeparationDist > Radius)
+			{
+				return false;
+			}
+		}
+
+		// Triangle Vertices
+		VectorRegister4Float SqrX2X1 = VectorDot3(X2X1, X2X1);
+		FRealSingle SqrRadius = Radius * Radius;
+		{
+			VectorRegister4Float TimeA = VectorClamp(VectorDivide(VectorDot3(X2X1, AX1), SqrX2X1), VectorZeroFloat(), VectorOneFloat());
+			VectorRegister4Float PA = VectorMultiplyAdd(X1, VectorSubtract(VectorOneFloat(), TimeA), VectorMultiply(X2, TimeA));
+
+			if (VectorDot3Scalar(VectorSubtract(PA, A), CentroidA) > 0.0f)
+			{
+				const VectorRegister4Float APA = VectorSubtract(PA, A);
+				if (VectorDot3Scalar(APA, APA) > SqrRadius)
+				{
+					return false;
+				}
+			}
+		}
+		{
+			const VectorRegister4Float TimeB = VectorClamp(VectorDivide(VectorDot3(X2X1, VectorSubtract(X1, B)), SqrX2X1), VectorZeroFloat(), VectorOneFloat());
+			const VectorRegister4Float PB = VectorMultiplyAdd(X1, VectorSubtract(VectorOneFloat(), TimeB), VectorMultiply(X2, TimeB));
+
+			if (VectorDot3Scalar(VectorSubtract(PB, B), CentroidB) > 0.0f)
+			{
+				const VectorRegister4Float BPB = VectorSubtract(PB, B);
+				if (VectorDot3Scalar(BPB, BPB) > SqrRadius)
+				{
+					return false;
+				}
+			}
+		}
+		{
+			const VectorRegister4Float TimeC = VectorClamp(VectorDivide(VectorDot3(X2X1, VectorSubtract(X1, C)), SqrX2X1), VectorZeroFloat(), VectorOneFloat());
+			const VectorRegister4Float PC = VectorMultiplyAdd(X1, VectorSubtract(VectorOneFloat(), TimeC), VectorMultiply(X2, TimeC));
+
+			if (VectorDot3Scalar(VectorSubtract(PC, C), CentroidC) > 0.0f)
+			{
+				const VectorRegister4Float CPC = VectorSubtract(PC, C);
+				if (VectorDot3Scalar(CPC, CPC) > SqrRadius)
+				{
+					return false;
+				}
+			}
+		}
+
+		return true;
+	}
+}
diff --git a/Engine/Source/Runtime/Experimental/Chaos/Private/Chaos/Collision/TriangleOverlap.h b/Engine/Source/Runtime/Experimental/Chaos/Private/Chaos/Collision/TriangleOverlap.h
new file mode 100644
--- /dev/null
+++ b/Engine/Source/Runtime/Experimental/Chaos/Private/Chaos/Collision/TriangleOverlap.h
@@ -0,0 +1,10 @@
+// Copyright Epic Games, Inc. All Rights Reserved.
+#pragma once
+
+#include "Chaos/Core.h"
+
+namespace Chaos
+{
+	bool ComputeCapsuleTriangleOverlapSimd(const VectorRegister4Float& A, const VectorRegister4Float& B, const VectorRegister4Float& C, const VectorRegister4Float& X1, const VectorRegister4Float& X2, FRealSingle Radius);
+}
+
diff --git a/Engine/Source/Runtime/Experimental/Chaos/Private/Chaos/HeightField.cpp b/Engine/Source/Runtime/Experimental/Chaos/Private/Chaos/HeightField.cpp
--- a/Engine/Source/Runtime/Experimental/Chaos/Private/Chaos/HeightField.cpp
+++ b/Engine/Source/Runtime/Experimental/Chaos/Private/Chaos/HeightField.cpp
@@ -4,6 +4,7 @@
 #include "Chaos/Collision/ContactPoint.h"
 #include "Chaos/Collision/ContactPointsMiscShapes.h"
 #include "Chaos/Collision/ContactTriangles.h"
+#include "Chaos/Collision/TriangleOverlap.h"
 #include "Chaos/Collision/PBDCollisionConstraint.h"
 #include "Chaos/CollisionOneShotManifolds.h"
 #include "Chaos/Core.h"
@@ -1989,6 +1990,26 @@ namespace Chaos
 			}
 			else
 			{
+				VectorRegister4Float X1;
+				VectorRegister4Float X2;
+				FRealSingle Radius;
+				if constexpr (std::is_same<QueryGeomType, FCapsule>::value)
+				{
+					Radius = FRealSingle(QueryGeom.GetRadius());
+					const FVec3f X1f = QueryGeom.GetX1();
+					const FVec3f X2f = QueryGeom.GetX2();
+					X1 = VectorLoadFloat3(&X1f.X);
+					X2 = VectorLoadFloat3(&X2f.X);
+				}
+				else if constexpr (std::is_same<QueryGeomType, TImplicitObjectScaled < FCapsule>>::value)
+				{
+					Radius = FRealSingle(QueryGeom.GetRadius());
+					const FVec3f X1f = QueryGeom.GetUnscaledObject()->GetX1() * QueryGeom.GetScale();
+					const FVec3f X2f = QueryGeom.GetUnscaledObject()->GetX2() * QueryGeom.GetScale();
+					X1 = VectorLoadFloat3(&X1f.X);
+					X2 = VectorLoadFloat3(&X2f.X);
+				}
+
 				VectorRegister4Float Points[4];
 				const UE::Math::TQuat<FReal>& RotationDouble = QueryTM.GetRotation();
 				VectorRegister4Float Rotation = MakeVectorRegisterFloatFromDouble(MakeVectorRegister(RotationDouble.X, RotationDouble.Y, RotationDouble.Z, RotationDouble.W));
@@ -2013,14 +2034,28 @@ namespace Chaos
 							P = VectorQuaternionRotateVector(InvRotation, P);
 						}
 
-						if (OverlapTriangleNoMTD(Points[0], Points[1], Points[3]))
+						if constexpr (std::is_same<QueryGeomType, FCapsule>::value || std::is_same<QueryGeomType, TImplicitObjectScaled < FCapsule>>::value)
 						{
-							return true;
+							if (ComputeCapsuleTriangleOverlapSimd(Points[0], Points[1], Points[3], X1, X2, Radius))
+							{
+								return true;
+							}
+							if (ComputeCapsuleTriangleOverlapSimd(Points[0], Points[3], Points[2], X1, X2, Radius))
+							{
+								return true;
+							}
 						}
-						if (OverlapTriangleNoMTD(Points[0], Points[3], Points[2]))
+						else
 						{
-							return true;
-						}
+							if (OverlapTriangleNoMTD(Points[0], Points[1], Points[3]))
+							{
+								return true;
+							}
+							if (OverlapTriangleNoMTD(Points[0], Points[3], Points[2]))
+							{
+								return true;
+							}
+						}	
 					}
 				}
 			}
diff --git a/Engine/Source/Runtime/Experimental/Chaos/Private/Chaos/TriangleMeshImplicitObject.cpp b/Engine/Source/Runtime/Experimental/Chaos/Private/Chaos/TriangleMeshImplicitObject.cpp
--- a/Engine/Source/Runtime/Experimental/Chaos/Private/Chaos/TriangleMeshImplicitObject.cpp
+++ b/Engine/Source/Runtime/Experimental/Chaos/Private/Chaos/TriangleMeshImplicitObject.cpp
@@ -2,6 +2,7 @@
 #include "Chaos/TriangleMeshImplicitObject.h"
 #include "Chaos/Collision/ContactPoint.h"
 #include "Chaos/Collision/ContactPointsMiscShapes.h"
+#include "Chaos/Collision/TriangleOverlap.h"
 #include "Chaos/Collision/PBDCollisionConstraint.h"
 #include "Chaos/CollisionOneShotManifolds.h"
 #include "Chaos/Capsule.h"
@@ -436,118 +437,6 @@ struct FTriangleMeshOverlapVisitorNoMTD
 };
 namespace
 {
-bool ComputeCapsuleTriangleCollisionSimd(const VectorRegister4Float& A, const VectorRegister4Float& B, const VectorRegister4Float& C, const VectorRegister4Float& X1, const VectorRegister4Float& X2, FRealSingle Radius)
-{
-
-	const VectorRegister4Float AB = VectorSubtract(B, A);
-	const VectorRegister4Float BC = VectorSubtract(C, B);
-
-	VectorRegister4Float Normal = VectorNormalize(VectorCross(AB, BC));
-
-
-	// Plane Triangle
-	const VectorRegister4Float AX1 = VectorSubtract(X1, A);
-	const VectorRegister4Float AX2 = VectorSubtract(X2, A);
-
-	FRealSingle AX1Dist = VectorDot3Scalar(AX1, Normal);
-	FRealSingle AX2Dist = VectorDot3Scalar(AX2, Normal);
-
-	if (FMath::Sign(AX1Dist) == FMath::Sign(AX2Dist))
-	{
-		FRealSingle ClosestDist = FMath::Min<FRealSingle>(FMath::Abs<FRealSingle>(AX1Dist), FMath::Abs<FRealSingle>(AX2Dist));
-		if (ClosestDist > Radius)
-		{
-			return false;
-		}
-	}
-
-	// Edges
-	constexpr FRealSingle ThirdFloat = 1.0f / 3.0f;
-	constexpr VectorRegister4Float Third = MakeVectorRegisterFloatConstant(ThirdFloat, ThirdFloat, ThirdFloat, ThirdFloat);
-	const VectorRegister4Float Centroid = VectorMultiply(VectorAdd(VectorAdd(A, B), C), Third);
-
-	const VectorRegister4Float X2X1 = VectorSubtract(X1, X2);
-	const VectorRegister4Float CA = VectorSubtract(A, C);
-	const VectorRegister4Float CentroidA = VectorSubtract(A, Centroid);
-	{
-		VectorRegister4Float EdgeSeparationAxis = VectorNormalize(VectorCross(X2X1, CA));
-		EdgeSeparationAxis = VectorSelect(VectorCompareGT(VectorZeroFloat(), VectorDot3(CentroidA, EdgeSeparationAxis)), VectorNegate(EdgeSeparationAxis), EdgeSeparationAxis);
-
-		const FRealSingle EdgeSeparationDist = VectorDot3Scalar(VectorSubtract(X1, C), EdgeSeparationAxis);
-		if (EdgeSeparationDist > Radius)
-		{
-			return false;
-		}
-	}
-	const VectorRegister4Float CentroidB = VectorSubtract(B, Centroid);
-	{
-		VectorRegister4Float EdgeSeparationAxis = VectorNormalize(VectorCross(X2X1, AB));
-		EdgeSeparationAxis = VectorSelect(VectorCompareGT(VectorZeroFloat(), VectorDot3(CentroidB, EdgeSeparationAxis)), VectorNegate(EdgeSeparationAxis), EdgeSeparationAxis);
-
-		const FRealSingle EdgeSeparationDist = VectorDot3Scalar(VectorSubtract(X1, A), EdgeSeparationAxis);
-		if (EdgeSeparationDist > Radius)
-		{
-			return false;
-		}
-	}
-	const VectorRegister4Float CentroidC = VectorSubtract(C, Centroid);
-	{
-		VectorRegister4Float EdgeSeparationAxis = VectorNormalize(VectorCross(X2X1, BC));
-		EdgeSeparationAxis = VectorSelect(VectorCompareGT(VectorZeroFloat(), VectorDot3(CentroidC, EdgeSeparationAxis)), VectorNegate(EdgeSeparationAxis), EdgeSeparationAxis);
-
-		const FRealSingle EdgeSeparationDist = VectorDot3Scalar(VectorSubtract(X1, B), EdgeSeparationAxis);
-		if (EdgeSeparationDist > Radius)
-		{
-			return false;
-		}
-	}
-
-	// Triangle Vertex
-	VectorRegister4Float SqrX2X1 = VectorDot3(X2X1, X2X1);
-	FRealSingle SqrRadius = Radius * Radius;
-	{
-		VectorRegister4Float TimeA = VectorClamp(VectorDivide(VectorDot3(X2X1, AX1), SqrX2X1), VectorZeroFloat(), VectorOneFloat());
-		VectorRegister4Float PA = VectorMultiplyAdd(X1, VectorSubtract(VectorOneFloat(), TimeA), VectorMultiply(X2, TimeA));
-
-		if (VectorDot3Scalar(VectorSubtract(PA, A), CentroidA) > 0.0f)
-		{
-			const VectorRegister4Float APA = VectorSubtract(PA, A);
-			if (VectorDot3Scalar(APA, APA) > SqrRadius)
-			{
-				return false;
-			}
-		}
-	}
-	{
-		const VectorRegister4Float TimeB = VectorClamp(VectorDivide(VectorDot3(X2X1, VectorSubtract(X1, B)), SqrX2X1), VectorZeroFloat(), VectorOneFloat());
-		const VectorRegister4Float PB = VectorMultiplyAdd(X1, VectorSubtract(VectorOneFloat(), TimeB), VectorMultiply(X2, TimeB));
-
-		if (VectorDot3Scalar(VectorSubtract(PB, B), CentroidB) > 0.0f)
-		{
-			const VectorRegister4Float BPB = VectorSubtract(PB, B);
-			if (VectorDot3Scalar(BPB, BPB) > SqrRadius)
-			{
-				return false;
-			}
-		}
-	}
-	{
-		const VectorRegister4Float TimeC = VectorClamp(VectorDivide(VectorDot3(X2X1, VectorSubtract(X1, C)), SqrX2X1), VectorZeroFloat(), VectorOneFloat());
-		const VectorRegister4Float PC = VectorMultiplyAdd(X1, VectorSubtract(VectorOneFloat(), TimeC), VectorMultiply(X2, TimeC));
-
-		if (VectorDot3Scalar(VectorSubtract(PC, C), CentroidC) > 0.0f)
-		{
-			const VectorRegister4Float CPC = VectorSubtract(PC, C);
-			if (VectorDot3Scalar(CPC, CPC) > SqrRadius)
-			{
-				return false;
-			}
-		}
-	}
-
-	return true;
-}
-
 template <typename IdxType, typename ParticlesType>
 inline void TriangleMeshTransformVertsHelperSimd(const VectorRegister4Float& TriMeshScale, int32 TriIdx, const ParticlesType& Particles,
 	const TArray<TVector<IdxType, 3>>& Elements, VectorRegister4Float& OutA, VectorRegister4Float& OutB, VectorRegister4Float& OutC)
@@ -559,7 +448,6 @@ inline void TriangleMeshTransformVertsHelperSimd(const VectorRegister4Float& Tri
 	OutB = VectorMultiply(OutB, TriMeshScale);
 	OutC = VectorMultiply(OutC, TriMeshScale);
 }
-
 }
 
 template <>
@@ -570,8 +458,8 @@ struct FTriangleMeshOverlapVisitorNoMTD<FCapsule>
 		, TriMesh(InTriMesh)
 		, bFoundIntersection(false)
 	{
-		FVec3f X1f = WorldScaleQueryTM.TransformPosition(InQueryGeom.GetX1());
-		FVec3f X2f = WorldScaleQueryTM.TransformPosition(InQueryGeom.GetX2());
+		const FVec3f X1f = WorldScaleQueryTM.TransformPosition(InQueryGeom.GetX1());
+		const FVec3f X2f = WorldScaleQueryTM.TransformPosition(InQueryGeom.GetX2());
 		X1 = VectorLoadFloat3(&X1f.X);
 		X2 = VectorLoadFloat3(&X2f.X);
 		TriMeshScale = VectorLoadFloat3(&InTriMeshScale.X);
@@ -589,7 +477,7 @@ struct FTriangleMeshOverlapVisitorNoMTD<FCapsule>
 			TriangleMeshTransformVertsHelperSimd(TriMeshScale, TriIdx, TriMesh->MParticles, TriMesh->MElements.GetSmallIndexBuffer(), A, B, C);
 		}
 
-		bFoundIntersection = ComputeCapsuleTriangleCollisionSimd(A, B, C, X1, X2, Radius);
+		bFoundIntersection = ComputeCapsuleTriangleOverlapSimd(A, B, C, X1, X2, Radius);
 		return !bFoundIntersection;
 	}
 
@@ -621,8 +509,8 @@ struct FTriangleMeshOverlapVisitorNoMTD<TImplicitObjectScaled<FCapsule>>
 		, TriMesh(InTriMesh)
 		, bFoundIntersection(false)
 	{
-		FVec3f X1f = WorldScaleQueryTM.TransformPosition(InQueryGeom.GetUnscaledObject()->GetX1() * InQueryGeom.GetScale());
-		FVec3f X2f = WorldScaleQueryTM.TransformPosition(InQueryGeom.GetUnscaledObject()->GetX2() * InQueryGeom.GetScale());
+		const FVec3f X1f = WorldScaleQueryTM.TransformPosition(InQueryGeom.GetUnscaledObject()->GetX1() * InQueryGeom.GetScale());
+		const FVec3f X2f = WorldScaleQueryTM.TransformPosition(InQueryGeom.GetUnscaledObject()->GetX2() * InQueryGeom.GetScale());
 		X1 = VectorLoadFloat3(&X1f.X);
 		X2 = VectorLoadFloat3(&X2f.X);
 		TriMeshScale = VectorLoadFloat3(&InTriMeshScale.X);
@@ -640,7 +528,7 @@ struct FTriangleMeshOverlapVisitorNoMTD<TImplicitObjectScaled<FCapsule>>
 			TriangleMeshTransformVertsHelperSimd(TriMeshScale, TriIdx, TriMesh->MParticles, TriMesh->MElements.GetSmallIndexBuffer(), A, B, C);
 		}
 
-		bFoundIntersection = ComputeCapsuleTriangleCollisionSimd(A, B, C, X1, X2, Radius);
+		bFoundIntersection = ComputeCapsuleTriangleOverlapSimd(A, B, C, X1, X2, Radius);
 		return !bFoundIntersection;
 	}
 
