Chaos - fix convex-heightfield triangle vertex indices
 This would have cauised issues with edge pruning.

diff --git a/Engine/Source/Runtime/Experimental/Chaos/Public/Chaos/HeightField.h b/Engine/Source/Runtime/Experimental/Chaos/Public/Chaos/HeightField.h
--- a/Engine/Source/Runtime/Experimental/Chaos/Public/Chaos/HeightField.h
+++ b/Engine/Source/Runtime/Experimental/Chaos/Public/Chaos/HeightField.h
@@ -155,16 +155,20 @@ namespace Chaos
 					const int32 FaceIndex0 = CellIndex * 2 + 0;
 					const int32 FaceIndex1 = CellIndex * 2 + 1;
 
-					Triangles[0] = FTriangle(Points[0], Points[1], Points[3]);
-					Triangles[1] = FTriangle(Points[0], Points[3], Points[2]);
-					if (!bStandardWinding)
+					if (bStandardWinding)
 					{
-						Triangles[0].ReverseWinding();
-						Triangles[1].ReverseWinding();
+						Triangles[0] = FTriangle(Points[0], Points[1], Points[3]);
+						Triangles[1] = FTriangle(Points[0], Points[3], Points[2]);
+						Visitor(Triangles[0], FaceIndex0, VertexIndex0, VertexIndex1, VertexIndex3);
+						Visitor(Triangles[1], FaceIndex1, VertexIndex0, VertexIndex3, VertexIndex2);
+					}
+					else
+					{
+						Triangles[0] = FTriangle(Points[0], Points[3], Points[1]);
+						Triangles[1] = FTriangle(Points[0], Points[2], Points[3]);
+						Visitor(Triangles[0], FaceIndex0, VertexIndex0, VertexIndex3, VertexIndex1);
+						Visitor(Triangles[1], FaceIndex1, VertexIndex0, VertexIndex2, VertexIndex3);
 					}
-
-					Visitor(Triangles[0], FaceIndex0, VertexIndex0, VertexIndex1, VertexIndex3);
-					Visitor(Triangles[1], FaceIndex1, VertexIndex0, VertexIndex3, VertexIndex2);
 				}
 			}
 		}
