Chaos Vehicles: Bug fix where updating aerodynamic drag or downforce at runtime has no effect.

diff --git a/Engine/Source/Runtime/Experimental/ChaosVehicles/ChaosVehiclesCore/Public/AerodynamicsSystem.h b/Engine/Source/Runtime/Experimental/ChaosVehicles/ChaosVehiclesCore/Public/AerodynamicsSystem.h
--- a/Engine/Source/Runtime/Experimental/ChaosVehicles/ChaosVehiclesCore/Public/AerodynamicsSystem.h
+++ b/Engine/Source/Runtime/Experimental/ChaosVehicles/ChaosVehiclesCore/Public/AerodynamicsSystem.h
@@ -51,11 +51,13 @@ namespace Chaos
 		void SetDragCoefficient(float InCoeffient)
 		{
 			DragCoefficient = InCoeffient;
+			EffectiveDragConstant = 0.5f * Setup().AreaMetresSquared * DragCoefficient;
 		}
 
 		void SetDownforceCoefficient(float InCoeffient)
 		{
 			DownforceCoefficient = InCoeffient;
+			EffectiveLiftConstant = 0.5f * Setup().AreaMetresSquared * DownforceCoefficient;
 		}
 
 		/** get the drag force generated at the given velocity */
