Fix: added [[maybe_unused]] to some variables as a fix/workaround for compile error C4101: 'X2': unreferenced local variable (VS2022)

diff --git a/Engine/Source/Runtime/Experimental/Chaos/Private/Chaos/HeightField.cpp b/Engine/Source/Runtime/Experimental/Chaos/Private/Chaos/HeightField.cpp
--- a/Engine/Source/Runtime/Experimental/Chaos/Private/Chaos/HeightField.cpp
+++ b/Engine/Source/Runtime/Experimental/Chaos/Private/Chaos/HeightField.cpp
@@ -1990,9 +1990,9 @@ namespace Chaos
 			}
 			else
 			{
-				VectorRegister4Float X1;
-				VectorRegister4Float X2;
-				FRealSingle Radius;
+				[[maybe_unused]] VectorRegister4Float X1;
+				[[maybe_unused]] VectorRegister4Float X2;
+				[[maybe_unused]] FRealSingle Radius;
 				if constexpr (std::is_same<QueryGeomType, FCapsule>::value)
 				{
 					Radius = FRealSingle(QueryGeom.GetRadius());
