Dataflow:
- Made member functions const

diff --git a/Engine/Source/Runtime/Experimental/Dataflow/Core/Private/Dataflow/DataflowSelection.cpp b/Engine/Source/Runtime/Experimental/Dataflow/Core/Private/Dataflow/DataflowSelection.cpp
--- a/Engine/Source/Runtime/Experimental/Dataflow/Core/Private/Dataflow/DataflowSelection.cpp
+++ b/Engine/Source/Runtime/Experimental/Dataflow/Core/Private/Dataflow/DataflowSelection.cpp
@@ -64,7 +64,7 @@ void FDataflowSelection::XOR(const FDataflowSelection& Other, FDataflowSelection
 	Result.SelectionArray = TBitArray<>::BitwiseXOR(SelectionArray, Other.SelectionArray, EBitwiseOperatorFlags::MaxSize);
 }
 
-int32 FDataflowSelection::NumSelected()
+int32 FDataflowSelection::NumSelected() const
 {
 	int32 NumSelectedBits = 0;
 
@@ -82,7 +82,7 @@ int32 FDataflowSelection::NumSelected()
 	return NumSelectedBits;
 }
 
-bool FDataflowSelection::AnySelected()
+bool FDataflowSelection::AnySelected() const
 {
 	TBitArray<>::FConstIterator It(SelectionArray);
 	while (It)
diff --git a/Engine/Source/Runtime/Experimental/Dataflow/Core/Public/Dataflow/DataflowSelection.h b/Engine/Source/Runtime/Experimental/Dataflow/Core/Public/Dataflow/DataflowSelection.h
--- a/Engine/Source/Runtime/Experimental/Dataflow/Core/Public/Dataflow/DataflowSelection.h
+++ b/Engine/Source/Runtime/Experimental/Dataflow/Core/Public/Dataflow/DataflowSelection.h
@@ -12,8 +12,8 @@ struct DATAFLOWCORE_API FDataflowSelection
 	void Initialize(int32 NumBits, bool Value);
 	void Initialize(const FDataflowSelection& Other);
 	int32 Num() const { return SelectionArray.Num(); }
-	int32 NumSelected();
-	bool AnySelected();
+	int32 NumSelected() const;
+	bool AnySelected() const;
 	bool IsSelected(int32 Idx) const { return SelectionArray[Idx]; }
 	void SetSelected(int32 Idx) { SelectionArray[Idx] = true; }
 	void SetNotSelected(int32 Idx) { SelectionArray[Idx] = false; }
