Fix potential bug in Triangle Overlap
Avoid division by zero

diff --git a/Engine/Source/Programs/HeadlessChaos/Private/HeadlessChaosTestOverlap.cpp b/Engine/Source/Programs/HeadlessChaos/Private/HeadlessChaosTestOverlap.cpp
--- a/Engine/Source/Programs/HeadlessChaos/Private/HeadlessChaosTestOverlap.cpp
+++ b/Engine/Source/Programs/HeadlessChaos/Private/HeadlessChaosTestOverlap.cpp
@@ -734,6 +734,49 @@ namespace ChaosTest
 					EXPECT_EQ(bResult, false);
 				}
 			}
+			{
+				const FVec3 X1 = { 0.0, 0.0, -250.0 };
+				const FVec3 X2 = { 0.0, 0.0, 250.0 };
+				const FReal Radius = 0.5;
+				const FCapsule Capsule = FCapsule(X1, X2, Radius);
+
+				{
+					const FRigidTransform3 QueryTM(FVec3(0.0, 0.0, 5.0), FQuat::Identity);
+					bool bResult = TriangleMesh->OverlapGeom(Capsule, QueryTM, 0.0);
+					EXPECT_TRUE(bResult);
+				}
+			}
+			{
+				const FVec3 X1 = { 0.0, 0.0, 0.0 };
+				const FVec3 X2 = { 0.0, 0.0, 0.0 };
+				const FReal Radius = 0.5;
+				const FCapsule Capsule = FCapsule(X1, X2, Radius);
+				{
+					const FRigidTransform3 QueryTM(FVec3(0.0, 0.0, 0.0), FQuat::Identity);
+					bool bResult = TriangleMesh->OverlapGeom(Capsule, QueryTM, 0.0);
+					EXPECT_TRUE(bResult);
+				}
+				{
+					const FRigidTransform3 QueryTM(FVec3(0.0, 0.0, 5.0), FQuat::Identity);
+					bool bResult = TriangleMesh->OverlapGeom(Capsule, QueryTM, 0.0);
+					EXPECT_FALSE(bResult);
+				}
+				{
+					const FRigidTransform3 QueryTM(FVec3(0.0, 0.0, 10.0), FQuat::Identity);
+					bool bResult = TriangleMesh->OverlapGeom(Capsule, QueryTM, 0.0);
+					EXPECT_TRUE(bResult);
+				}
+				{
+					const FRigidTransform3 QueryTM(FVec3(5.0, 10.25, 0.0), FQuat::Identity);
+					bool bResult = TriangleMesh->OverlapGeom(Capsule, QueryTM, 0.0);
+					EXPECT_TRUE(bResult);
+				}
+				{
+					const FRigidTransform3 QueryTM(FVec3(5.0, 10.6, 0.0), FQuat::Identity);
+					bool bResult = TriangleMesh->OverlapGeom(Capsule, QueryTM, 0.0);
+					EXPECT_FALSE(bResult);
+				}
+			}
 		}
 	}
 }
diff --git a/Engine/Source/Runtime/Experimental/Chaos/Private/Chaos/Collision/TriangleOverlap.cpp b/Engine/Source/Runtime/Experimental/Chaos/Private/Chaos/Collision/TriangleOverlap.cpp
--- a/Engine/Source/Runtime/Experimental/Chaos/Private/Chaos/Collision/TriangleOverlap.cpp
+++ b/Engine/Source/Runtime/Experimental/Chaos/Private/Chaos/Collision/TriangleOverlap.cpp
@@ -1,5 +1,6 @@
 // Copyright Epic Games, Inc. All Rights Reserved.
 #include "Chaos/Collision/TriangleOverlap.h"
+#include "Chaos/VectorUtility.h"
 
 namespace Chaos
 {
@@ -69,10 +70,13 @@ namespace Chaos
 
 		// Triangle Vertices
 		VectorRegister4Float SqrX2X1 = VectorDot3(X2X1, X2X1);
+		VectorRegister4Float ZeroMask = VectorCompareEQ(VectorZeroFloat(), SqrX2X1);
+
 		FRealSingle SqrRadius = Radius * Radius;
 		{
 			VectorRegister4Float TimeA = VectorClamp(VectorDivide(VectorDot3(X2X1, AX1), SqrX2X1), VectorZeroFloat(), VectorOneFloat());
-			VectorRegister4Float PA = VectorMultiplyAdd(X1, VectorSubtract(VectorOneFloat(), TimeA), VectorMultiply(X2, TimeA));
+			TimeA = VectorBitwiseNotAnd(ZeroMask, TimeA);
+			const VectorRegister4Float PA = VectorMultiplyAdd(X1, VectorSubtract(VectorOneFloat(), TimeA), VectorMultiply(X2, TimeA));
 
 			if (VectorDot3Scalar(VectorSubtract(PA, A), CentroidA) > 0.0f)
 			{
@@ -84,7 +88,8 @@ namespace Chaos
 			}
 		}
 		{
-			const VectorRegister4Float TimeB = VectorClamp(VectorDivide(VectorDot3(X2X1, VectorSubtract(X1, B)), SqrX2X1), VectorZeroFloat(), VectorOneFloat());
+			VectorRegister4Float TimeB = VectorClamp(VectorDivide(VectorDot3(X2X1, VectorSubtract(X1, B)), SqrX2X1), VectorZeroFloat(), VectorOneFloat());
+			TimeB = VectorBitwiseNotAnd(ZeroMask, TimeB);
 			const VectorRegister4Float PB = VectorMultiplyAdd(X1, VectorSubtract(VectorOneFloat(), TimeB), VectorMultiply(X2, TimeB));
 
 			if (VectorDot3Scalar(VectorSubtract(PB, B), CentroidB) > 0.0f)
@@ -97,7 +102,8 @@ namespace Chaos
 			}
 		}
 		{
-			const VectorRegister4Float TimeC = VectorClamp(VectorDivide(VectorDot3(X2X1, VectorSubtract(X1, C)), SqrX2X1), VectorZeroFloat(), VectorOneFloat());
+			VectorRegister4Float TimeC = VectorClamp(VectorDivide(VectorDot3(X2X1, VectorSubtract(X1, C)), SqrX2X1), VectorZeroFloat(), VectorOneFloat());
+			TimeC = VectorBitwiseNotAnd(ZeroMask, TimeC);
 			const VectorRegister4Float PC = VectorMultiplyAdd(X1, VectorSubtract(VectorOneFloat(), TimeC), VectorMultiply(X2, TimeC));
 
 			if (VectorDot3Scalar(VectorSubtract(PC, C), CentroidC) > 0.0f)
