geometrycollection runtime optimization
- Optimize managed array collection NumElement to only fetch once from the group map
- Make sure to skip IncrementSleepTimer and ImcrementBreakTimer method is GC is not broken yet
- Early out in RefreshEmbeddedGeometry if no examplar is present
- Optimize BuildGeometry to reduce cache missed when running too small batch across large number of threads ( 2.5 time perf boost )

diff --git a/Engine/Source/Runtime/Experimental/Chaos/Private/GeometryCollection/ManagedArrayCollection.cpp b/Engine/Source/Runtime/Experimental/Chaos/Private/GeometryCollection/ManagedArrayCollection.cpp
--- a/Engine/Source/Runtime/Experimental/Chaos/Private/GeometryCollection/ManagedArrayCollection.cpp
+++ b/Engine/Source/Runtime/Experimental/Chaos/Private/GeometryCollection/ManagedArrayCollection.cpp
@@ -149,12 +149,12 @@ TArray<FName> FManagedArrayCollection::AttributeNames(FName Group) const
 	return AttributeNames;
 }
 
-int32 FManagedArrayCollection::NumElements(FName Group) const
+int32 FManagedArrayCollection::NumElements(FName GroupName) const
 {
 	int32 Num = 0;
-	if (GroupInfo.Contains(Group))
+	if (const FGroupInfo* Group = GroupInfo.Find(GroupName))
 	{
-		Num = GroupInfo[Group].Size;
+		Num = Group->Size;
 	}
 	return Num;
 }
diff --git a/Engine/Source/Runtime/Experimental/GeometryCollectionEngine/Private/GeometryCollection/GeometryCollectionComponent.cpp b/Engine/Source/Runtime/Experimental/GeometryCollectionEngine/Private/GeometryCollection/GeometryCollectionComponent.cpp
--- a/Engine/Source/Runtime/Experimental/GeometryCollectionEngine/Private/GeometryCollection/GeometryCollectionComponent.cpp
+++ b/Engine/Source/Runtime/Experimental/GeometryCollectionEngine/Private/GeometryCollection/GeometryCollectionComponent.cpp
@@ -1064,6 +1064,12 @@ UPhysicalMaterial* UGeometryCollectionComponent::GetPhysicalMaterial() const
 
 void UGeometryCollectionComponent::RefreshEmbeddedGeometry()
 {
+	const int32 ExemplarCount = EmbeddedGeometryComponents.Num();
+	if (ExemplarCount == 0)
+	{
+		return;
+	}
+
 	const TManagedArray<int32>& ExemplarIndexArray = GetExemplarIndexArray();
 	const int32 TransformCount = GlobalMatrices.Num();
 	if (!ensureMsgf(TransformCount == ExemplarIndexArray.Num(), TEXT("GlobalMatrices (Num=%d) cached on GeometryCollectionComponent are not in sync with ExemplarIndexArray (Num=%d) on underlying GeometryCollection; likely missed a dynamic data update"), TransformCount, ExemplarIndexArray.Num()))
@@ -1081,7 +1087,6 @@ void UGeometryCollectionComponent::RefreshEmbeddedGeometry()
 	EmbeddedInstanceIndex.Init(INDEX_NONE, RestCollection->GetGeometryCollection()->NumElements(FGeometryCollection::TransformGroup));
 #endif
 
-	const int32 ExemplarCount = EmbeddedGeometryComponents.Num();
 	for (int32 ExemplarIndex = 0; ExemplarIndex < ExemplarCount; ++ExemplarIndex)
 	{		
 #if WITH_EDITOR
@@ -2214,8 +2219,11 @@ void UGeometryCollectionComponent::TickComponent(float DeltaTime, enum ELevelTic
 		// In editor mode we have no DynamicCollection so this test is necessary
 		if(DynamicCollection) //, TEXT("No dynamic collection available for component %s during tick."), *GetName()))
 		{
-			IncrementSleepTimer(DeltaTime);
-			IncrementBreakTimer(DeltaTime);
+			if (IsRootBroken())
+			{
+				IncrementSleepTimer(DeltaTime);
+				IncrementBreakTimer(DeltaTime);
+			}
 
 			// todo(chaos) : find a way to only update that of transform have changed
 			// right now this does not work properly because the dirty flags may not be updated at the right time
@@ -4071,6 +4079,18 @@ void UGeometryCollectionComponent::RefreshISMPoolInstances()
 	}
 }
 
+bool UGeometryCollectionComponent::IsRootBroken() const
+{
+	if (DynamicCollection && DynamicCollection->Active.Num() > 0)
+	{
+		const int32 RootIndex = GetRootIndex();
+		if (RootIndex != INDEX_NONE)
+		{
+			return !DynamicCollection->Active[RootIndex];
+		}
+	}
+	return false;
+}
 
 struct FGeometryCollectionDecayContext
 {
@@ -4421,7 +4441,7 @@ int32 UGeometryCollectionComponent::GetInitialLevel(int32 ItemIndex)
 	return Level;
 }
 
-int32 UGeometryCollectionComponent::GetRootIndex()
+int32 UGeometryCollectionComponent::GetRootIndex() const
 {
 	if (RestCollection && RestCollection->GetGeometryCollection())
 	{
diff --git a/Engine/Source/Runtime/Experimental/GeometryCollectionEngine/Private/GeometryCollection/GeometryCollectionSceneProxy.cpp b/Engine/Source/Runtime/Experimental/GeometryCollectionEngine/Private/GeometryCollection/GeometryCollectionSceneProxy.cpp
--- a/Engine/Source/Runtime/Experimental/GeometryCollectionEngine/Private/GeometryCollection/GeometryCollectionSceneProxy.cpp
+++ b/Engine/Source/Runtime/Experimental/GeometryCollectionEngine/Private/GeometryCollection/GeometryCollectionSceneProxy.cpp
@@ -86,6 +86,22 @@ FAutoConsoleVariableRef CVarRayTracingGeometryCollectionProxyMeshes(
 	ECVF_RenderThreadSafe
 );
 
+static int32 GeometryCollectionBuildGeometryParallelVertexProcessingMinBatch = 5000;
+FAutoConsoleVariableRef CVarGeometryCollectionBuildGeometryParallelVertexProcessingMinBatch(
+	TEXT("r.GeometryCollectionBuildGeometry.ParallelVertexProcessingMinBatch"),
+	GeometryCollectionBuildGeometryParallelVertexProcessingMinBatch,
+	TEXT("When building geometry collection rendering geometry, minimum number of vertices to batch for parallel processing performance tuning"),
+	ECVF_RenderThreadSafe
+);
+
+static int32 GeometryCollectionBuildGeometryParallelIndexProcessingMinBatch = 5000;
+FAutoConsoleVariableRef CVarGeometryCollectionBuildGeometryParallelIndexProcessingMinBatch(
+	TEXT("r.GeometryCollectionBuildGeometry.ParallelIndexProcessingMinBatch"),
+	GeometryCollectionBuildGeometryParallelIndexProcessingMinBatch,
+	TEXT("When building geometry collection rendering geometry, minimum number of indices to batch for parallel processing performance tuning"),
+	ECVF_RenderThreadSafe
+);
+
 #if !defined(CHAOS_GEOMETRY_COLLECTION_SET_DYNAMIC_DATA_ISPC_ENABLED_DEFAULT)
 #define CHAOS_GEOMETRY_COLLECTION_SET_DYNAMIC_DATA_ISPC_ENABLED_DEFAULT 1
 #endif
@@ -440,20 +456,28 @@ void FGeometryCollectionSceneProxy::ReleaseResources()
 	VertexFactory.ReleaseResource();
 }
 
+inline static FColor GetOnlyVertexColors(const FGeometryCollectionConstantData* ConstantDataIn, int32 PointIdx, bool bShowBoneColors)
+{
+	return ConstantDataIn->Colors[PointIdx].ToFColor(true);
+};
+
+inline static FColor GetBonesColors(const FGeometryCollectionConstantData* ConstantDataIn, int32 PointIdx, bool bShowBoneColors)
+{
+	return bShowBoneColors
+		? ConstantDataIn->BoneColors[PointIdx].ToFColor(true)
+		: (ConstantDataIn->BoneColors[PointIdx] * ConstantDataIn->Colors[PointIdx]).ToFColor(true);
+};
+
 void FGeometryCollectionSceneProxy::BuildGeometry( const FGeometryCollectionConstantData* ConstantDataIn, TArray<FDynamicMeshVertex>& OutVertices, TArray<int32>& OutIndices, TArray<int32> &OutOriginalMeshIndices)
 {
-	OutVertices.SetNumUninitialized(ConstantDataIn->Vertices.Num());
-	ParallelFor(ConstantData->Vertices.Num(), [&](int32 PointIdx)
+	const bool bCapturedShowBoneColors = this->bShowBoneColors;
+	const auto SetOutVertex = [&OutVertices, &ConstantDataIn, bCapturedShowBoneColors](int32 PointIdx, FColor (*GetColorFunc)(const FGeometryCollectionConstantData*, int32, bool) )
 	{
 		OutVertices[PointIdx] =
 			FDynamicMeshVertex(
 				ConstantDataIn->Vertices[PointIdx],
 				ConstantDataIn->UVs[PointIdx][0],
-				(bShowBoneColors||bEnableBoneSelection) 
-					? bShowBoneColors 
-						? ConstantDataIn->BoneColors[PointIdx].ToFColor(true)
-						: (ConstantDataIn->BoneColors[PointIdx] * ConstantDataIn->Colors[PointIdx]).ToFColor(true)
-					: ConstantDataIn->Colors[PointIdx].ToFColor(true)
+				GetColorFunc(ConstantDataIn, PointIdx, bCapturedShowBoneColors)
 			);
 		OutVertices[PointIdx].SetTangents(ConstantDataIn->TangentU[PointIdx], ConstantDataIn->TangentV[PointIdx], ConstantDataIn->Normals[PointIdx]);
 
@@ -464,25 +488,47 @@ void FGeometryCollectionSceneProxy::BuildGeometry( const FGeometryCollectionCons
 				OutVertices[PointIdx].TextureCoordinate[UVLayerIdx] = ConstantDataIn->UVs[PointIdx][UVLayerIdx];
 			}
 		}
-	});
+	};
 
+	OutVertices.SetNumUninitialized(ConstantDataIn->Vertices.Num());
+	const int32 MinVertexBatchSize = GeometryCollectionBuildGeometryParallelVertexProcessingMinBatch;
+	if (bShowBoneColors || bEnableBoneSelection)
+	{
+		ParallelFor(TEXT("GCSceneProxy::BuildGeometry(Vertices)"), ConstantData->Vertices.Num(), MinVertexBatchSize,
+			[&SetOutVertex](int32 PointIdx)
+			{
+				SetOutVertex(PointIdx, GetBonesColors);
+			});
+	}
+	else
+	{
+		ParallelFor(TEXT("GCSceneProxy::BuildGeometry(Vertices)"), ConstantData->Vertices.Num(), MinVertexBatchSize,
+			[&SetOutVertex](int32 PointIdx)
+			{
+				SetOutVertex(PointIdx, GetOnlyVertexColors);
+			});
+	}
 	check(ConstantDataIn->Indices.Num() * 3 == NumIndices);
 
+	const int32 MinIndexBatchSize = GeometryCollectionBuildGeometryParallelIndexProcessingMinBatch;
+
 	OutIndices.SetNumUninitialized(NumIndices);
-	ParallelFor(ConstantDataIn->Indices.Num(), [&](int32 IndexIdx)
-	{
-		OutIndices[IndexIdx * 3 ]    = ConstantDataIn->Indices[IndexIdx].X;
-		OutIndices[IndexIdx * 3 + 1] = ConstantDataIn->Indices[IndexIdx].Y;
-		OutIndices[IndexIdx * 3 + 2] = ConstantDataIn->Indices[IndexIdx].Z;
-	});
+	ParallelFor(TEXT("GCSceneProxy::BuildGeometry(Indices)"), ConstantDataIn->Indices.Num(), MinIndexBatchSize, 
+		[&OutIndices, &ConstantDataIn](int32 IndexIdx)
+		{
+			OutIndices[IndexIdx * 3 ]    = ConstantDataIn->Indices[IndexIdx].X;
+			OutIndices[IndexIdx * 3 + 1] = ConstantDataIn->Indices[IndexIdx].Y;
+			OutIndices[IndexIdx * 3 + 2] = ConstantDataIn->Indices[IndexIdx].Z;
+		});
 	
 	OutOriginalMeshIndices.SetNumUninitialized(ConstantDataIn->OriginalMeshIndices.Num() * 3);
-	ParallelFor(ConstantDataIn->OriginalMeshIndices.Num(), [&](int32 IndexIdx)
-	{
-		OutOriginalMeshIndices[IndexIdx * 3] = ConstantDataIn->OriginalMeshIndices[IndexIdx].X;
-		OutOriginalMeshIndices[IndexIdx * 3 + 1] = ConstantDataIn->OriginalMeshIndices[IndexIdx].Y;
-		OutOriginalMeshIndices[IndexIdx * 3 + 2] = ConstantDataIn->OriginalMeshIndices[IndexIdx].Z;
-	});
+	ParallelFor(TEXT("GCSceneProxy::BuildGeometry(OriginalMeshIndices)"), ConstantDataIn->OriginalMeshIndices.Num(), MinIndexBatchSize, 
+		[&OutOriginalMeshIndices, &ConstantDataIn](int32 IndexIdx)
+		{
+			OutOriginalMeshIndices[IndexIdx * 3] = ConstantDataIn->OriginalMeshIndices[IndexIdx].X;
+			OutOriginalMeshIndices[IndexIdx * 3 + 1] = ConstantDataIn->OriginalMeshIndices[IndexIdx].Y;
+			OutOriginalMeshIndices[IndexIdx * 3 + 2] = ConstantDataIn->OriginalMeshIndices[IndexIdx].Z;
+		});
 }
 
 void FGeometryCollectionSceneProxy::SetConstantData_RenderThread(FGeometryCollectionConstantData* NewConstantData, bool ForceInit)
diff --git a/Engine/Source/Runtime/Experimental/GeometryCollectionEngine/Public/GeometryCollection/GeometryCollectionComponent.h b/Engine/Source/Runtime/Experimental/GeometryCollectionEngine/Public/GeometryCollection/GeometryCollectionComponent.h
--- a/Engine/Source/Runtime/Experimental/GeometryCollectionEngine/Public/GeometryCollection/GeometryCollectionComponent.h
+++ b/Engine/Source/Runtime/Experimental/GeometryCollectionEngine/Public/GeometryCollection/GeometryCollectionComponent.h
@@ -616,7 +616,7 @@ public:
 	
 	/** Get the root item index of the hierarchy */
 	UFUNCTION(BlueprintCallable, Category = "ChaosPhysics")
-	int32 GetRootIndex();
+	int32 GetRootIndex() const;
 
 	/**
 	* Get mass and extent of a specific piece
@@ -1169,6 +1169,8 @@ private:
 	/** update ISM transforms */
 	void RefreshISMPoolInstances();
 
+	/** return true if the root cluster is not longer active at runtime */
+	bool IsRootBroken() const;
 
 	void IncrementSleepTimer(float DeltaTime);
 	void IncrementBreakTimer(float DeltaTime);
