Iris - Implemented globally unique NetHandles.

diff --git a/Engine/Source/Runtime/Experimental/Iris/Core/Private/Iris/ReplicationSystem/DirtyNetObjectTracker.cpp b/Engine/Source/Runtime/Experimental/Iris/Core/Private/Iris/ReplicationSystem/DirtyNetObjectTracker.cpp
--- a/Engine/Source/Runtime/Experimental/Iris/Core/Private/Iris/ReplicationSystem/DirtyNetObjectTracker.cpp
+++ b/Engine/Source/Runtime/Experimental/Iris/Core/Private/Iris/ReplicationSystem/DirtyNetObjectTracker.cpp
@@ -4,7 +4,9 @@
 #include "HAL/PlatformAtomics.h"
 #include "Iris/IrisConstants.h"
 #include "Iris/Core/IrisLog.h"
+#include "Iris/Core/IrisProfiler.h"
 #include "Iris/ReplicationSystem/ReplicationSystem.h"
+#include "Iris/ReplicationSystem/ReplicationSystemInternal.h"
 #include "Traits/IntType.h"
 #include <atomic>
 
@@ -25,39 +27,30 @@
 namespace UE::Net::Private
 {
 
-#if UE_NET_ALLOW_MULTIPLE_REPLICATION_SYSTEMS
-constexpr uint32 DirtyNetObjectTrackerCount = FReplicationSystemFactory::MaxReplicationSystemCount;
-#else
-constexpr uint32 DirtyNetObjectTrackerCount = 1;
-#endif
-static FDirtyNetObjectTracker* DirtyNetObjectTrackers[DirtyNetObjectTrackerCount + 1]; // nullptr as last tracker
-
 FDirtyNetObjectTracker::FDirtyNetObjectTracker()
-: ReplicationSystemId(InvalidReplicationSystemId)
+: NetRefHandleManager(nullptr)
 , DirtyNetObjectContainer(nullptr)
+, ReplicationSystemId(InvalidReplicationSystemId)
 , DirtyNetObjectWordCount(0)
 , NetObjectIdRangeStart(0U)
 , NetObjectIdRangeEnd(0U)
 , NetObjectIdCount(0)
+, bHasPolledGlobalDirtyTracker(false)
 {
 }
 
 FDirtyNetObjectTracker::~FDirtyNetObjectTracker()
 {
-	if (ReplicationSystemId < DirtyNetObjectTrackerCount && DirtyNetObjectTrackers[ReplicationSystemId] == this)
-	{
-		DirtyNetObjectTrackers[ReplicationSystemId] = nullptr;
-	}
-
-	Shutdown();
+	Deinit();
 }
 
 void FDirtyNetObjectTracker::Init(const FDirtyNetObjectTrackerInitParams& Params)
 {
-	check(Params.ReplicationSystemId < DirtyNetObjectTrackerCount);
 	check(Params.NetObjectIndexRangeEnd >= Params.NetObjectIndexRangeStart);
+	check(Params.NetRefHandleManager != nullptr);
 	check(DirtyNetObjectContainer == nullptr);
 
+	NetRefHandleManager = Params.NetRefHandleManager;
 	ReplicationSystemId = Params.ReplicationSystemId;
 	NetObjectIdRangeStart = Params.NetObjectIndexRangeStart;
 	NetObjectIdRangeEnd = Params.NetObjectIndexRangeEnd;
@@ -67,26 +60,47 @@ void FDirtyNetObjectTracker::Init(const FDirtyNetObjectTrackerInitParams& Params
 	 */
 	NetObjectIdCount = NetObjectIdRangeEnd + 1;
 
+	GlobalDirtyTrackerPollHandle = FGlobalDirtyNetObjectTracker::CreatePoller();
+
 	DirtyNetObjectWordCount = (NetObjectIdCount + StorageTypeBitCount - 1)/StorageTypeBitCount;
 	DirtyNetObjectContainer = new StorageType[DirtyNetObjectWordCount];
 	ClearDirtyNetObjects();
 
-	if (DirtyNetObjectTrackers[ReplicationSystemId] != nullptr)
-	{
-		LowLevelFatalError(TEXT("DirtyNetObjectTrackerAlready initialized for ReplicationSystemId %u"), ReplicationSystemId);
-	}
-	
-	DirtyNetObjectTrackers[ReplicationSystemId] = this;
-
 	UE_LOG_DIRTYOBJECTTRACKER(TEXT("FDirtyNetObjectTracker::Init %u Id, Start:%u, End: %u"), ReplicationSystemId, NetObjectIdRangeStart, NetObjectIdRangeEnd);
 }
 
-void FDirtyNetObjectTracker::Shutdown()
+void FDirtyNetObjectTracker::Deinit()
 {
+	GlobalDirtyTrackerPollHandle.Destroy();
+	bHasPolledGlobalDirtyTracker = false;
+
 	delete[] DirtyNetObjectContainer;
 	DirtyNetObjectContainer = nullptr;
 }
 
+void FDirtyNetObjectTracker::UpdateDirtyNetObjects()
+{
+	if (!GlobalDirtyTrackerPollHandle.IsValid())
+	{
+		return;
+	}
+
+	IRIS_PROFILER_SCOPE(FDirtyNetObjectTracker_UpdateDirtyNetObjects)
+
+	bHasPolledGlobalDirtyTracker = true;
+	const TSet<FNetHandle>& GlobalDirtyNetObjects = FGlobalDirtyNetObjectTracker::GetDirtyNetObjects(GlobalDirtyTrackerPollHandle);
+	for (FNetHandle NetHandle : GlobalDirtyNetObjects)
+	{
+		const FInternalNetRefIndex NetObjectIndex = NetRefHandleManager->GetInternalIndexFromNetHandle(NetHandle);
+		if (NetObjectIndex != FNetRefHandleManager::InvalidInternalIndex)
+		{
+			const uint32 BitOffset = NetObjectIndex;
+			const StorageType BitMask = StorageType(1) << (BitOffset & (StorageTypeBitCount - 1));
+			DirtyNetObjectContainer[BitOffset/StorageTypeBitCount] |= BitMask;
+		}
+	}
+}
+
 void FDirtyNetObjectTracker::MarkNetObjectDirty(uint32 NetObjectIndex)
 {
 	if ((NetObjectIndex >= NetObjectIdRangeStart) & (NetObjectIndex <= NetObjectIdRangeEnd))
@@ -108,6 +122,12 @@ FNetBitArrayView FDirtyNetObjectTracker::GetDirtyNetObjects() const
 
 void FDirtyNetObjectTracker::ClearDirtyNetObjects()
 {
+	if (bHasPolledGlobalDirtyTracker)
+	{
+		bHasPolledGlobalDirtyTracker = false;
+		FGlobalDirtyNetObjectTracker::ResetDirtyNetObjects(GlobalDirtyTrackerPollHandle);
+	}
+
 	FMemory::Memset(DirtyNetObjectContainer, 0, DirtyNetObjectWordCount*sizeof(StorageType));
 
 	std::atomic_thread_fence(std::memory_order_seq_cst);
@@ -115,9 +135,10 @@ void FDirtyNetObjectTracker::ClearDirtyNetObjects()
 
 void MarkNetObjectStateDirty(uint32 ReplicationSystemId, uint32 NetObjectIndex)
 {
-	if (FDirtyNetObjectTracker* DirtyNetObjectTracker = DirtyNetObjectTrackers[FMath::Min(ReplicationSystemId, DirtyNetObjectTrackerCount)])
+	if (UReplicationSystem* ReplicationSystem = GetReplicationSystem(ReplicationSystemId))
 	{
-		DirtyNetObjectTracker->MarkNetObjectDirty(NetObjectIndex);
+		FDirtyNetObjectTracker& DirtyNetObjectTracker = ReplicationSystem->GetReplicationSystemInternal()->GetDirtyNetObjectTracker();
+		DirtyNetObjectTracker.MarkNetObjectDirty(NetObjectIndex);
 	}
 }
 
diff --git a/Engine/Source/Runtime/Experimental/Iris/Core/Private/Iris/ReplicationSystem/DirtyNetObjectTracker.h b/Engine/Source/Runtime/Experimental/Iris/Core/Private/Iris/ReplicationSystem/DirtyNetObjectTracker.h
--- a/Engine/Source/Runtime/Experimental/Iris/Core/Private/Iris/ReplicationSystem/DirtyNetObjectTracker.h
+++ b/Engine/Source/Runtime/Experimental/Iris/Core/Private/Iris/ReplicationSystem/DirtyNetObjectTracker.h
@@ -3,6 +3,12 @@
 #pragma once
 
 #include "Net/Core/NetBitArray.h"
+#include "Net/Core/DirtyNetObjectTracker/GlobalDirtyNetObjectTracker.h"
+
+namespace UE::Net::Private
+{
+	class FNetRefHandleManager;
+}
 
 namespace UE::Net::Private
 {
@@ -11,9 +17,10 @@ IRISCORE_API void MarkNetObjectStateDirty(uint32 ReplicationSystemId, uint32 Net
 
 struct FDirtyNetObjectTrackerInitParams
 {
-	uint32 ReplicationSystemId;
-	uint32 NetObjectIndexRangeStart;
-	uint32 NetObjectIndexRangeEnd;
+	const FNetRefHandleManager* NetRefHandleManager = nullptr;
+	uint32 ReplicationSystemId = 0;
+	uint32 NetObjectIndexRangeStart = 0;
+	uint32 NetObjectIndexRangeEnd = 0;
 };
 
 class FDirtyNetObjectTracker
@@ -24,6 +31,9 @@ public:
 
 	void Init(const FDirtyNetObjectTrackerInitParams& Params);
 
+	/* Update dirty objects with the set of globally marked dirty objects. **/
+	void UpdateDirtyNetObjects();
+
 	// N.B. These methods do not use atomics.
 	// It's up to the user to determine whether it's safe to call the functions or not.
 	FNetBitArrayView GetDirtyNetObjects() const;
@@ -35,15 +45,18 @@ private:
 	using StorageType = FNetBitArrayView::StorageWordType;
 	static constexpr uint32 StorageTypeBitCount = FNetBitArrayView::WordBitCount;
 
-	void Shutdown();
+	void Deinit();
 	void MarkNetObjectDirty(uint32 NetObjectIndex);
 
-	uint32 ReplicationSystemId;
+	const FNetRefHandleManager* NetRefHandleManager;
 	StorageType* DirtyNetObjectContainer;
+	FGlobalDirtyNetObjectTracker::FPollHandle GlobalDirtyTrackerPollHandle;
+	uint32 ReplicationSystemId;
 	uint32 DirtyNetObjectWordCount;
 	uint32 NetObjectIdRangeStart;
 	uint32 NetObjectIdRangeEnd;
 	uint32 NetObjectIdCount;
+	bool bHasPolledGlobalDirtyTracker;
 };
 
 }
diff --git a/Engine/Source/Runtime/Experimental/Iris/Core/Private/Iris/ReplicationSystem/Filtering/NetObjectGroups.cpp b/Engine/Source/Runtime/Experimental/Iris/Core/Private/Iris/ReplicationSystem/Filtering/NetObjectGroups.cpp
--- a/Engine/Source/Runtime/Experimental/Iris/Core/Private/Iris/ReplicationSystem/Filtering/NetObjectGroups.cpp
+++ b/Engine/Source/Runtime/Experimental/Iris/Core/Private/Iris/ReplicationSystem/Filtering/NetObjectGroups.cpp
@@ -179,7 +179,7 @@ FNetObjectGroup* FNetObjectGroups::GetGroup(FNetObjectGroupHandle GroupHandle)
 	return Groups.IsValidIndex(GroupHandle) ? &Groups[GroupHandle] : nullptr;
 }
 
-bool FNetObjectGroups::Contains(FNetObjectGroupHandle GroupHandle, FInternalNetHandle InternalIndex) const
+bool FNetObjectGroups::Contains(FNetObjectGroupHandle GroupHandle, FInternalNetRefIndex InternalIndex) const
 {
 	// If the group does not exist we cannot be in it..
 	const FNetObjectGroup* Group = GetGroup(GroupHandle);
@@ -193,10 +193,10 @@ bool FNetObjectGroups::Contains(FNetObjectGroupHandle GroupHandle, FInternalNetH
 	return InternalIndex ? IsMemberOf(GroupMemberships[InternalIndex], GroupHandle) : false;
 }
 
-void FNetObjectGroups::AddToGroup(FNetObjectGroupHandle GroupHandle, FInternalNetHandle InternalIndex)
+void FNetObjectGroups::AddToGroup(FNetObjectGroupHandle GroupHandle, FInternalNetRefIndex InternalIndex)
 {
 	FNetObjectGroup* Group = GetGroup(GroupHandle);
-	if (InternalIndex != FNetHandleManager::InvalidInternalIndex && Group)
+	if (InternalIndex != FNetRefHandleManager::InvalidInternalIndex && Group)
 	{
 		if (AddGroupMembership(GroupMemberships[InternalIndex], GroupHandle))
 		{
@@ -209,10 +209,10 @@ void FNetObjectGroups::AddToGroup(FNetObjectGroupHandle GroupHandle, FInternalNe
 	}
 }
 
-void FNetObjectGroups::RemoveFromGroup(FNetObjectGroupHandle GroupHandle, FInternalNetHandle InternalIndex)
+void FNetObjectGroups::RemoveFromGroup(FNetObjectGroupHandle GroupHandle, FInternalNetRefIndex InternalIndex)
 {
 	FNetObjectGroup* Group = GetGroup(GroupHandle);
-	if (InternalIndex != FNetHandleManager::InvalidInternalIndex && Group)
+	if (InternalIndex != FNetRefHandleManager::InvalidInternalIndex && Group)
 	{
 		checkSlow(IsMemberOf(GroupMemberships[InternalIndex], GroupHandle));
 
@@ -221,7 +221,7 @@ void FNetObjectGroups::RemoveFromGroup(FNetObjectGroupHandle GroupHandle, FInter
 	}
 }
 
-uint32 FNetObjectGroups::GetNumGroupMemberships(FInternalNetHandle InternalIndex) const
+uint32 FNetObjectGroups::GetNumGroupMemberships(FInternalNetRefIndex InternalIndex) const
 {
 	if (InternalIndex >= (uint32)GroupMemberships.Num())
 	{
@@ -237,7 +237,7 @@ uint32 FNetObjectGroups::GetNumGroupMemberships(FInternalNetHandle InternalIndex
 	return Count;
 }
 
-const FNetObjectGroupHandle* FNetObjectGroups::GetGroupMemberships(FInternalNetHandle InternalIndex, uint32& GroupCount) const
+const FNetObjectGroupHandle* FNetObjectGroups::GetGroupMemberships(FInternalNetRefIndex InternalIndex, uint32& GroupCount) const
 {
 	GroupCount = GetNumGroupMemberships(InternalIndex);
 	if (GroupCount)
diff --git a/Engine/Source/Runtime/Experimental/Iris/Core/Private/Iris/ReplicationSystem/Filtering/NetObjectGroups.h b/Engine/Source/Runtime/Experimental/Iris/Core/Private/Iris/ReplicationSystem/Filtering/NetObjectGroups.h
--- a/Engine/Source/Runtime/Experimental/Iris/Core/Private/Iris/ReplicationSystem/Filtering/NetObjectGroups.h
+++ b/Engine/Source/Runtime/Experimental/Iris/Core/Private/Iris/ReplicationSystem/Filtering/NetObjectGroups.h
@@ -13,7 +13,7 @@ namespace UE::Net
 	typedef uint16 FNetObjectGroupHandle;
 	namespace Private
 	{
-		typedef uint32 FInternalNetHandle;
+		typedef uint32 FInternalNetRefIndex;
 	}
 }
 
@@ -63,13 +63,13 @@ public:
 
 	inline bool IsValidGroup(FNetObjectGroupHandle GroupHandle) const { return GroupHandle && Groups.IsValidIndex(GroupHandle); }
 
-	bool Contains(FNetObjectGroupHandle GroupHandle, FInternalNetHandle InternalIndex) const;
-	void AddToGroup(FNetObjectGroupHandle GroupHandle, FInternalNetHandle InternalIndex);
-	void RemoveFromGroup(FNetObjectGroupHandle GroupHandle, FInternalNetHandle InternalIndex);
+	bool Contains(FNetObjectGroupHandle GroupHandle, FInternalNetRefIndex InternalIndex) const;
+	void AddToGroup(FNetObjectGroupHandle GroupHandle, FInternalNetRefIndex InternalIndex);
+	void RemoveFromGroup(FNetObjectGroupHandle GroupHandle, FInternalNetRefIndex InternalIndex);
 
 	// Returns how many groups the given handle is a member of
-	uint32 GetNumGroupMemberships(FInternalNetHandle InternalIndex) const;
-	const FNetObjectGroupHandle* GetGroupMemberships(FInternalNetHandle InternalIndex, uint32& GroupCount) const;
+	uint32 GetNumGroupMemberships(FInternalNetRefIndex InternalIndex) const;
+	const FNetObjectGroupHandle* GetGroupMemberships(FInternalNetRefIndex InternalIndex, uint32& GroupCount) const;
 
 	// Create and manage named groups, only groups created as a named group will be findable by name
 	FNetObjectGroupHandle CreateNamedGroup(FName GroupName);
diff --git a/Engine/Source/Runtime/Experimental/Iris/Core/Private/Iris/ReplicationSystem/Filtering/ReplicationFiltering.cpp b/Engine/Source/Runtime/Experimental/Iris/Core/Private/Iris/ReplicationSystem/Filtering/ReplicationFiltering.cpp
--- a/Engine/Source/Runtime/Experimental/Iris/Core/Private/Iris/ReplicationSystem/Filtering/ReplicationFiltering.cpp
+++ b/Engine/Source/Runtime/Experimental/Iris/Core/Private/Iris/ReplicationSystem/Filtering/ReplicationFiltering.cpp
@@ -19,11 +19,11 @@
 namespace UE::Net::Private
 {
 
-inline static ENetFilterStatus GetDependentObjectFilterStatus(const FNetHandleManager* NetHandleManager, const FNetBitArray& ObjectsInScope, FInternalNetHandle ObjectIndex)
+inline static ENetFilterStatus GetDependentObjectFilterStatus(const FNetRefHandleManager* NetRefHandleManager, const FNetBitArray& ObjectsInScope, FInternalNetRefIndex ObjectIndex)
 {
-	for (const FInternalNetHandle ParentObjectIndex : NetHandleManager->GetDependentObjectParents(ObjectIndex))
+	for (const FInternalNetRefIndex ParentObjectIndex : NetRefHandleManager->GetDependentObjectParents(ObjectIndex))
 	{
-		if (GetDependentObjectFilterStatus(NetHandleManager, ObjectsInScope, ParentObjectIndex) == ENetFilterStatus::Allow)
+		if (GetDependentObjectFilterStatus(NetRefHandleManager, ObjectsInScope, ParentObjectIndex) == ENetFilterStatus::Allow)
 		{
 			return ENetFilterStatus::Allow;
 		}
@@ -62,8 +62,8 @@ public:
 		FReplicationInstanceProtocol const** InstanceProtocols;
 	};
 
-	FUpdateDirtyObjectsBatchHelper(const FNetHandleManager* InNetHandleManager, uint32 FilterCount)
-		: NetHandleManager(InNetHandleManager)
+	FUpdateDirtyObjectsBatchHelper(const FNetRefHandleManager* InNetRefHandleManager, uint32 FilterCount)
+		: NetRefHandleManager(InNetRefHandleManager)
 	{
 		PerFilterInfos.SetNumUninitialized(FilterCount);
 		ObjectIndicesStorage.SetNumUninitialized(FilterCount * MaxObjectCountPerBatch);
@@ -91,7 +91,7 @@ public:
 				continue;
 			}
 
-			if (const FReplicationInstanceProtocol* InstanceProtocol = NetHandleManager->GetReplicatedObjectDataNoCheck(ObjectIndex).InstanceProtocol)
+			if (const FReplicationInstanceProtocol* InstanceProtocol = NetRefHandleManager->GetReplicatedObjectDataNoCheck(ObjectIndex).InstanceProtocol)
 			{
 				FPerFilterInfo& PerFilterInfo = PerFilterInfosData[FilterIndex];
 				PerFilterInfo.ObjectIndices[PerFilterInfo.ObjectCount] = ObjectIndex;
@@ -114,7 +114,7 @@ private:
 
 	TArray<uint32> ObjectIndicesStorage;
 	TArray<const FReplicationInstanceProtocol*> InstanceProtocolsStorage;
-	const FNetHandleManager* NetHandleManager;
+	const FNetRefHandleManager* NetRefHandleManager;
 };
 
 DEFINE_LOG_CATEGORY_STATIC(LogIrisFiltering, Log, All);
@@ -142,7 +142,7 @@ void FReplicationFiltering::Init(FReplicationFilteringInitParams& Params)
 	ReplicationSystem = Params.ReplicationSystem;
 
 	Connections = Params.Connections;
-	NetHandleManager = Params.NetHandleManager;
+	NetRefHandleManager = Params.NetRefHandleManager;
 	BaselineInvalidationTracker = Params.BaselineInvalidationTracker;
 	Groups = Params.Groups;
 
@@ -239,7 +239,7 @@ FNetBitArrayView FReplicationFiltering::GetGroupFilteredOutObjects(uint32 Connec
 	return MakeNetBitArrayView(ConnectionInfo.GroupFilteredOutObjects);
 }
 
-void FReplicationFiltering::SetOwningConnection(uint32 ObjectIndex, uint32 ConnectionId)
+void FReplicationFiltering::SetOwningConnection(FInternalNetRefIndex ObjectIndex, uint32 ConnectionId)
 {
 	// We allow valid connections as well as 0, which would prevent the object from being replicated to anyone.
 	if ((ConnectionId != 0U) && !Connections->IsValidConnection(ConnectionId))
@@ -264,7 +264,7 @@ void FReplicationFiltering::SetOwningConnection(uint32 ObjectIndex, uint32 Conne
 	}
 }
 
-bool FReplicationFiltering::SetFilter(uint32 ObjectIndex, FNetObjectFilterHandle Filter)
+bool FReplicationFiltering::SetFilter(FInternalNetRefIndex ObjectIndex, FNetObjectFilterHandle Filter)
 {
 	if (Filter == ConnectionFilterHandle)
 	{
@@ -289,16 +289,16 @@ bool FReplicationFiltering::SetFilter(uint32 ObjectIndex, FNetObjectFilterHandle
 		return false;
 	}
 
-	const FNetHandleManager::FReplicatedObjectData& ObjectData = NetHandleManager->GetReplicatedObjectDataNoCheck(ObjectIndex);
-	const TArray<uint8*>& ReplicatedObjectsStateBuffers = NetHandleManager->GetReplicatedObjectStateBuffers();
+	const FNetRefHandleManager::FReplicatedObjectData& ObjectData = NetRefHandleManager->GetReplicatedObjectDataNoCheck(ObjectIndex);
+	const TArray<uint8*>& ReplicatedObjectsStateBuffers = NetRefHandleManager->GetReplicatedObjectStateBuffers();
 	// Let subobjects be filtered like their owners.
-	if (bWantsToUseDynamicFilter && (ObjectData.SubObjectRootIndex != FNetHandleManager::InvalidInternalIndex))
+	if (bWantsToUseDynamicFilter && (ObjectData.SubObjectRootIndex != FNetRefHandleManager::InvalidInternalIndex))
 	{
-		UE_LOG(LogIrisFiltering, Warning, TEXT("Ignoring request to use dynamic filter on object %s due to it being a subobject. Any filtering already successfully applied remains as is."), ToCStr(ObjectData.Handle.ToString()));
+		UE_LOG(LogIrisFiltering, Warning, TEXT("Ignoring request to use dynamic filter on object %s due to it being a subobject. Any filtering already successfully applied remains as is."), ToCStr(ObjectData.RefHandle.ToString()));
 		return false;
 	}
 
-	const FNetBitArray& ObjectsInScope = NetHandleManager->GetScopableInternalIndices();
+	const FNetBitArray& ObjectsInScope = NetRefHandleManager->GetScopableInternalIndices();
 	auto ClearAndDirtyStaticFilter = [this](uint32 ObjIndex)
 	{
 		this->bHasDirtyConnectionFilter = 1;
@@ -372,7 +372,7 @@ bool FReplicationFiltering::SetFilter(uint32 ObjectIndex, FNetObjectFilterHandle
 	return false;
 }
 
-bool FReplicationFiltering::SetConnectionFilter(uint32 ObjectIndex, const FNetBitArrayView& ConnectionIndices, ENetFilterStatus ReplicationStatus)
+bool FReplicationFiltering::SetConnectionFilter(FInternalNetRefIndex ObjectIndex, const FNetBitArrayView& ConnectionIndices, ENetFilterStatus ReplicationStatus)
 {
 	// If we have some sort of other filtering set then remove that. From now on we want to use connection filtering.
 	if (HasOwnerFilter(ObjectIndex))
@@ -511,12 +511,12 @@ void FReplicationFiltering::InitNewConnections()
 	auto InitNewConnection = [this](uint32 ConnectionId)
 	{
 		// Copy default scope
-		const FNetBitArray& ScopableInternalIndices = NetHandleManager->GetScopableInternalIndices();
+		const FNetBitArray& ScopableInternalIndices = NetRefHandleManager->GetScopableInternalIndices();
 		FPerConnectionInfo& ConnectionInfo = this->ConnectionInfos[ConnectionId];
 
 		ConnectionInfo.ConnectionFilteredObjects.Init(MaxObjectCount);
 		ConnectionInfo.ConnectionFilteredObjects.Copy(ScopableInternalIndices);
-		ConnectionInfo.ConnectionFilteredObjects.ClearBit(FNetHandleManager::InvalidInternalIndex);
+		ConnectionInfo.ConnectionFilteredObjects.ClearBit(FNetRefHandleManager::InvalidInternalIndex);
 
 		// Everything enabled by default in groups, i.e. not filtered out.
 		ConnectionInfo.GroupFilteredOutObjects.Init(MaxObjectCount);
@@ -549,7 +549,7 @@ void FReplicationFiltering::InitNewConnections()
 					{
 						GroupFilteredOutObjects.SetBit(ObjectIndex);
 						// Filter subobjects
-						for (const FInternalNetHandle SubObjectIndex : NetHandleManager->GetSubObjects(ObjectIndex))
+						for (const FInternalNetRefIndex SubObjectIndex : NetRefHandleManager->GetSubObjects(ObjectIndex))
 						{
 							GroupFilteredOutObjects.SetBit(SubObjectIndex);
 						}
@@ -625,8 +625,8 @@ void FReplicationFiltering::ResetRemovedConnections()
 
 void FReplicationFiltering::UpdateObjectsInScope()
 {
-	const FNetBitArray& ObjectsInScope = NetHandleManager->GetScopableInternalIndices();
-	const FNetBitArray& PrevObjectsInScope = NetHandleManager->GetPrevFrameScopableInternalIndices();
+	const FNetBitArray& ObjectsInScope = NetRefHandleManager->GetScopableInternalIndices();
+	const FNetBitArray& PrevObjectsInScope = NetRefHandleManager->GetPrevFrameScopableInternalIndices();
 
 	/**
 	 * It's possible for an object to be created, have some filtering applied and then be removed later the same frame.
@@ -673,8 +673,8 @@ void FReplicationFiltering::UpdateObjectsInScope()
 	{
 		uint32* ObjectsWithOwnerFilterStorage = ObjectsWithOwnerFilter.GetData();
 		uint16* ObjectIndexToOwningConnectionStorage = ObjectIndexToOwningConnection.GetData();
-		const uint32* SubObjectInternalIndicesStorage = NetHandleManager->GetSubObjectInternalIndices().GetData();
-		uint32 PrevParentIndex = FNetHandleManager::InvalidInternalIndex;
+		const uint32* SubObjectInternalIndicesStorage = NetRefHandleManager->GetSubObjectInternalIndices().GetData();
+		uint32 PrevParentIndex = FNetRefHandleManager::InvalidInternalIndex;
 		for (uint32 WordIt = 0, WordEndIt = ModifiedWordIndex; WordIt != WordEndIt; ++WordIt)
 		{
 			const uint32 WordIndex = ModifiedWordsStorage[WordIt];
@@ -719,9 +719,9 @@ void FReplicationFiltering::UpdateObjectsInScope()
 					AddedSubObjects ^= LeastSignificantBit;
 
 					const uint32 ObjectIndex = BitOffset + FPlatformMath::CountTrailingZeros(LeastSignificantBit);
-					const uint32 ParentIndex = NetHandleManager->GetReplicatedObjectDataNoCheck(ObjectIndex).SubObjectRootIndex;
+					const uint32 ParentIndex = NetRefHandleManager->GetReplicatedObjectDataNoCheck(ObjectIndex).SubObjectRootIndex;
 					
-					if (ParentIndex == PrevParentIndex || ParentIndex == FNetHandleManager::InvalidInternalIndex)
+					if (ParentIndex == PrevParentIndex || ParentIndex == FNetRefHandleManager::InvalidInternalIndex)
 					{
 						continue;
 					}
@@ -806,12 +806,12 @@ void FReplicationFiltering::UpdateObjectsInScope()
 uint32 FReplicationFiltering::GetOwningConnectionIfDirty(uint32 ObjectIndex) const
 {
 	// If this is a subobject we must check if our parent is updated or not
-	if (NetHandleManager->GetSubObjectInternalIndices().GetBit(ObjectIndex))
+	if (NetRefHandleManager->GetSubObjectInternalIndices().GetBit(ObjectIndex))
 	{
-		const FNetHandleManager::FReplicatedObjectData& ObjectData = NetHandleManager->GetReplicatedObjectDataNoCheck(ObjectIndex);
+		const FNetRefHandleManager::FReplicatedObjectData& ObjectData = NetRefHandleManager->GetReplicatedObjectDataNoCheck(ObjectIndex);
 		if (ObjectData.IsSubObject())
 		{
-			const uint32 ParentIndex = NetHandleManager->GetReplicatedObjectDataNoCheck(ObjectIndex).SubObjectRootIndex;
+			const uint32 ParentIndex = NetRefHandleManager->GetReplicatedObjectDataNoCheck(ObjectIndex).SubObjectRootIndex;
 			return ObjectIndexToOwningConnection[ParentIndex];
 		}
 	}
@@ -834,9 +834,9 @@ void FReplicationFiltering::UpdateOwnerAndConnectionFiltering()
 		{
 			const uint32 OwningConnectionId = ObjectIndexToOwningConnectionStorage[ObjectIndex];
 
-			for (const FInternalNetHandle SubObjectIndex : NetHandleManager->GetSubObjects(ObjectIndex))
+			for (const FInternalNetRefIndex SubObjectIndex : NetRefHandleManager->GetSubObjects(ObjectIndex))
 			{
-				const FNetHandleManager::FReplicatedObjectData& ObjectData = NetHandleManager->GetReplicatedObjectDataNoCheck(SubObjectIndex);
+				const FNetRefHandleManager::FReplicatedObjectData& ObjectData = NetRefHandleManager->GetReplicatedObjectDataNoCheck(SubObjectIndex);
 				if (ObjectData.IsSubObject())
 				{
 					ObjectIndexToOwningConnectionStorage[SubObjectIndex] = OwningConnectionId;
@@ -847,7 +847,7 @@ void FReplicationFiltering::UpdateOwnerAndConnectionFiltering()
 		ObjectsWithDirtyOwner.ForAllSetBits(UpdateOwners);
 	}
 
-	const FNetBitArray& GlobalObjectsInScope = NetHandleManager->GetScopableInternalIndices();
+	const FNetBitArray& GlobalObjectsInScope = NetRefHandleManager->GetScopableInternalIndices();
 
 	// Update filtering
 	if (bHasDirtyConnectionFilter)
@@ -886,7 +886,7 @@ void FReplicationFiltering::UpdateOwnerAndConnectionFiltering()
 					}
 
 					// Subobjects follow suit.
-					for (const FInternalNetHandle SubObjectIndex : NetHandleManager->GetSubObjects(ObjectIndex))
+					for (const FInternalNetRefIndex SubObjectIndex : NetRefHandleManager->GetSubObjects(ObjectIndex))
 					{
 						const bool bEnableObject = bObjectIsInScope && GlobalObjectsInScope.GetBit(SubObjectIndex);
 						const bool bIsGroupEnabled = !GroupFilteredOutObjects.GetBit(SubObjectIndex);
@@ -928,14 +928,14 @@ void FReplicationFiltering::UpdateGroupFiltering()
 				{
 					FNetBitArray& GroupFilteredOutObjects = LocalConnectionInfos[ConnectionId].GroupFilteredOutObjects;
 					FNetBitArray& ObjectsInScopeBeforeDynamicFiltering = LocalConnectionInfos[ConnectionId].ObjectsInScopeBeforeDynamicFiltering;
-					const FNetBitArray& ScopableObjects = NetHandleManager->GetScopableInternalIndices();
+					const FNetBitArray& ScopableObjects = NetRefHandleManager->GetScopableInternalIndices();
 					for (uint32 ObjectIndex : MakeArrayView(Group->Members.GetData(), Group->Members.Num()))
 					{
 						GroupFilteredOutObjects.SetBit(ObjectIndex);
 						ObjectsInScopeBeforeDynamicFiltering.ClearBit(ObjectIndex);
 
 						// Filter subobjects
-						for (const FInternalNetHandle SubObjectIndex : NetHandleManager->GetSubObjects(ObjectIndex))
+						for (const FInternalNetRefIndex SubObjectIndex : NetRefHandleManager->GetSubObjects(ObjectIndex))
 						{
 							const bool bIsScopable = ScopableObjects.GetBit(SubObjectIndex);
 							GroupFilteredOutObjects.SetBitValue(SubObjectIndex, bIsScopable);
@@ -997,8 +997,8 @@ void FReplicationFiltering::UpdateDynamicFiltering()
 
 	// Subobjects will never be added to a dynamic filter, but objects can become dependent at any time.
 	// We need to make sure they are not filtered out.
-	const uint32* SubObjectsData = NetHandleManager->GetSubObjectInternalIndices().GetData();
-	const uint32* DependentObjectsData = NetHandleManager->GetDependentObjectInternalIndices().GetData();
+	const uint32* SubObjectsData = NetRefHandleManager->GetSubObjectInternalIndices().GetData();
+	const uint32* DependentObjectsData = NetRefHandleManager->GetDependentObjectInternalIndices().GetData();
 	const uint32* ObjectsRequiringDynamicFilterUpdateData = ObjectsRequiringDynamicFilterUpdate.GetData();
 
 	uint32* ConnectionIds = static_cast<uint32*>(FMemory_Alloca(ValidConnections.GetNumBits() * sizeof(uint32)));
@@ -1031,7 +1031,7 @@ void FReplicationFiltering::UpdateDynamicFiltering()
 				FNetObjectFilteringParams FilteringParams(MakeNetBitArrayView(Info.FilteredObjects));
 				FilteringParams.OutAllowedObjects = AllowedObjects;
 				FilteringParams.FilteringInfos = NetObjectFilteringInfos.GetData();
-				FilteringParams.StateBuffers = NetHandleManager->GetReplicatedObjectStateBuffers().GetData();
+				FilteringParams.StateBuffers = NetRefHandleManager->GetReplicatedObjectStateBuffers().GetData();
 				FilteringParams.ConnectionId = ConnId;
 				FilteringParams.View = Connections->GetReplicationView(ConnId);
 				Info.Filter->Filter(FilteringParams);
@@ -1091,7 +1091,7 @@ void FReplicationFiltering::UpdateDynamicFiltering()
 					DisabledObjects ^= LeastSignificantBit;
 
 					const uint32 ObjectIndex = BitOffset + FPlatformMath::CountTrailingZeros(LeastSignificantBit);
-					for (const FInternalNetHandle SubObjectIndex : NetHandleManager->GetSubObjects(ObjectIndex))
+					for (const FInternalNetRefIndex SubObjectIndex : NetRefHandleManager->GetSubObjects(ObjectIndex))
 					{						
 						DynamicFilteredOutObjects.SetBit(SubObjectIndex);
 					}
@@ -1109,7 +1109,7 @@ void FReplicationFiltering::UpdateDynamicFiltering()
 					EnabledObjects ^= LeastSignificantBit;
 
 					const uint32 ObjectIndex = BitOffset + FPlatformMath::CountTrailingZeros(LeastSignificantBit);
-					for (const FInternalNetHandle SubObjectIndex : NetHandleManager->GetSubObjects(ObjectIndex))
+					for (const FInternalNetRefIndex SubObjectIndex : NetRefHandleManager->GetSubObjects(ObjectIndex))
 					{
 						DynamicFilteredOutObjects.ClearBit(SubObjectIndex);
 					}
@@ -1133,10 +1133,10 @@ void FReplicationFiltering::UpdateDynamicFiltering()
 		{
 			for (const uint32 DependentObjectIndex : DisabledDependentObjects)
 			{
-				if (GetDependentObjectFilterStatus(NetHandleManager, ConnectionInfo.ObjectsInScope, DependentObjectIndex) == ENetFilterStatus::Allow)
+				if (GetDependentObjectFilterStatus(NetRefHandleManager, ConnectionInfo.ObjectsInScope, DependentObjectIndex) == ENetFilterStatus::Allow)
 				{
 					ConnectionInfo.ObjectsInScope.SetBit(DependentObjectIndex);
-					for (const FInternalNetHandle SubObjectIndex : NetHandleManager->GetSubObjects(DependentObjectIndex))
+					for (const FInternalNetRefIndex SubObjectIndex : NetRefHandleManager->GetSubObjects(DependentObjectIndex))
 					{
 						ConnectionInfo.ObjectsInScope.SetBit(SubObjectIndex);
 					}
@@ -1169,7 +1169,7 @@ void FReplicationFiltering::PostUpdateDynamicFiltering()
 
 void FReplicationFiltering::NotifyFiltersOfDirtyObjects(const FNetBitArrayView& DirtyObjects)
 {
-	FUpdateDirtyObjectsBatchHelper BatchHelper(NetHandleManager, DynamicFilterInfos.Num());
+	FUpdateDirtyObjectsBatchHelper BatchHelper(NetRefHandleManager, DynamicFilterInfos.Num());
 
 	constexpr SIZE_T MaxBatchObjectCount = FUpdateDirtyObjectsBatchHelper::Constants::MaxObjectCountPerBatch;
 	uint32 ObjectIndices[MaxBatchObjectCount];
@@ -1192,7 +1192,7 @@ void FReplicationFiltering::BatchNotifyFiltersOfDirtyObjects(FUpdateDirtyObjects
 	BatchHelper.PrepareBatch(ObjectIndices, ObjectCount, ObjectIndexToDynamicFilterIndex.GetData());
 
 	FNetObjectFilterUpdateParams UpdateParameters;
-	UpdateParameters.StateBuffers = NetHandleManager->GetReplicatedObjectStateBuffers().GetData();
+	UpdateParameters.StateBuffers = NetRefHandleManager->GetReplicatedObjectStateBuffers().GetData();
 	UpdateParameters.FilteringInfos = NetObjectFilteringInfos.GetData();
 
 	for (const FUpdateDirtyObjectsBatchHelper::FPerFilterInfo& PerFilterInfo : BatchHelper.PerFilterInfos)
@@ -1614,7 +1614,7 @@ bool FReplicationFiltering::UpdateGroupFilterEffectsForObject(uint32 ObjectIndex
 		ClearGroupFilterForObject(ObjectIndex);
 	
 		// Filter subobjects
-		for (const FInternalNetHandle SubObjectIndex : NetHandleManager->GetSubObjects(ObjectIndex))
+		for (const FInternalNetRefIndex SubObjectIndex : NetRefHandleManager->GetSubObjects(ObjectIndex))
 		{
 			if (!GetIsFilteredOutByAnyGroup(SubObjectIndex, ConnectionId))
 			{
@@ -1664,13 +1664,13 @@ void FReplicationFiltering::InternalSetGroupFilterStatus(FNetObjectGroupHandle G
 					FPerConnectionInfo& ConnectionInfo = ConnectionInfos[ConnectionId];
 					FNetBitArray& GroupFilteredOutObjects = ConnectionInfo.GroupFilteredOutObjects;
 					FNetBitArray& ObjectsInScopeBeforeDynamicFiltering = ConnectionInfo.ObjectsInScopeBeforeDynamicFiltering;
-					const FNetBitArray& ScopableObjects = NetHandleManager->GetScopableInternalIndices();
+					const FNetBitArray& ScopableObjects = NetRefHandleManager->GetScopableInternalIndices();
 					for (uint32 ObjectIndex : MakeArrayView(Group->Members.GetData(), Group->Members.Num()))
 					{
 						GroupFilteredOutObjects.SetBit(ObjectIndex);
 						ObjectsInScopeBeforeDynamicFiltering.ClearBit(ObjectIndex);
 						// Filter subobjects
-						for (const FInternalNetHandle SubObjectIndex : NetHandleManager->GetSubObjects(ObjectIndex))
+						for (const FInternalNetRefIndex SubObjectIndex : NetRefHandleManager->GetSubObjects(ObjectIndex))
 						{
 							GroupFilteredOutObjects.SetBitValue(SubObjectIndex, ScopableObjects.GetBit(SubObjectIndex));
 							ObjectsInScopeBeforeDynamicFiltering.ClearBit(SubObjectIndex);
@@ -1699,7 +1699,7 @@ void FReplicationFiltering::InternalSetGroupFilterStatus(FNetObjectGroupHandle G
 					const FNetObjectGroup* Group = Groups->GetGroup(GroupHandle);
 
 					// Update group filter mask for the connection
-					const FNetBitArray& ScopedObjects = NetHandleManager->GetScopableInternalIndices();
+					const FNetBitArray& ScopedObjects = NetRefHandleManager->GetScopableInternalIndices();
 			
 					for (uint32 ObjectIndex : MakeArrayView(Group->Members.GetData(), Group->Members.Num()))
 					{
@@ -1724,29 +1724,29 @@ bool FReplicationFiltering::GetGroupFilterStatus(FNetObjectGroupHandle GroupHand
 	return true;
 }
 
-void FReplicationFiltering::NotifyObjectAddedToGroup(FNetObjectGroupHandle GroupHandle, uint32 ObjectIndex)
+void FReplicationFiltering::NotifyObjectAddedToGroup(FNetObjectGroupHandle GroupHandle, FInternalNetRefIndex ObjectIndex)
 {
 	// As we want to avoid the fundamentally bad case of doing tons of single object adds to a group and updating group filter for every call we deffer the update until next call to filter
 	if (FilterGroups.GetBit(GroupHandle))
 	{
 		DirtyFilterGroups.SetBit(GroupHandle);
 
-		UE_LOG(LogIrisFiltering, Verbose, TEXT("ReplicationFiltering::NotifyObjectAddedToGroup Added %s to Groupfilter with GroupIndex: %u"), *(NetHandleManager->GetReplicatedObjectDataNoCheck(ObjectIndex).Handle.ToString()), GroupHandle);
+		UE_LOG(LogIrisFiltering, Verbose, TEXT("ReplicationFiltering::NotifyObjectAddedToGroup Added %s to Groupfilter with GroupIndex: %u"), *(NetRefHandleManager->GetReplicatedObjectDataNoCheck(ObjectIndex).RefHandle.ToString()), GroupHandle);
 	}
 	else if (SubObjectFilterGroups.GetBit(GroupHandle))
 	{
-		UE_LOG(LogIrisFiltering, Verbose, TEXT("ReplicationFiltering::NotifyObjectAddedToGroup Added %s to SubObjectFilter with GroupIndex: %u"), *(NetHandleManager->GetReplicatedObjectDataNoCheck(ObjectIndex).Handle.ToString()), GroupHandle);
+		UE_LOG(LogIrisFiltering, Verbose, TEXT("ReplicationFiltering::NotifyObjectAddedToGroup Added %s to SubObjectFilter with GroupIndex: %u"), *(NetRefHandleManager->GetReplicatedObjectDataNoCheck(ObjectIndex).RefHandle.ToString()), GroupHandle);
 	}
 }
 
-void FReplicationFiltering::NotifyObjectRemovedFromGroup(FNetObjectGroupHandle GroupHandle, uint32 ObjectIndex)
+void FReplicationFiltering::NotifyObjectRemovedFromGroup(FNetObjectGroupHandle GroupHandle, FInternalNetRefIndex ObjectIndex)
 {
-	UE_LOG(LogIrisFiltering, Verbose, TEXT("ReplicationFiltering::NotifyObjectRemovedFromGroup GroupIndex: %u, %s"), GroupHandle, *(NetHandleManager->GetReplicatedObjectDataNoCheck(ObjectIndex).Handle.ToString()));
+	UE_LOG(LogIrisFiltering, Verbose, TEXT("ReplicationFiltering::NotifyObjectRemovedFromGroup GroupIndex: %u, %s"), GroupHandle, *(NetRefHandleManager->GetReplicatedObjectDataNoCheck(ObjectIndex).RefHandle.ToString()));
 
 	if (FilterGroups.GetBit(GroupHandle))
 	{
 		FPerObjectInfo* ConnectionState = GetPerObjectInfo(GroupInfos[GroupHandle].ConnectionStateIndex);
-		const FNetBitArray& ScopedObjects = NetHandleManager->GetScopableInternalIndices();
+		const FNetBitArray& ScopedObjects = NetRefHandleManager->GetScopableInternalIndices();
 
 		auto MarkGroupsDirty = [this, ObjectIndex, GroupHandle, ConnectionState, &ScopedObjects](uint32 ConnectionId)
 		{
@@ -1766,12 +1766,12 @@ void FReplicationFiltering::NotifyObjectRemovedFromGroup(FNetObjectGroupHandle G
 	}
 }
 
-void FReplicationFiltering::NotifyAddedDependentObject(uint32 ObjectIndex)
+void FReplicationFiltering::NotifyAddedDependentObject(FInternalNetRefIndex ObjectIndex)
 {
 	ObjectsRequiringDynamicFilterUpdate.SetBit(ObjectIndex);
 }
 
-void FReplicationFiltering::NotifyRemovedDependentObject(uint32 ObjectIndex)
+void FReplicationFiltering::NotifyRemovedDependentObject(FInternalNetRefIndex ObjectIndex)
 {
 	ObjectsRequiringDynamicFilterUpdate.SetBit(ObjectIndex);
 }
@@ -1876,7 +1876,7 @@ void FReplicationFiltering::RemoveFromDynamicFilter(uint32 ObjectIndex, uint32 F
 void FReplicationFiltering::InvalidateBaselinesForObject(uint32 ObjectIndex, uint32 NewOwningConnectionId, uint32 PrevOwningConnectionId)
 {
 	{
-		const FNetHandleManager::FReplicatedObjectData& ObjectData = NetHandleManager->GetReplicatedObjectDataNoCheck(ObjectIndex);
+		const FNetRefHandleManager::FReplicatedObjectData& ObjectData = NetRefHandleManager->GetReplicatedObjectDataNoCheck(ObjectIndex);
 		// $IRIS TODO Only invalidate baselines if the object has conditions based on owner/non-owner.
 		if (EnumHasAnyFlags(ObjectData.Protocol->ProtocolTraits, EReplicationProtocolTraits::HasLifetimeConditionals))
 		{
@@ -1891,9 +1891,9 @@ void FReplicationFiltering::InvalidateBaselinesForObject(uint32 ObjectIndex, uin
 	}
 
 	// Invalidate baselines for subobjects
-	for (const FInternalNetHandle SubObjectIndex : NetHandleManager->GetSubObjects(ObjectIndex))
+	for (const FInternalNetRefIndex SubObjectIndex : NetRefHandleManager->GetSubObjects(ObjectIndex))
 	{
-		const FNetHandleManager::FReplicatedObjectData& SubObjectData = NetHandleManager->GetReplicatedObjectDataNoCheck(SubObjectIndex);
+		const FNetRefHandleManager::FReplicatedObjectData& SubObjectData = NetRefHandleManager->GetReplicatedObjectDataNoCheck(SubObjectIndex);
 		if (SubObjectData.IsSubObject())
 		{
 			if (EnumHasAnyFlags(SubObjectData.Protocol->ProtocolTraits, EReplicationProtocolTraits::HasLifetimeConditionals))
diff --git a/Engine/Source/Runtime/Experimental/Iris/Core/Private/Iris/ReplicationSystem/Filtering/ReplicationFiltering.h b/Engine/Source/Runtime/Experimental/Iris/Core/Private/Iris/ReplicationSystem/Filtering/ReplicationFiltering.h
--- a/Engine/Source/Runtime/Experimental/Iris/Core/Private/Iris/ReplicationSystem/Filtering/ReplicationFiltering.h
+++ b/Engine/Source/Runtime/Experimental/Iris/Core/Private/Iris/ReplicationSystem/Filtering/ReplicationFiltering.h
@@ -16,9 +16,10 @@ namespace UE::Net
 	namespace Private
 	{
 		class FDeltaCompressionBaselineInvalidationTracker;
-		class FNetHandleManager;
+		class FNetRefHandleManager;
 		class FNetObjectGroups;
 		class FReplicationConnections;
+		typedef uint32 FInternalNetRefIndex;
 	}
 }
 
@@ -28,7 +29,7 @@ namespace UE::Net::Private
 struct FReplicationFilteringInitParams
 {
 	TObjectPtr<UReplicationSystem> ReplicationSystem;
-	const FNetHandleManager* NetHandleManager = nullptr;
+	const FNetRefHandleManager* NetRefHandleManager = nullptr;
 	FNetObjectGroups* Groups = nullptr;
 	FDeltaCompressionBaselineInvalidationTracker* BaselineInvalidationTracker = nullptr;
 	FReplicationConnections* Connections = nullptr;
@@ -54,14 +55,14 @@ public:
 	FNetBitArrayView GetGroupFilteredOutObjects(uint32 ConnectionId) const;
 
 	// Who owns what?
-	void SetOwningConnection(uint32 ObjectIndex, uint32 ConnectionId);
-	uint32 GetOwningConnection(uint32 ObjectIndex) const { return !bHasDirtyOwner ? ObjectIndexToOwningConnection[ObjectIndex] : GetOwningConnectionIfDirty(ObjectIndex); }
+	void SetOwningConnection(FInternalNetRefIndex ObjectIndex, uint32 ConnectionId);
+	uint32 GetOwningConnection(FInternalNetRefIndex ObjectIndex) const { return !bHasDirtyOwner ? ObjectIndexToOwningConnection[ObjectIndex] : GetOwningConnectionIfDirty(ObjectIndex); }
 
 	// Various filters
-	bool SetFilter(uint32 ObjectIndex, FNetObjectFilterHandle Filter);
+	bool SetFilter(FInternalNetRefIndex ObjectIndex, FNetObjectFilterHandle Filter);
 
 	// Set whether an object is allowed to be replicated to certain connections or not.
-	bool SetConnectionFilter(uint32 ObjectIndex, const FNetBitArrayView& ConnectionIndices, ENetFilterStatus ReplicationStatus);
+	bool SetConnectionFilter(FInternalNetRefIndex ObjectIndex, const FNetBitArrayView& ConnectionIndices, ENetFilterStatus ReplicationStatus);
 
 	FNetObjectFilterHandle GetFilterHandle(const FName FilterName) const;
 	UNetObjectFilter* GetFilter(const FName FilterName) const;
@@ -80,11 +81,11 @@ public:
 	void SetGroupFilterStatus(FNetObjectGroupHandle GroupHandle, uint32 ConnectionId, ENetFilterStatus ReplicationStatus);
 	bool GetGroupFilterStatus(FNetObjectGroupHandle GroupHandle, uint32 ConnectionId, ENetFilterStatus& OutReplicationStatus) const;
 
-	void NotifyObjectAddedToGroup(FNetObjectGroupHandle GroupHandle, uint32 ObjectIndex);
-	void NotifyObjectRemovedFromGroup(FNetObjectGroupHandle GroupHandle, uint32 ObjectIndex);
+	void NotifyObjectAddedToGroup(FNetObjectGroupHandle GroupHandle, FInternalNetRefIndex ObjectIndex);
+	void NotifyObjectRemovedFromGroup(FNetObjectGroupHandle GroupHandle, FInternalNetRefIndex ObjectIndex);
 
-	void NotifyAddedDependentObject(uint32 ObjectIndex);
-	void NotifyRemovedDependentObject(uint32 ObjectIndex);
+	void NotifyAddedDependentObject(FInternalNetRefIndex ObjectIndex);
+	void NotifyRemovedDependentObject(FInternalNetRefIndex ObjectIndex);
 	
 	// SubObjectFilter status
 	void AddSubObjectFilter(FNetObjectGroupHandle GroupHandle);
@@ -190,7 +191,7 @@ private:
 
 	// General
 	TObjectPtr<UReplicationSystem> ReplicationSystem;
-	const FNetHandleManager* NetHandleManager = nullptr;
+	const FNetRefHandleManager* NetRefHandleManager = nullptr;
 
 	// Groups
 	FNetObjectGroups* Groups = nullptr;
diff --git a/Engine/Source/Runtime/Experimental/Iris/Core/Private/Iris/ReplicationSystem/LegacyPushModel.cpp b/Engine/Source/Runtime/Experimental/Iris/Core/Private/Iris/ReplicationSystem/LegacyPushModel.cpp
--- a/Engine/Source/Runtime/Experimental/Iris/Core/Private/Iris/ReplicationSystem/LegacyPushModel.cpp
+++ b/Engine/Source/Runtime/Experimental/Iris/Core/Private/Iris/ReplicationSystem/LegacyPushModel.cpp
@@ -7,6 +7,7 @@
 #include "Iris/ReplicationSystem/ReplicationSystem.h"
 #include "Iris/ReplicationSystem/ReplicationSystemInternal.h"
 #include "Iris/ReplicationSystem/NetHandleManager.h"
+#include "Net/Core/DirtyNetObjectTracker/GlobalDirtyNetObjectTracker.h"
 #include "Containers/ArrayView.h"
 
 #if WITH_PUSH_MODEL
@@ -47,9 +48,9 @@ void FNetHandleLegacyPushModelHelper::ShutdownPushModel()
 	UE_NET_SET_IRIS_MARK_PROPERTIES_DIRTY_DELEGATE({});
 }
 
-void FNetHandleLegacyPushModelHelper::SetNetPushID(UObject* Object, FNetPushObjectHandle NetHandle)
+void FNetHandleLegacyPushModelHelper::SetNetPushID(UObject* Object, FNetPushObjectHandle Handle)
 {
-	FObjectNetPushIdHelper::SetNetPushIdDynamic(Object, UEPushModelPrivate::FNetPushObjectId(NetHandle.GetValue()).GetValue());
+	FObjectNetPushIdHelper::SetNetPushIdDynamic(Object, UEPushModelPrivate::FNetPushObjectId(Handle.GetPushObjectId()).GetValue());
 }
 
 void FNetHandleLegacyPushModelHelper::ClearNetPushID(UObject* Object)
@@ -60,13 +61,13 @@ void FNetHandleLegacyPushModelHelper::ClearNetPushID(UObject* Object)
 void FNetHandleLegacyPushModelHelper::MarkPropertyOwnerDirty(const UObject* Object, UEPushModelPrivate::FNetIrisPushObjectId PushId, const int32 RepIndex)
 {
 	const FNetPushObjectHandle Handle(PushId);
-	MarkNetObjectStateDirty(Handle.GetReplicationSystemId(), Handle.GetInternalIndex());
+	MarkNetObjectStateDirty(Handle.GetNetHandle());
 }
 
 void FNetHandleLegacyPushModelHelper::MarkPropertiesOwnerDirty(const UObject* Object, UEPushModelPrivate::FNetIrisPushObjectId PushId, const int32 StartRepIndex, const int32 EndRepIndex)
 {
 	const FNetPushObjectHandle Handle(PushId);
-	MarkNetObjectStateDirty(Handle.GetReplicationSystemId(), Handle.GetInternalIndex());
+	MarkNetObjectStateDirty(Handle.GetNetHandle());
 }
 
 void FNetHandleLegacyPushModelHelper::OptionallyMarkPropertyOwnerDirty(const UObject* Object, UEPushModelPrivate::FNetIrisPushObjectId PushId, const int32 RepIndex)
@@ -74,7 +75,7 @@ void FNetHandleLegacyPushModelHelper::OptionallyMarkPropertyOwnerDirty(const UOb
 	if (IsIrisPushModelEnabled(true))
 	{
 		const FNetPushObjectHandle Handle(PushId);
-		MarkNetObjectStateDirty(Handle.GetReplicationSystemId(), Handle.GetInternalIndex());
+		MarkNetObjectStateDirty(Handle.GetNetHandle());
 	}
 }
 
@@ -83,7 +84,7 @@ void FNetHandleLegacyPushModelHelper::OptionallyMarkPropertiesOwnerDirty(const U
 	if (IsIrisPushModelEnabled(true))
 	{
 		const FNetPushObjectHandle Handle(PushId);
-		MarkNetObjectStateDirty(Handle.GetReplicationSystemId(), Handle.GetInternalIndex());
+		MarkNetObjectStateDirty(Handle.GetNetHandle());
 	}
 }	
 
diff --git a/Engine/Source/Runtime/Experimental/Iris/Core/Private/Iris/ReplicationSystem/LegacyPushModel.h b/Engine/Source/Runtime/Experimental/Iris/Core/Private/Iris/ReplicationSystem/LegacyPushModel.h
--- a/Engine/Source/Runtime/Experimental/Iris/Core/Private/Iris/ReplicationSystem/LegacyPushModel.h
+++ b/Engine/Source/Runtime/Experimental/Iris/Core/Private/Iris/ReplicationSystem/LegacyPushModel.h
@@ -5,7 +5,7 @@
 #if WITH_PUSH_MODEL
 
 #include "CoreTypes.h"
-#include "Iris/ReplicationSystem/NetHandle.h"
+#include "Net/Core/NetHandle/NetHandle.h"
 #include "Net/Core/PushModel/PushModel.h"
 
 namespace UE::Net
@@ -21,34 +21,21 @@ extern IRISCORE_API bool bIsIrisPushModelForceEnabled;
 extern IRISCORE_API int IrisPushModelMode;
 inline bool IsIrisPushModelEnabled(bool bIsPushModelEnabled = IS_PUSH_MODEL_ENABLED()) { return (bIsIrisPushModelForceEnabled | (IrisPushModelMode > 0)) & bIsPushModelEnabled; }
 
-struct FNetPushObjectHandle
+class FNetPushObjectHandle
 {
 public:
-	FNetPushObjectHandle() : Value(Invalid) {}
-	FNetPushObjectHandle(uint32 InIndex, uint32 InReplicationSystemId) : Index(InIndex), ReplicationSystemId(InReplicationSystemId + 1U), ZeroPadding(0) {}
+	FNetPushObjectHandle(FNetHandle InNetHandle);
 
-	uint32 GetInternalIndex() const { return Index; }
-	uint32 GetReplicationSystemId() const { check(ReplicationSystemId != 0U); return ReplicationSystemId - 1U; }
-	bool IsValid() const { return Value != Invalid; }
+	bool IsValid() const;
 
 private:
 	friend struct FNetHandleLegacyPushModelHelper;
 
-	enum : uint64 { Invalid = 0 };
-
-	explicit FNetPushObjectHandle(UEPushModelPrivate::FNetIrisPushObjectId PushId) : Value(PushId) {}
-	uint64 GetValue() const { return Value; }
-
-	union 
-	{
-		struct
-		{
-			uint64 Index : FNetHandle::IdBits;								// Ever increasing index, we use different indices for static and dynamic object
-			uint64 ReplicationSystemId : FNetHandle::ReplicationSystemIdBits;	// ReplicationSystemId, when running in pie, we track the owning instance
-			uint64 ZeroPadding : (64 - FNetHandle::IdBits - FNetHandle::ReplicationSystemIdBits);
-		};
-		uint64 Value : 64;
-	};
+	explicit FNetPushObjectHandle(UEPushModelPrivate::FNetIrisPushObjectId PushId);
+	FNetHandle GetNetHandle() const;
+	UEPushModelPrivate::FNetIrisPushObjectId GetPushObjectId() const;
+
+	FNetHandle NetHandle;
 };
 
 struct FNetHandleLegacyPushModelHelper
@@ -70,6 +57,32 @@ private:
 	static void OptionallyMarkPropertiesOwnerDirty(const UObject* Object, UEPushModelPrivate::FNetIrisPushObjectId PushId, const int32 StartRepIndex, const int32 EndRepIndex);
 };
 
+
+inline FNetPushObjectHandle::FNetPushObjectHandle(FNetHandle InNetHandle)
+: NetHandle(InNetHandle)
+{
+}
+
+inline FNetPushObjectHandle::FNetPushObjectHandle(UEPushModelPrivate::FNetIrisPushObjectId PushId)
+: NetHandle(static_cast<uint64>(PushId))
+{
+}
+
+inline bool FNetPushObjectHandle::IsValid() const
+{
+	return NetHandle.IsValid();
+}
+
+inline FNetHandle FNetPushObjectHandle::GetNetHandle() const
+{
+	return NetHandle;
+}
+
+inline UEPushModelPrivate::FNetIrisPushObjectId FNetPushObjectHandle::GetPushObjectId() const
+{
+	return NetHandle.GetInternalValue();
+}
+
 }
 
 #define UE_NET_IRIS_SET_PUSH_ID(...) UE::Net::Private::FNetHandleLegacyPushModelHelper::SetNetPushID(__VA_ARGS__)
diff --git a/Engine/Source/Runtime/Experimental/Iris/Core/Private/Iris/ReplicationSystem/NetBlob/NetBlob.cpp b/Engine/Source/Runtime/Experimental/Iris/Core/Private/Iris/ReplicationSystem/NetBlob/NetBlob.cpp
--- a/Engine/Source/Runtime/Experimental/Iris/Core/Private/Iris/ReplicationSystem/NetBlob/NetBlob.cpp
+++ b/Engine/Source/Runtime/Experimental/Iris/Core/Private/Iris/ReplicationSystem/NetBlob/NetBlob.cpp
@@ -67,12 +67,12 @@ void FNetBlob::DeserializeCreationInfo(FNetSerializationContext& Context, FNetBl
 	OutCreationInfo.Flags = ENetBlobFlags::None;
 }
 
-void FNetBlob::SerializeWithObject(FNetSerializationContext& Context, FNetHandle NetHandle) const
+void FNetBlob::SerializeWithObject(FNetSerializationContext& Context, FNetRefHandle RefHandle) const
 {
 	SerializeBlob(Context);
 }
 
-void FNetBlob::DeserializeWithObject(FNetSerializationContext& Context, FNetHandle NetHandle)
+void FNetBlob::DeserializeWithObject(FNetSerializationContext& Context, FNetRefHandle RefHandle)
 {
 	check(BlobDescriptor.IsValid());
 	DeserializeBlob(Context);
@@ -149,14 +149,14 @@ void FNetObjectAttachment::DeserializeObjectReference(FNetSerializationContext&
 	ReadFullNetObjectReference(Context, TargetObjectReference);
 }
 
-void FNetObjectAttachment::SerializeSubObjectReference(FNetSerializationContext& Context, FNetHandle NetHandle) const
+void FNetObjectAttachment::SerializeSubObjectReference(FNetSerializationContext& Context, FNetRefHandle RefHandle) const
 {
 	WriteFullNetObjectReference(Context, TargetObjectReference);
 }
 
-void FNetObjectAttachment::DeserializeSubObjectReference(FNetSerializationContext& Context, FNetHandle NetHandle)
+void FNetObjectAttachment::DeserializeSubObjectReference(FNetSerializationContext& Context, FNetRefHandle RefHandle)
 {
-	NetObjectReference = Private::FObjectReferenceCache::MakeNetObjectReference(NetHandle);
+	NetObjectReference = Private::FObjectReferenceCache::MakeNetObjectReference(RefHandle);
 	ReadFullNetObjectReference(Context, TargetObjectReference);
 }
 
diff --git a/Engine/Source/Runtime/Experimental/Iris/Core/Private/Iris/ReplicationSystem/NetBlob/NetBlobAssembler.cpp b/Engine/Source/Runtime/Experimental/Iris/Core/Private/Iris/ReplicationSystem/NetBlob/NetBlobAssembler.cpp
--- a/Engine/Source/Runtime/Experimental/Iris/Core/Private/Iris/ReplicationSystem/NetBlob/NetBlobAssembler.cpp
+++ b/Engine/Source/Runtime/Experimental/Iris/Core/Private/Iris/ReplicationSystem/NetBlob/NetBlobAssembler.cpp
@@ -21,7 +21,7 @@ void FNetBlobAssembler::Init(const FNetBlobAssemblerInitParams& InitParams)
 	ensureMsgf(PartialNetBlobHandlerConfig != nullptr, TEXT("NetBlobAssembler requires a PartialNetBlobHandlerConfig"));
 }
 
-void FNetBlobAssembler::AddPartialNetBlob(FNetSerializationContext& Context, FNetHandle InNetHandle, const TRefCountPtr<FPartialNetBlob>& PartialNetBlob)
+void FNetBlobAssembler::AddPartialNetBlob(FNetSerializationContext& Context, FNetRefHandle InRefHandle, const TRefCountPtr<FPartialNetBlob>& PartialNetBlob)
 {
 	bIsReadyToAssemble = false;
 
@@ -61,7 +61,7 @@ void FNetBlobAssembler::AddPartialNetBlob(FNetSerializationContext& Context, FNe
 			return;
 		}
 
-		NetHandle = InNetHandle;
+		RefHandle = InRefHandle;
 
 		NetBlobCreationInfo = PartialNetBlob->GetOriginalCreationInfo();
 
@@ -78,7 +78,7 @@ void FNetBlobAssembler::AddPartialNetBlob(FNetSerializationContext& Context, FNe
 	}
 	else
 	{
-		if (InNetHandle != NetHandle)
+		if (InRefHandle != RefHandle)
 		{
 			Context.SetError(NetError_PartialNetBlobSequenceError);
 			return;
@@ -139,9 +139,9 @@ TRefCountPtr<FNetBlob> FNetBlobAssembler::Assemble(FNetSerializationContext& Con
 		FNetSerializationContext ReadContext = Context.MakeSubContext(&BitReader);
 		ReadContext.SetTraceCollector(nullptr);
 
-		if (NetHandle.IsValid())
+		if (RefHandle.IsValid())
 		{
-			NetBlob->DeserializeWithObject(ReadContext, NetHandle);
+			NetBlob->DeserializeWithObject(ReadContext, RefHandle);
 		}
 		else
 		{
diff --git a/Engine/Source/Runtime/Experimental/Iris/Core/Private/Iris/ReplicationSystem/NetBlob/NetBlobManager.cpp b/Engine/Source/Runtime/Experimental/Iris/Core/Private/Iris/ReplicationSystem/NetBlob/NetBlobManager.cpp
--- a/Engine/Source/Runtime/Experimental/Iris/Core/Private/Iris/ReplicationSystem/NetBlob/NetBlobManager.cpp
+++ b/Engine/Source/Runtime/Experimental/Iris/Core/Private/Iris/ReplicationSystem/NetBlob/NetBlobManager.cpp
@@ -23,7 +23,7 @@ FNetBlobManager::FNetBlobManager()
 , ObjectReferenceCache(nullptr)
 , Connections(nullptr)
 , PartialNetObjectAttachmentHandlerConfig(nullptr)
-, NetHandleManager(nullptr)
+, NetRefHandleManager(nullptr)
 , bIsServer(false)
 , bSendAttachmentsWithObject(false)
 {
@@ -35,7 +35,7 @@ void FNetBlobManager::Init(FNetBlobManagerInitParams& InitParams)
 
 	ReplicationSystem = InitParams.ReplicationSystem;
 	Connections = &InitParams.ReplicationSystem->GetReplicationSystemInternal()->GetConnections();
-	NetHandleManager = &InitParams.ReplicationSystem->GetReplicationSystemInternal()->GetNetHandleManager();
+	NetRefHandleManager = &InitParams.ReplicationSystem->GetReplicationSystemInternal()->GetNetRefHandleManager();
 	ObjectReferenceCache = &InitParams.ReplicationSystem->GetReplicationSystemInternal()->GetObjectReferenceCache();
 	bIsServer = ReplicationSystem->IsServer();
 	bSendAttachmentsWithObject = InitParams.bSendAttachmentsWithObject; 
@@ -63,8 +63,8 @@ bool FNetBlobManager::QueueNetObjectAttachment(uint32 ConnectionId, const FNetOb
 		return true;
 	}
 
-	FInternalNetHandle OwnerIndex = 0;
-	FInternalNetHandle SubObjectIndex = 0;
+	FInternalNetRefIndex OwnerIndex = 0;
+	FInternalNetRefIndex SubObjectIndex = 0;
 	const FNetObjectReference& NetObjectReference = Attachment->GetNetObjectReference();
 
 	const FNetObjectReference OwnerOrSubObjectReference = TargetRef;	
@@ -120,8 +120,8 @@ bool FNetBlobManager::SendRPC(const UObject* Object, const UObject* SubObject, c
 		return true;
 	}
 
-	FInternalNetHandle OwnerIndex = 0;
-	FInternalNetHandle SubObjectIndex = 0;
+	FInternalNetRefIndex OwnerIndex = 0;
+	FInternalNetRefIndex SubObjectIndex = 0;
 
 	const FNetObjectReference OwnerOrSubObjectReference = ObjectReferenceCache->GetOrCreateObjectReference(SubObject ? SubObject : Object);	
 	FNetObjectReference OwnerReference = OwnerOrSubObjectReference;
@@ -179,8 +179,8 @@ bool FNetBlobManager::SendRPC(uint32 ConnectionId, const UObject* Object, const
 		return true;
 	}
 
-	FInternalNetHandle OwnerIndex = 0;
-	FInternalNetHandle SubObjectIndex = 0;
+	FInternalNetRefIndex OwnerIndex = 0;
+	FInternalNetRefIndex SubObjectIndex = 0;
 
 	const FNetObjectReference OwnerOrSubObjectReference = ObjectReferenceCache->GetOrCreateObjectReference(SubObject ? SubObject : Object);	
 	FNetObjectReference OwnerReference = OwnerOrSubObjectReference;
@@ -214,21 +214,21 @@ bool FNetBlobManager::SendRPC(uint32 ConnectionId, const UObject* Object, const
 	return true;
 }
 
-bool FNetBlobManager::GetOwnerAndSubObjectIndicesFromHandle(FNetHandle NetHandle, FInternalNetHandle& OutOwnerIndex, FInternalNetHandle& OutSubObjectIndex)
+bool FNetBlobManager::GetOwnerAndSubObjectIndicesFromHandle(FNetRefHandle RefHandle, FInternalNetRefIndex& OutOwnerIndex, FInternalNetRefIndex& OutSubObjectIndex)
 {
-	if (!NetHandle.IsValid())
+	if (!RefHandle.IsValid())
 	{
 		return false;
 	}
 
-	const FInternalNetHandle ObjectIndex = NetHandleManager->GetInternalIndex(NetHandle);
-	if (ObjectIndex == FNetHandleManager::InvalidInternalIndex)
+	const FInternalNetRefIndex ObjectIndex = NetRefHandleManager->GetInternalIndex(RefHandle);
+	if (ObjectIndex == FNetRefHandleManager::InvalidInternalIndex)
 	{
 		return false;
 	}
 
-	const FNetHandleManager::FReplicatedObjectData& ObjectData = NetHandleManager->GetReplicatedObjectDataNoCheck(ObjectIndex);
-	if (ObjectData.SubObjectRootIndex != FNetHandleManager::InvalidInternalIndex)
+	const FNetRefHandleManager::FReplicatedObjectData& ObjectData = NetRefHandleManager->GetReplicatedObjectDataNoCheck(ObjectIndex);
+	if (ObjectData.SubObjectRootIndex != FNetRefHandleManager::InvalidInternalIndex)
 	{
 		OutOwnerIndex = ObjectData.SubObjectRootIndex;
 		OutSubObjectIndex = ObjectIndex;
@@ -236,7 +236,7 @@ bool FNetBlobManager::GetOwnerAndSubObjectIndicesFromHandle(FNetHandle NetHandle
 	else
 	{
 		OutOwnerIndex = ObjectIndex;
-		OutSubObjectIndex = FNetHandleManager::InvalidInternalIndex;
+		OutSubObjectIndex = FNetRefHandleManager::InvalidInternalIndex;
 	}
 
 	return true;
@@ -244,7 +244,7 @@ bool FNetBlobManager::GetOwnerAndSubObjectIndicesFromHandle(FNetHandle NetHandle
 
 void FNetBlobManager::ProcessNetObjectAttachmentSendQueue(EProcessMode ProcessMode)
 {
-	AttachmentSendQueue.PrepareProcessQueue(Connections, NetHandleManager);
+	AttachmentSendQueue.PrepareProcessQueue(Connections, NetRefHandleManager);
 	AttachmentSendQueue.ProcessQueue(ProcessMode);
 }
 
@@ -320,7 +320,7 @@ void FNetBlobManager::FNetObjectAttachmentSendQueue::Init(FNetBlobManager* InMan
 	Manager = InManager;
 }
 
-void FNetBlobManager::FNetObjectAttachmentSendQueue::Enqueue(uint32 ConnectionId, FInternalNetHandle OwnerIndex, FInternalNetHandle SubObjectIndex, const TRefCountPtr<FNetObjectAttachment>& Attachment)
+void FNetBlobManager::FNetObjectAttachmentSendQueue::Enqueue(uint32 ConnectionId, FInternalNetRefIndex OwnerIndex, FInternalNetRefIndex SubObjectIndex, const TRefCountPtr<FNetObjectAttachment>& Attachment)
 {
 	FNetObjectAttachmentQueueEntry& QueueEntry = AttachmentQueue.AddDefaulted_GetRef();
 	QueueEntry.ConnectionId = ConnectionId;
@@ -329,7 +329,7 @@ void FNetBlobManager::FNetObjectAttachmentSendQueue::Enqueue(uint32 ConnectionId
 	QueueEntry.Attachment = Attachment;
 }
 
-void FNetBlobManager::FNetObjectAttachmentSendQueue::Enqueue(FInternalNetHandle OwnerIndex, FInternalNetHandle SubObjectIndex, const TRefCountPtr<FNetObjectAttachment>& Attachment)
+void FNetBlobManager::FNetObjectAttachmentSendQueue::Enqueue(FInternalNetRefIndex OwnerIndex, FInternalNetRefIndex SubObjectIndex, const TRefCountPtr<FNetObjectAttachment>& Attachment)
 {
 	FNetObjectAttachmentQueueEntry& QueueEntry = AttachmentQueue.AddDefaulted_GetRef();
 	QueueEntry.ConnectionId = 0;
@@ -340,7 +340,7 @@ void FNetBlobManager::FNetObjectAttachmentSendQueue::Enqueue(FInternalNetHandle
 	bHasMulticastAttachments = true;
 }
 
-void FNetBlobManager::FNetObjectAttachmentSendQueue::PrepareProcessQueue(FReplicationConnections* InConnections, const FNetHandleManager* InNetHandleManager)
+void FNetBlobManager::FNetObjectAttachmentSendQueue::PrepareProcessQueue(FReplicationConnections* InConnections, const FNetRefHandleManager* InNetRefHandleManager)
 {
 	if (ProcessContext.IsValid())
 	{
@@ -348,7 +348,7 @@ void FNetBlobManager::FNetObjectAttachmentSendQueue::PrepareProcessQueue(FReplic
 	}
 
 	ProcessContext.Connections = InConnections;
-	ProcessContext.NetHandleManager = InNetHandleManager;
+	ProcessContext.NetRefHandleManager = InNetRefHandleManager;
 
 	ProcessContext.AttachmentsToObjectsGoingOutOfScope.Init(AttachmentQueue.Num());
 	ProcessContext.AttachmentsToObjectsInScope.Init(AttachmentQueue.Num());
@@ -368,13 +368,13 @@ void FNetBlobManager::FNetObjectAttachmentSendQueue::PrepareProcessQueue(FReplic
 	}
 
 	// Figure out if we have any attachments to objects going out of scope.
-	const FNetBitArray& ScopableObjects = InNetHandleManager->GetScopableInternalIndices();
-	const FNetBitArray& PrevScopableObjects = InNetHandleManager->GetPrevFrameScopableInternalIndices();
+	const FNetBitArray& ScopableObjects = InNetRefHandleManager->GetScopableInternalIndices();
+	const FNetBitArray& PrevScopableObjects = InNetRefHandleManager->GetPrevFrameScopableInternalIndices();
 	
 	uint32 CurrentEntryIndex = 0U;
 	for (const FNetObjectAttachmentQueueEntry& Entry : MakeArrayView(AttachmentQueue))
 	{
-		const uint32 TargetInternalObjectIndex = Entry.SubObjectIndex != FNetHandleManager::InvalidInternalIndex ? Entry.SubObjectIndex : Entry.OwnerIndex;
+		const uint32 TargetInternalObjectIndex = Entry.SubObjectIndex != FNetRefHandleManager::InvalidInternalIndex ? Entry.SubObjectIndex : Entry.OwnerIndex;
 
 		const bool bIsAttachmentToObjectGoingOutOfScope = !ScopableObjects.GetBit(TargetInternalObjectIndex) && PrevScopableObjects.GetBit(TargetInternalObjectIndex);
 		if (bIsAttachmentToObjectGoingOutOfScope)
diff --git a/Engine/Source/Runtime/Experimental/Iris/Core/Private/Iris/ReplicationSystem/NetBlob/NetBlobManager.h b/Engine/Source/Runtime/Experimental/Iris/Core/Private/Iris/ReplicationSystem/NetBlob/NetBlobManager.h
--- a/Engine/Source/Runtime/Experimental/Iris/Core/Private/Iris/ReplicationSystem/NetBlob/NetBlobManager.h
+++ b/Engine/Source/Runtime/Experimental/Iris/Core/Private/Iris/ReplicationSystem/NetBlob/NetBlobManager.h
@@ -4,7 +4,7 @@
 #include "CoreTypes.h"
 #include "Iris/ReplicationSystem/NetBlob/NetBlob.h"
 #include "Iris/ReplicationSystem/NetBlob/NetBlobHandlerManager.h"
-#include "Iris/ReplicationSystem/NetHandleManager.h" // For FInternalNetHandle
+#include "Iris/ReplicationSystem/NetHandleManager.h" // For FInternalNetRefIndex
 #include "UObject/StrongObjectPtr.h"
 
 class UNetBlobHandler;
@@ -17,7 +17,7 @@ namespace UE::Net
 	class FNetObjectReference;
 	namespace Private
 	{
-		class FNetHandleManager;
+		class FNetRefHandleManager;
 		class FObjectReferenceCache;
 		class FReplicationConnections;
 	}
@@ -68,7 +68,7 @@ public:
 private:
 	void RegisterDefaultHandlers();
 
-	bool GetOwnerAndSubObjectIndicesFromHandle(FNetHandle NetHandle, FInternalNetHandle& OutOwnerIndex, FInternalNetHandle& OutSubObjectIndex);
+	bool GetOwnerAndSubObjectIndicesFromHandle(FNetRefHandle RefHandle, FInternalNetRefIndex& OutOwnerIndex, FInternalNetRefIndex& OutSubObjectIndex);
 
 	class FNetObjectAttachmentSendQueue
 	{
@@ -78,12 +78,12 @@ private:
 		void Init(FNetBlobManager* Manager);
 
 		// Unicast
-		void Enqueue(uint32 ConnectionId, FInternalNetHandle OwnerIndex, FInternalNetHandle SubObjectIndex, const TRefCountPtr<FNetObjectAttachment>& Attachment);
+		void Enqueue(uint32 ConnectionId, FInternalNetRefIndex OwnerIndex, FInternalNetRefIndex SubObjectIndex, const TRefCountPtr<FNetObjectAttachment>& Attachment);
 
 		// Multicast
-		void Enqueue(FInternalNetHandle OwnerIndex, FInternalNetHandle SubObjectIndex, const TRefCountPtr<FNetObjectAttachment>& Attachment);
+		void Enqueue(FInternalNetRefIndex OwnerIndex, FInternalNetRefIndex SubObjectIndex, const TRefCountPtr<FNetObjectAttachment>& Attachment);
 
-		void PrepareProcessQueue(FReplicationConnections* InConnections, const FNetHandleManager* InNetHandleManager);
+		void PrepareProcessQueue(FReplicationConnections* InConnections, const FNetRefHandleManager* InNetRefHandleManager);
 		void ProcessQueue(EProcessMode ProcessMode);
 		void ResetProcessQueue();
 	
@@ -91,8 +91,8 @@ private:
 		struct FNetObjectAttachmentQueueEntry
 		{
 			uint32 ConnectionId;
-			FInternalNetHandle OwnerIndex;
-			FInternalNetHandle SubObjectIndex;
+			FInternalNetRefIndex OwnerIndex;
+			FInternalNetRefIndex SubObjectIndex;
 			TRefCountPtr<FNetObjectAttachment> Attachment;
 		};
 		typedef TArray<FNetObjectAttachmentQueueEntry> FQueue;
@@ -110,15 +110,15 @@ private:
 
 			TArray<uint32> ConnectionIds;
 			FReplicationConnections* Connections = nullptr;
-			const FNetHandleManager* NetHandleManager = nullptr;
+			const FNetRefHandleManager* NetRefHandleManager = nullptr;
 
 			void Reset()
 			{
 				Connections = nullptr;
-				NetHandleManager = nullptr;
+				NetRefHandleManager = nullptr;
 			}
 
-			bool IsValid() const { return NetHandleManager != nullptr; }
+			bool IsValid() const { return NetRefHandleManager != nullptr; }
 		};
 		FProcessQueueContext ProcessContext;
 	};
@@ -134,7 +134,7 @@ private:
 	FObjectReferenceCache* ObjectReferenceCache;
 	FReplicationConnections* Connections;
 	const UPartialNetObjectAttachmentHandlerConfig* PartialNetObjectAttachmentHandlerConfig;
-	const FNetHandleManager* NetHandleManager;
+	const FNetRefHandleManager* NetRefHandleManager;
 	bool bIsServer;
 	bool bSendAttachmentsWithObject;
 };
diff --git a/Engine/Source/Runtime/Experimental/Iris/Core/Private/Iris/ReplicationSystem/NetBlob/NetRPC.cpp b/Engine/Source/Runtime/Experimental/Iris/Core/Private/Iris/ReplicationSystem/NetBlob/NetRPC.cpp
--- a/Engine/Source/Runtime/Experimental/Iris/Core/Private/Iris/ReplicationSystem/NetBlob/NetRPC.cpp
+++ b/Engine/Source/Runtime/Experimental/Iris/Core/Private/Iris/ReplicationSystem/NetBlob/NetRPC.cpp
@@ -50,7 +50,7 @@ TArrayView<const FNetObjectReference> FNetRPC::GetExports() const
 	return ReferencesToExport.IsValid() ? MakeArrayView(*ReferencesToExport) : MakeArrayView<const FNetObjectReference>(nullptr, 0);
 }
 
-void FNetRPC::SerializeWithObject(FNetSerializationContext& Context, FNetHandle NetHandle) const
+void FNetRPC::SerializeWithObject(FNetSerializationContext& Context, FNetRefHandle RefHandle) const
 {
 	UE_NET_TRACE_DYNAMIC_NAME_SCOPE(BlobDescriptor->DebugName, *Context.GetBitStreamWriter(), Context.GetTraceCollector(), ENetTraceVerbosity::Verbose);
 
@@ -69,7 +69,7 @@ void FNetRPC::SerializeWithObject(FNetSerializationContext& Context, FNetHandle
 		return;
 	}
 
-	InternalSerializeSubObjectReference(Context, NetHandle);
+	InternalSerializeSubObjectReference(Context, RefHandle);
 	if (Context.HasError())
 	{
 		return;
@@ -90,7 +90,7 @@ void FNetRPC::SerializeWithObject(FNetSerializationContext& Context, FNetHandle
 	}
 }
 
-void FNetRPC::DeserializeWithObject(FNetSerializationContext& Context, FNetHandle NetHandle)
+void FNetRPC::DeserializeWithObject(FNetSerializationContext& Context, FNetRefHandle RefHandle)
 {
 	// We don't know the function name until much later.
 	UE_NET_TRACE_NAMED_SCOPE(TraceScope, RPC, *Context.GetBitStreamReader(), Context.GetTraceCollector(), ENetTraceVerbosity::Trace);
@@ -110,7 +110,7 @@ void FNetRPC::DeserializeWithObject(FNetSerializationContext& Context, FNetHandl
 	}
 
 	// $IRIS TODO Fix. May need to send subobject information
-	InternalDeserializeSubObjectReference(Context, NetHandle);
+	InternalDeserializeSubObjectReference(Context, RefHandle);
 	if (Context.HasErrorOrOverflow())
 	{
 		return;
@@ -289,16 +289,16 @@ void FNetRPC::InternalDeserializeObjectReference(FNetSerializationContext& Conte
 	DeserializeObjectReference(Context);
 }
 
-void FNetRPC::InternalSerializeSubObjectReference(FNetSerializationContext& Context, FNetHandle NetHandle) const
+void FNetRPC::InternalSerializeSubObjectReference(FNetSerializationContext& Context, FNetRefHandle RefHandle) const
 {
 	UE_NET_TRACE_SCOPE(TargetObject, *(Context.GetBitStreamWriter()), Context.GetTraceCollector(), ENetTraceVerbosity::Trace);
-	SerializeSubObjectReference(Context, NetHandle);
+	SerializeSubObjectReference(Context, RefHandle);
 }
 
-void FNetRPC::InternalDeserializeSubObjectReference(FNetSerializationContext& Context, FNetHandle NetHandle)
+void FNetRPC::InternalDeserializeSubObjectReference(FNetSerializationContext& Context, FNetRefHandle RefHandle)
 {
 	UE_NET_TRACE_SCOPE(TargetObject, *(Context.GetBitStreamReader()), Context.GetTraceCollector(), ENetTraceVerbosity::Trace);
-	DeserializeSubObjectReference(Context, NetHandle);
+	DeserializeSubObjectReference(Context, RefHandle);
 }
 
 void FNetRPC::InternalSerializeBlob(FNetSerializationContext& Context) const
@@ -315,7 +315,7 @@ void FNetRPC::InternalDeserializeBlob(FNetSerializationContext& Context)
 
 bool FNetRPC::ResolveFunctionAndObject(FNetSerializationContext& Context)
 {
-	// At this point we need a valid NetHandle and FunctionLocator
+	// At this point we need a valid handle and FunctionLocator
 	if (!NetObjectReference.GetRefHandle().IsValid())
 	{
 		// This can occur if sending side had queued up rpcs to object being invalidated
@@ -524,7 +524,7 @@ void FNetRPC::CallFunction(FNetSerializationContext& Context)
 
 bool FNetRPC::IsServerAllowedToExecuteRPC(FNetSerializationContext& Context) const
 {
-	const FNetHandle Handle = NetObjectReference.GetRefHandle();
+	const FNetRefHandle Handle = NetObjectReference.GetRefHandle();
 	const UReplicationSystem* ReplicationSystem = Context.GetInternalContext()->ReplicationSystem;
 
 	const uint32 OwningConnectionId = ReplicationSystem->GetOwningNetConnection(Handle);
diff --git a/Engine/Source/Runtime/Experimental/Iris/Core/Private/Iris/ReplicationSystem/NetBlob/NetRPC.h b/Engine/Source/Runtime/Experimental/Iris/Core/Private/Iris/ReplicationSystem/NetBlob/NetRPC.h
--- a/Engine/Source/Runtime/Experimental/Iris/Core/Private/Iris/ReplicationSystem/NetBlob/NetRPC.h
+++ b/Engine/Source/Runtime/Experimental/Iris/Core/Private/Iris/ReplicationSystem/NetBlob/NetRPC.h
@@ -32,8 +32,8 @@ private:
 
 	virtual TArrayView<const FNetObjectReference> GetExports() const override final;
 
-	virtual void SerializeWithObject(FNetSerializationContext& Context, FNetHandle NetHandle) const override;
-	virtual void DeserializeWithObject(FNetSerializationContext& Context, FNetHandle NetHandle) override;
+	virtual void SerializeWithObject(FNetSerializationContext& Context, FNetRefHandle RefHandle) const override;
+	virtual void DeserializeWithObject(FNetSerializationContext& Context, FNetRefHandle RefHandle) override;
 
 	virtual void Serialize(FNetSerializationContext& Context) const override;
 	virtual void Deserialize(FNetSerializationContext& Context) override;
@@ -47,8 +47,8 @@ private:
 	void InternalSerializeObjectReference(FNetSerializationContext& Context) const;
 	void InternalDeserializeObjectReference(FNetSerializationContext& Context);
 
-	void InternalSerializeSubObjectReference(FNetSerializationContext& Context, FNetHandle NetHandle) const;
-	void InternalDeserializeSubObjectReference(FNetSerializationContext& Context, FNetHandle NetHandle);
+	void InternalSerializeSubObjectReference(FNetSerializationContext& Context, FNetRefHandle RefHandle) const;
+	void InternalDeserializeSubObjectReference(FNetSerializationContext& Context, FNetRefHandle RefHandle);
 
 	void InternalSerializeBlob(FNetSerializationContext& Context) const;
 	void InternalDeserializeBlob(FNetSerializationContext& Context);
diff --git a/Engine/Source/Runtime/Experimental/Iris/Core/Private/Iris/ReplicationSystem/NetBlob/PartialNetBlob.cpp b/Engine/Source/Runtime/Experimental/Iris/Core/Private/Iris/ReplicationSystem/NetBlob/PartialNetBlob.cpp
--- a/Engine/Source/Runtime/Experimental/Iris/Core/Private/Iris/ReplicationSystem/NetBlob/PartialNetBlob.cpp
+++ b/Engine/Source/Runtime/Experimental/Iris/Core/Private/Iris/ReplicationSystem/NetBlob/PartialNetBlob.cpp
@@ -32,12 +32,12 @@ TArrayView<const FNetObjectReference> FPartialNetBlob::GetExports() const
 	}
 }
 
-void FPartialNetBlob::SerializeWithObject(FNetSerializationContext& Context, FNetHandle NetHandle) const
+void FPartialNetBlob::SerializeWithObject(FNetSerializationContext& Context, FNetRefHandle RefHandle) const
 {
 	InternalSerialize(Context);
 }
 
-void FPartialNetBlob::DeserializeWithObject(FNetSerializationContext& Context, FNetHandle NetHandle)
+void FPartialNetBlob::DeserializeWithObject(FNetSerializationContext& Context, FNetRefHandle RefHandle)
 {
 	InternalDeserialize(Context);
 }
diff --git a/Engine/Source/Runtime/Experimental/Iris/Core/Private/Iris/ReplicationSystem/NetBlob/RawDataNetBlob.cpp b/Engine/Source/Runtime/Experimental/Iris/Core/Private/Iris/ReplicationSystem/NetBlob/RawDataNetBlob.cpp
--- a/Engine/Source/Runtime/Experimental/Iris/Core/Private/Iris/ReplicationSystem/NetBlob/RawDataNetBlob.cpp
+++ b/Engine/Source/Runtime/Experimental/Iris/Core/Private/Iris/ReplicationSystem/NetBlob/RawDataNetBlob.cpp
@@ -44,12 +44,12 @@ void FRawDataNetBlob::InternalDeserialize(FNetSerializationContext& Context)
 	Reader->ReadBitStream(RawData.GetData(), RawDataBitCount);
 }
 
-void FRawDataNetBlob::SerializeWithObject(FNetSerializationContext& Context, FNetHandle NetHandle) const
+void FRawDataNetBlob::SerializeWithObject(FNetSerializationContext& Context, FNetRefHandle RefHandle) const
 {
 	InternalSerialize(Context);
 }
 
-void FRawDataNetBlob::DeserializeWithObject(FNetSerializationContext& Context, FNetHandle NetHandle)
+void FRawDataNetBlob::DeserializeWithObject(FNetSerializationContext& Context, FNetRefHandle RefHandle)
 {
 	InternalDeserialize(Context);
 }
diff --git a/Engine/Source/Runtime/Experimental/Iris/Core/Private/Iris/ReplicationSystem/NetBlob/ReliableNetBlobQueue.cpp b/Engine/Source/Runtime/Experimental/Iris/Core/Private/Iris/ReplicationSystem/NetBlob/ReliableNetBlobQueue.cpp
--- a/Engine/Source/Runtime/Experimental/Iris/Core/Private/Iris/ReplicationSystem/NetBlob/ReliableNetBlobQueue.cpp
+++ b/Engine/Source/Runtime/Experimental/Iris/Core/Private/Iris/ReplicationSystem/NetBlob/ReliableNetBlobQueue.cpp
@@ -49,16 +49,16 @@ bool FReliableNetBlobQueue::IsSafeToDestroy() const
 
 uint32 FReliableNetBlobQueue::Serialize(FNetSerializationContext& Context, FReliableNetBlobQueue::ReplicationRecord& OutRecord)
 {
-	FNetHandle InvalidNetHandle;
+	FNetRefHandle InvalidNetHandle;
 	return SerializeInternal(Context, InvalidNetHandle, OutRecord, false);
 }
 
-uint32 FReliableNetBlobQueue::SerializeWithObject(FNetSerializationContext& Context, FNetHandle NetHandle, FReliableNetBlobQueue::ReplicationRecord& OutRecord)
+uint32 FReliableNetBlobQueue::SerializeWithObject(FNetSerializationContext& Context, FNetRefHandle RefHandle, FReliableNetBlobQueue::ReplicationRecord& OutRecord)
 {
-	return SerializeInternal(Context, NetHandle, OutRecord, true);
+	return SerializeInternal(Context, RefHandle, OutRecord, true);
 }
 
-uint32 FReliableNetBlobQueue::SerializeInternal(FNetSerializationContext& Context, FNetHandle NetHandle, FReliableNetBlobQueue::ReplicationRecord& OutRecord, const bool bSerializeWithObject)
+uint32 FReliableNetBlobQueue::SerializeInternal(FNetSerializationContext& Context, FNetRefHandle RefHandle, FReliableNetBlobQueue::ReplicationRecord& OutRecord, const bool bSerializeWithObject)
 {
 	FNetBitStreamWriter* Writer = Context.GetBitStreamWriter();
 	const FObjectReferenceCache* ObjectReferenceCache = Context.GetInternalContext()->ObjectReferenceCache;
@@ -113,7 +113,7 @@ uint32 FReliableNetBlobQueue::SerializeInternal(FNetSerializationContext& Contex
 		Attachment->SerializeCreationInfo(Context, Attachment->GetCreationInfo());
 		if (bSerializeWithObject)
 		{
-			Attachment->SerializeWithObject(Context, NetHandle);
+			Attachment->SerializeWithObject(Context, RefHandle);
 		}
 		else
 		{
@@ -179,16 +179,16 @@ void FReliableNetBlobQueue::CommitReplicationRecord(FReliableNetBlobQueue::Repli
 
 uint32 FReliableNetBlobQueue::Deserialize(FNetSerializationContext& Context)
 {
-	FNetHandle InvalidNetHandle;
+	FNetRefHandle InvalidNetHandle;
 	return DeserializeInternal(Context, InvalidNetHandle, false);
 }
 
-uint32 FReliableNetBlobQueue::DeserializeWithObject(FNetSerializationContext& Context, FNetHandle NetHandle)
+uint32 FReliableNetBlobQueue::DeserializeWithObject(FNetSerializationContext& Context, FNetRefHandle RefHandle)
 {
-	return DeserializeInternal(Context, NetHandle, true);
+	return DeserializeInternal(Context, RefHandle, true);
 }
 
-uint32 FReliableNetBlobQueue::DeserializeInternal(FNetSerializationContext& Context, FNetHandle NetHandle, const bool bSerializeWithObject)
+uint32 FReliableNetBlobQueue::DeserializeInternal(FNetSerializationContext& Context, FNetRefHandle RefHandle, const bool bSerializeWithObject)
 {
 	INetBlobReceiver* BlobReceiver = Context.GetNetBlobReceiver();
 	checkSlow(BlobReceiver != nullptr);
@@ -241,7 +241,7 @@ uint32 FReliableNetBlobQueue::DeserializeInternal(FNetSerializationContext& Cont
 
 		if (bSerializeWithObject)
 		{
-			Blob->DeserializeWithObject(Context, NetHandle);
+			Blob->DeserializeWithObject(Context, RefHandle);
 		}
 		else
 		{
@@ -253,7 +253,7 @@ uint32 FReliableNetBlobQueue::DeserializeInternal(FNetSerializationContext& Cont
 		// Break if something went wrong with deserialization.
 		if (Context.HasErrorOrOverflow())
 		{
-			UE_LOG(LogIris, Warning, TEXT("Failed to deserialize reliable attachments for %s"), *NetHandle.ToString());
+			UE_LOG(LogIris, Warning, TEXT("Failed to deserialize reliable attachments for %s"), *RefHandle.ToString());
 			break;
 		}
 
diff --git a/Engine/Source/Runtime/Experimental/Iris/Core/Private/Iris/ReplicationSystem/NetBlob/ShrinkWrapNetBlob.cpp b/Engine/Source/Runtime/Experimental/Iris/Core/Private/Iris/ReplicationSystem/NetBlob/ShrinkWrapNetBlob.cpp
--- a/Engine/Source/Runtime/Experimental/Iris/Core/Private/Iris/ReplicationSystem/NetBlob/ShrinkWrapNetBlob.cpp
+++ b/Engine/Source/Runtime/Experimental/Iris/Core/Private/Iris/ReplicationSystem/NetBlob/ShrinkWrapNetBlob.cpp
@@ -21,11 +21,11 @@ TArrayView<const FNetObjectReference> FShrinkWrapNetBlob::GetExports() const
 	return OriginalBlob->CallGetExports();
 }
 
-void FShrinkWrapNetBlob::SerializeWithObject(FNetSerializationContext& Context, FNetHandle NetHandle) const
+void FShrinkWrapNetBlob::SerializeWithObject(FNetSerializationContext& Context, FNetRefHandle RefHandle) const
 {
 	if (Context.GetTraceCollector() != nullptr && OriginalBlob.IsValid())
 	{
-		OriginalBlob->SerializeWithObject(Context, NetHandle);
+		OriginalBlob->SerializeWithObject(Context, RefHandle);
 	}
 	else
 	{
@@ -33,7 +33,7 @@ void FShrinkWrapNetBlob::SerializeWithObject(FNetSerializationContext& Context,
 	}
 }
 
-void FShrinkWrapNetBlob::DeserializeWithObject(FNetSerializationContext& Context, FNetHandle NetHandle)
+void FShrinkWrapNetBlob::DeserializeWithObject(FNetSerializationContext& Context, FNetRefHandle RefHandle)
 {
 	checkf(false, TEXT("%s"), TEXT("This function should not be called. Contact the networking team."));
 }
diff --git a/Engine/Source/Runtime/Experimental/Iris/Core/Private/Iris/ReplicationSystem/NetDependencyData.cpp b/Engine/Source/Runtime/Experimental/Iris/Core/Private/Iris/ReplicationSystem/NetDependencyData.cpp
--- a/Engine/Source/Runtime/Experimental/Iris/Core/Private/Iris/ReplicationSystem/NetDependencyData.cpp
+++ b/Engine/Source/Runtime/Experimental/Iris/Core/Private/Iris/ReplicationSystem/NetDependencyData.cpp
@@ -11,9 +11,9 @@ FNetDependencyData::FNetDependencyData()
 {
 }
 
-void FNetDependencyData::FreeStoredDependencyDataForObject(FInternalNetHandle InternalHandle)
+void FNetDependencyData::FreeStoredDependencyDataForObject(FInternalNetRefIndex InternalIndex)
 {
-	if (FDependencyInfo* Entry = DependencyInfos.Find(InternalHandle))
+	if (FDependencyInfo* Entry = DependencyInfos.Find(InternalIndex))
 	{
 		for (const uint32 ArrayIndex : Entry->ArrayIndices)
 		{
@@ -29,26 +29,26 @@ void FNetDependencyData::FreeStoredDependencyDataForObject(FInternalNetHandle In
 			SubObjectConditionalsStorage.RemoveAt(Entry->SubObjectConditionalArrayIndex);
 		}
 
-		DependencyInfos.Remove(InternalHandle);
+		DependencyInfos.Remove(InternalIndex);
 	}
 }
 
-FNetDependencyData::FDependencyInfo& FNetDependencyData::GetOrCreateCacheEntry(FInternalNetHandle InternalHandle)
+FNetDependencyData::FDependencyInfo& FNetDependencyData::GetOrCreateCacheEntry(FInternalNetRefIndex InternalIndex)
 {
-	FDependencyInfo* Entry = DependencyInfos.Find(InternalHandle);
+	FDependencyInfo* Entry = DependencyInfos.Find(InternalIndex);
 
 	if (!Entry)
 	{
-		Entry = &DependencyInfos.Add(InternalHandle);
+		Entry = &DependencyInfos.Add(InternalIndex);
 		*Entry = FDependencyInfo();
 	}
 
 	return *Entry;
 }
 
-FNetDependencyData::FSubObjectConditionalsArray& FNetDependencyData::GetOrCreateSubObjectConditionalsArray(FInternalNetHandle OwnerHandle)
+FNetDependencyData::FSubObjectConditionalsArray& FNetDependencyData::GetOrCreateSubObjectConditionalsArray(FInternalNetRefIndex OwnerIndex)
 {
-	FDependencyInfo& Entry = GetOrCreateCacheEntry(OwnerHandle);
+	FDependencyInfo& Entry = GetOrCreateCacheEntry(OwnerIndex);
 	check(Entry.ArrayIndices[EArrayType::ChildSubObjects] != FDependencyInfo::InvalidCacheIndex);
 	
 	if (Entry.SubObjectConditionalArrayIndex == FDependencyInfo::InvalidCacheIndex)
@@ -70,18 +70,18 @@ FNetDependencyData::FSubObjectConditionalsArray& FNetDependencyData::GetOrCreate
 	}
 }
 
-FNetDependencyData::FInternalNetHandleArray& FNetDependencyData::GetOrCreateInternalChildSubObjectsArray(FInternalNetHandle OwnerHandle, FSubObjectConditionalsArray*& OutSubObjectConditionals)
+FNetDependencyData::FInternalNetRefIndexArray& FNetDependencyData::GetOrCreateInternalChildSubObjectsArray(FInternalNetRefIndex OwnerIndex, FSubObjectConditionalsArray*& OutSubObjectConditionals)
 {
-	FDependencyInfo& Entry = GetOrCreateCacheEntry(OwnerHandle);
+	FDependencyInfo& Entry = GetOrCreateCacheEntry(OwnerIndex);
 	if (Entry.ArrayIndices[ChildSubObjects] == FDependencyInfo::InvalidCacheIndex)
 	{
 		// Allocate storage for subobjects / atomic dependencies
 		FSparseArrayAllocationInfo AllocInfo = DependentObjectsStorage.AddUninitialized();
 		Entry.ArrayIndices[ChildSubObjects] = AllocInfo.Index;
 
-		FInternalNetHandleArray* InternalHandleArray = new (AllocInfo.Pointer) FInternalNetHandleArray();
+		FInternalNetRefIndexArray* InternalIndexArray = new (AllocInfo.Pointer) FInternalNetRefIndexArray();
 		OutSubObjectConditionals = nullptr;
-		return *InternalHandleArray;
+		return *InternalIndexArray;
 	}
 	else
 	{
@@ -90,18 +90,18 @@ FNetDependencyData::FInternalNetHandleArray& FNetDependencyData::GetOrCreateInte
 	}
 }
 
-FNetDependencyData::FInternalNetHandleArray& FNetDependencyData::GetOrCreateInternalHandleArray(FInternalNetHandle OwnerHandle, EArrayType ArrayType)
+FNetDependencyData::FInternalNetRefIndexArray& FNetDependencyData::GetOrCreateInternalIndexArray(FInternalNetRefIndex OwnerIndex, EArrayType ArrayType)
 {
-	FDependencyInfo& Entry = GetOrCreateCacheEntry(OwnerHandle);
+	FDependencyInfo& Entry = GetOrCreateCacheEntry(OwnerIndex);
 	if (Entry.ArrayIndices[ArrayType] == FDependencyInfo::InvalidCacheIndex)
 	{
 		// Allocate storage for subobjects / atomic dependencies
 		FSparseArrayAllocationInfo AllocInfo = DependentObjectsStorage.AddUninitialized();
 		Entry.ArrayIndices[ArrayType] = AllocInfo.Index;
 
-		FInternalNetHandleArray* InternalHandleArray = new (AllocInfo.Pointer) FInternalNetHandleArray();
+		FInternalNetRefIndexArray* InternalIndexArray = new (AllocInfo.Pointer) FInternalNetRefIndexArray();
 
-		return *InternalHandleArray;
+		return *InternalIndexArray;
 	}
 	else
 	{
diff --git a/Engine/Source/Runtime/Experimental/Iris/Core/Private/Iris/ReplicationSystem/NetDependencyData.h b/Engine/Source/Runtime/Experimental/Iris/Core/Private/Iris/ReplicationSystem/NetDependencyData.h
--- a/Engine/Source/Runtime/Experimental/Iris/Core/Private/Iris/ReplicationSystem/NetDependencyData.h
+++ b/Engine/Source/Runtime/Experimental/Iris/Core/Private/Iris/ReplicationSystem/NetDependencyData.h
@@ -9,7 +9,7 @@
 
 namespace UE::Net::Private
 {
-	typedef uint32 FInternalNetHandle;
+	typedef uint32 FInternalNetRefIndex;
 	typedef int8 FLifeTimeConditionStorage;
 }
 
@@ -18,7 +18,7 @@ namespace UE::Net::Private
 
 struct FChildSubObjectsInfo
 {
-	const FInternalNetHandle* ChildSubObjects = nullptr;
+	const FInternalNetRefIndex* ChildSubObjects = nullptr;
 	const FLifeTimeConditionStorage* SubObjectLifeTimeConditions = nullptr;
 	uint32 NumSubObjects = 0U;
 };
@@ -28,7 +28,7 @@ class FNetDependencyData
 public:
 	FNetDependencyData();
 
-	typedef TArray<FInternalNetHandle, TInlineAllocator<8>> FInternalNetHandleArray;
+	typedef TArray<FInternalNetRefIndex, TInlineAllocator<8>> FInternalNetRefIndexArray;
 
 	typedef TArray<FLifeTimeConditionStorage, TInlineAllocator<8>> FSubObjectConditionalsArray;
 
@@ -42,19 +42,19 @@ public:
 	};
 
 	template<EArrayType TypeIndex>
-	FInternalNetHandleArray& GetOrCreateInternalHandleArray(FInternalNetHandle InternalHandle)
+	FInternalNetRefIndexArray& GetOrCreateInternalIndexArray(FInternalNetRefIndex InternalIndex)
 	{
 		static_assert(TypeIndex != EArrayType::Count, "Invalid array type index");
-		return GetOrCreateInternalHandleArray(InternalHandle, TypeIndex);
+		return GetOrCreateInternalIndexArray(InternalIndex, TypeIndex);
 	};
 
-	FSubObjectConditionalsArray& GetOrCreateSubObjectConditionalsArray(FInternalNetHandle InternalHandle);
+	FSubObjectConditionalsArray& GetOrCreateSubObjectConditionalsArray(FInternalNetRefIndex InternalIndex);
 
-	FInternalNetHandleArray& GetOrCreateInternalChildSubObjectsArray(FInternalNetHandle InternalHandle, FSubObjectConditionalsArray*& OutSubObjectConditionals);
+	FInternalNetRefIndexArray& GetOrCreateInternalChildSubObjectsArray(FInternalNetRefIndex InternalIndex, FSubObjectConditionalsArray*& OutSubObjectConditionals);
 
-	bool GetInternalChildSubObjectAndConditionalArrays(FInternalNetHandle InternalHandle, FInternalNetHandleArray*& OutChildSubObjectsArray, FSubObjectConditionalsArray*& OutSubObjectConditionals)
+	bool GetInternalChildSubObjectAndConditionalArrays(FInternalNetRefIndex InternalIndex, FInternalNetRefIndexArray*& OutChildSubObjectsArray, FSubObjectConditionalsArray*& OutSubObjectConditionals)
 	{
-		const FDependencyInfo* Entry = DependencyInfos.Find(InternalHandle);
+		const FDependencyInfo* Entry = DependencyInfos.Find(InternalIndex);
 		const uint32 ArrayIndex = Entry ? Entry->ArrayIndices[ChildSubObjects] : FDependencyInfo::InvalidCacheIndex;
 
 		if (ArrayIndex == FDependencyInfo::InvalidCacheIndex)
@@ -67,9 +67,9 @@ public:
 		return true;
 	}
 
-	bool GetChildSubObjects(FInternalNetHandle InternalHandle, FChildSubObjectsInfo& OutInfo) const
+	bool GetChildSubObjects(FInternalNetRefIndex InternalIndex, FChildSubObjectsInfo& OutInfo) const
 	{
-		const FDependencyInfo* Entry = DependencyInfos.Find(InternalHandle);
+		const FDependencyInfo* Entry = DependencyInfos.Find(InternalIndex);
 		const uint32 ArrayIndex = Entry ? Entry->ArrayIndices[ChildSubObjects] : FDependencyInfo::InvalidCacheIndex;
 		if (ArrayIndex == FDependencyInfo::InvalidCacheIndex)
 		{
@@ -84,10 +84,10 @@ public:
 
 
 	template<EArrayType TypeIndex>
-	FInternalNetHandleArray* GetInternalHandleArray(FInternalNetHandle InternalHandle)
+	FInternalNetRefIndexArray* GetInternalIndexArray(FInternalNetRefIndex InternalIndex)
 	{
 		static_assert(TypeIndex != EArrayType::Count, "Invalid array type index");
-		const FDependencyInfo* Entry = DependencyInfos.Find(InternalHandle);
+		const FDependencyInfo* Entry = DependencyInfos.Find(InternalIndex);
 		const uint32 ArrayIndex = Entry ? Entry->ArrayIndices[TypeIndex] : FDependencyInfo::InvalidCacheIndex;
 		if (ArrayIndex != FDependencyInfo::InvalidCacheIndex)
 		{
@@ -97,20 +97,20 @@ public:
 	}
 
 	template<EArrayType TypeIndex>
-	TArrayView<const FInternalNetHandle> GetInternalHandleArray(FInternalNetHandle InternalHandle) const
+	TArrayView<const FInternalNetRefIndex> GetInternalIndexArray(FInternalNetRefIndex InternalIndex) const
 	{
 		static_assert(TypeIndex != EArrayType::Count, "Invalid array type index");
 
-		const FDependencyInfo* Entry = DependencyInfos.Find(InternalHandle);
+		const FDependencyInfo* Entry = DependencyInfos.Find(InternalIndex);
 		const uint32 ArrayIndex = Entry ? Entry->ArrayIndices[TypeIndex] : FDependencyInfo::InvalidCacheIndex;
 		if (ArrayIndex != FDependencyInfo::InvalidCacheIndex)
 		{
 			return MakeArrayView(DependentObjectsStorage[ArrayIndex]);
 		}
-		return MakeArrayView<const FInternalNetHandle>(nullptr, 0);
+		return MakeArrayView<const FInternalNetRefIndex>(nullptr, 0);
 	}
 
-	void FreeStoredDependencyDataForObject(FInternalNetHandle InternalHandle);
+	void FreeStoredDependencyDataForObject(FInternalNetRefIndex InternalIndex);
 	
 private:
 	
@@ -122,13 +122,13 @@ private:
 	};
 
 private:
-	FInternalNetHandleArray& GetOrCreateInternalHandleArray(FInternalNetHandle InternalHandle, EArrayType Type);
-	FDependencyInfo& GetOrCreateCacheEntry(FInternalNetHandle InternalHandle);
+	FInternalNetRefIndexArray& GetOrCreateInternalIndexArray(FInternalNetRefIndex InternalIndex, EArrayType Type);
+	FDependencyInfo& GetOrCreateCacheEntry(FInternalNetRefIndex InternalIndex);
 	
 	// Map to track the replicated objects with subObjects or dependencies
-	TMap<FInternalNetHandle, FDependencyInfo> DependencyInfos;
+	TMap<FInternalNetRefIndex, FDependencyInfo> DependencyInfos;
 	// Storage for DependentObjects and SubObjects
-	TSparseArray<FInternalNetHandleArray> DependentObjectsStorage;
+	TSparseArray<FInternalNetRefIndexArray> DependentObjectsStorage;
 	// Storage for SubObject conditionals
 	TSparseArray<FSubObjectConditionalsArray> SubObjectConditionalsStorage;
 };
diff --git a/Engine/Source/Runtime/Experimental/Iris/Core/Private/Iris/ReplicationSystem/NetExports.cpp b/Engine/Source/Runtime/Experimental/Iris/Core/Private/Iris/ReplicationSystem/NetExports.cpp
--- a/Engine/Source/Runtime/Experimental/Iris/Core/Private/Iris/ReplicationSystem/NetExports.cpp
+++ b/Engine/Source/Runtime/Experimental/Iris/Core/Private/Iris/ReplicationSystem/NetExports.cpp
@@ -15,7 +15,7 @@ void FNetExports::InitExportRecordForPacket()
 void FNetExports::CommitExportsToRecord(FExportScope& ExportScope)
 {
 	const FNetExportContext::FBatchExports& BatchExports = ExportScope.ExportContext.GetBatchExports();
-	for (FNetHandle Handle : BatchExports.HandlesExportedInCurrentBatch)
+	for (FNetRefHandle Handle : BatchExports.HandlesExportedInCurrentBatch)
 	{
 		NetHandleExports.Enqueue(Handle);
 	}
@@ -57,7 +57,7 @@ void FNetExports::ProcessPacketDeliveryStatus(UE::Net::EPacketDeliveryStatus Sta
 		uint32 NetHandleExportCount = ExportInfo.NetHandleExportCount;
 		while (NetHandleExportCount)
 		{
-			FNetHandle Handle = NetHandleExports.PeekNoCheck();
+			FNetRefHandle Handle = NetHandleExports.PeekNoCheck();
 			AcknowledgedExports.AcknowledgedExportedHandles.Add(Handle);
 			NetHandleExports.PopNoCheck();
 			--NetHandleExportCount;
diff --git a/Engine/Source/Runtime/Experimental/Iris/Core/Private/Iris/ReplicationSystem/NetExports.h b/Engine/Source/Runtime/Experimental/Iris/Core/Private/Iris/ReplicationSystem/NetExports.h
--- a/Engine/Source/Runtime/Experimental/Iris/Core/Private/Iris/ReplicationSystem/NetExports.h
+++ b/Engine/Source/Runtime/Experimental/Iris/Core/Private/Iris/ReplicationSystem/NetExports.h
@@ -65,7 +65,7 @@ private:
 private:
 
 	TResizableCircularQueue<FExportInfo> ExportRecord;
-	TResizableCircularQueue<FNetHandle> NetHandleExports;
+	TResizableCircularQueue<FNetRefHandle> NetHandleExports;
 	TResizableCircularQueue<FNetToken> NetTokenExports;
 
 	FExportInfo CurrentExportInfo;
diff --git a/Engine/Source/Runtime/Experimental/Iris/Core/Private/Iris/ReplicationSystem/NetHandle.cpp b/Engine/Source/Runtime/Experimental/Iris/Core/Private/Iris/ReplicationSystem/NetHandle.cpp
--- a/Engine/Source/Runtime/Experimental/Iris/Core/Private/Iris/ReplicationSystem/NetHandle.cpp
+++ b/Engine/Source/Runtime/Experimental/Iris/Core/Private/Iris/ReplicationSystem/NetHandle.cpp
@@ -7,38 +7,38 @@
 namespace UE::Net
 {
 
-FString FNetHandle::ToString() const
+FString FNetRefHandle::ToString() const
 {
 	FString Result;
 #if UE_NET_ALLOW_MULTIPLE_REPLICATION_SYSTEMS
 	const uint32 ReplicationSystemIdToDisplay = ReplicationSystemId - 1U;
-	Result = FString::Printf(TEXT("NetHandle (Id=%u):(RepSystemId=%u)"), GetId(), ReplicationSystemIdToDisplay);
+	Result = FString::Printf(TEXT("NetRefHandle (Id=%u):(RepSystemId=%u)"), GetId(), ReplicationSystemIdToDisplay);
 #else
-	Result = FString::Printf(TEXT("NetHandle (Id=%u)"), GetId());
+	Result = FString::Printf(TEXT("NetRefHandle (Id=%u)"), GetId());
 #endif
 	return Result;
 }
 
 }
 
-FStringBuilderBase& operator<<(FStringBuilderBase& Builder, const UE::Net::FNetHandle& NetHandle) 
+FStringBuilderBase& operator<<(FStringBuilderBase& Builder, const UE::Net::FNetRefHandle& NetRefHandle) 
 { 	
 #if UE_NET_ALLOW_MULTIPLE_REPLICATION_SYSTEMS
-	const uint32 ReplicationSystemIdToDisplay = NetHandle.GetReplicationSystemId() - 1U;
-	Builder.Appendf(TEXT("NetHandle (Id=%u):(RepSystemId=%u)"), NetHandle.GetId(), ReplicationSystemIdToDisplay);
+	const uint32 ReplicationSystemIdToDisplay = NetRefHandle.GetReplicationSystemId() - 1U;
+	Builder.Appendf(TEXT("NetRefHandle (Id=%u):(RepSystemId=%u)"), NetRefHandle.GetId(), ReplicationSystemIdToDisplay);
 #else
-	Builder.Appendf(TEXT("NetHandle (Id=%u)"), NetHandle.GetId());
+	Builder.Appendf(TEXT("NetRefHandle (Id=%u)"), NetRefHandle.GetId());
 #endif
 	return Builder;
 }
 
-FAnsiStringBuilderBase& operator<<(FAnsiStringBuilderBase& Builder, const UE::Net::FNetHandle& NetHandle)
+FAnsiStringBuilderBase& operator<<(FAnsiStringBuilderBase& Builder, const UE::Net::FNetRefHandle& NetRefHandle)
 {
-	#if UE_NET_ALLOW_MULTIPLE_REPLICATION_SYSTEMS
-		const uint32 ReplicationSystemIdToDisplay = NetHandle.GetReplicationSystemId() - 1U;
-		Builder.Appendf("NetHandle (Id=%u):(RepSystemId=%u)", NetHandle.GetId(), ReplicationSystemIdToDisplay);
-	#else
-		Builder.Appendf("NetHandle (Id=%u)", NetHandle.GetId());
-	#endif
-		return Builder;
+#if UE_NET_ALLOW_MULTIPLE_REPLICATION_SYSTEMS
+	const uint32 ReplicationSystemIdToDisplay = NetRefHandle.GetReplicationSystemId() - 1U;
+	Builder.Appendf("NetRefHandle (Id=%u):(RepSystemId=%u)", NetRefHandle.GetId(), ReplicationSystemIdToDisplay);
+#else
+	Builder.Appendf("NetRefHandle (Id=%u)", NetRefHandle.GetId());
+#endif
+	return Builder;
 }
