Iris - Implemented globally unique NetHandles.

diff --git a/Engine/Source/Runtime/Experimental/Iris/Core/Private/Iris/ReplicationSystem/NetHandleManager.cpp b/Engine/Source/Runtime/Experimental/Iris/Core/Private/Iris/ReplicationSystem/NetHandleManager.cpp
--- a/Engine/Source/Runtime/Experimental/Iris/Core/Private/Iris/ReplicationSystem/NetHandleManager.cpp
+++ b/Engine/Source/Runtime/Experimental/Iris/Core/Private/Iris/ReplicationSystem/NetHandleManager.cpp
@@ -8,6 +8,7 @@
 #include "Iris/Serialization/NetSerializationContext.h"
 #include "Iris/Serialization/InternalNetSerializationContext.h"
 #include "Net/Core/Trace/NetTrace.h"
+#include "UObject/CoreNetTypes.h"
 #include "UObject/Object.h"
 #include "UObject/UObjectGlobals.h"
 #include "UObject/CoreNetTypes.h"
@@ -15,7 +16,7 @@
 namespace UE::Net::Private
 {
 
-FNetHandleManager::FNetHandleManager(FReplicationProtocolManager& InReplicationProtocolManager, uint32 InReplicationSystemId, uint32 InMaxActiveObjectCount)
+FNetRefHandleManager::FNetRefHandleManager(FReplicationProtocolManager& InReplicationProtocolManager, uint32 InReplicationSystemId, uint32 InMaxActiveObjectCount)
 : ActiveObjectCount(0)
 , MaxActiveObjectCount(InMaxActiveObjectCount)
 , ReplicationSystemId(InReplicationSystemId)
@@ -31,7 +32,8 @@ FNetHandleManager::FNetHandleManager(FReplicationProtocolManager& InReplicationP
 , NextDynamicHandleIndex(1U)
 , ReplicationProtocolManager(InReplicationProtocolManager)
 {
-	// first is reserved for invalid handle
+	static_assert(InvalidInternalIndex == 0, "FNetRefHandleManager::InvalidInternalIndex has an unexpected value");
+	// Mark the invalid index as used
 	AssignedInternalIndices.SetBit(0);
 
 	ReplicatedObjectData.SetNumUninitialized(MaxActiveObjectCount);
@@ -44,10 +46,10 @@ FNetHandleManager::FNetHandleManager(FReplicationProtocolManager& InReplicationP
 	ReplicatedObjectData[InvalidInternalIndex] = FReplicatedObjectData();
 }
 
-uint32 FNetHandleManager::GetNextHandleId(uint32 HandleId) const
+uint32 FNetRefHandleManager::GetNextNetRefHandleId(uint32 HandleId) const
 {
 	// Since we use the lowest bit in the index to indicate if the handle is static or dynamic we can not use all bits as the index
-	constexpr uint32 NetHandleIdIndexBitMask = (1U << (FNetHandle::IdBits - 1U)) - 1U;
+	constexpr uint32 NetHandleIdIndexBitMask = (1U << (FNetRefHandle::IdBits - 1U)) - 1U;
 
 	uint32 NextHandleId = (HandleId + 1) & NetHandleIdIndexBitMask;
 	if (NextHandleId == 0)
@@ -57,23 +59,23 @@ uint32 FNetHandleManager::GetNextHandleId(uint32 HandleId) const
 	return NextHandleId;
 }
 
-FInternalNetHandle FNetHandleManager::GetNextFreeInternalIndex() const
+FInternalNetRefIndex FNetRefHandleManager::GetNextFreeInternalIndex() const
 {
 	const uint32 NextFreeIndex = AssignedInternalIndices.FindFirstZero();
 	return NextFreeIndex != FNetBitArray::InvalidIndex ? NextFreeIndex : InvalidInternalIndex;
 }
 
-uint32 FNetHandleManager::InternalCreateNetObject(const FNetHandle NetHandle, const FReplicationProtocol* ReplicationProtocol)
+FInternalNetRefIndex FNetRefHandleManager::InternalCreateNetObject(const FNetRefHandle NetRefHandle, const FNetHandle GlobalHandle, const FReplicationProtocol* ReplicationProtocol)
 {
 	if (ActiveObjectCount >= MaxActiveObjectCount)
 	{
 		return InvalidInternalIndex;
 	}
 
-	// verify that the handle is free
-	if (HandleMap.Contains(NetHandle))
+	// Verify that the handle is free
+	if (RefHandleToInternalIndex.Contains(NetRefHandle))
 	{
-		ensureAlwaysMsgf(false, TEXT("NetHandleManager::InternalCreateNetObject %s already exists"), *NetHandle.ToString());
+		ensureAlwaysMsgf(false, TEXT("NetRefHandleManager::InternalCreateNetObject %s already exists"), *NetRefHandle.ToString());
 		return InvalidInternalIndex;
 	}
 
@@ -85,7 +87,8 @@ uint32 FNetHandleManager::InternalCreateNetObject(const FNetHandle NetHandle, co
 
 		Data = FReplicatedObjectData();
 
-		Data.Handle = NetHandle;
+		Data.RefHandle = NetRefHandle;
+		Data.NetHandle = GlobalHandle;
 		Data.Protocol = ReplicationProtocol;
 		Data.InstanceProtocol = nullptr;
 		ReplicatedObjectStateBuffers.GetData()[InternalIndex] = nullptr;
@@ -94,8 +97,14 @@ uint32 FNetHandleManager::InternalCreateNetObject(const FNetHandle NetHandle, co
 
 		++ActiveObjectCount;
 
-		// Add map entry so that we can map from NetHandle to InternalIndex
-		HandleMap.Add(Data.Handle, InternalIndex);
+		// Add map entry so that we can map from NetRefHandle to InternalIndex
+		RefHandleToInternalIndex.Add(NetRefHandle, InternalIndex);
+		
+		// Add mapping from global handle to InternalIndex to speed up lookups for ReplicationSystem public API
+		if (GlobalHandle.IsValid())
+		{
+			NetHandleToInternalIndex.Add(GlobalHandle, InternalIndex);
+		}
 
 		// Mark Handle index as assigned and scopable for now
 		AssignedInternalIndices.SetBit(InternalIndex);
@@ -112,7 +121,7 @@ uint32 FNetHandleManager::InternalCreateNetObject(const FNetHandle NetHandle, co
 	return InvalidInternalIndex;
 }
 
-void FNetHandleManager::AttachInstanceProtocol(FInternalNetHandle InternalIndex, const FReplicationInstanceProtocol* InstanceProtocol, UObject* Instance)
+void FNetRefHandleManager::AttachInstanceProtocol(FInternalNetRefIndex InternalIndex, const FReplicationInstanceProtocol* InstanceProtocol, UObject* Instance)
 {
 	if (ensureAlways((InternalIndex != InvalidInternalIndex) && InstanceProtocol))
 	{
@@ -124,7 +133,7 @@ void FNetHandleManager::AttachInstanceProtocol(FInternalNetHandle InternalIndex,
 	}
 }
 
-const FReplicationInstanceProtocol* FNetHandleManager::DetachInstanceProtocol(FInternalNetHandle InternalIndex)
+const FReplicationInstanceProtocol* FNetRefHandleManager::DetachInstanceProtocol(FInternalNetRefIndex InternalIndex)
 {
 	if (ensure(InternalIndex != InvalidInternalIndex))
 	{
@@ -139,40 +148,38 @@ const FReplicationInstanceProtocol* FNetHandleManager::DetachInstanceProtocol(FI
 	return nullptr;
 }
 
-FNetHandle FNetHandleManager::AllocateNetHandle(bool bIsStatic)
+FNetRefHandle FNetRefHandleManager::AllocateNetRefHandle(bool bIsStatic)
 {
 	uint32& NextHandleId = bIsStatic ? NextStaticHandleIndex : NextDynamicHandleIndex;
 
-	const uint32 NewHandleId = MakeNetHandleId(NextHandleId, bIsStatic);
-	FNetHandle NewHandle = MakeNetHandle(NewHandleId, ReplicationSystemId);
+	const uint32 NewHandleId = MakeNetRefHandleId(NextHandleId, bIsStatic);
+	FNetRefHandle NewHandle = MakeNetRefHandle(NewHandleId, ReplicationSystemId);
 
 	// Verify that the handle is free
-	if (HandleMap.Contains(NewHandle))
+	if (RefHandleToInternalIndex.Contains(NewHandle))
 	{
-		checkf(false, TEXT("FNetHandleManager::AllocateNetHandle - Handle %s already exists!"), *NewHandle.ToString());
+		checkf(false, TEXT("FNetRefHandleManager::AllocateNetHandle - Handle %s already exists!"), *NewHandle.ToString());
 
-		return FNetHandle();
+		return FNetRefHandle();
 	}
 
 	// Bump NextHandleId
-	NextHandleId = GetNextHandleId(NextHandleId);
+	NextHandleId = GetNextNetRefHandleId(NextHandleId);
 
 	return NewHandle;
 }
 
-FNetHandle FNetHandleManager::CreateNetObject(FNetHandle WantedHandle, const FReplicationProtocol* ReplicationProtocol)
+FNetRefHandle FNetRefHandleManager::CreateNetObject(FNetRefHandle WantedHandle, FNetHandle GlobalHandle, const FReplicationProtocol* ReplicationProtocol)
 {
-	FNetHandle NetHandle = WantedHandle;
+	FNetRefHandle NetRefHandle = WantedHandle;
 
-	const FInternalNetHandle InternalIndex = InternalCreateNetObject(NetHandle, ReplicationProtocol);
+	const FInternalNetRefIndex InternalIndex = InternalCreateNetObject(NetRefHandle, GlobalHandle, ReplicationProtocol);
 	if (InternalIndex != InvalidInternalIndex)
 	{
 		FReplicatedObjectData& Data = ReplicatedObjectData[InternalIndex];
 		
 		// Allocate storage for outgoing data. We need a valid pointer even if the size is zero.
-		uint8* StateBuffer = (uint8*)FMemory::Malloc(FPlatformMath::Max(ReplicationProtocol->InternalTotalSize, 1U), ReplicationProtocol->InternalTotalAlignment);
-		// Clear state buffer.
-		FMemory::Memzero(StateBuffer, ReplicationProtocol->InternalTotalSize);
+		uint8* StateBuffer = (uint8*)FMemory::MallocZeroed(FPlatformMath::Max(ReplicationProtocol->InternalTotalSize, 1U), ReplicationProtocol->InternalTotalAlignment);
 		if (EnumHasAnyFlags(ReplicationProtocol->ProtocolTraits, EReplicationProtocolTraits::HasConditionalChangeMask))
 		{
 			// Enable all conditions by default. This is to be compatible with the implementation of FRepChangedPropertyTracker where we have hooks into ReplicationSystem.
@@ -185,25 +192,25 @@ FNetHandle FNetHandleManager::CreateNetObject(FNetHandle WantedHandle, const FRe
 		// Bump protocol refcount
 		ReplicationProtocol->AddRef();
 
-		return NetHandle;
+		return NetRefHandle;
 	}
 	else
 	{
-		return FNetHandle();
+		return FNetRefHandle();
 	}
 }
 
-// Create NetHandle not owned by us
-FNetHandle FNetHandleManager::CreateNetObjectFromRemote(FNetHandle WantedHandle, const FReplicationProtocol* ReplicationProtocol)
+// Create NetRefHandle not owned by us
+FNetRefHandle FNetRefHandleManager::CreateNetObjectFromRemote(FNetRefHandle WantedHandle, const FReplicationProtocol* ReplicationProtocol)
 {
-	if (!ensureAlwaysMsgf(WantedHandle.IsValid() && !WantedHandle.IsCompleteHandle(), TEXT("FNetHandleManager::CreateNetObjectFromRemote Expected WantedHandle %s to be valid and incomplete"), *WantedHandle.ToString()))
+	if (!ensureAlwaysMsgf(WantedHandle.IsValid() && !WantedHandle.IsCompleteHandle(), TEXT("FNetRefHandleManager::CreateNetObjectFromRemote Expected WantedHandle %s to be valid and incomplete"), *WantedHandle.ToString()))
 	{
-		return FNetHandle();
+		return FNetRefHandle();
 	}
 
-	FNetHandle NetHandle = MakeNetHandle(WantedHandle.GetId(), ReplicationSystemId);
+	FNetRefHandle NetRefHandle = MakeNetRefHandle(WantedHandle.GetId(), ReplicationSystemId);
 
-	const FInternalNetHandle InternalIndex = InternalCreateNetObject(NetHandle, ReplicationProtocol);
+	const FInternalNetRefIndex InternalIndex = InternalCreateNetObject(NetRefHandle, FNetHandle(), ReplicationProtocol);
 	if (InternalIndex != InvalidInternalIndex)
 	{
 		FReplicatedObjectData& Data = ReplicatedObjectData[InternalIndex];
@@ -222,19 +229,19 @@ FNetHandle FNetHandleManager::CreateNetObjectFromRemote(FNetHandle WantedHandle,
 		// Bump protocol refcount
 		ReplicationProtocol->AddRef();
 
-		return NetHandle;
+		return NetRefHandle;
 	}
 	else
 	{
-		return FNetHandle();
+		return FNetRefHandle();
 	}
 }
 
-void FNetHandleManager::InternalDestroyNetObject(FInternalNetHandle InternalIndex)
+void FNetRefHandleManager::InternalDestroyNetObject(FInternalNetRefIndex InternalIndex)
 {
 	FReplicatedObjectData& Data = ReplicatedObjectData[InternalIndex];
 
-	UE_LOG(LogIris, Verbose, TEXT("FNetHandleManager::InternalDestroyNetObject ( InternalIndex: %u ) %s"), InternalIndex, *Data.Handle.ToString());
+	UE_LOG(LogIris, Verbose, TEXT("FNetRefHandleManager::InternalDestroyNetObject ( InternalIndex: %u ) %s"), InternalIndex, *Data.RefHandle.ToString());
 
 	uint8* StateBuffer = ReplicatedObjectStateBuffers.GetData()[InternalIndex];
 	// Free any allocated resources
@@ -254,16 +261,16 @@ void FNetHandleManager::InternalDestroyNetObject(FInternalNetHandle InternalInde
 	}
 	
 	// If we have subobjects remove them from our list
-	if (FNetDependencyData::FInternalNetHandleArray* SubObjectArray = SubObjects.GetInternalHandleArray<FNetDependencyData::EArrayType::SubObjects>(InternalIndex))
+	if (FNetDependencyData::FInternalNetRefIndexArray* SubObjectArray = SubObjects.GetInternalIndexArray<FNetDependencyData::EArrayType::SubObjects>(InternalIndex))
 	{
-		for (FInternalNetHandle SubObjectInternalHandle : *SubObjectArray)
+		for (FInternalNetRefIndex SubObjectInternalIndex : *SubObjectArray)
 		{
-			InternalRemoveSubObject(InternalIndex, SubObjectInternalHandle, false);
+			InternalRemoveSubObject(InternalIndex, SubObjectInternalIndex, false);
 		}
 		SubObjectArray->Reset(0U);
 	}
 	// Clear ChildSubObjectArray
-	if (FNetDependencyData::FInternalNetHandleArray* ChildSubObjectArray = SubObjects.GetInternalHandleArray<FNetDependencyData::EArrayType::ChildSubObjects>(InternalIndex))
+	if (FNetDependencyData::FInternalNetRefIndexArray* ChildSubObjectArray = SubObjects.GetInternalIndexArray<FNetDependencyData::EArrayType::ChildSubObjects>(InternalIndex))
 	{
 		ChildSubObjectArray->Reset(0U);
 	}
@@ -293,7 +300,7 @@ void FNetHandleManager::InternalDestroyNetObject(FInternalNetHandle InternalInde
 	// Clear pointer to state buffer
 	ReplicatedObjectStateBuffers.GetData()[InternalIndex] = nullptr;
 
-	UE_NET_TRACE_NETHANDLE_DESTROYED(Data.Handle);
+	UE_NET_TRACE_NETHANDLE_DESTROYED(Data.RefHandle);
 
 	Data = FReplicatedObjectData();
 
@@ -319,11 +326,12 @@ void FNetHandleManager::InternalDestroyNetObject(FInternalNetHandle InternalInde
 	--ActiveObjectCount;
 }
 
-FNetHandle FNetHandleManager::CreateHandleForDestructionInfo(FNetHandle Handle, const FReplicationProtocol* DestroydObjectProtocol)
+FNetRefHandle FNetRefHandleManager::CreateHandleForDestructionInfo(FNetRefHandle Handle, const FReplicationProtocol* DestroyedObjectProtocol)
 {
 	// Create destruction info handle carrying destruction info
-	FNetHandle AllocatedHandle = AllocateNetHandle(false);
-	FNetHandle DestructionInfoHandle = CreateNetObject(AllocatedHandle, DestroydObjectProtocol);
+	constexpr bool bIsStaticHandle = false;
+	FNetRefHandle AllocatedHandle = AllocateNetRefHandle(bIsStaticHandle);
+	FNetRefHandle DestructionInfoHandle = CreateNetObject(AllocatedHandle, FNetHandle(), DestroyedObjectProtocol);
 
 	if (DestructionInfoHandle.IsValid())
 	{
@@ -347,7 +355,7 @@ FNetHandle FNetHandleManager::CreateHandleForDestructionInfo(FNetHandle Handle,
 	return DestructionInfoHandle;
 }
 
-void FNetHandleManager::RemoveFromScope(FInternalNetHandle InternalIndex)
+void FNetRefHandleManager::RemoveFromScope(FInternalNetRefIndex InternalIndex)
 {
 	// Can only remove an object from scope if it is assignable
 	if (ensure(AssignedInternalIndices.GetBit(InternalIndex)))
@@ -356,13 +364,16 @@ void FNetHandleManager::RemoveFromScope(FInternalNetHandle InternalIndex)
 	}
 }
 
-void FNetHandleManager::DestroyNetObject(FNetHandle Handle)
+void FNetRefHandleManager::DestroyNetObject(FNetRefHandle RefHandle)
 {
-	FInternalNetHandle InternalIndex = HandleMap.FindAndRemoveChecked(Handle);
+	FInternalNetRefIndex InternalIndex = RefHandleToInternalIndex.FindAndRemoveChecked(RefHandle);
 	if (ensure(AssignedInternalIndices.GetBit(InternalIndex)))
 	{
 		FReplicatedObjectData& Data = ReplicatedObjectData[InternalIndex];
-		check(Data.Handle == Handle);
+		check(Data.RefHandle == RefHandle);
+
+		// Remove mapping from global handle to internal index
+		NetHandleToInternalIndex.Remove(Data.NetHandle);
 
 		// Remove from scopable objects if not already done
 		ScopableInternalIndices.ClearBit(InternalIndex);
@@ -372,12 +383,12 @@ void FNetHandleManager::DestroyNetObject(FNetHandle Handle)
 	}
 }
 
-void FNetHandleManager::DestroyObjectsPendingDestroy()
+void FNetRefHandleManager::DestroyObjectsPendingDestroy()
 {
 	// Destroy Objects pending destroy
-	for (int32 It = PendingDestroyInternalIndices.Num() -1; It >= 0; --It)
+	for (int32 It = PendingDestroyInternalIndices.Num() - 1; It >= 0; --It)
 	{
-		const FInternalNetHandle InternalIndex = PendingDestroyInternalIndices[It];
+		const FInternalNetRefIndex InternalIndex = PendingDestroyInternalIndices[It];
 		if (ReplicatedObjectRefCount[InternalIndex] == 0)
 		{
 			InternalDestroyNetObject(InternalIndex);
@@ -386,13 +397,13 @@ void FNetHandleManager::DestroyObjectsPendingDestroy()
 	}
 }
 
-bool FNetHandleManager::AddSubObject(FNetHandle OwnerHandle, FNetHandle SubObjectHandle, FNetHandle RelativeOtherSubObjectHandle, EAddSubObjectFlags Flags)
+bool FNetRefHandleManager::AddSubObject(FNetRefHandle OwnerHandle, FNetRefHandle SubObjectHandle, FNetRefHandle RelativeOtherSubObjectHandle, EAddSubObjectFlags Flags)
 {
 	check(OwnerHandle != SubObjectHandle);
 
 	// validate objects
-	const FInternalNetHandle OwnerInternalIndex = GetInternalIndex(OwnerHandle);
-	const FInternalNetHandle SubObjectInternalIndex = GetInternalIndex(SubObjectHandle);
+	const FInternalNetRefIndex OwnerInternalIndex = GetInternalIndex(OwnerHandle);
+	const FInternalNetRefIndex SubObjectInternalIndex = GetInternalIndex(SubObjectHandle);
 
 	const bool bIsValidOwner = ensure(OwnerInternalIndex != InvalidInternalIndex);
 	const bool bIsValidSubObejct = ensure(SubObjectInternalIndex != InvalidInternalIndex);
@@ -402,25 +413,32 @@ bool FNetHandleManager::AddSubObject(FNetHandle OwnerHandle, FNetHandle SubObjec
 		return false;
 	}
 
+	const FInternalNetRefIndex RelativeOtherSubObjectInternalIndex = EnumHasAnyFlags(Flags, EAddSubObjectFlags::ReplicateWithSubObject) ? GetInternalIndex(RelativeOtherSubObjectHandle) : InvalidInternalIndex;
+	return InternalAddSubObject(OwnerInternalIndex, SubObjectInternalIndex, RelativeOtherSubObjectInternalIndex, Flags);
+}
+
+bool FNetRefHandleManager::InternalAddSubObject(FInternalNetRefIndex OwnerInternalIndex, FInternalNetRefIndex SubObjectInternalIndex, FInternalNetRefIndex RelativeOtherSubObjectInternalIndex, EAddSubObjectFlags Flags)
+{
+	using namespace UE::Net::Private;
+
 	FReplicatedObjectData& SubObjectData = GetReplicatedObjectDataNoCheck(SubObjectInternalIndex);
-	if (!ensureAlwaysMsgf(SubObjectData.SubObjectRootIndex == InvalidInternalIndex, TEXT("FNetHandleManager::AddSubObject %s is already marked as a subobject"), *SubObjectHandle.ToString()))
+	if (!ensureAlwaysMsgf(SubObjectData.SubObjectRootIndex == InvalidInternalIndex, TEXT("FNetRefHandleManager::AddSubObject %s is already marked as a subobject"), ToCStr(SubObjectData.RefHandle.ToString())))
 	{
 		return false;
 	}
 
-	FNetDependencyData::FInternalNetHandleArray& SubObjectArray = SubObjects.GetOrCreateInternalHandleArray<FNetDependencyData::SubObjects>(OwnerInternalIndex);
+	FNetDependencyData::FInternalNetRefIndexArray& SubObjectArray = SubObjects.GetOrCreateInternalIndexArray<FNetDependencyData::SubObjects>(OwnerInternalIndex);
 	SubObjectArray.Add(SubObjectInternalIndex);
 	SubObjectData.SubObjectRootIndex = OwnerInternalIndex;
 	SubObjectData.bDestroySubObjectWithOwner = EnumHasAnyFlags(Flags, EAddSubObjectFlags::DestroyWithOwner);
 	// Mark the object as a subobject
 	SetIsSubObject(SubObjectInternalIndex, true);
 
-	const FInternalNetHandle RelativeOtherSubObjectInternalIndex = EnumHasAnyFlags(Flags, EAddSubObjectFlags::ReplicateWithSubObject) ? GetInternalIndex(RelativeOtherSubObjectHandle) : InvalidInternalIndex;
-	if (RelativeOtherSubObjectInternalIndex != InvalidInternalIndex && ensureAlwaysMsgf(SubObjectArray.Find(RelativeOtherSubObjectInternalIndex) != INDEX_NONE, TEXT("RelativeOtherSubObjectHandle %s Must be a Subobject of %s"), *RelativeOtherSubObjectHandle.ToString(), *OwnerHandle.ToString()))
+	if (RelativeOtherSubObjectInternalIndex != InvalidInternalIndex && ensureAlwaysMsgf(SubObjectArray.Find(RelativeOtherSubObjectInternalIndex) != INDEX_NONE, TEXT("RelativeOtherSubObjectHandle %s Must be a Subobject of %s"), ToCStr(GetNetRefHandleFromInternalIndex(RelativeOtherSubObjectInternalIndex).ToString()), ToCStr(GetNetRefHandleFromInternalIndex(OwnerInternalIndex).ToString())))
 	{
 		// Add to child array of RelativeOtherSubObjectRelativeIndex for hierarchical replication order
 		FNetDependencyData::FSubObjectConditionalsArray* SubObjectConditionalsArray = nullptr;
-		FNetDependencyData::FInternalNetHandleArray& ChildSubObjectArray = SubObjects.GetOrCreateInternalChildSubObjectsArray(RelativeOtherSubObjectInternalIndex, SubObjectConditionalsArray);
+		FNetDependencyData::FInternalNetRefIndexArray& ChildSubObjectArray = SubObjects.GetOrCreateInternalChildSubObjectsArray(RelativeOtherSubObjectInternalIndex, SubObjectConditionalsArray);
 		ChildSubObjectArray.Add(SubObjectInternalIndex);
 		if (SubObjectConditionalsArray)
 		{
@@ -432,7 +450,7 @@ bool FNetHandleManager::AddSubObject(FNetHandle OwnerHandle, FNetHandle SubObjec
 	{
 		// Add to child array of root
 		FNetDependencyData::FSubObjectConditionalsArray* SubObjectConditionalsArray = nullptr;
-		FNetDependencyData::FInternalNetHandleArray& ChildSubObjectArray = SubObjects.GetOrCreateInternalChildSubObjectsArray(OwnerInternalIndex, SubObjectConditionalsArray);
+		FNetDependencyData::FInternalNetRefIndexArray& ChildSubObjectArray = SubObjects.GetOrCreateInternalChildSubObjectsArray(OwnerInternalIndex, SubObjectConditionalsArray);
 		ChildSubObjectArray.Add(SubObjectInternalIndex);
 		if (SubObjectConditionalsArray)
 		{
@@ -444,7 +462,7 @@ bool FNetHandleManager::AddSubObject(FNetHandle OwnerHandle, FNetHandle SubObjec
 	return true;
 }
 
-void FNetHandleManager::InternalRemoveSubObject(FInternalNetHandle OwnerInternalIndex, FInternalNetHandle SubObjectInternalIndex, bool bRemoveFromSubObjectArray)
+void FNetRefHandleManager::InternalRemoveSubObject(FInternalNetRefIndex OwnerInternalIndex, FInternalNetRefIndex SubObjectInternalIndex, bool bRemoveFromSubObjectArray)
 {
 	// both must be valid
 	if (OwnerInternalIndex != InvalidInternalIndex && SubObjectInternalIndex != InvalidInternalIndex)
@@ -455,14 +473,14 @@ void FNetHandleManager::InternalRemoveSubObject(FInternalNetHandle OwnerInternal
 		if (bRemoveFromSubObjectArray)
 		{
 			// Remove from root parent
-			if (FNetDependencyData::FInternalNetHandleArray* SubObjectArray = SubObjects.GetInternalHandleArray<FNetDependencyData::EArrayType::SubObjects>(OwnerInternalIndex))
+			if (FNetDependencyData::FInternalNetRefIndexArray* SubObjectArray = SubObjects.GetInternalIndexArray<FNetDependencyData::EArrayType::SubObjects>(OwnerInternalIndex))
 			{
 				SubObjectArray->Remove(SubObjectInternalIndex);
 			}
 			// Remove from parents ChildSubObjectArray
 			if (SubObjectData.SubObjectParentIndex != InvalidInternalIndex)
 			{
-				FNetDependencyData::FInternalNetHandleArray* ChildSubObjectArray;
+				FNetDependencyData::FInternalNetRefIndexArray* ChildSubObjectArray;
 				FNetDependencyData::FSubObjectConditionalsArray* SubObjectConditionsArray;
 
 				if (SubObjects.GetInternalChildSubObjectAndConditionalArrays(SubObjectData.SubObjectParentIndex, ChildSubObjectArray, SubObjectConditionsArray))
@@ -488,9 +506,9 @@ void FNetHandleManager::InternalRemoveSubObject(FInternalNetHandle OwnerInternal
 	}
 }
 
-void FNetHandleManager::RemoveSubObject(FNetHandle Handle)
+void FNetRefHandleManager::RemoveSubObject(FNetRefHandle Handle)
 {
-	FInternalNetHandle SubObjectInternalIndex = GetInternalIndex(Handle);
+	FInternalNetRefIndex SubObjectInternalIndex = GetInternalIndex(Handle);
 	checkSlow(SubObjectInternalIndex);
 
 	if (SubObjectInternalIndex != InvalidInternalIndex)
@@ -503,14 +521,14 @@ void FNetHandleManager::RemoveSubObject(FNetHandle Handle)
 	}
 }
 
-bool FNetHandleManager::SetSubObjectNetCondition(FInternalNetHandle SubObjectInternalIndex, FLifeTimeConditionStorage SubObjectCondition)
+bool FNetRefHandleManager::SetSubObjectNetCondition(FInternalNetRefIndex SubObjectInternalIndex, FLifeTimeConditionStorage SubObjectCondition)
 {
 	if (ensure(SubObjectInternalIndex != InvalidInternalIndex))
 	{
-		const FInternalNetHandle SubObjectParentIndex = ReplicatedObjectData[SubObjectInternalIndex].SubObjectParentIndex;
+		const FInternalNetRefIndex SubObjectParentIndex = ReplicatedObjectData[SubObjectInternalIndex].SubObjectParentIndex;
 		if (ensure(SubObjectParentIndex != InvalidInternalIndex))
 		{
-			FNetDependencyData::FInternalNetHandleArray* SubObjectsArray;
+			FNetDependencyData::FInternalNetRefIndexArray* SubObjectsArray;
 			FNetDependencyData::FSubObjectConditionalsArray* SubObjectConditionals;
 			
 			if (SubObjects.GetInternalChildSubObjectAndConditionalArrays(SubObjectParentIndex, SubObjectsArray, SubObjectConditionals))
@@ -536,21 +554,21 @@ bool FNetHandleManager::SetSubObjectNetCondition(FInternalNetHandle SubObjectInt
 	return false;
 }
 
-FNetHandle FNetHandleManager::GetSubObjectOwner(FNetHandle SubObjectHandle) const
+FNetRefHandle FNetRefHandleManager::GetSubObjectOwner(FNetRefHandle SubObjectRefHandle) const
 {
-	const FInternalNetHandle SubObjectInternalIndex = GetInternalIndex(SubObjectHandle);
-	const FInternalNetHandle OwnerInternalIndex = SubObjectInternalIndex != InvalidInternalIndex ? ReplicatedObjectData[SubObjectInternalIndex].SubObjectRootIndex : InvalidInternalIndex;
+	const FInternalNetRefIndex SubObjectInternalIndex = GetInternalIndex(SubObjectRefHandle);
+	const FInternalNetRefIndex OwnerInternalIndex = SubObjectInternalIndex != InvalidInternalIndex ? ReplicatedObjectData[SubObjectInternalIndex].SubObjectRootIndex : InvalidInternalIndex;
 
-	return OwnerInternalIndex != InvalidInternalIndex ? ReplicatedObjectData[OwnerInternalIndex].Handle : FNetHandle();
+	return OwnerInternalIndex != InvalidInternalIndex ? ReplicatedObjectData[OwnerInternalIndex].RefHandle : FNetRefHandle();
 }
 
-bool FNetHandleManager::AddDependentObject(FNetHandle ParentHandle, FNetHandle DependentObjectHandle, EAddDependentObjectFlags Flags)
+bool FNetRefHandleManager::AddDependentObject(FNetRefHandle ParentRefHandle, FNetRefHandle DependentObjectRefHandle, EAddDependentObjectFlags Flags)
 {
-	check(ParentHandle != DependentObjectHandle);
+	check(ParentRefHandle != DependentObjectRefHandle);
 
 	// validate objects
-	FInternalNetHandle ParentInternalIndex = GetInternalIndex(ParentHandle);
-	FInternalNetHandle DependentObjectInternalIndex = GetInternalIndex(DependentObjectHandle);
+	FInternalNetRefIndex ParentInternalIndex = GetInternalIndex(ParentRefHandle);
+	FInternalNetRefIndex DependentObjectInternalIndex = GetInternalIndex(DependentObjectRefHandle);
 
 	const bool bIsValidOwner = ensure(ParentInternalIndex != InvalidInternalIndex);
 	const bool bIsValidDependentObject = ensure(DependentObjectInternalIndex != InvalidInternalIndex);
@@ -569,14 +587,14 @@ bool FNetHandleManager::AddDependentObject(FNetHandle ParentHandle, FNetHandle D
 	check(!SubObjectInternalIndices.GetBit(ParentInternalIndex));
 
 	// Add dependent to parents dependent object list
-	FNetDependencyData::FInternalNetHandleArray& ParentDependentObjectArray = SubObjects.GetOrCreateInternalHandleArray<FNetDependencyData::DependentObjects>(ParentInternalIndex);
+	FNetDependencyData::FInternalNetRefIndexArray& ParentDependentObjectArray = SubObjects.GetOrCreateInternalIndexArray<FNetDependencyData::DependentObjects>(ParentInternalIndex);
 	const bool bDependentIsAlreadyDependant = ParentDependentObjectArray.Find(DependentObjectInternalIndex) != INDEX_NONE;
 	if (!bDependentIsAlreadyDependant)
 	{
 		ParentDependentObjectArray.Add(DependentObjectInternalIndex);
 	}
 
-	FNetDependencyData::FInternalNetHandleArray& DependentParentObjectArray = SubObjects.GetOrCreateInternalHandleArray<FNetDependencyData::ParentObjects>(DependentObjectInternalIndex);
+	FNetDependencyData::FInternalNetRefIndexArray& DependentParentObjectArray = SubObjects.GetOrCreateInternalIndexArray<FNetDependencyData::ParentObjects>(DependentObjectInternalIndex);
 	const bool bDependentHadParentAlready = DependentParentObjectArray.Find(ParentInternalIndex) != INDEX_NONE;
 	if (!bDependentHadParentAlready)
 	{
@@ -593,15 +611,15 @@ bool FNetHandleManager::AddDependentObject(FNetHandle ParentHandle, FNetHandle D
 	{
 		// If this gets out of sync something is messed up
 		check(bDependentHadParentAlready == bDependentIsAlreadyDependant);
-		ensureAlwaysMsgf(false, TEXT("FNetHandleManager::AddDependentObject %s already is a child of %s parent %s"), *DependentObjectHandle.ToString(), *GetReplicatedObjectDataNoCheck(DependentObjectData.SubObjectRootIndex).Handle.ToString(), *ParentHandle.ToString());
+		ensureAlwaysMsgf(false, TEXT("FNetRefHandleManager::AddDependentObject %s already is a child of %s parent %s"), *DependentObjectRefHandle.ToString(), *GetReplicatedObjectDataNoCheck(DependentObjectData.SubObjectRootIndex).RefHandle.ToString(), *ParentRefHandle.ToString());
 	}
 
 	return true;
 }
 
-void FNetHandleManager::RemoveDependentObject(FNetHandle DependentHandle)
+void FNetRefHandleManager::RemoveDependentObject(FNetRefHandle DependentHandle)
 {
-	FInternalNetHandle DependentInternalIndex = GetInternalIndex(DependentHandle);
+	FInternalNetRefIndex DependentInternalIndex = GetInternalIndex(DependentHandle);
 
 	if (DependentInternalIndex != InvalidInternalIndex)
 	{
@@ -609,11 +627,11 @@ void FNetHandleManager::RemoveDependentObject(FNetHandle DependentHandle)
 	}
 }
 
-void FNetHandleManager::InternalRemoveDependentObject(FInternalNetHandle ParentInternalIndex, FInternalNetHandle DependentInternalIndex, ERemoveDependentObjectFlags Flags)
+void FNetRefHandleManager::InternalRemoveDependentObject(FInternalNetRefIndex ParentInternalIndex, FInternalNetRefIndex DependentInternalIndex, ERemoveDependentObjectFlags Flags)
 {
 	if (EnumHasAnyFlags(Flags, ERemoveDependentObjectFlags::RemoveFromDependentParentObjects))
 	{
-		if (FNetDependencyData::FInternalNetHandleArray* ParentObjectArray = SubObjects.GetInternalHandleArray<FNetDependencyData::ParentObjects>(DependentInternalIndex))
+		if (FNetDependencyData::FInternalNetRefIndexArray* ParentObjectArray = SubObjects.GetInternalIndexArray<FNetDependencyData::ParentObjects>(DependentInternalIndex))
 		{
 			ParentObjectArray->Remove(ParentInternalIndex);
 			if (ParentObjectArray->Num() == 0)
@@ -627,7 +645,7 @@ void FNetHandleManager::InternalRemoveDependentObject(FInternalNetHandle ParentI
 
 	if (EnumHasAnyFlags(Flags, ERemoveDependentObjectFlags::RemoveFromParentDependentObjects))
 	{
-		if (FNetDependencyData::FInternalNetHandleArray* ParentDependentObjectArray = SubObjects.GetInternalHandleArray<FNetDependencyData::DependentObjects>(ParentInternalIndex))
+		if (FNetDependencyData::FInternalNetRefIndexArray* ParentDependentObjectArray = SubObjects.GetInternalIndexArray<FNetDependencyData::DependentObjects>(ParentInternalIndex))
 		{
 			ParentDependentObjectArray->Remove(DependentInternalIndex);
 			if (ParentDependentObjectArray->Num() == 0)
@@ -640,12 +658,12 @@ void FNetHandleManager::InternalRemoveDependentObject(FInternalNetHandle ParentI
 	}
 }
 
-void FNetHandleManager::InternalRemoveDependentObject(FInternalNetHandle DependentInternalIndex)
+void FNetRefHandleManager::InternalRemoveDependentObject(FInternalNetRefIndex DependentInternalIndex)
 {
 	// Remove from all parents
-	if (FNetDependencyData::FInternalNetHandleArray* ParentObjectArray = SubObjects.GetInternalHandleArray<FNetDependencyData::ParentObjects>(DependentInternalIndex))
+	if (FNetDependencyData::FInternalNetRefIndexArray* ParentObjectArray = SubObjects.GetInternalIndexArray<FNetDependencyData::ParentObjects>(DependentInternalIndex))
 	{
-		for (FInternalNetHandle ParentInternalIndex : *ParentObjectArray)
+		for (FInternalNetRefIndex ParentInternalIndex : *ParentObjectArray)
 		{
 			// Flag is set to only update data on the parent to avoid modifying the array we iterate over
 			InternalRemoveDependentObject(ParentInternalIndex, DependentInternalIndex, ERemoveDependentObjectFlags::RemoveFromParentDependentObjects);
@@ -654,9 +672,9 @@ void FNetHandleManager::InternalRemoveDependentObject(FInternalNetHandle Depende
 	}
 
 	// Remove from our dependents
-	if (FNetDependencyData::FInternalNetHandleArray* DependentObjectArray = SubObjects.GetInternalHandleArray<FNetDependencyData::DependentObjects>(DependentInternalIndex))
+	if (FNetDependencyData::FInternalNetRefIndexArray* DependentObjectArray = SubObjects.GetInternalIndexArray<FNetDependencyData::DependentObjects>(DependentInternalIndex))
 	{
-		for (FInternalNetHandle ChildDepedentInternalIndex : *DependentObjectArray)
+		for (FInternalNetRefIndex ChildDepedentInternalIndex : *DependentObjectArray)
 		{
 			// Flag is set to only update data on the childDependentObject to avoid modifying the array we iterate over
 			InternalRemoveDependentObject(DependentInternalIndex, ChildDepedentInternalIndex, ERemoveDependentObjectFlags::RemoveFromDependentParentObjects);
@@ -673,11 +691,11 @@ void FNetHandleManager::InternalRemoveDependentObject(FInternalNetHandle Depende
 	DependentObjectInternalIndices.ClearBit(DependentInternalIndex);
 }
 
-void FNetHandleManager::RemoveDependentObject(FNetHandle ParentHandle, FNetHandle DependentHandle)
+void FNetRefHandleManager::RemoveDependentObject(FNetRefHandle ParentHandle, FNetRefHandle DependentHandle)
 {
 	// Validate objects
-	FInternalNetHandle ParentInternalIndex = GetInternalIndex(ParentHandle);
-	FInternalNetHandle DependentInternalIndex = GetInternalIndex(DependentHandle);
+	FInternalNetRefIndex ParentInternalIndex = GetInternalIndex(ParentHandle);
+	FInternalNetRefIndex DependentInternalIndex = GetInternalIndex(DependentHandle);
 
 	if ((ParentInternalIndex == InvalidInternalIndex) || (DependentInternalIndex == InvalidInternalIndex))
 	{
@@ -687,7 +705,7 @@ void FNetHandleManager::RemoveDependentObject(FNetHandle ParentHandle, FNetHandl
 	InternalRemoveDependentObject(ParentInternalIndex, DependentInternalIndex);
 }
 
-void FNetHandleManager::SetShouldPropagateChangedStates(FInternalNetHandle ObjectInternalIndex, bool bShouldPropagateChangedStates)
+void FNetRefHandleManager::SetShouldPropagateChangedStates(FInternalNetRefIndex ObjectInternalIndex, bool bShouldPropagateChangedStates)
 {
 	if (ObjectInternalIndex != InvalidInternalIndex)
 	{
@@ -705,28 +723,28 @@ void FNetHandleManager::SetShouldPropagateChangedStates(FInternalNetHandle Objec
 	}
 }
 
-void FNetHandleManager::SetShouldPropagateChangedStates(FNetHandle Handle, bool bShouldPropagateChangedStates)
+void FNetRefHandleManager::SetShouldPropagateChangedStates(FNetRefHandle Handle, bool bShouldPropagateChangedStates)
 {
-	FInternalNetHandle ObjectInternalIndex = GetInternalIndex(Handle);
+	FInternalNetRefIndex ObjectInternalIndex = GetInternalIndex(Handle);
 	return SetShouldPropagateChangedStates(ObjectInternalIndex, bShouldPropagateChangedStates);
 }
 
-void FNetHandleManager::AddReferencedObjects(FReferenceCollector& Collector)
+void FNetRefHandleManager::AddReferencedObjects(FReferenceCollector& Collector)
 {
 	Collector.AddReferencedObjects(ReplicatedInstances);
 }
 
-uint32 FNetHandleManager::MakeNetHandleId(uint32 Id, bool bIsStatic)
+uint32 FNetRefHandleManager::MakeNetRefHandleId(uint32 Id, bool bIsStatic)
 {
 	return (Id << 1U) | (bIsStatic ? 1U : 0U);
 }
 
-FNetHandle FNetHandleManager::MakeNetHandle(uint32 Id, uint32 ReplicationSystemId)
+FNetRefHandle FNetRefHandleManager::MakeNetRefHandle(uint32 Id, uint32 ReplicationSystemId)
 {
-	check((Id & FNetHandle::IdMask) == Id);
-	check(ReplicationSystemId < FNetHandle::MaxReplicationSystemId);
+	check((Id & FNetRefHandle::IdMask) == Id);
+	check(ReplicationSystemId < FNetRefHandle::MaxReplicationSystemId);
 
-	FNetHandle Handle;
+	FNetRefHandle Handle;
 
 	Handle.Id = Id;
 	Handle.ReplicationSystemId = ReplicationSystemId + 1U;
@@ -734,11 +752,11 @@ FNetHandle FNetHandleManager::MakeNetHandle(uint32 Id, uint32 ReplicationSystemI
 	return Handle;
 }
 
-FNetHandle FNetHandleManager::MakeNetHandleFromId(uint32 Id)
+FNetRefHandle FNetRefHandleManager::MakeNetRefHandleFromId(uint32 Id)
 {
-	check((Id & FNetHandle::IdMask) == Id);
+	check((Id & FNetRefHandle::IdMask) == Id);
 
-	FNetHandle Handle;
+	FNetRefHandle Handle;
 
 	Handle.Id = Id;
 	Handle.ReplicationSystemId = 0U;
diff --git a/Engine/Source/Runtime/Experimental/Iris/Core/Private/Iris/ReplicationSystem/NetHandleManager.h b/Engine/Source/Runtime/Experimental/Iris/Core/Private/Iris/ReplicationSystem/NetHandleManager.h
--- a/Engine/Source/Runtime/Experimental/Iris/Core/Private/Iris/ReplicationSystem/NetHandleManager.h
+++ b/Engine/Source/Runtime/Experimental/Iris/Core/Private/Iris/ReplicationSystem/NetHandleManager.h
@@ -7,6 +7,7 @@
 #include "Containers/Map.h"
 #include "Containers/ArrayView.h"
 #include "Net/Core/NetBitArray.h"
+#include "Net/Core/NetHandle/NetHandle.h"
 #include "Iris/ReplicationState/ReplicationStateDescriptor.h"
 #include "Iris/ReplicationSystem/NetHandle.h"
 #include "Iris/ReplicationSystem/NetDependencyData.h"
@@ -26,7 +27,7 @@ namespace UE::Net
 namespace UE::Net::Private
 {
 
-typedef uint32 FInternalNetHandle;
+typedef uint32 FInternalNetRefIndex;
 
 enum class EAddSubObjectFlags : uint32
 {
@@ -56,10 +57,10 @@ enum class ERemoveDependentObjectFlags : uint32
 ENUM_CLASS_FLAGS(ERemoveDependentObjectFlags);
 
 // Internal class to manage NetHandles and their internal data
-class FNetHandleManager
+class FNetRefHandleManager
 {
 public:
-	// We will never assign index 0
+	// We will never assign InvalidInternalIndex
 	enum { InvalidInternalIndex = 0U };
 
 	// We need to store some internal data for Replicated objects
@@ -76,7 +77,8 @@ public:
 		{
 		}
 
-		FNetHandle Handle;
+		FNetRefHandle RefHandle;
+		FNetHandle NetHandle;
 		const FReplicationProtocol* Protocol;
 		const FReplicationInstanceProtocol* InstanceProtocol;
 		uint8* ReceiveStateBuffer;
@@ -105,77 +107,81 @@ public:
 	};
 
 	// We want to handle both actual Networked and "network addressable" handles in the same way
-	// As we expect the number of FNetHandles to be much greater than the actual number of Replicated Objects we
+	// As we expect the number of FNetRefHandles to be much greater than the actual number of Replicated Objects we
 	// use a map to map from Handle to InternalIndex, only Replicated Objects will ever be assigned an InternalNetHandle
-	typedef TMap<FNetHandle, FInternalNetHandle> FHandleMap;
+	typedef TMap<FNetRefHandle, FInternalNetRefIndex> FRefHandleMap;
+	typedef TMap<FNetHandle, FInternalNetRefIndex> FNetHandleMap;
 
 public:
-	FNetHandleManager(FReplicationProtocolManager& InReplicationProtocolManager, uint32 InReplicationSystemId, uint32 MaxActiveObjects);
+	FNetRefHandleManager(FReplicationProtocolManager& InReplicationProtocolManager, uint32 InReplicationSystemId, uint32 MaxActiveObjects);
 
 	// Return true if this is a scopable index
-	bool IsScopableIndex(FInternalNetHandle InternalIndex) const { return ScopableInternalIndices.GetBit(InternalIndex); }
+	bool IsScopableIndex(FInternalNetRefIndex InternalIndex) const { return ScopableInternalIndices.GetBit(InternalIndex); }
 
 	// Return true if this index is assigned
-	bool IsAssignedIndex(FInternalNetHandle InternalIndex) const { return InternalIndex && ScopableInternalIndices.GetBit(InternalIndex); }
+	bool IsAssignedIndex(FInternalNetRefIndex InternalIndex) const { return InternalIndex && ScopableInternalIndices.GetBit(InternalIndex); }
 
-	static FNetHandle MakeNetHandle(uint32 Id, uint32 ReplicationSystemId);
-	static FNetHandle MakeNetHandleFromId(uint32 Id);
+	static FNetRefHandle MakeNetRefHandle(uint32 Id, uint32 ReplicationSystemId);
+	static FNetRefHandle MakeNetRefHandleFromId(uint32 Id);
 
 	// Returns a valid handle if the wanted handle can be allocated
-	FNetHandle AllocateNetHandle(bool bIsStatic);
+	FNetRefHandle AllocateNetRefHandle(bool bIsStatic);
 
 	// Create local Net Object
-	FNetHandle CreateNetObject(FNetHandle WantedHandle, const FReplicationProtocol* ReplicationProtocol);
+	FNetRefHandle CreateNetObject(FNetRefHandle WantedHandle, FNetHandle GlobalHandle, const FReplicationProtocol* ReplicationProtocol);
 
 	// Create NetObject on request from remote
-	FNetHandle CreateNetObjectFromRemote(FNetHandle WantedHandle, const FReplicationProtocol* ReplicationProtocol);
+	FNetRefHandle CreateNetObjectFromRemote(FNetRefHandle WantedHandle, const FReplicationProtocol* ReplicationProtocol);
 
 	// Attach Instance protocol to handle
 	// Instance can be null, we only track the Instance for legacy support
-	void AttachInstanceProtocol(FInternalNetHandle InternalIndex, const FReplicationInstanceProtocol* InstanceProtocol, UObject* Instance);
+	void AttachInstanceProtocol(FInternalNetRefIndex InternalIndex, const FReplicationInstanceProtocol* InstanceProtocol, UObject* Instance);
 
 	// Used when we want to detach the instance protocol due to async destroy
-	const FReplicationInstanceProtocol* DetachInstanceProtocol(FInternalNetHandle InternalIndex);
+	const FReplicationInstanceProtocol* DetachInstanceProtocol(FInternalNetRefIndex InternalIndex);
 
 	// Creates a DestructionInfo Handle which is used to replicate persistently destroyed static objects when late joining or streaming in levels with static objects that already has been destroyed
 	// Returns a new handle used to replicate the destruction info, and adds a cross-reference to the original handle being destroyed if that still exists 
 	// Handle is a reference to an already replicated Handle that is about to be destroyed
-	FNetHandle CreateHandleForDestructionInfo(FNetHandle Handle, const FReplicationProtocol* DestroydObjectProtocol);
+	FNetRefHandle CreateHandleForDestructionInfo(FNetRefHandle Handle, const FReplicationProtocol* DestroydObjectProtocol);
 	
-	void DestroyNetObject(FNetHandle Handle);
+	void DestroyNetObject(FNetRefHandle Handle);
 
 	// Mark object as no longer scopable, object will be removed from scope for all connections
-	void RemoveFromScope(FInternalNetHandle InternalIndex);
+	void RemoveFromScope(FInternalNetRefIndex InternalIndex);
 
-	TArrayView<const FInternalNetHandle> GetObjectsPendingDestroy() const { return MakeArrayView(PendingDestroyInternalIndices); }
+	TArrayView<const FInternalNetRefIndex> GetObjectsPendingDestroy() const { return MakeArrayView(PendingDestroyInternalIndices); }
 	void DestroyObjectsPendingDestroy();
 
-	const FReplicatedObjectData& GetReplicatedObjectDataNoCheck(FInternalNetHandle InternalIndex) const { return ReplicatedObjectData[InternalIndex]; }
-	FReplicatedObjectData& GetReplicatedObjectDataNoCheck(FInternalNetHandle InternalIndex) { return ReplicatedObjectData[InternalIndex]; }
+	const FReplicatedObjectData& GetReplicatedObjectDataNoCheck(FInternalNetRefIndex InternalIndex) const { return ReplicatedObjectData[InternalIndex]; }
+	FReplicatedObjectData& GetReplicatedObjectDataNoCheck(FInternalNetRefIndex InternalIndex) { return ReplicatedObjectData[InternalIndex]; }
 
-	inline const FReplicatedObjectData& GetReplicatedObjectData(FInternalNetHandle InternalIndex) const;
+	inline const FReplicatedObjectData& GetReplicatedObjectData(FInternalNetRefIndex InternalIndex) const;
 
-	inline const uint8* GetReplicatedObjectStateBufferNoCheck(FInternalNetHandle InternalObjectIndex) const { return ReplicatedObjectStateBuffers.GetData()[InternalObjectIndex]; }
-	inline uint8* GetReplicatedObjectStateBufferNoCheck(FInternalNetHandle InternalObjectIndex) { return ReplicatedObjectStateBuffers.GetData()[InternalObjectIndex]; }
+	inline const uint8* GetReplicatedObjectStateBufferNoCheck(FInternalNetRefIndex InternalObjectIndex) const { return ReplicatedObjectStateBuffers.GetData()[InternalObjectIndex]; }
+	inline uint8* GetReplicatedObjectStateBufferNoCheck(FInternalNetRefIndex InternalObjectIndex) { return ReplicatedObjectStateBuffers.GetData()[InternalObjectIndex]; }
 	inline const TArray<uint8*>& GetReplicatedObjectStateBuffers() const { return ReplicatedObjectStateBuffers; }
 
-	// Verify a handle against internal handle, A handle is valid if it matches internal storage
-	inline bool IsValidNetHandle(FNetHandle Handle) const;
+	/** Verify a handle against internal handle. A handle is valid if it matches internal storage. */
+	inline bool IsValidNetRefHandle(FNetRefHandle Handle) const;
 
-	// Returns true if the handle is for local replicated object
-	inline bool IsLocalNetHandle(FNetHandle Handle) const;
+	/** Returns true if the handle is for a local replicated object. */
+	inline bool IsLocalNetRefHandle(FNetRefHandle Handle) const;
 
-	// Returns true if the handle is for a remote replicated object
-	inline bool IsRemoteNetHandle(FNetHandle Handle) const;
+	/** Returns true if the handle is for a remote replicated object. */
+	inline bool IsRemoteNetRefHandle(FNetRefHandle Handle) const;
 
-	// Extract an Full handle from an incomplete one, a incomplete handle is a handle with only an index
-	inline FNetHandle GetCompleteNetHandle(FNetHandle IncompleteHandle) const;
+	/** Extract a full handle from an incomplete one consisting of only an index. */
+	inline FNetRefHandle GetCompleteNetRefHandle(FNetRefHandle IncompleteHandle) const;
 	
 	// Get Handle from internal index
-	inline FNetHandle GetNetHandleFromInternalIndex(FInternalNetHandle InternalIndex) const;
+	inline FNetRefHandle GetNetRefHandleFromInternalIndex(FInternalNetRefIndex InternalIndex) const;
 
 	// Get internal index from handle
-	inline FInternalNetHandle GetInternalIndex(FNetHandle Handle) const;
+	inline FInternalNetRefIndex GetInternalIndex(FNetRefHandle Handle) const;
+
+	// Get internal index from NetHandle
+	inline FInternalNetRefIndex GetInternalIndexFromNetHandle(FNetHandle Handle) const;
 
 	// Get bitarray for all currently scopable internal indices
 	const FNetBitArray& GetScopableInternalIndices() const { return ScopableInternalIndices; }
@@ -189,56 +195,57 @@ public:
 
 	// SubObjects
 	const FNetBitArray& GetSubObjectInternalIndices() const { return SubObjectInternalIndices; }	
-	bool AddSubObject(FNetHandle OwnerHandle, FNetHandle SubObjectHandle, FNetHandle RelativeOtherSubObjectHandle, EAddSubObjectFlags Flags = EAddSubObjectFlags::Default);
-	inline bool AddSubObject(FNetHandle OwnerHandle, FNetHandle SubObjectHandle, EAddSubObjectFlags Flags = EAddSubObjectFlags::Default) { return AddSubObject(OwnerHandle, SubObjectHandle, FNetHandle(), Flags); }
-	void RemoveSubObject(FNetHandle SubObjectHandle);
-	FNetHandle GetSubObjectOwner(FNetHandle SubObjectHandle) const;
-	bool SetSubObjectNetCondition(FInternalNetHandle SubObjectInternalIndex, FLifeTimeConditionStorage SubObjectCondition);
+
+	bool AddSubObject(FNetRefHandle OwnerHandle, FNetRefHandle SubObjectHandle, FNetRefHandle RelativeOtherSubObjectHandle, EAddSubObjectFlags Flags = EAddSubObjectFlags::Default);
+	bool AddSubObject(FNetRefHandle OwnerHandle, FNetRefHandle SubObjectHandle, EAddSubObjectFlags Flags = EAddSubObjectFlags::Default);
+	void RemoveSubObject(FNetRefHandle SubObjectHandle);
+	FNetRefHandle GetSubObjectOwner(FNetRefHandle SubObjectHandle) const;
+	bool SetSubObjectNetCondition(FInternalNetRefIndex SubObjectInternalIndex, FLifeTimeConditionStorage SubObjectCondition);
 
 	// DependentObjects
 	const FNetBitArray& GetDependentObjectInternalIndices() const { return DependentObjectInternalIndices; }
 	const FNetBitArray& GetObjectsWithDependentObjectsInternalIndices() const { return ObjectsWithDependentObjectsInternalIndices; }
-	bool AddDependentObject(FNetHandle ParentHandle, FNetHandle DependentHandle, EAddDependentObjectFlags Flags = EAddDependentObjectFlags::WarnIfAlreadyDependentObject);
-	void RemoveDependentObject(FNetHandle ParentHandle, FNetHandle DependentHandle);
+	bool AddDependentObject(FNetRefHandle ParentHandle, FNetRefHandle DependentHandle, EAddDependentObjectFlags Flags = EAddDependentObjectFlags::WarnIfAlreadyDependentObject);
+	void RemoveDependentObject(FNetRefHandle ParentHandle, FNetRefHandle DependentHandle);
 
 	// Remove DependentHandles from all dependent object tracking
 	// $TODO: This might need a pass ObjectReplicationBridge as well since most likely also must restore filter/polling logic
-	void RemoveDependentObject(FNetHandle DependentHandle);
+	void RemoveDependentObject(FNetRefHandle DependentHandle);
 	
-	void SetShouldPropagateChangedStates(FNetHandle Handle, bool bShouldPropagateChangedStates);
-	void SetShouldPropagateChangedStates(FInternalNetHandle ObjectInternalIndex, bool bShouldPropagateChangedStates);
+	void SetShouldPropagateChangedStates(FNetRefHandle Handle, bool bShouldPropagateChangedStates);
+	void SetShouldPropagateChangedStates(FInternalNetRefIndex ObjectInternalIndex, bool bShouldPropagateChangedStates);
 
 	uint32 GetMaxActiveObjectCount() const { return MaxActiveObjectCount; }
 	uint32 GetActiveObjectCount() const { return ActiveObjectCount; }
 
 	// We do refcount objects tracked by each connection in order to know when it is safe to reuse an InternalIndex
-	void AddNetObjectRef(FInternalNetHandle InternalIndex) { ++ReplicatedObjectRefCount[InternalIndex]; }
-	void ReleaseNetObjectRef(FInternalNetHandle InternalIndex) { check(ReplicatedObjectRefCount[InternalIndex] > 0); --ReplicatedObjectRefCount[InternalIndex]; }
-	uint16 GetNetObjectRefCount(FInternalNetHandle ObjectInternalIndex) const { return ReplicatedObjectRefCount[ObjectInternalIndex]; }
+	void AddNetObjectRef(FInternalNetRefIndex InternalIndex) { ++ReplicatedObjectRefCount[InternalIndex]; }
+	void ReleaseNetObjectRef(FInternalNetRefIndex InternalIndex) { check(ReplicatedObjectRefCount[InternalIndex] > 0); --ReplicatedObjectRefCount[InternalIndex]; }
+	uint16 GetNetObjectRefCount(FInternalNetRefIndex ObjectInternalIndex) const { return ReplicatedObjectRefCount[ObjectInternalIndex]; }
 
 	// Get dependent objects for the given ParentIndex
-	inline TArrayView<const FInternalNetHandle> GetDependentObjects(FInternalNetHandle ParentIndex) const;
+	inline TArrayView<const FInternalNetRefIndex> GetDependentObjects(FInternalNetRefIndex ParentIndex) const;
 
 	// Get all parents of the given DependentIndex
-	inline TArrayView<const FInternalNetHandle> GetDependentObjectParents(FInternalNetHandle DependentIndex) const;
+	inline TArrayView<const FInternalNetRefIndex> GetDependentObjectParents(FInternalNetRefIndex DependentIndex) const;
 
 	// Get all subobjects for the given OwnerIndex, Note: only valid for the root
-	inline TArrayView<const FInternalNetHandle> GetSubObjects(FInternalNetHandle OwnerIndex) const;
+	inline TArrayView<const FInternalNetRefIndex> GetSubObjects(FInternalNetRefIndex OwnerIndex) const;
 
 	// Get child subobjects for a object, used when we do hierarchical operations such as conditional serialization
-	inline TArrayView<const FInternalNetHandle> GetChildSubObjects(FInternalNetHandle ParentObjectIndex) const;
+	inline TArrayView<const FInternalNetRefIndex> GetChildSubObjects(FInternalNetRefIndex ParentObjectIndex) const;
 
 	// Get child subobjects and condtionals array if one exists, if there are no child subobjects the method returns false
-	inline bool GetChildSubObjects(FInternalNetHandle OwnerIndex, FChildSubObjectsInfo& OutInfo) const;
+	inline bool GetChildSubObjects(FInternalNetRefIndex OwnerIndex, FChildSubObjectsInfo& OutInfo) const;
 
-	const FHandleMap& GetReplicatedHandles() const { return HandleMap; }
+	const FRefHandleMap& GetReplicatedHandles() const { return RefHandleToInternalIndex; }
 
 	const TArray<UObject*>& GetReplicatedInstances() const { return ReplicatedInstances; }
 
 	void AddReferencedObjects(FReferenceCollector& Collector);
 
-	bool GetIsDestroyedStartupObject(FInternalNetHandle InternalIndex) const { return DestroyedStartupObjectInternalIndices.GetBit(InternalIndex); }
-	inline uint32 GetOriginalDestroyedStartupObjectIndex(FInternalNetHandle InternalIndex) const;
+	bool GetIsDestroyedStartupObject(FInternalNetRefIndex InternalIndex) const { return DestroyedStartupObjectInternalIndices.GetBit(InternalIndex); }
+	inline uint32 GetOriginalDestroyedStartupObjectIndex(FInternalNetRefIndex InternalIndex) const;
 
 	const FNetBitArray& GetDestroyedStartupObjectInternalIndices() const { return DestroyedStartupObjectInternalIndices; }
 
@@ -247,11 +254,11 @@ public:
 
 	// Iterate over all dependent objects and their dependent objects
 	template <typename T>
-	void ForAllDependentObjectsRecursive(FInternalNetHandle ObjectIndex, T&& Functor) const
+	void ForAllDependentObjectsRecursive(FInternalNetRefIndex ObjectIndex, T&& Functor) const
 	{
 		if (ObjectsWithDependentObjectsInternalIndices.GetBit(ObjectIndex))
 		{
-			for (const FInternalNetHandle DependentObjectIndex : GetDependentObjects(ObjectIndex))
+			for (const FInternalNetRefIndex DependentObjectIndex : GetDependentObjects(ObjectIndex))
 			{
 				Functor(DependentObjectIndex);
 				ForAllDependentObjectsRecursive(DependentObjectIndex, Functor);
@@ -261,20 +268,22 @@ public:
 	};
 
 private:
-	FInternalNetHandle InternalCreateNetObject(const FNetHandle NetHandle, const FReplicationProtocol* ReplicationProtocol);
-	void InternalDestroyNetObject(FInternalNetHandle InternalIndex);
+	FInternalNetRefIndex InternalCreateNetObject(const FNetRefHandle NetRefHandle, const FNetHandle GlobalHandle, const FReplicationProtocol* ReplicationProtocol);
+	void InternalDestroyNetObject(FInternalNetRefIndex InternalIndex);
 
-	static uint32 MakeNetHandleId(uint32 Seed, bool bIsStatic);
-	uint32 GetNextHandleId(uint32 HandleIndex) const;
+	static uint32 MakeNetRefHandleId(uint32 Seed, bool bIsStatic);
+	uint32 GetNextNetRefHandleId(uint32 HandleIndex) const;
 
 	// Get the next free internal index, returns InvalidInternalIndex if a free one cannot be found
-	FInternalNetHandle GetNextFreeInternalIndex() const;
+	FInternalNetRefIndex GetNextFreeInternalIndex() const;
+
+	bool InternalAddSubObject(FInternalNetRefIndex OwnerInternalIndex, FInternalNetRefIndex SubObjectInternalIndex, FInternalNetRefIndex RelativeOtherSubObjectInternalIndex, EAddSubObjectFlags Flags);
 
-	inline void SetIsSubObject(FInternalNetHandle InternalIndex, bool IsSubObject) { SubObjectInternalIndices.SetBitValue(InternalIndex, IsSubObject); }
-	void InternalRemoveSubObject(FInternalNetHandle OwnerInternalIndex, FInternalNetHandle SubObjectInternalIndex, bool bRemoveFromSubObjectArray = true);
+	inline void SetIsSubObject(FInternalNetRefIndex InternalIndex, bool IsSubObject) { SubObjectInternalIndices.SetBitValue(InternalIndex, IsSubObject); }
+	void InternalRemoveSubObject(FInternalNetRefIndex OwnerInternalIndex, FInternalNetRefIndex SubObjectInternalIndex, bool bRemoveFromSubObjectArray = true);
 
-	void InternalRemoveDependentObject(FInternalNetHandle ParentInternalIndex, FInternalNetHandle DependentInternalIndex, ERemoveDependentObjectFlags Flags = ERemoveDependentObjectFlags::All);
-	void InternalRemoveDependentObject(FInternalNetHandle DependentInternalIndex);
+	void InternalRemoveDependentObject(FInternalNetRefIndex ParentInternalIndex, FInternalNetRefIndex DependentInternalIndex, ERemoveDependentObjectFlags Flags = ERemoveDependentObjectFlags::All);
+	void InternalRemoveDependentObject(FInternalNetRefIndex DependentInternalIndex);
 
 	// The current replicated object count
 	uint32 ActiveObjectCount;
@@ -284,7 +293,8 @@ private:
 
 	uint32 ReplicationSystemId;
 
-	FHandleMap HandleMap;
+	FRefHandleMap RefHandleToInternalIndex;
+	FNetHandleMap NetHandleToInternalIndex;
 
 	// Bitset used in order to track assigned internal which are scopable
 	FNetBitArray ScopableInternalIndices;
@@ -316,10 +326,11 @@ private:
 	// Map OriginalIndex -> DestructionInfoIndex
 	// Map DestructionInfoIndex -> Orignal
 	TMap<uint32, uint32> DestroyedStartupObject;
+
 private:
 
 	// Array used in order to track objects pending destroy
-	TArray<FInternalNetHandle> PendingDestroyInternalIndices;
+	TArray<FInternalNetRefIndex> PendingDestroyInternalIndices;
 
 	// Just an array containing data about our replicated objects
 	TArray<FReplicatedObjectData> ReplicatedObjectData;
@@ -342,44 +353,50 @@ private:
 	FReplicationProtocolManager& ReplicationProtocolManager;
 };
 
-const FNetHandleManager::FReplicatedObjectData& FNetHandleManager::GetReplicatedObjectData(FInternalNetHandle InternalIndex) const
+const FNetRefHandleManager::FReplicatedObjectData& FNetRefHandleManager::GetReplicatedObjectData(FInternalNetRefIndex InternalIndex) const
 {
 	check(AssignedInternalIndices.GetBit(InternalIndex));
 	return GetReplicatedObjectDataNoCheck(InternalIndex);
 }
 
-FInternalNetHandle FNetHandleManager::GetInternalIndex(FNetHandle Handle) const
+FInternalNetRefIndex FNetRefHandleManager::GetInternalIndex(FNetRefHandle Handle) const
 {
-	const FInternalNetHandle* InternalIndex = HandleMap.Find(Handle);
-	return InternalIndex ? *InternalIndex : FNetHandleManager::InvalidInternalIndex;
+	const FInternalNetRefIndex* InternalIndex = RefHandleToInternalIndex.Find(Handle);
+	return InternalIndex ? *InternalIndex : FNetRefHandleManager::InvalidInternalIndex;
 }
 
-FNetHandle FNetHandleManager::GetNetHandleFromInternalIndex(FInternalNetHandle InternalIndex) const
+FInternalNetRefIndex FNetRefHandleManager::GetInternalIndexFromNetHandle(FNetHandle Handle) const
+{
+	const FInternalNetRefIndex* InternalIndex = NetHandleToInternalIndex.Find(Handle);
+	return InternalIndex ? *InternalIndex : FNetRefHandleManager::InvalidInternalIndex;
+}
+
+FNetRefHandle FNetRefHandleManager::GetNetRefHandleFromInternalIndex(FInternalNetRefIndex InternalIndex) const
 {
 	check(AssignedInternalIndices.GetBit(InternalIndex));
-	return GetReplicatedObjectDataNoCheck(InternalIndex).Handle;
+	return GetReplicatedObjectDataNoCheck(InternalIndex).RefHandle;
 }
 
-FNetHandle FNetHandleManager::GetCompleteNetHandle(FNetHandle IncompleteHandle) const
+FNetRefHandle FNetRefHandleManager::GetCompleteNetRefHandle(FNetRefHandle IncompleteHandle) const
 {
-	if (const FInternalNetHandle* InternalIndex = HandleMap.Find(IncompleteHandle))
+	if (const FInternalNetRefIndex* InternalIndex = RefHandleToInternalIndex.Find(IncompleteHandle))
 	{
-		return GetReplicatedObjectDataNoCheck(*InternalIndex).Handle;
+		return GetReplicatedObjectDataNoCheck(*InternalIndex).RefHandle;
 	}
 	else
 	{
-		return FNetHandle();
+		return FNetRefHandle();
 	}
 }
 
-bool FNetHandleManager::IsValidNetHandle(FNetHandle Handle) const
+bool FNetRefHandleManager::IsValidNetRefHandle(FNetRefHandle Handle) const
 {
-	return HandleMap.Contains(Handle);
+	return RefHandleToInternalIndex.Contains(Handle);
 }
 
-bool FNetHandleManager::IsLocalNetHandle(FNetHandle Handle) const
+bool FNetRefHandleManager::IsLocalNetRefHandle(FNetRefHandle Handle) const
 {
-	if (const uint32 InternalIndex = GetInternalIndex(Handle))
+	if (const FInternalNetRefIndex InternalIndex = GetInternalIndex(Handle))
 	{
 		// For the time being only replicated objects owned by this peer has a state buffer
 		return ReplicatedObjectStateBuffers[InternalIndex] != nullptr;
@@ -388,7 +405,7 @@ bool FNetHandleManager::IsLocalNetHandle(FNetHandle Handle) const
 	return false;
 }
 
-bool FNetHandleManager::IsRemoteNetHandle(FNetHandle Handle) const
+bool FNetRefHandleManager::IsRemoteNetRefHandle(FNetRefHandle Handle) const
 {
 	if (const uint32 InternalIndex = GetInternalIndex(Handle))
 	{
@@ -399,35 +416,40 @@ bool FNetHandleManager::IsRemoteNetHandle(FNetHandle Handle) const
 	return false;
 }
 
-uint32 FNetHandleManager::GetOriginalDestroyedStartupObjectIndex(FInternalNetHandle InternalIndex) const
+uint32 FNetRefHandleManager::GetOriginalDestroyedStartupObjectIndex(FInternalNetRefIndex InternalIndex) const
 {
 	const uint32* FoundOriginalInternalIndex = DestroyedStartupObject.Find(InternalIndex);
 	return FoundOriginalInternalIndex ? *FoundOriginalInternalIndex : 0U;
 }
 
-TArrayView<const FInternalNetHandle> FNetHandleManager::GetSubObjects(FInternalNetHandle OwnerIndex) const
+inline bool FNetRefHandleManager::AddSubObject(FNetRefHandle OwnerHandle, FNetRefHandle SubObjectHandle, EAddSubObjectFlags Flags)
+{
+	return AddSubObject(OwnerHandle, SubObjectHandle, FNetRefHandle(), Flags);
+}
+
+TArrayView<const FInternalNetRefIndex> FNetRefHandleManager::GetSubObjects(FInternalNetRefIndex OwnerIndex) const
 {
-	return SubObjects.GetInternalHandleArray<FNetDependencyData::EArrayType::SubObjects>(OwnerIndex);
+	return SubObjects.GetInternalIndexArray<FNetDependencyData::EArrayType::SubObjects>(OwnerIndex);
 }
 
-bool FNetHandleManager::GetChildSubObjects(FInternalNetHandle OwnerIndex, FChildSubObjectsInfo& OutInfo) const
+bool FNetRefHandleManager::GetChildSubObjects(FInternalNetRefIndex OwnerIndex, FChildSubObjectsInfo& OutInfo) const
 {
 	return SubObjects.GetChildSubObjects(OwnerIndex, OutInfo);
 }
 
-TArrayView<const FInternalNetHandle> FNetHandleManager::GetChildSubObjects(FInternalNetHandle OwnerIndex) const
+TArrayView<const FInternalNetRefIndex> FNetRefHandleManager::GetChildSubObjects(FInternalNetRefIndex OwnerIndex) const
 {
-	return SubObjects.GetInternalHandleArray<FNetDependencyData::EArrayType::ChildSubObjects>(OwnerIndex);
+	return SubObjects.GetInternalIndexArray<FNetDependencyData::EArrayType::ChildSubObjects>(OwnerIndex);
 }
 
-TArrayView<const FInternalNetHandle> FNetHandleManager::GetDependentObjects(FInternalNetHandle ParentIndex) const
+TArrayView<const FInternalNetRefIndex> FNetRefHandleManager::GetDependentObjects(FInternalNetRefIndex ParentIndex) const
 {
-	return SubObjects.GetInternalHandleArray<FNetDependencyData::EArrayType::DependentObjects>(ParentIndex);
+	return SubObjects.GetInternalIndexArray<FNetDependencyData::EArrayType::DependentObjects>(ParentIndex);
 }
 
-TArrayView<const FInternalNetHandle> FNetHandleManager::GetDependentObjectParents(FInternalNetHandle DependentIndex) const
+TArrayView<const FInternalNetRefIndex> FNetRefHandleManager::GetDependentObjectParents(FInternalNetRefIndex DependentIndex) const
 {
-	return SubObjects.GetInternalHandleArray<FNetDependencyData::EArrayType::ParentObjects>(DependentIndex);
+	return SubObjects.GetInternalIndexArray<FNetDependencyData::EArrayType::ParentObjects>(DependentIndex);
 }
 	
 }
diff --git a/Engine/Source/Runtime/Experimental/Iris/Core/Private/Iris/ReplicationSystem/ObjectPollFrequencyLimiter.h b/Engine/Source/Runtime/Experimental/Iris/Core/Private/Iris/ReplicationSystem/ObjectPollFrequencyLimiter.h
--- a/Engine/Source/Runtime/Experimental/Iris/Core/Private/Iris/ReplicationSystem/ObjectPollFrequencyLimiter.h
+++ b/Engine/Source/Runtime/Experimental/Iris/Core/Private/Iris/ReplicationSystem/ObjectPollFrequencyLimiter.h
@@ -8,7 +8,7 @@
 
 namespace UE::Net::Private
 {
-	typedef uint32 FInternalNetHandle;
+	typedef uint32 FInternalNetRefIndex;
 }
 
 namespace UE::Net::Private
@@ -25,9 +25,9 @@ public:
 	void Init(uint32 MaxActiveObjectCount);
 	void Deinit();
 
-	void SetPollFramePeriod(FInternalNetHandle InternalHandle, uint8 PollFramePeriod);
+	void SetPollFramePeriod(FInternalNetRefIndex InternalIndex, uint8 PollFramePeriod);
 
-	void SetPollWithObject(FInternalNetHandle ObjectToPollWithInternalHandle, FInternalNetHandle InternalHandle);
+	void SetPollWithObject(FInternalNetRefIndex ObjectToPollWithInternalIndex, FInternalNetRefIndex InternalIndex);
 
 	void Update(const FNetBitArrayView& ScopableObjects, const FNetBitArrayView& DirtyObjects, FNetBitArrayView& OutObjectsToPoll);
 
@@ -44,23 +44,23 @@ private:
 };
 
 
-inline void FObjectPollFrequencyLimiter::SetPollFramePeriod(FInternalNetHandle InternalHandle, uint8 PollFramePeriod)
+inline void FObjectPollFrequencyLimiter::SetPollFramePeriod(FInternalNetRefIndex InternalIndex, uint8 PollFramePeriod)
 {
-	MaxInternalHandle = FPlatformMath::Max(MaxInternalHandle, InternalHandle);
+	MaxInternalHandle = FPlatformMath::Max(MaxInternalHandle, InternalIndex);
 
-	FramesBetweenUpdates[InternalHandle] = PollFramePeriod;
+	FramesBetweenUpdates[InternalIndex] = PollFramePeriod;
 	// Spread the polling of objects with the same frequency so that if you add lots of objects the same frame they won't be polled at the same time.
 	const uint32 FrameOffset = FrameIndexOffsets[PollFramePeriod]++;
-	FrameCounters[InternalHandle] = FrameOffset % (uint32(PollFramePeriod) + 1U);
+	FrameCounters[InternalIndex] = FrameOffset % (uint32(PollFramePeriod) + 1U);
 }
 
-inline void FObjectPollFrequencyLimiter::SetPollWithObject(FInternalNetHandle ObjectToPollWithInternalHandle, FInternalNetHandle InternalHandle)
+inline void FObjectPollFrequencyLimiter::SetPollWithObject(FInternalNetRefIndex ObjectToPollWithInternalIndex, FInternalNetRefIndex InternalIndex)
 {
-	MaxInternalHandle = FPlatformMath::Max(MaxInternalHandle, InternalHandle);
+	MaxInternalHandle = FPlatformMath::Max(MaxInternalHandle, InternalIndex);
 
 	// Copy state from object to poll with
-	FramesBetweenUpdates[InternalHandle] = FramesBetweenUpdates[ObjectToPollWithInternalHandle];
-	FrameCounters[InternalHandle] = FrameCounters[ObjectToPollWithInternalHandle];
+	FramesBetweenUpdates[InternalIndex] = FramesBetweenUpdates[ObjectToPollWithInternalIndex];
+	FrameCounters[InternalIndex] = FrameCounters[ObjectToPollWithInternalIndex];
 }
 
 }
