Iris - Implemented globally unique NetHandles.

diff --git a/Engine/Source/Runtime/Experimental/Iris/Core/Private/Iris/ReplicationSystem/ObjectReferenceCache.cpp b/Engine/Source/Runtime/Experimental/Iris/Core/Private/Iris/ReplicationSystem/ObjectReferenceCache.cpp
--- a/Engine/Source/Runtime/Experimental/Iris/Core/Private/Iris/ReplicationSystem/ObjectReferenceCache.cpp
+++ b/Engine/Source/Runtime/Experimental/Iris/Core/Private/Iris/ReplicationSystem/ObjectReferenceCache.cpp
@@ -57,7 +57,7 @@ FObjectReferenceCache::FObjectReferenceCache()
 , ReplicationBridge(nullptr)
 , NetTokenStore(nullptr)
 , StringTokenStore(nullptr)
-, NetHandleManager(nullptr)
+, NetRefHandleManager(nullptr)
 , bIsAuthority(false)
 {
 }
@@ -68,7 +68,7 @@ void FObjectReferenceCache::Init(UReplicationSystem* InReplicationSystem)
 	ReplicationBridge = InReplicationSystem->GetReplicationBridgeAs<UObjectReplicationBridge>();
 	NetTokenStore =  &InReplicationSystem->GetReplicationSystemInternal()->GetNetTokenStore();
 	StringTokenStore = InReplicationSystem->GetStringTokenStore();
-	NetHandleManager = &InReplicationSystem->GetReplicationSystemInternal()->GetNetHandleManager();
+	NetRefHandleManager = &InReplicationSystem->GetReplicationSystemInternal()->GetNetRefHandleManager();
 	bIsAuthority = InReplicationSystem->IsServer();
 }
 
@@ -122,7 +122,7 @@ bool FObjectReferenceCache::CanClientLoadObjectInternal(const UObject* Object, b
 	return true;
 }
 
-bool FObjectReferenceCache::ShouldIgnoreWhenMissing(FNetHandle RefHandle) const
+bool FObjectReferenceCache::ShouldIgnoreWhenMissing(FNetRefHandle RefHandle) const
 {
 	if (RefHandle.IsDynamic())
 	{
@@ -142,9 +142,9 @@ bool FObjectReferenceCache::ShouldIgnoreWhenMissing(FNetHandle RefHandle) const
 
 	const FCachedNetObjectReference* OutermostCacheObject = CacheObject;
 
-	while (OutermostCacheObject != nullptr && OutermostCacheObject->OuterNetHandle.IsValid())
+	while (OutermostCacheObject != nullptr && OutermostCacheObject->OuterNetRefHandle.IsValid())
 	{
-		OutermostCacheObject = ReferenceHandleToCachedReference.Find(OutermostCacheObject->OuterNetHandle);
+		OutermostCacheObject = ReferenceHandleToCachedReference.Find(OutermostCacheObject->OuterNetRefHandle);
 	}
 
 	if (OutermostCacheObject != nullptr)
@@ -171,7 +171,7 @@ bool FObjectReferenceCache::RenamePathForPie(uint32 ConnectionId, FString& Str,
 	return GPlayInEditorID != -1 ? ReplicationBridge->RemapPathForPIE(ConnectionId, Str, bReading) : false;
 }
 
-FNetHandle FObjectReferenceCache::CreateObjectReferenceHandle(const UObject* Object)
+FNetRefHandle FObjectReferenceCache::CreateObjectReferenceHandle(const UObject* Object)
 {
 	FNetObjectReference Ref;
 	CreateObjectReferenceInternal(Object, Ref);
@@ -189,20 +189,20 @@ bool FObjectReferenceCache::CreateObjectReferenceInternal(const UObject* Object,
 	CA_ASSUME(Object != nullptr);
 
 	// Check if we are already know about this object
-	if (FNetHandle* ReferenceHandle = ObjectToNetReferenceHandle.Find(Object))
+	if (FNetRefHandle* ReferenceHandle = ObjectToNetReferenceHandle.Find(Object))
 	{
 		// Verify that this is the same object and clean up if it is not
 		FCachedNetObjectReference& CachedObject = ReferenceHandleToCachedReference.FindChecked(*ReferenceHandle);
 
 		if (CachedObject.Object.Get() == Object)
 		{
-			UE_LOG_REFERENCECACHE(VeryVerbose, TEXT("ObjectReferenceCache::CreateObjectReferenceHandle Found existing %s for Object %s, OuterNetHandle: %s"), *CachedObject.NetHandle.ToString(), ToCStr(Object->GetPathName()), *CachedObject.OuterNetHandle.ToString());
+			UE_LOG_REFERENCECACHE(VeryVerbose, TEXT("ObjectReferenceCache::CreateObjectReferenceHandle Found existing %s for Object %s, OuterNetRefHandle: %s"), *CachedObject.NetRefHandle.ToString(), ToCStr(Object->GetPathName()), *CachedObject.OuterNetRefHandle.ToString());
 			OutReference = MakeNetObjectReference(CachedObject);			
 			return true;
 		}
 		else
 		{
-			UE_LOG_REFERENCECACHE(Verbose, TEXT("Detected stale cached Object removing %s from ObjectToNetReferenceHandle lookup"), *CachedObject.NetHandle.ToString());
+			UE_LOG_REFERENCECACHE(Verbose, TEXT("Detected stale cached Object removing %s from ObjectToNetReferenceHandle lookup"), *CachedObject.NetRefHandle.ToString());
 
 			ObjectToNetReferenceHandle.Remove(CachedObject.ObjectKey);
 			ObjectToNetReferenceHandle.Remove(Object);
@@ -253,7 +253,7 @@ bool FObjectReferenceCache::CreateObjectReferenceInternal(const UObject* Object,
 		// Find known outer if there is one, and express reference as a relative path to the outer
 		for (const UObject* Outer = Object; Outer != nullptr; Outer = Outer->GetOuter())
 		{
-			if (FNetHandle* ReplicatedOuter = ObjectToNetReferenceHandle.Find(Outer))
+			if (FNetRefHandle* ReplicatedOuter = ObjectToNetReferenceHandle.Find(Outer))
 			{
 				FString RelativeObjectPath = Object->GetPathName(Outer);
 				constexpr bool bReading = false;
@@ -266,7 +266,7 @@ bool FObjectReferenceCache::CreateObjectReferenceInternal(const UObject* Object,
 
 		if (bIsDynamic)
 		{
-			// If the object is dynamic and we did not find an outer handle there is not much we can do
+			// If the object is dynamic and we did not find an outer handle there's nothing we can do
 			return false;
 		}
 		else
@@ -276,16 +276,16 @@ bool FObjectReferenceCache::CreateObjectReferenceInternal(const UObject* Object,
 			constexpr bool bReading = false;
 			// The connection ID isn't used unless reading.
 			RenamePathForPie(UE::Net::InvalidConnectionId, ObjectPath, bReading);
-			OutReference = MakeNetObjectReference(FNetHandle(), StringTokenStore->GetOrCreateToken(ObjectPath));	
+			OutReference = MakeNetObjectReference(FNetRefHandle(), StringTokenStore->GetOrCreateToken(ObjectPath));	
 			return true;
 		}
 	}
 
-	// Generate a new Handle
-	FNetHandle NetHandle = NetHandleManager->AllocateNetHandle(bIsStatic);
-	FNetToken PathToken;
-	
-	check(NetHandle.IsValid());
+	// Generate a new RefHandle
+	FNetRefHandle NetRefHandle = NetRefHandleManager->AllocateNetRefHandle(bIsStatic);
+	check(NetRefHandle.IsValid());
+
+	FNetToken PathToken;	
 
 	// If our name is stable, or at least relative to our outer
 	if (bIsFullNameStableForNetworking || Object->IsNameStableForNetworking())
@@ -302,9 +302,9 @@ bool FObjectReferenceCache::CreateObjectReferenceInternal(const UObject* Object,
 
 	CachedObject.Object = MakeWeakObjectPtr(const_cast<UObject*>(Object));
 	CachedObject.ObjectKey = Object;
-	CachedObject.NetHandle = NetHandle;
+	CachedObject.NetRefHandle = NetRefHandle;
 	CachedObject.RelativePath = PathToken;
-	CachedObject.OuterNetHandle = CreateObjectReferenceHandle(Object->GetOuter());
+	CachedObject.OuterNetRefHandle = CreateObjectReferenceHandle(Object->GetOuter());
 	CachedObject.bIsPackage = bIsStatic && !Object->GetOuter();
 	CachedObject.bNoLoad = !CanClientLoadObjectInternal(Object, bIsDynamic); // We are probably going to deal with this as a dependency
 	CachedObject.bIgnoreWhenMissing = !CachedObject.bNoLoad;
@@ -314,19 +314,19 @@ bool FObjectReferenceCache::CreateObjectReferenceInternal(const UObject* Object,
 	// Make sure it really is a package
 	check(CachedObject.bIsPackage == (Cast<UPackage>(Object) != nullptr));
 
-	ReferenceHandleToCachedReference.Add(NetHandle, CachedObject);
-	ObjectToNetReferenceHandle.Add(Object, NetHandle);
+	ReferenceHandleToCachedReference.Add(NetRefHandle, CachedObject);
+	ObjectToNetReferenceHandle.Add(Object, NetRefHandle);
 
 #if UE_NET_TRACE_ENABLED
-	if (NetHandle.IsValid())
+	if (NetRefHandle.IsValid())
 	{
 		TStringBuilder<512> Builder;
 		Builder.Appendf(TEXT("Ref_%s"), PathToken.IsValid() ? StringTokenStore->ResolveToken(PathToken) : ToCStr(Object->GetName()));
-		UE_NET_TRACE_NETHANDLE_CREATED(NetHandle, CreatePersistentNetDebugName(Builder.ToString()), 0, IsAuthority() ? 1U : 0U);
+		UE_NET_TRACE_NETHANDLE_CREATED(NetRefHandle, CreatePersistentNetDebugName(Builder.ToString()), 0, IsAuthority() ? 1U : 0U);
 	}
 #endif
 	
-	UE_LOG_REFERENCECACHE(Verbose, TEXT("ObjectReferenceCache::CreateObjectReferenceHandle %s for Object %s, OuterNetHandle: %s"), *CachedObject.NetHandle.ToString(), ToCStr(Object->GetPathName()), *CachedObject.OuterNetHandle.ToString());
+	UE_LOG_REFERENCECACHE(Verbose, TEXT("ObjectReferenceCache::CreateObjectReferenceHandle %s for Object %s, OuterNetRefHandle: %s"), *CachedObject.NetRefHandle.ToString(), ToCStr(Object->GetPathName()), *CachedObject.OuterNetRefHandle.ToString());
 
 	// Create reference
 	OutReference = MakeNetObjectReference(CachedObject);
@@ -335,10 +335,10 @@ bool FObjectReferenceCache::CreateObjectReferenceInternal(const UObject* Object,
 	// We want to track references with a dynamic root so that we can purge them when the dynamic root is removed
 	if (bIsDynamic)
 	{
-		const FNetHandle ReplicatedOuterHandle = GetDynamicRoot(NetHandle);
+		const FNetRefHandle ReplicatedOuterHandle = GetDynamicRoot(NetRefHandle);
 		if (ReplicatedOuterHandle.IsValid())
 		{
-			HandleToDynamicOuter.Add(ReplicatedOuterHandle, NetHandle);
+			HandleToDynamicOuter.Add(ReplicatedOuterHandle, NetRefHandle);
 		}
 	}
 #endif
@@ -346,15 +346,15 @@ bool FObjectReferenceCache::CreateObjectReferenceInternal(const UObject* Object,
 	return true;
 }
 
-FNetHandle FObjectReferenceCache::GetObjectReferenceHandleFromObject(const UObject* Object) const
+FNetRefHandle FObjectReferenceCache::GetObjectReferenceHandleFromObject(const UObject* Object) const
 {
 	if (Object == nullptr)
 	{
-		return FNetHandle();
+		return FNetRefHandle();
 	}
 
-	// Check if we are already know about this object
-	if (const FNetHandle* Reference = ObjectToNetReferenceHandle.Find(Object))
+	// Check if we already know about this object
+	if (const FNetRefHandle* Reference = ObjectToNetReferenceHandle.Find(Object))
 	{
 		// Verify that this is the same object
 		const FCachedNetObjectReference& CachedObject = ReferenceHandleToCachedReference.FindChecked(*Reference);
@@ -362,20 +362,20 @@ FNetHandle FObjectReferenceCache::GetObjectReferenceHandleFromObject(const UObje
 		const UObject* ExistingObject = CachedObject.Object.Get();
 		if (ExistingObject == Object)
 		{
-			UE_LOG_REFERENCECACHE(VeryVerbose, TEXT("ObjectReferenceCache::GetObjectReferenceHandleFromObject Found existing %s for Object %s, OuterNetHandle: %s"), *CachedObject.NetHandle.ToString(), ToCStr(Object->GetPathName()), *CachedObject.OuterNetHandle.ToString());
+			UE_LOG_REFERENCECACHE(VeryVerbose, TEXT("ObjectReferenceCache::GetObjectReferenceHandleFromObject Found existing %s for Object %s, OuterNetRefHandle: %s"), *CachedObject.NetRefHandle.ToString(), ToCStr(Object->GetPathName()), *CachedObject.OuterNetRefHandle.ToString());
 			
 			return *Reference;
 		}
 		else
 		{
-			UE_LOG_REFERENCECACHE(Verbose, TEXT("ObjectReferenceCache::Found %s cached RefHandle %s for object %s"), CachedObject.Object.IsStale() ? TEXT("Stale") : TEXT("UnResolved"), *CachedObject.NetHandle.ToString(), ToCStr(Object->GetName()));
+			UE_LOG_REFERENCECACHE(Verbose, TEXT("ObjectReferenceCache::Found %s cached RefHandle %s for object %s"), CachedObject.Object.IsStale() ? TEXT("Stale") : TEXT("UnResolved"), *CachedObject.NetRefHandle.ToString(), ToCStr(Object->GetName()));
 		}
 	}
 
-	return FNetHandle();
+	return FNetRefHandle();
 }
 
-void FObjectReferenceCache::AddRemoteReference(FNetHandle RefHandle, const UObject* Object)
+void FObjectReferenceCache::AddRemoteReference(FNetRefHandle RefHandle, const UObject* Object)
 {
 	// Only clients should get here
 	check(!IsAuthority()); 
@@ -388,12 +388,12 @@ void FObjectReferenceCache::AddRemoteReference(FNetHandle RefHandle, const UObje
 
 	#if UE_NET_VALIDATE_REFERENCECACHE
 		// Dynamic object might have been referenced before it gets replicated, then it will then already be registered (with a path relative to its outer)
-		if (FNetHandle* ExistingHandle = ObjectToNetReferenceHandle.Find(Object))
+		if (FNetRefHandle* ExistingHandle = ObjectToNetReferenceHandle.Find(Object))
 		{
 			FCachedNetObjectReference& ExistingCachedObject = ReferenceHandleToCachedReference.FindChecked(*ExistingHandle);
-			if (ExistingCachedObject.NetHandle != RefHandle)
+			if (ExistingCachedObject.NetRefHandle != RefHandle)
 			{
-				UE_LOG_REFERENCECACHE_WARNING(TEXT("AddRemoteReference: Detected conflicting Ref %s"), *ExistingCachedObject.NetHandle.ToString());
+				UE_LOG_REFERENCECACHE_WARNING(TEXT("AddRemoteReference: Detected conflicting Ref %s"), *ExistingCachedObject.NetRefHandle.ToString());
 			}
 		}
 	#endif
@@ -413,9 +413,9 @@ void FObjectReferenceCache::AddRemoteReference(FNetHandle RefHandle, const UObje
 	
 			CachedObject.Object = MakeWeakObjectPtr(const_cast<UObject*>(Object));
 			CachedObject.ObjectKey = Object;
-			CachedObject.NetHandle = RefHandle;
+			CachedObject.NetRefHandle = RefHandle;
 			CachedObject.RelativePath = FNetToken();
-			CachedObject.OuterNetHandle = FNetHandle();
+			CachedObject.OuterNetRefHandle = FNetRefHandle();
 			CachedObject.bIsPackage = false;
 			CachedObject.bNoLoad = false;
 			CachedObject.bIsBroken = false;
@@ -429,7 +429,7 @@ void FObjectReferenceCache::AddRemoteReference(FNetHandle RefHandle, const UObje
 	}
 }
 
-void FObjectReferenceCache::RemoveReference(FNetHandle RefHandle, const UObject* Object)
+void FObjectReferenceCache::RemoveReference(FNetRefHandle RefHandle, const UObject* Object)
 {
 	// Cleanup cached data for dynamic instance
 	if (ensure(RefHandle.IsDynamic()))
@@ -437,7 +437,7 @@ void FObjectReferenceCache::RemoveReference(FNetHandle RefHandle, const UObject*
 		UE_LOG_REFERENCECACHE(Verbose, TEXT("ObjectReferenceCache::RemoveReference %s, Object: %s Pointer: %p"), *RefHandle.ToString(), Object ? *Object->GetName() : TEXT("NULL"), Object);
 
 		// We create an iterator to avoid having to find the key twice
-		TMap<FNetHandle, FCachedNetObjectReference>::TKeyIterator It = ReferenceHandleToCachedReference.CreateKeyIterator(RefHandle);
+		TMap<FNetRefHandle, FCachedNetObjectReference>::TKeyIterator It = ReferenceHandleToCachedReference.CreateKeyIterator(RefHandle);
 		if (It)
 		{
 			FCachedNetObjectReference& CachedItemToRemove = It.Value();
@@ -453,11 +453,11 @@ void FObjectReferenceCache::RemoveReference(FNetHandle RefHandle, const UObject*
 
 #if UE_NET_CLEANUP_REFERENCES_WITH_DYNAMIC_OUTER
 			// We want to remove cached references that referenced the removed dynamic references
-			TArray<FNetHandle> SubReferences;
+			TArray<FNetRefHandle> SubReferences;
 			HandleToDynamicOuter.MultiFind(RefHandle, SubReferences);
 			if (SubReferences.Num())
 			{
-				for (FNetHandle SubRefHandle : SubReferences)
+				for (FNetRefHandle SubRefHandle : SubReferences)
 				{
 					FCachedNetObjectReference RemovedSubObjectRef;
 					if (ReferenceHandleToCachedReference.RemoveAndCopyValue(SubRefHandle, RemovedSubObjectRef))
@@ -489,10 +489,10 @@ void FObjectReferenceCache::RemoveReference(FNetHandle RefHandle, const UObject*
 		if (Object)
 		{
 			// We create an iterator to avoid having to find the key twice
-			TMap<const UObject*, FNetHandle>::TKeyIterator ObjectIt = ObjectToNetReferenceHandle.CreateKeyIterator(Object);
+			TMap<const UObject*, FNetRefHandle>::TKeyIterator ObjectIt = ObjectToNetReferenceHandle.CreateKeyIterator(Object);
 			if (ObjectIt)
 			{
-				const FNetHandle ExistingRefHandle = ObjectIt.Value();
+				const FNetRefHandle ExistingRefHandle = ObjectIt.Value();
 				if (ExistingRefHandle == RefHandle)
 				{
 					ObjectIt.RemoveCurrent();
@@ -508,7 +508,7 @@ void FObjectReferenceCache::RemoveReference(FNetHandle RefHandle, const UObject*
 	}
 }
 
-UObject* FObjectReferenceCache::GetObjectFromReferenceHandle(FNetHandle RefHandle)
+UObject* FObjectReferenceCache::GetObjectFromReferenceHandle(FNetRefHandle RefHandle)
 {
 	FCachedNetObjectReference* CachedObject = ReferenceHandleToCachedReference.Find(RefHandle);
 
@@ -518,7 +518,7 @@ UObject* FObjectReferenceCache::GetObjectFromReferenceHandle(FNetHandle RefHandl
 }
 
 // $IRIS: $TODO: Most of the logic comes from GUIDCache::GetObjectFromNetGUID so we want to keep this in sync
-UObject* FObjectReferenceCache::ResolveObjectReferenceHandleInternal(FNetHandle RefHandle, const FNetObjectResolveContext& ResolveContext, bool& bOutMustBeMapped)
+UObject* FObjectReferenceCache::ResolveObjectReferenceHandleInternal(FNetRefHandle RefHandle, const FNetObjectResolveContext& ResolveContext, bool& bOutMustBeMapped)
 {
 	FCachedNetObjectReference* CacheObjectPtr = ReferenceHandleToCachedReference.Find(RefHandle);
 	
@@ -545,7 +545,7 @@ UObject* FObjectReferenceCache::ResolveObjectReferenceHandleInternal(FNetHandle
 			RenamePathForPie(ResolveContext.ConnectionId, ObjectPath, bReading);
 			check(ObjectPath == Object->GetFName().ToString() || RefHandle.IsDynamic() || IsAuthority());
 
-			FNetHandle* ExistingRefHandle = ObjectToNetReferenceHandle.Find(Object);
+			FNetRefHandle* ExistingRefHandle = ObjectToNetReferenceHandle.Find(Object);
 			if (ExistingRefHandle && *ExistingRefHandle != RefHandle)
 			{
 				UE_LOG_REFERENCECACHE_WARNING(TEXT("ObjectReferenceCache::ResolveObjectRefererenceHandle detected potential conflicting references: %s and %s"), *RefHandle.ToString(), *(*ExistingRefHandle).ToString());
@@ -585,14 +585,14 @@ UObject* FObjectReferenceCache::ResolveObjectReferenceHandleInternal(FNetHandle
 	// First, resolve the outer
 	UObject* ObjOuter = nullptr;
 
-	if (CacheObjectPtr->OuterNetHandle.IsValid())
+	if (CacheObjectPtr->OuterNetRefHandle.IsValid())
 	{
 		// If we get here, we depend on an outer to fully load, don't go further until we know we have a fully loaded outer
-		FCachedNetObjectReference* OuterCacheObject = ReferenceHandleToCachedReference.Find(CacheObjectPtr->OuterNetHandle);
+		FCachedNetObjectReference* OuterCacheObject = ReferenceHandleToCachedReference.Find(CacheObjectPtr->OuterNetRefHandle);
 		if (OuterCacheObject == nullptr)
 		{
 			// Shouldn't be possible, but just in case...
-			if (CacheObjectPtr->OuterNetHandle.IsStatic())
+			if (CacheObjectPtr->OuterNetRefHandle.IsStatic())
 			{
 				UE_LOG(LogIris, Error, TEXT("GetObjectFromRefHandle: Static outer not registered. FullPath: %s"), ToCStr(FullPath(RefHandle, ResolveContext)))
 				CacheObjectPtr->bIsBroken = 1;	// Set this to 1 so that we don't keep spamming
@@ -611,7 +611,7 @@ UObject* FObjectReferenceCache::ResolveObjectReferenceHandleInternal(FNetHandle
 
 		// Try to resolve the outer
 		bool bOuterMustBeMapped;
-		ObjOuter = ResolveObjectReferenceHandleInternal(CacheObjectPtr->OuterNetHandle, ResolveContext, bOuterMustBeMapped);
+		ObjOuter = ResolveObjectReferenceHandleInternal(CacheObjectPtr->OuterNetRefHandle, ResolveContext, bOuterMustBeMapped);
 
 		// If we can't resolve the outer, we cannot do more, we are probably waiting for a replicated object or we are waiting for a stringtoken
 		if (ObjOuter == nullptr)
@@ -815,7 +815,7 @@ UObject* FObjectReferenceCache::ResolveObjectReferenceHandleInternal(FNetHandle
 	//}
 
 	// Promote to full handle
-	FNetHandle CompleteRefHandle = FNetHandleManager::MakeNetHandle(RefHandle.GetId(), ReplicationSystem->GetId());
+	FNetRefHandle CompleteRefHandle = FNetRefHandleManager::MakeNetRefHandle(RefHandle.GetId(), ReplicationSystem->GetId());
 
 	// If we already have an existing handle associated with this object, we need do do a few additional checks.
 	// Due to level streaming (streaming in and out the same level) we might have multiple path based references 
@@ -825,20 +825,20 @@ UObject* FObjectReferenceCache::ResolveObjectReferenceHandleInternal(FNetHandle
 	// and simply use the highest index as the correct reference	
 
 	bool bShouldAssign = true;
-	FNetHandle* ExistingHandle = ObjectToNetReferenceHandle.Find(Object);
+	FNetRefHandle* ExistingHandle = ObjectToNetReferenceHandle.Find(Object);
 	if (ExistingHandle && (*ExistingHandle != CompleteRefHandle))
 	{		
 		FCachedNetObjectReference& ExistingCachedObject = ReferenceHandleToCachedReference.FindChecked(*ExistingHandle);
 		if (ExistingCachedObject.Object.IsStale())
 		{
-			UE_LOG_REFERENCECACHE_WARNING(TEXT("Invalidating object for Stale Ref %s with %s as the object has been reused"), *ExistingCachedObject.NetHandle.ToString(), *CompleteRefHandle.ToString());
+			UE_LOG_REFERENCECACHE_WARNING(TEXT("Invalidating object for Stale Ref %s with %s as the object has been reused"), *ExistingCachedObject.NetRefHandle.ToString(), *CompleteRefHandle.ToString());
 		}
 		else if (ExistingCachedObject.Object.HasSameIndexAndSerialNumber(CacheObjectPtr->Object))
 		{
 			bShouldAssign = CompleteRefHandle.GetId() > (*ExistingHandle).GetId();
 			if (bShouldAssign)
 			{
-				UE_LOG_REFERENCECACHE_WARNING(TEXT("Invalidating object for conflicting Ref %s with %s"), *ExistingCachedObject.NetHandle.ToString(), *CompleteRefHandle.ToString());
+				UE_LOG_REFERENCECACHE_WARNING(TEXT("Invalidating object for conflicting Ref %s with %s"), *ExistingCachedObject.NetRefHandle.ToString(), *CompleteRefHandle.ToString());
 			}	
 		}
 	}
@@ -911,37 +911,37 @@ FNetObjectReference FObjectReferenceCache::GetReplicatedOuter(const FNetObjectRe
 	}
 
 	const FCachedNetObjectReference* CachedObject = ReferenceHandleToCachedReference.Find(Reference.GetRefHandle());
-	FNetHandle OuterNetHandle = CachedObject ? CachedObject->OuterNetHandle : FNetHandle();
+	FNetRefHandle OuterNetRefHandle = CachedObject ? CachedObject->OuterNetRefHandle : FNetRefHandle();
 
-	while (OuterNetHandle.IsValid())
+	while (OuterNetRefHandle.IsValid())
 	{
-		if (NetHandleManager->GetInternalIndex(OuterNetHandle))
+		if (NetRefHandleManager->GetInternalIndex(OuterNetRefHandle))
 		{
-			return FNetObjectReference(OuterNetHandle);
+			return FNetObjectReference(OuterNetRefHandle);
 		}
 
-		CachedObject = ReferenceHandleToCachedReference.Find(OuterNetHandle);
-		OuterNetHandle = CachedObject ? CachedObject->OuterNetHandle : FNetHandle();
+		CachedObject = ReferenceHandleToCachedReference.Find(OuterNetRefHandle);
+		OuterNetRefHandle = CachedObject ? CachedObject->OuterNetRefHandle : FNetRefHandle();
 	};
 
 	return FNetObjectReference();
 }
 
-FNetHandle FObjectReferenceCache::GetDynamicRoot(const FNetHandle Handle) const
+FNetRefHandle FObjectReferenceCache::GetDynamicRoot(const FNetRefHandle Handle) const
 {
 	const FCachedNetObjectReference* CachedObject = ReferenceHandleToCachedReference.Find(Handle);
-	while (CachedObject && CachedObject->OuterNetHandle.IsValid())
+	while (CachedObject && CachedObject->OuterNetRefHandle.IsValid())
 	{
-		CachedObject = ReferenceHandleToCachedReference.Find(CachedObject->OuterNetHandle);
+		CachedObject = ReferenceHandleToCachedReference.Find(CachedObject->OuterNetRefHandle);
 	}
 
-	if (CachedObject && CachedObject->NetHandle.IsDynamic())
+	if (CachedObject && CachedObject->NetRefHandle.IsDynamic())
 	{
-		return CachedObject->NetHandle;
+		return CachedObject->NetRefHandle;
 	}
 	else
 	{
-		return FNetHandle();
+		return FNetRefHandle();
 	}
 }
 
@@ -965,9 +965,9 @@ FNetObjectReference FObjectReferenceCache::GetOrCreateObjectReference(const FStr
 	FNetObjectReference OutReference;
 
 	const bool bIsStatic = true;
-	FNetHandle Handle = NetHandleManager->AllocateNetHandle(bIsStatic);
+	FNetRefHandle RefHandle = NetRefHandleManager->AllocateNetRefHandle(bIsStatic);
 
-	if (Handle.IsValid())
+	if (RefHandle.IsValid())
 	{
 		LLM_SCOPE_BYTAG(Iris);
 
@@ -975,21 +975,21 @@ FNetObjectReference FObjectReferenceCache::GetOrCreateObjectReference(const FStr
 		FCachedNetObjectReference CachedObject;
 
 		CachedObject.ObjectKey = nullptr;
-		CachedObject.NetHandle = Handle;
+		CachedObject.NetRefHandle = RefHandle;
 		CachedObject.RelativePath = StringTokenStore->GetOrCreateToken(ObjectPath);
-		CachedObject.OuterNetHandle = CreateObjectReferenceHandle(Outer);
+		CachedObject.OuterNetRefHandle = CreateObjectReferenceHandle(Outer);
 		CachedObject.bIsPackage = 0U;
 		CachedObject.bNoLoad = 1U;
 		CachedObject.bIgnoreWhenMissing = 1U;
 		CachedObject.bIsPending = 0U;
 		CachedObject.bIsBroken = 0U;
 
-		ReferenceHandleToCachedReference.Add(Handle, CachedObject);
+		ReferenceHandleToCachedReference.Add(RefHandle, CachedObject);
 	
-		UE_LOG_REFERENCECACHE(Verbose, TEXT("ObjectReferenceCache::GetOrCreateObjectReference(FromPath) %s for Object %s, %s"), *CachedObject.NetHandle.ToString(), ToCStr(ObjectPath), *CachedObject.OuterNetHandle.ToString());
+		UE_LOG_REFERENCECACHE(Verbose, TEXT("ObjectReferenceCache::GetOrCreateObjectReference(FromPath) %s for Object %s, %s"), *CachedObject.NetRefHandle.ToString(), ToCStr(ObjectPath), *CachedObject.OuterNetRefHandle.ToString());
 	}
 
-	OutReference = FNetObjectReference(Handle);
+	OutReference = FNetObjectReference(RefHandle);
 
 	return OutReference;
 }
@@ -1034,7 +1034,7 @@ void FObjectReferenceCache::ConditionalReadNetTokenData(FNetSerializationContext
 void FObjectReferenceCache::WriteFullReferenceInternal(FNetSerializationContext& Context, const FNetObjectReference& Ref) const
 {
 	FNetBitStreamWriter* Writer = Context.GetBitStreamWriter();
-	const FNetHandle RefHandle = Ref.GetRefHandle();
+	const FNetRefHandle RefHandle = Ref.GetRefHandle();
 
 	UE_NET_TRACE_OBJECT_SCOPE(RefHandle, *Writer, Context.GetTraceCollector(), ENetTraceVerbosity::Verbose);
 	
@@ -1045,13 +1045,12 @@ void FObjectReferenceCache::WriteFullReferenceInternal(FNetSerializationContext&
 	{
 		UE_CLOG(RefHandle.IsValid(), LogIris, Log, TEXT("ObjectReferenceCache::WriteFullReference Trying to write Stale handle, %s"), *RefHandle.ToString());
 
-		WriteNetHandle(Context, FNetHandle());
-
+		WriteNetRefHandle(Context, FNetRefHandle());
 		return;			
 	}
 
 	// Write handle
-	WriteNetHandle(Context, RefHandle);
+	WriteNetRefHandle(Context, RefHandle);
 
 	const bool bHasPath = CachedObject->RelativePath.IsValid();
 
@@ -1060,7 +1059,7 @@ void FObjectReferenceCache::WriteFullReferenceInternal(FNetSerializationContext&
 	// We always export if we do not have a ExportContext
 	const bool bMustExport = bHasPath && IsAuthority() && (!ExportContext || !ExportContext->IsExported(RefHandle));
 
-	UE_LOG_REFERENCECACHE(VeryVerbose, TEXT("ObjectReferenceCache::WriteFullReferenceInternal Auth: %u %s, Outer: %s Export: %u"), IsAuthority() ? 1U : 0U, *CachedObject->NetHandle.ToString(), *CachedObject->OuterNetHandle.ToString(), bMustExport ? 1U : 0U);
+	UE_LOG_REFERENCECACHE(VeryVerbose, TEXT("ObjectReferenceCache::WriteFullReferenceInternal Auth: %u %s, Outer: %s Export: %u"), IsAuthority() ? 1U : 0U, *CachedObject->NetRefHandle.ToString(), *CachedObject->OuterNetRefHandle.ToString(), bMustExport ? 1U : 0U);
 
 	if (Writer->WriteBool(bMustExport))
 	{
@@ -1077,7 +1076,7 @@ void FObjectReferenceCache::WriteFullReferenceInternal(FNetSerializationContext&
 		{
 			WriteNetToken(Writer, CachedObject->RelativePath);
 			ConditionalWriteNetTokenData(Context, ExportContext, CachedObject->RelativePath);
-			WriteFullReferenceInternal(Context, FNetObjectReference(CachedObject->OuterNetHandle));
+			WriteFullReferenceInternal(Context, FNetObjectReference(CachedObject->OuterNetRefHandle));
 		}
 	
 		if (ExportContext)
@@ -1101,11 +1100,11 @@ void FObjectReferenceCache::ReadFullReferenceInternal(FNetSerializationContext&
 		return;
 	}
 
-	UE_NET_TRACE_NAMED_OBJECT_SCOPE(ReferenceScope, FNetHandle(), *Context.GetBitStreamReader(), Context.GetTraceCollector(), ENetTraceVerbosity::Verbose);
+	UE_NET_TRACE_NAMED_OBJECT_SCOPE(ReferenceScope, FNetRefHandle(), *Context.GetBitStreamReader(), Context.GetTraceCollector(), ENetTraceVerbosity::Verbose);
 
-	const FNetHandle NetHandle = ReadNetHandle(Context);
+	const FNetRefHandle NetRefHandle = ReadNetRefHandle(Context);
 
-	if (!NetHandle.IsValid())
+	if (!NetRefHandle.IsValid())
 	{
 		OutObjectRef = FNetObjectReference();
 		return;
@@ -1119,10 +1118,10 @@ void FObjectReferenceCache::ReadFullReferenceInternal(FNetSerializationContext&
 		return;
 	}
 
-	UE_LOG_REFERENCECACHE(VeryVerbose, TEXT("ObjectReferenceCache::ReadFullReferenceInternal Auth: %u %s, IsExported: %u"), IsAuthority() ? 1U : 0U, *NetHandle.ToString(), bIsExported ? 1U : 0U);
+	UE_LOG_REFERENCECACHE(VeryVerbose, TEXT("ObjectReferenceCache::ReadFullReferenceInternal Auth: %u %s, IsExported: %u"), IsAuthority() ? 1U : 0U, *NetRefHandle.ToString(), bIsExported ? 1U : 0U);
 
 	// For dynamic we assume no load, otherwise we will get it through export
-	uint8 bNoLoad = NetHandle.IsDynamic();
+	uint8 bNoLoad = NetRefHandle.IsDynamic();
 
 	// Do we have a path
 	FNetObjectReference OuterRef;
@@ -1155,24 +1154,24 @@ void FObjectReferenceCache::ReadFullReferenceInternal(FNetSerializationContext&
 		return;
 	}
 
-	UE_NET_TRACE_SET_SCOPE_OBJECTID(ReferenceScope, NetHandle);
+	UE_NET_TRACE_SET_SCOPE_OBJECTID(ReferenceScope, NetRefHandle);
 
 	// Already known? Validate that it is correct!
-	if (const FCachedNetObjectReference* CachedObject = ReferenceHandleToCachedReference.Find(NetHandle))
+	if (const FCachedNetObjectReference* CachedObject = ReferenceHandleToCachedReference.Find(NetRefHandle))
 	{
-		// Dynamic objects can be referenced by either NetHandle or ParentNetHandle/PathCombo
-		check(!bIsExported || ((CachedObject->NetHandle == NetHandle && CachedObject->RelativePath == RelativePath && CachedObject->OuterNetHandle == OuterRef.GetRefHandle()) || 
-				(NetHandle.IsDynamic() && (NetHandle == CachedObject->NetHandle))));
+		// Dynamic objects can be referenced by either NetRefHandle or parent NetRefHandle and path combo
+		check(!bIsExported || ((CachedObject->NetRefHandle == NetRefHandle && CachedObject->RelativePath == RelativePath && CachedObject->OuterNetRefHandle == OuterRef.GetRefHandle()) || 
+				(NetRefHandle.IsDynamic() && (NetRefHandle == CachedObject->NetRefHandle))));
 
-		OutObjectRef = FNetObjectReference(NetHandle);
+		OutObjectRef = FNetObjectReference(NetRefHandle);
 
 		return;
 	}
 	else
 	{
-		if (!(bIsExported || NetHandle.IsDynamic()))
+		if (!(bIsExported || NetRefHandle.IsDynamic()))
 		{
-			UE_LOG_REFERENCECACHE(Error, TEXT("ObjectReferenceCache::ReadFullReferenceInternal Fail reading %s IsExported:%u"), *NetHandle.ToString(), bIsExported ? 1U : 0U);
+			UE_LOG_REFERENCECACHE(Error, TEXT("ObjectReferenceCache::ReadFullReferenceInternal Fail reading %s IsExported:%u"), *NetRefHandle.ToString(), bIsExported ? 1U : 0U);
 			Context.SetError(GNetError_InvalidNetHandle);
 			return;
 		}
@@ -1185,28 +1184,28 @@ void FObjectReferenceCache::ReadFullReferenceInternal(FNetSerializationContext&
 		FCachedNetObjectReference CachedObject;
 
 		CachedObject.ObjectKey = nullptr;
-		CachedObject.NetHandle = NetHandle;
+		CachedObject.NetRefHandle = NetRefHandle;
 		CachedObject.RelativePath = RelativePath;
-		CachedObject.OuterNetHandle = OuterRef.GetRefHandle();
-		CachedObject.bIsPackage = NetHandle.IsStatic() && !OuterRef.GetRefHandle().IsValid();
+		CachedObject.OuterNetRefHandle = OuterRef.GetRefHandle();
+		CachedObject.bIsPackage = NetRefHandle.IsStatic() && !OuterRef.GetRefHandle().IsValid();
 		CachedObject.bNoLoad = bNoLoad;
 		CachedObject.bIsBroken = false;
 		CachedObject.bIsPending = false;
 		CachedObject.bIgnoreWhenMissing = bNoLoad;
 
-		ReferenceHandleToCachedReference.Add(NetHandle, CachedObject);
+		ReferenceHandleToCachedReference.Add(NetRefHandle, CachedObject);
 	}
 
-	OutObjectRef = FNetObjectReference(NetHandle);
+	OutObjectRef = FNetObjectReference(NetRefHandle);
 
 #if UE_NET_CLEANUP_REFERENCES_WITH_DYNAMIC_OUTER
 	// We want to track references with a dynamic root so that we can purge them when the dynamic root is removed
-	if (NetHandle.IsDynamic() && OuterRef.IsValid())
+	if (NetRefHandle.IsDynamic() && OuterRef.IsValid())
 	{
-		FNetHandle DynamicOuterHandle = GetDynamicRoot(NetHandle);
+		FNetRefHandle DynamicOuterHandle = GetDynamicRoot(NetRefHandle);
 		if (DynamicOuterHandle.IsValid())
 		{
-			HandleToDynamicOuter.Add(DynamicOuterHandle, NetHandle);
+			HandleToDynamicOuter.Add(DynamicOuterHandle, NetRefHandle);
 		}
 	}
 #endif
@@ -1218,17 +1217,17 @@ void FObjectReferenceCache::ReadFullReference(FNetSerializationContext& Context,
 	FNetBitStreamReader* Reader = Context.GetBitStreamReader();
 
 	UE_NET_TRACE_SCOPE(FullNetObjectReference, *Reader, Context.GetTraceCollector(), ENetTraceVerbosity::Trace);
-	UE_NET_TRACE_NAMED_OBJECT_SCOPE(ReferenceScope, FNetHandle(), *Reader, Context.GetTraceCollector(), ENetTraceVerbosity::Verbose);
+	UE_NET_TRACE_NAMED_OBJECT_SCOPE(ReferenceScope, FNetRefHandle(), *Reader, Context.GetTraceCollector(), ENetTraceVerbosity::Verbose);
 	LLM_SCOPE_BYTAG(Iris);
 
 	// Handle client assigned reference
 	if (const bool bIsClientAssignedReference = Reader->ReadBool())
 	{
-		const FNetHandle NetHandle = ReadNetHandle(Context);
+		const FNetRefHandle NetRefHandle = ReadNetRefHandle(Context);
 		FNetToken RelativePath = ReadNetToken(Reader);
 		ConditionalReadNetTokenData(Context, RelativePath);
 
-		OutRef.RefHandle = NetHandle;
+		OutRef.RefHandle = NetRefHandle;
 		OutRef.PathToken = RelativePath;
 
 		return;
@@ -1250,7 +1249,7 @@ void FObjectReferenceCache::WriteFullReference(FNetSerializationContext& Context
 	const bool bIsClientAssignedReference = Ref.PathToken.IsValid();
 	if (Writer->WriteBool(bIsClientAssignedReference))
 	{
-		WriteNetHandle(Context, Ref.GetRefHandle());
+		WriteNetRefHandle(Context, Ref.GetRefHandle());
 		WriteNetToken(Writer, Ref.PathToken);
 		FNetExportContext* ExportContext = Context.GetExportContext();
 		ConditionalWriteNetTokenData(Context, ExportContext, Ref.PathToken);
@@ -1267,13 +1266,13 @@ void FObjectReferenceCache::WriteReference(FNetSerializationContext& Context, FN
 	UE_NET_TRACE_OBJECT_SCOPE(Ref.GetRefHandle(), *Context.GetBitStreamWriter(), Context.GetTraceCollector(), ENetTraceVerbosity::Trace);
 
 	FNetBitStreamWriter* Writer = Context.GetBitStreamWriter();
-	const FNetHandle RefHandle = Ref.GetRefHandle();
+	const FNetRefHandle RefHandle = Ref.GetRefHandle();
 
 	// Only client assigned references contains a pathtoken
 	const bool bIsClientAssignedReference = Ref.PathToken.IsValid();
 	if (Writer->WriteBool(bIsClientAssignedReference))
 	{
-		WriteNetHandle(Context, Ref.GetRefHandle());
+		WriteNetRefHandle(Context, Ref.GetRefHandle());
 		WriteNetToken(Writer, Ref.PathToken);
 
 		return;
@@ -1287,7 +1286,7 @@ void FObjectReferenceCache::WriteReference(FNetSerializationContext& Context, FN
 	{
 		UE_CLOG(RefHandle.IsValid(), LogIris, Verbose, TEXT("ObjectReferenceCache::WriteReference Trying to write Stale handle, %s"), *RefHandle.ToString());
 
-		WriteNetHandle(Context, FNetHandle());
+		WriteNetRefHandle(Context, FNetRefHandle());
 
 		return;
 	}
@@ -1297,29 +1296,29 @@ void FObjectReferenceCache::WriteReference(FNetSerializationContext& Context, FN
 	checkSlow(Ref.CanBeExported() || RefHandle.IsDynamic());
 
 	// Write handle
-	WriteNetHandle(Context, RefHandle);
+	WriteNetRefHandle(Context, RefHandle);
 }
 
 void FObjectReferenceCache::ReadReference(FNetSerializationContext& Context, FNetObjectReference& OutRef)
 {
 	UE_NET_TRACE_SCOPE(NetObjectReference, *Context.GetBitStreamReader(), Context.GetTraceCollector(), ENetTraceVerbosity::Trace);
-	UE_NET_TRACE_NAMED_OBJECT_SCOPE(ReferenceScope, FNetHandle(), *Context.GetBitStreamReader(), Context.GetTraceCollector(), ENetTraceVerbosity::Verbose);
+	UE_NET_TRACE_NAMED_OBJECT_SCOPE(ReferenceScope, FNetRefHandle(), *Context.GetBitStreamReader(), Context.GetTraceCollector(), ENetTraceVerbosity::Verbose);
 	LLM_SCOPE_BYTAG(Iris);
 
 	FNetBitStreamReader* Reader = Context.GetBitStreamReader();
 
 	if (const bool bIsClientAssignedReference = Reader->ReadBool())
 	{
-		const FNetHandle NetHandle = ReadNetHandle(Context);
+		const FNetRefHandle NetRefHandle = ReadNetRefHandle(Context);
 		FNetToken RelativePath = ReadNetToken(Reader);
 
-		OutRef.RefHandle = NetHandle;
+		OutRef.RefHandle = NetRefHandle;
 		OutRef.PathToken = RelativePath;
 
 		return;
 	}
 
-	const FNetHandle NetHandle = ReadNetHandle(Context);
+	const FNetRefHandle NetRefHandle = ReadNetRefHandle(Context);
 
 	if (Reader->IsOverflown())
 	{
@@ -1328,19 +1327,19 @@ void FObjectReferenceCache::ReadReference(FNetSerializationContext& Context, FNe
 		return;
 	}
 
-	if (!NetHandle.IsValid() )
+	if (!NetRefHandle.IsValid())
 	{
 		OutRef = FNetObjectReference();
 		return;
 	}
 
-	UE_NET_TRACE_SET_SCOPE_OBJECTID(ReferenceScope, NetHandle);
+	UE_NET_TRACE_SET_SCOPE_OBJECTID(ReferenceScope, NetRefHandle);
 
 	// The handle must exist in the cache or it must be a dynamic handle that will be explicitly registered
-	if (ensureAlwaysMsgf(NetHandle.IsDynamic() || ReferenceHandleToCachedReference.Find(NetHandle), TEXT("FObjectReferenceCache::ReadReference: Handle %s must already be known or be dynamic"), *NetHandle.ToString()))
+	if (ensureAlwaysMsgf(NetRefHandle.IsDynamic() || ReferenceHandleToCachedReference.Find(NetRefHandle), TEXT("FObjectReferenceCache::ReadReference: Handle %s must already be known or be dynamic"), *NetRefHandle.ToString()))
 	{
-		UE_LOG_REFERENCECACHE(Verbose, TEXT("ObjectReferenceCache::ReadReference Auth: %u %s"), IsAuthority() ? 1U : 0U, *NetHandle.ToString());
-		OutRef = FNetObjectReference(NetHandle);
+		UE_LOG_REFERENCECACHE(Verbose, TEXT("ObjectReferenceCache::ReadReference Auth: %u %s"), IsAuthority() ? 1U : 0U, *NetRefHandle.ToString());
+		OutRef = FNetObjectReference(NetRefHandle);
 	}
 	else
 	{
@@ -1359,7 +1358,7 @@ bool FObjectReferenceCache::WriteExports(FNetSerializationContext& Context, TArr
 		for (const FNetObjectReference& Reference : ExportsView)
 		{
 			const bool bIsClientAssigned = Reference.PathToken.IsValid();
-			const FNetHandle& Handle = Reference.GetRefHandle();
+			const FNetRefHandle Handle = Reference.GetRefHandle();
 
 			if ((bIsClientAssigned && !ExportContext->IsExported(Reference.GetPathToken())) || (IsAuthority() && Reference.CanBeExported() && !ExportContext->IsExported(Handle)))
 			{
@@ -1392,7 +1391,7 @@ bool FObjectReferenceCache::ReadExports(FNetSerializationContext& Context)
 	return !Context.HasErrorOrOverflow();
 }
 
-FString FObjectReferenceCache::FullPath(FNetHandle RefHandle, const FNetObjectResolveContext& ResolveContext) const
+FString FObjectReferenceCache::FullPath(FNetRefHandle RefHandle, const FNetObjectResolveContext& ResolveContext) const
 {
 	FString FullPath;
 
@@ -1401,7 +1400,7 @@ FString FObjectReferenceCache::FullPath(FNetHandle RefHandle, const FNetObjectRe
 	return FullPath;
 }
 
-void FObjectReferenceCache::GenerateFullPath_r(FNetHandle RefHandle, const FNetObjectResolveContext& ResolveContext, FString& FullPath) const
+void FObjectReferenceCache::GenerateFullPath_r(FNetRefHandle RefHandle, const FNetObjectResolveContext& ResolveContext, FString& FullPath) const
 {
 	if ( !RefHandle.IsValid() )
 	{
@@ -1419,7 +1418,7 @@ void FObjectReferenceCache::GenerateFullPath_r(FNetHandle RefHandle, const FNetO
 		return;
 	}
 
-	GenerateFullPath_r(CacheObject->OuterNetHandle, ResolveContext, FullPath);
+	GenerateFullPath_r(CacheObject->OuterNetRefHandle, ResolveContext, FullPath);
 
 	if (!FullPath.IsEmpty())
 	{
diff --git a/Engine/Source/Runtime/Experimental/Iris/Core/Private/Iris/ReplicationSystem/ObjectReferenceCache.h b/Engine/Source/Runtime/Experimental/Iris/Core/Private/Iris/ReplicationSystem/ObjectReferenceCache.h
--- a/Engine/Source/Runtime/Experimental/Iris/Core/Private/Iris/ReplicationSystem/ObjectReferenceCache.h
+++ b/Engine/Source/Runtime/Experimental/Iris/Core/Private/Iris/ReplicationSystem/ObjectReferenceCache.h
@@ -16,7 +16,7 @@ namespace UE::Net
 	class FNetSerializationContext;
 	namespace Private
 	{
-		class FNetHandleManager;
+		class FNetRefHandleManager;
 		class FNetExportContext;
 	}
 }
@@ -41,16 +41,16 @@ public:
 	bool IsAuthority() const;
 
 	// Create and assign a new NetHandle to the object
-	FNetHandle CreateObjectReferenceHandle(const UObject* Object);
+	FNetRefHandle CreateObjectReferenceHandle(const UObject* Object);
 
 	// Get existing handle for object
-	FNetHandle GetObjectReferenceHandleFromObject(const UObject* Object) const;
+	FNetRefHandle GetObjectReferenceHandleFromObject(const UObject* Object) const;
 
 	// Get object from handle, only if the object is in the cache.
-	UObject* GetObjectFromReferenceHandle(FNetHandle RefHandle);
+	UObject* GetObjectFromReferenceHandle(FNetRefHandle RefHandle);
 
 	// Try to resolve the object reference and try to load it if the object cannot be found
-	UObject* ResolveObjectReferenceHandle(FNetHandle RefHandle, const FNetObjectResolveContext& ResolveContext);
+	UObject* ResolveObjectReferenceHandle(FNetRefHandle RefHandle, const FNetObjectResolveContext& ResolveContext);
 
 	UObject* ResolveObjectReference(const FNetObjectReference& ObjectRef, const FNetObjectResolveContext& ResolveContext);
 	ENetObjectReferenceResolveResult ResolveObjectReference(const FNetObjectReference& ObjectRef, const FNetObjectResolveContext& ResolveContext, UObject*& OutResolvedObject);
@@ -65,10 +65,10 @@ public:
 	FNetObjectReference GetOrCreateObjectReference(const FString& ObjectPath, const UObject* Outer);
 
 	// Add reference for dynamically spawned object
-	void AddRemoteReference(FNetHandle RefHandle, const UObject* Object);
+	void AddRemoteReference(FNetRefHandle RefHandle, const UObject* Object);
 
 	// Remove references to dynamic objects
-	void RemoveReference(FNetHandle RefHandle, const UObject* Object);
+	void RemoveReference(FNetRefHandle RefHandle, const UObject* Object);
 
 	// Write full chain of object references for RefHandle
 	void WriteFullReference(FNetSerializationContext& Context, FNetObjectReference Ref) const;
@@ -88,23 +88,23 @@ public:
 
 	bool ReadExports(FNetSerializationContext& Context);
 
-	static FNetObjectReference MakeNetObjectReference(FNetHandle Handle) { return FNetObjectReference(Handle); }
+	static FNetObjectReference MakeNetObjectReference(FNetRefHandle Handle);
 
 private:
 
 	struct FCachedNetObjectReference
 	{
 		TWeakObjectPtr<UObject> Object;
-		const UObject* ObjectKey;
+		const UObject* ObjectKey = nullptr;
 
-		// NetHandle
-		FNetHandle NetHandle;
+		// NetRefHandle
+		FNetRefHandle NetRefHandle;
 
 		// RelativePath to outer
 		FNetToken RelativePath;
 
 		// Ref to outer
-		FNetHandle OuterNetHandle;
+		FNetRefHandle OuterNetRefHandle;
 
 		// Flags
 		uint8 bNoLoad : 1;				// Don't load this, only do a find
@@ -121,46 +121,46 @@ private:
 
 	void ReadFullReferenceInternal(FNetSerializationContext& Context, FNetObjectReference& OutRef, uint32 RecursionCount);
 	void WriteFullReferenceInternal(FNetSerializationContext& Context, const FNetObjectReference& Ref) const;
-	UObject* ResolveObjectReferenceHandleInternal(FNetHandle RefHandle, const FNetObjectResolveContext& ResolveContext, bool& bOutMustBeMapped);
+	UObject* ResolveObjectReferenceHandleInternal(FNetRefHandle RefHandle, const FNetObjectResolveContext& ResolveContext, bool& bOutMustBeMapped);
 	bool IsDynamicInternal(const UObject* Object) const;
 	bool SupportsObjectInternal(const UObject* Object) const;
 	bool CanClientLoadObjectInternal(const UObject* Object, bool bIsDynamic) const;
-	bool ShouldIgnoreWhenMissing(FNetHandle RefHandle) const;
+	bool ShouldIgnoreWhenMissing(FNetRefHandle RefHandle) const;
 	bool RenamePathForPie(uint32 ConnectionId, FString& Str, bool bReading);
 
 	// Get the string path of RefHandle
-	FString FullPath(FNetHandle RefHandle, const FNetObjectResolveContext& ResolveContext) const;
-	void GenerateFullPath_r(FNetHandle RefHandle, const FNetObjectResolveContext& ResolveContext, FString& OutFullPath) const;
+	FString FullPath(FNetRefHandle RefHandle, const FNetObjectResolveContext& ResolveContext) const;
+	void GenerateFullPath_r(FNetRefHandle RefHandle, const FNetObjectResolveContext& ResolveContext, FString& OutFullPath) const;
 
 	// Find dynamic root
-	FNetHandle GetDynamicRoot(const FNetHandle Handle) const;
+	FNetRefHandle GetDynamicRoot(const FNetRefHandle Handle) const;
 
-	static FNetObjectReference MakeNetObjectReference(FNetHandle NetHandle, FNetToken RelativePath) { return FNetObjectReference(NetHandle, RelativePath, RelativePath.IsValid() ? ENetObjectReferenceTraits::CanBeExported : ENetObjectReferenceTraits::None); }
-	static FNetObjectReference MakeNetObjectReference(const FCachedNetObjectReference& CachedReference) { return FNetObjectReference(CachedReference.NetHandle, FNetToken(), CachedReference.RelativePath.IsValid() ? ENetObjectReferenceTraits::CanBeExported : ENetObjectReferenceTraits::None); }
+	static FNetObjectReference MakeNetObjectReference(FNetRefHandle RefHandle, FNetToken RelativePath);
+	static FNetObjectReference MakeNetObjectReference(const FCachedNetObjectReference& CachedReference);
 
 private:
 
 	// Map raw UObject pointer -> Handle
 	// To verify that the reference is valid we need to check the weakpointer stored in the cache
-	TMap<const UObject*, FNetHandle> ObjectToNetReferenceHandle;
+	TMap<const UObject*, FNetRefHandle> ObjectToNetReferenceHandle;
 
 	// Map ReferenceHandle -> CachedReference
-	TMap<FNetHandle, FCachedNetObjectReference> ReferenceHandleToCachedReference;
+	TMap<FNetRefHandle, FCachedNetObjectReference> ReferenceHandleToCachedReference;
 
 	// To properly clean up stale references referencing dynamic objects we need to track them
-	TMultiMap<FNetHandle, FNetHandle> HandleToDynamicOuter;
+	TMultiMap<FNetRefHandle, FNetRefHandle> HandleToDynamicOuter;
 	
 	UReplicationSystem* ReplicationSystem;
 	UObjectReplicationBridge* ReplicationBridge;
 	FNetTokenStore* NetTokenStore;
 	FStringTokenStore* StringTokenStore;
-	FNetHandleManager* NetHandleManager;
+	FNetRefHandleManager* NetRefHandleManager;
 	
 	// Do we have authority to create references?
 	uint32 bIsAuthority : 1;
 };
 
-inline UObject* FObjectReferenceCache::ResolveObjectReferenceHandle(FNetHandle RefHandle, const FNetObjectResolveContext& ResolveContext)
+inline UObject* FObjectReferenceCache::ResolveObjectReferenceHandle(FNetRefHandle RefHandle, const FNetObjectResolveContext& ResolveContext)
 {
 	bool bMustBeMapped;
 	return ResolveObjectReferenceHandleInternal(RefHandle, ResolveContext, bMustBeMapped);
@@ -173,4 +173,21 @@ inline UObject* FObjectReferenceCache::ResolveObjectReference(const FNetObjectRe
 	return ResolvedObject;
 }
 
+inline FNetObjectReference FObjectReferenceCache::MakeNetObjectReference(FNetRefHandle Handle)
+{
+	return FNetObjectReference(Handle);
+}
+
+inline FNetObjectReference FObjectReferenceCache::MakeNetObjectReference(FNetRefHandle RefHandle, FNetToken RelativePath)
+{
+	const ENetObjectReferenceTraits Traits = RelativePath.IsValid() ? ENetObjectReferenceTraits::CanBeExported : ENetObjectReferenceTraits::None;
+	return FNetObjectReference(RefHandle, RelativePath, Traits);
+}
+
+inline FNetObjectReference FObjectReferenceCache::MakeNetObjectReference(const FCachedNetObjectReference& CachedReference)
+{
+	const ENetObjectReferenceTraits Traits = CachedReference.RelativePath.IsValid() ? ENetObjectReferenceTraits::CanBeExported : ENetObjectReferenceTraits::None;
+	return FNetObjectReference(CachedReference.NetRefHandle, FNetToken(), Traits);
+}
+
 }
diff --git a/Engine/Source/Runtime/Experimental/Iris/Core/Private/Iris/ReplicationSystem/ObjectReplicationBridge.cpp b/Engine/Source/Runtime/Experimental/Iris/Core/Private/Iris/ReplicationSystem/ObjectReplicationBridge.cpp
--- a/Engine/Source/Runtime/Experimental/Iris/Core/Private/Iris/ReplicationSystem/ObjectReplicationBridge.cpp
+++ b/Engine/Source/Runtime/Experimental/Iris/Core/Private/Iris/ReplicationSystem/ObjectReplicationBridge.cpp
@@ -9,9 +9,10 @@
 #include "Iris/Core/IrisMemoryTracker.h"
 #include "Iris/Core/IrisProfiler.h"
 
+#include "Net/Core/NetHandle/NetHandleManager.h"
+#include "Net/Core/PropertyConditions/PropertyConditionsDelegates.h"
 #include "Net/Core/Trace/NetTrace.h"
 #include "Net/Core/Trace/NetDebugName.h"
-#include "Net/Core/PropertyConditions/PropertyConditionsDelegates.h"
 
 #include "Iris/ReplicationSystem/LegacyPushModel.h"
 #include "Iris/ReplicationSystem/ObjectPollFrequencyLimiter.h"
@@ -31,11 +32,7 @@
 #include "Iris/Serialization/InternalNetSerializationContext.h"
 #include "Iris/Serialization/NetBitStreamUtil.h"
 
-#if UE_NET_ALLOW_MULTIPLE_REPLICATION_SYSTEMS
 #define UE_LOG_OBJECTREPLICATIONBRIDGE(Category, Format, ...)  UE_LOG(LogIrisBridge, Category, TEXT("ObjectReplicationBridge(%u)::") Format, GetReplicationSystem()->GetId(), ##__VA_ARGS__)
-#else
-#define UE_LOG_OBJECTREPLICATIONBRIDGE(Category, Format, ...)  UE_LOG(LogIrisBridge, Category, TEXT("ObjectReplicationBridge::") Format, ##__VA_ARGS__)
-#endif
 
 static bool bUseFrequencyBasedPolling = true;
 static FAutoConsoleVariableRef CVarUseFrequencyBasedPolling(
@@ -65,7 +62,7 @@ static FAutoConsoleVariableRef CVarEnableFilterMappings(
 		TEXT("Whether we honor filter mappings set in ObjectReplicationBridgeConfig. If filter mappings are enabled then objects may also be assigned the default spatial filter even if there aren't any specific mappings. Default is true.")
 		);
 
-UObjectReplicationBridge::FCreateNetHandleParams UObjectReplicationBridge::DefaultCreateNetHandleParams =
+UObjectReplicationBridge::FCreateNetRefHandleParams UObjectReplicationBridge::DefaultCreateNetRefHandleParams =
 {
 	0U, // bCanReceive
 	0U, // bNeedsPreUpdate
@@ -117,7 +114,7 @@ void UObjectReplicationBridge::Initialize(UReplicationSystem* InReplicationSyste
 
 	Super::Initialize(InReplicationSystem);
 
-	const uint32 MaxActiveObjectCount = NetHandleManager->GetMaxActiveObjectCount();
+	const uint32 MaxActiveObjectCount = NetRefHandleManager->GetMaxActiveObjectCount();
 	PollFrequencyLimiter->Init(MaxActiveObjectCount);
 	ObjectsWithObjectReferences.Init(MaxActiveObjectCount);
 	GarbageCollectionAffectedObjects.Init(MaxActiveObjectCount);
@@ -125,19 +122,19 @@ void UObjectReplicationBridge::Initialize(UReplicationSystem* InReplicationSyste
 	LoadConfig();
 
 	// Hookup delegate when a property custom conditions is changed
-	OnCustomConditionChangedHandle = UE::Net::Private::FPropertyConditionDelegates::GetOnPropertyCustomConditionChangedDelegate().AddLambda([this, InReplicationSystem](const UObject* Owner, uint16 RepIndex, bool bEnable)
+	OnCustomConditionChangedHandle = UE::Net::Private::FPropertyConditionDelegates::GetOnPropertyCustomConditionChangedDelegate().AddLambda([this](const UObject* Owner, uint16 RepIndex, bool bEnable)
 	{
 		using namespace UE::Net;
 		using namespace UE::Net::Private;
 
-		const FNetHandle Handle = this->GetReplicatedHandle(Owner);
-		if (Handle.IsValid())
+		const FNetRefHandle RefHandle = this->GetReplicatedRefHandle(Owner);
+		if (RefHandle.IsValid())
 		{
-			FReplicationSystemInternal* ReplicationSystemInternal = GetReplicationSystem()->GetReplicationSystemInternal();
-			const FNetHandleManager& LocalNetHandleManager = ReplicationSystemInternal->GetNetHandleManager();
+			FReplicationSystemInternal* ReplicationSystemInternal = this->GetReplicationSystem()->GetReplicationSystemInternal();
+			const FNetRefHandleManager& LocalNetRefHandleManager = ReplicationSystemInternal->GetNetRefHandleManager();
 			FReplicationConditionals& Conditionals = ReplicationSystemInternal->GetConditionals();
 
-			Conditionals.SetPropertyCustomCondition(LocalNetHandleManager.GetInternalIndex(Handle), Owner, RepIndex, bEnable);
+			Conditionals.SetPropertyCustomCondition(LocalNetRefHandleManager.GetInternalIndex(RefHandle), Owner, RepIndex, bEnable);
 		}
 	});
 }
@@ -150,9 +147,9 @@ void UObjectReplicationBridge::Deinitialize()
 	Super::Deinitialize();
 }
 
-UObject* UObjectReplicationBridge::GetObjectFromReferenceHandle(FNetHandle NetHandle) const
+UObject* UObjectReplicationBridge::GetObjectFromReferenceHandle(FNetRefHandle RefHandle) const
 {
-	return GetObjectReferenceCache()->GetObjectFromReferenceHandle(NetHandle);
+	return GetObjectReferenceCache()->GetObjectFromReferenceHandle(RefHandle);
 }
 
 UObject* UObjectReplicationBridge::ResolveObjectReference(const UE::Net::FNetObjectReference& Reference, const UE::Net::FNetObjectResolveContext& ResolveContext)
@@ -179,31 +176,52 @@ void UObjectReplicationBridge::AddStaticDestructionInfo(const FString& ObjectPat
 	}
 }
 
-UE::Net::FNetHandle UObjectReplicationBridge::GetReplicatedHandle(const UObject* Object) const
+UObject* UObjectReplicationBridge::GetReplicatedObject(FNetRefHandle Handle) const
+{
+	return IsReplicatedHandle(Handle) ? GetObjectFromReferenceHandle(Handle) : nullptr;
+};
+
+UE::Net::FNetRefHandle UObjectReplicationBridge::GetReplicatedRefHandle(const UObject* Object) const
 {
-	FNetHandle Handle = GetObjectReferenceCache()->GetObjectReferenceHandleFromObject(Object);
+	FNetRefHandle Handle = GetObjectReferenceCache()->GetObjectReferenceHandleFromObject(Object);
 
-	return IsReplicatedHandle(Handle) ? Handle : FNetHandle();
+	return IsReplicatedHandle(Handle) ? Handle : FNetRefHandle();
 }
 
-UObject* UObjectReplicationBridge::GetReplicatedObject(FNetHandle Handle) const
+UE::Net::FNetRefHandle UObjectReplicationBridge::GetReplicatedRefHandle(FNetHandle Handle) const
 {
-	return IsReplicatedHandle(Handle) ? GetObjectFromReferenceHandle(Handle) : nullptr;
-};
+	using namespace UE::Net;
+	using namespace UE::Net::Private;
+
+	// If the object is replicated by the owning ReplicationSystem the internal handle should be valid.
+	FInternalNetRefIndex ObjectInternalIndex = NetRefHandleManager->GetInternalIndexFromNetHandle(Handle);
+	if (ObjectInternalIndex == FNetRefHandleManager::InvalidInternalIndex)
+	{
+		return FNetRefHandle();
+	}
+
+	const FNetRefHandleManager::FReplicatedObjectData& ObjectData = NetRefHandleManager->GetReplicatedObjectDataNoCheck(ObjectInternalIndex);
+	return ObjectData.RefHandle;
+}
 
-UE::Net::FNetHandle UObjectReplicationBridge::BeginReplication(UObject* Instance, const FCreateNetHandleParams& Params)
+UE::Net::FNetRefHandle UObjectReplicationBridge::BeginReplication(UObject* Instance, const FCreateNetRefHandleParams& Params)
 {
 	LLM_SCOPE_BYTAG(IrisState);
 
 	using namespace UE::Net;
 	using namespace UE::Net::Private;
 
-	FNetHandle AllocatedHandle = GetObjectReferenceCache()->CreateObjectReferenceHandle(Instance);
+	FNetRefHandle AllocatedRefHandle = ObjectReferenceCache->CreateObjectReferenceHandle(Instance);
 
 	// If we failed to assign a handle, or if the Handle already is replicating, just return the handle
-	if (!AllocatedHandle.IsValid() || IsReplicatedHandle(AllocatedHandle))
+	if (!AllocatedRefHandle.IsValid())
 	{
-		return AllocatedHandle;
+		return FNetRefHandle();
+	}
+
+	if (IsReplicatedHandle(AllocatedRefHandle))
+	{
+		return AllocatedRefHandle;
 	}
 
 	// Register fragments
@@ -214,7 +232,7 @@ UE::Net::FNetHandle UObjectReplicationBridge::BeginReplication(UObject* Instance
 	
 	FFragmentRegistrationContext FragmentRegistrationContext(GetReplicationStateDescriptorRegistry(), Traits);
 
-	// For everything derived from UObject we can call the virtual function call RegisterReplicationFragments	
+	// For everything derived from UObject we can call the virtual function RegisterReplicationFragments	
 	CallRegisterReplicationFragments(Instance, FragmentRegistrationContext, EFragmentRegistrationFlags::None);
 
 	const FReplicationFragments& RegisteredFragments = FFragmentRegistrationContextPrivateAccessor::GetReplicationFragments(FragmentRegistrationContext);
@@ -240,7 +258,7 @@ UE::Net::FNetHandle UObjectReplicationBridge::BeginReplication(UObject* Instance
 		if (!bIsValidProtocol)
 		{
 			UE_LOG_OBJECTREPLICATIONBRIDGE(Error, TEXT("BeginReplication Found invalid protocol ProtocolId:0x%" UINT64_x_FMT " for Object named %s"), ReplicationProtocol->ProtocolIdentifier, *Instance->GetName());
-			return FNetHandle();
+			return FNetRefHandle();
 		}
 	}
 #endif
@@ -248,15 +266,18 @@ UE::Net::FNetHandle UObjectReplicationBridge::BeginReplication(UObject* Instance
 	if (ReplicationProtocol)
 	{
 		// Create NetHandle and bind instance
-		// $TODO: Rename to BeginReplication 
-		FNetHandle Handle = InternalCreateNetObject(AllocatedHandle, ReplicationProtocol);
-		if (Handle.IsValid())
+		FNetHandle NetHandle = FNetHandleManager::GetOrCreateNetHandle(Instance);
+		FNetRefHandle RefHandle = InternalCreateNetObject(AllocatedRefHandle, NetHandle, ReplicationProtocol);
+		if (RefHandle.IsValid())
 		{
 			// Attach the instance and Bind the Instance protocol to dirty tracking
 			constexpr bool bBindInstanceProtocol = true;
-			InternalAttachInstanceToNetHandle(Handle, bBindInstanceProtocol, InstanceProtocol.Get(), Instance);
-			
-			const FInternalNetHandle InternalReplicationIndex = NetHandleManager->GetInternalIndex(Handle);
+			InternalAttachInstanceToNetHandle(RefHandle, bBindInstanceProtocol, InstanceProtocol.Get(), Instance);
+#if WITH_PUSH_MODEL
+			SetNetPushIdOnInstance(InstanceProtocol.Get(), NetHandle);
+#endif
+
+			const FInternalNetRefIndex InternalReplicationIndex = NetRefHandleManager->GetInternalIndex(RefHandle);
 
 			// Initialize conditionals
 			FReplicationSystemInternal* ReplicationSystemInternal = GetReplicationSystem()->GetReplicationSystemInternal();
@@ -270,7 +291,7 @@ UE::Net::FNetHandle UObjectReplicationBridge::BeginReplication(UObject* Instance
 			FindPollInfo(Instance->GetClass(), PollFramePeriod);
 			SetPollFramePeriod(InternalReplicationIndex, PollFramePeriod);
 
-			UE_LOG_OBJECTREPLICATIONBRIDGE(Verbose, TEXT("BeginReplication Created %s with ProtocolId:0x%" UINT64_x_FMT " for Object named %s"), *Handle.ToString(), ReplicationProtocol->ProtocolIdentifier, *Instance->GetName());
+			UE_LOG_OBJECTREPLICATIONBRIDGE(Verbose, TEXT("BeginReplication Created %s with ProtocolId:0x%" UINT64_x_FMT " for Object named %s"), *RefHandle.ToString(), ReplicationProtocol->ProtocolIdentifier, *Instance->GetName());
 
 			{
 				FWorldLocations& WorldLocations = ReplicationSystem->GetReplicationSystemInternal()->GetWorldLocations();
@@ -280,18 +301,18 @@ UE::Net::FNetHandle UObjectReplicationBridge::BeginReplication(UObject* Instance
 			// Set prioritizer
 			if (Params.StaticPriority > 0.0f)
 			{
-				ReplicationSystem->SetStaticPriority(Handle, Params.StaticPriority);
+				ReplicationSystem->SetStaticPriority(RefHandle, Params.StaticPriority);
 			}
 			else
 			{
 				const FNetObjectPrioritizerHandle PrioritizerHandle = GetPrioritizer(Instance->GetClass());
 				if (PrioritizerHandle != InvalidNetObjectPrioritizerHandle)
 				{
-					ReplicationSystem->SetPrioritizer(Handle, PrioritizerHandle);
+					ReplicationSystem->SetPrioritizer(RefHandle, PrioritizerHandle);
 				}
 				else if (Params.bNeedsWorldLocationUpdate || HasRepTag(ReplicationProtocol, RepTag_WorldLocation))
 				{
-					ReplicationSystem->SetPrioritizer(Handle, DefaultSpatialNetObjectPrioritizerHandle);
+					ReplicationSystem->SetPrioritizer(RefHandle, DefaultSpatialNetObjectPrioritizerHandle);
 				}
 			}
 
@@ -300,31 +321,30 @@ UE::Net::FNetHandle UObjectReplicationBridge::BeginReplication(UObject* Instance
 				const FNetObjectFilterHandle FilterHandle = GetDynamicFilter(Instance->GetClass());
 				if (FilterHandle != InvalidNetObjectFilterHandle)
 				{
-					ReplicationSystem->SetFilter(Handle, FilterHandle);
+					ReplicationSystem->SetFilter(RefHandle, FilterHandle);
 				}
 			}
 
 			if (ShouldClassBeDeltaCompressed(Instance->GetClass()))
 			{
-				ReplicationSystem->SetDeltaCompressionStatus(Handle, ENetObjectDeltaCompressionStatus::Allow);
+				ReplicationSystem->SetDeltaCompressionStatus(RefHandle, ENetObjectDeltaCompressionStatus::Allow);
 			}
 
 			// Release instance protocol from the uniquePtr as it is now successfully bound to the handle
 			InstanceProtocol.Release();
 
-			return Handle;
+			return RefHandle;
 		}
-		UE_LOG_OBJECTREPLICATIONBRIDGE(Warning, TEXT("BeginReplication Failed to create Handle with ProtocolId:0x%" UINT64_x_FMT " for Object named %s"), (ReplicationProtocol != nullptr ? ReplicationProtocol->ProtocolIdentifier : FReplicationProtocolIdentifier(0)), *Instance->GetName());
-
+		UE_LOG_OBJECTREPLICATIONBRIDGE(Warning, TEXT("BeginReplication Failed to create NetRefHandle with ProtocolId:0x%" UINT64_x_FMT " for Object named %s"), (ReplicationProtocol != nullptr ? ReplicationProtocol->ProtocolIdentifier : FReplicationProtocolIdentifier(0)), *Instance->GetName());
 	}
 
-	// If we get here, it means that we failed to assign an internal handle for the object, we probably have ran out of handles which currently is a fatal error.
-	UE_LOG(LogIris, Error, TEXT("UObjectReplicationBridge::CreateNetHandle - Failed to create NetHandle for object %s"), ToCStr(Instance->GetPathName()));
+	// If we get here, it means that we failed to assign an internal handle for the object. We've probably run out of handles which currently is a fatal error.
+	UE_LOG(LogIris, Error, TEXT("UObjectReplicationBridge::BeginReplication - Failed to create NetRefHandle for object %s"), ToCStr(Instance->GetPathName()));
 
-	return FNetHandle();
+	return FNetRefHandle();
 }
 
-UE::Net::FNetHandle UObjectReplicationBridge::BeginReplication(FNetHandle OwnerHandle, UObject* Instance, FNetHandle InsertRelativeToSubObjectHandle, ESubObjectInsertionOrder InsertionOrder, const FCreateNetHandleParams& Params)
+UE::Net::FNetRefHandle UObjectReplicationBridge::BeginReplication(FNetRefHandle OwnerRefHandle, UObject* Instance, FNetRefHandle InsertRelativeToSubObjectRefHandle, ESubObjectInsertionOrder InsertionOrder, const FCreateNetRefHandleParams& Params)
 {
 	LLM_SCOPE_BYTAG(IrisState);
 
@@ -332,48 +352,48 @@ UE::Net::FNetHandle UObjectReplicationBridge::BeginReplication(FNetHandle OwnerH
 	using namespace UE::Net::Private;
 
 	const FReplicationSystemInternal* ReplicationSystemInternal = GetReplicationSystem()->GetReplicationSystemInternal();
-	const FNetHandleManager& LocalNetHandleManager = ReplicationSystemInternal->GetNetHandleManager();
+	const FNetRefHandleManager& LocalNetRefHandleManager = ReplicationSystemInternal->GetNetRefHandleManager();
 
 	// Owner must be replicated
-	check(IsReplicatedHandle(OwnerHandle));
-	// verify assumptions
+	check(IsReplicatedHandle(OwnerRefHandle));
+	// Verify assumptions
 	check(!Instance->HasAnyFlags(RF_ArchetypeObject | RF_ClassDefaultObject));
 
-	FNetHandle SubObjectHandle = GetReplicatedHandle(Instance);
-	if (SubObjectHandle.IsValid())
+	FNetRefHandle SubObjectRefHandle = GetReplicatedRefHandle(Instance);
+	if (SubObjectRefHandle.IsValid())
 	{
 		// Verify that the existing object is a subobject of the owner
-		check(OwnerHandle == LocalNetHandleManager.GetSubObjectOwner(SubObjectHandle));
-		return SubObjectHandle;
+		check(OwnerRefHandle == LocalNetRefHandleManager.GetSubObjectOwner(SubObjectRefHandle));
+		return SubObjectRefHandle;
 	}
 	else
 	{
-		FCreateNetHandleParams SubObjectCreateParams = Params;
+		FCreateNetRefHandleParams SubObjectCreateParams = Params;
 		// The filtering system ignores subobjects so let's not waste cycles figuring out which filter to use.
 		SubObjectCreateParams.bAllowDynamicFilter = 0U;
-		SubObjectHandle = BeginReplication(Instance, SubObjectCreateParams);
+		SubObjectRefHandle = BeginReplication(Instance, SubObjectCreateParams);
 	}
 
-	if (SubObjectHandle.IsValid())
+	if (SubObjectRefHandle.IsValid())
 	{
 		// Add subobject
-		UE_LOG_OBJECTREPLICATIONBRIDGE(Verbose, TEXT("BeginReplication Added SubObject %s to Owner %s RelativeToSubObjectHandle %s"), *SubObjectHandle.ToString(), *OwnerHandle.ToString(), *InsertRelativeToSubObjectHandle.ToString());
+		UE_LOG_OBJECTREPLICATIONBRIDGE(Verbose, TEXT("BeginReplication Added SubObject %s to Owner %s RelativeToSubObjectHandle %s"), *SubObjectRefHandle.ToString(), *OwnerRefHandle.ToString(), *InsertRelativeToSubObjectRefHandle.ToString());
 
-		InternalAddSubObject(OwnerHandle, SubObjectHandle, InsertRelativeToSubObjectHandle, InsertionOrder);
+		InternalAddSubObject(OwnerRefHandle, SubObjectRefHandle, InsertRelativeToSubObjectRefHandle, InsertionOrder);
 
 		// SubObjects should always poll with owner
-		SetPollWithObject(OwnerHandle, SubObjectHandle);
+		SetPollWithObject(OwnerRefHandle, SubObjectRefHandle);
 
 		// Copy pending dormancy from owner
-		SetObjectWantsToBeDormant(SubObjectHandle, GetObjectWantsToBeDormant(OwnerHandle));
+		SetObjectWantsToBeDormant(SubObjectRefHandle, GetObjectWantsToBeDormant(OwnerRefHandle));
 	
-		return SubObjectHandle;
+		return SubObjectRefHandle;
 	}
 
-	return FNetHandle();
+	return FNetRefHandle();
 }
 
-void UObjectReplicationBridge::SetSubObjectNetCondition(FNetHandle SubObjectHandle, ELifetimeCondition Condition)
+void UObjectReplicationBridge::SetSubObjectNetCondition(FNetRefHandle SubObjectRefHandle, ELifetimeCondition Condition)
 {
 	using namespace UE::Net::Private;
 
@@ -381,32 +401,32 @@ void UObjectReplicationBridge::SetSubObjectNetCondition(FNetHandle SubObjectHand
 	static_assert(ELifetimeCondition::COND_Max <= 127);
 
 	FReplicationSystemInternal* ReplicationSystemInternal = GetReplicationSystem()->GetReplicationSystemInternal();
-	FNetHandleManager& LocalNetHandleManager = ReplicationSystemInternal->GetNetHandleManager();
+	FNetRefHandleManager& LocalNetRefHandleManager = ReplicationSystemInternal->GetNetRefHandleManager();
 
-	const FInternalNetHandle SubObjectInternalIndex = LocalNetHandleManager.GetInternalIndex(SubObjectHandle);
-	if (LocalNetHandleManager.SetSubObjectNetCondition(SubObjectInternalIndex, (int8)Condition))
+	const FInternalNetRefIndex SubObjectInternalIndex = LocalNetRefHandleManager.GetInternalIndex(SubObjectRefHandle);
+	if (LocalNetRefHandleManager.SetSubObjectNetCondition(SubObjectInternalIndex, (int8)Condition))
 	{
-		UE_LOG_OBJECTREPLICATIONBRIDGE(Verbose, TEXT("SetSubObjectNetCondition for SubObject %s Condition %s"), *SubObjectHandle.ToString(), *UEnum::GetValueAsString<ELifetimeCondition>(Condition));
+		UE_LOG_OBJECTREPLICATIONBRIDGE(Verbose, TEXT("SetSubObjectNetCondition for SubObject %s Condition %s"), *SubObjectRefHandle.ToString(), *UEnum::GetValueAsString<ELifetimeCondition>(Condition));
 		MarkNetObjectStateDirty(ReplicationSystem->GetId(), SubObjectInternalIndex);
 	}
 	else
 	{
-		UE_LOG_OBJECTREPLICATIONBRIDGE(Verbose, TEXT("Failed to SetSubObjectNetCondition for SubObject %s Condition %s"), *SubObjectHandle.ToString(), *UEnum::GetValueAsString<ELifetimeCondition>(Condition));
+		UE_LOG_OBJECTREPLICATIONBRIDGE(Verbose, TEXT("Failed to SetSubObjectNetCondition for SubObject %s Condition %s"), *SubObjectRefHandle.ToString(), *UEnum::GetValueAsString<ELifetimeCondition>(Condition));
 	}
 }
 
-void UObjectReplicationBridge::AddDependentObject(FNetHandle ParentHandle, FNetHandle DependentHandle, EDependentWarnFlags WarnFlags)
+void UObjectReplicationBridge::AddDependentObject(FNetRefHandle ParentHandle, FNetRefHandle DependentHandle, EDependentWarnFlags WarnFlags)
 {
 	using namespace UE::Net::Private;
 
 	FReplicationSystemInternal* ReplicationSystemInternal = GetReplicationSystem()->GetReplicationSystemInternal();
-	FNetHandleManager& LocalNetHandleManager = ReplicationSystemInternal->GetNetHandleManager();
+	FNetRefHandleManager& LocalNetRefHandleManager = ReplicationSystemInternal->GetNetRefHandleManager();
 
 	// New logic to treat Dependent objects 
-	if (LocalNetHandleManager.AddDependentObject(ParentHandle, DependentHandle))
+	if (LocalNetRefHandleManager.AddDependentObject(ParentHandle, DependentHandle))
 	{
 		FReplicationFiltering& Filtering = ReplicationSystemInternal->GetFiltering();
-		const FInternalNetHandle DependentInternalIndex = LocalNetHandleManager.GetInternalIndex(DependentHandle);
+		const FInternalNetRefIndex DependentInternalIndex = LocalNetRefHandleManager.GetInternalIndex(DependentHandle);
 		Filtering.NotifyAddedDependentObject(DependentInternalIndex);
 
 		UE_LOG_OBJECTREPLICATIONBRIDGE(Verbose, TEXT("AddDependentObject Added dependent object %s to parent %s"), *DependentHandle.ToString(), *ParentHandle.ToString());
@@ -417,7 +437,7 @@ void UObjectReplicationBridge::AddDependentObject(FNetHandle ParentHandle, FNetH
 	}
 }
 
-void UObjectReplicationBridge::RemoveDependentObject(FNetHandle ParentHandle, FNetHandle DependentHandle)
+void UObjectReplicationBridge::RemoveDependentObject(FNetRefHandle ParentHandle, FNetRefHandle DependentHandle)
 {
 	using namespace UE::Net::Private;
 
@@ -426,18 +446,18 @@ void UObjectReplicationBridge::RemoveDependentObject(FNetHandle ParentHandle, FN
 	FReplicationSystemInternal* ReplicationSystemInternal = GetReplicationSystem()->GetReplicationSystemInternal();
 
 	// Remove dependent object
-	FNetHandleManager& LocalNetHandleManager = ReplicationSystemInternal->GetNetHandleManager();
-	LocalNetHandleManager.RemoveDependentObject(ParentHandle, DependentHandle);
+	FNetRefHandleManager& LocalNetRefHandleManager = ReplicationSystemInternal->GetNetRefHandleManager();
+	LocalNetRefHandleManager.RemoveDependentObject(ParentHandle, DependentHandle);
 
-	const FInternalNetHandle DependentInternalIndex = LocalNetHandleManager.GetInternalIndex(DependentHandle);
-	if (DependentInternalIndex != FNetHandleManager::InvalidInternalIndex)
+	const FInternalNetRefIndex DependentInternalIndex = LocalNetRefHandleManager.GetInternalIndex(DependentHandle);
+	if (DependentInternalIndex != FNetRefHandleManager::InvalidInternalIndex)
 	{
 		FReplicationFiltering& Filtering = ReplicationSystemInternal->GetFiltering();
 		Filtering.NotifyRemovedDependentObject(DependentInternalIndex);
 	}
 }
 
-bool UObjectReplicationBridge::WriteNetHandleCreationInfo(FReplicationBridgeSerializationContext& Context, FNetHandle Handle)
+bool UObjectReplicationBridge::WriteNetRefHandleCreationInfo(FReplicationBridgeSerializationContext& Context, FNetRefHandle Handle)
 {
 	// Write ProtocolId
 	const UE::Net::FReplicationProtocol* Protocol = GetReplicationSystem()->GetReplicationProtocol(Handle);
@@ -449,22 +469,24 @@ bool UObjectReplicationBridge::WriteNetHandleCreationInfo(FReplicationBridgeSeri
 
 void UObjectReplicationBridge::EndReplication(UObject* Instance, EEndReplicationFlags EndReplicationFlags, FEndReplicationParameters* Parameters)
 {
-	const FNetHandle Handle = GetReplicatedHandle(Instance);
+	const FNetRefHandle Handle = GetReplicatedRefHandle(Instance);
 	UReplicationBridge::EndReplication(Handle, EndReplicationFlags, Parameters);
 }
 
-void UObjectReplicationBridge::DetachInstanceFromRemote(FNetHandle Handle, bool bTearOff, bool bShouldDestroyInstance)
+void UObjectReplicationBridge::DetachInstanceFromRemote(FNetRefHandle Handle, bool bTearOff, bool bShouldDestroyInstance)
 {
-	UE_LOG_OBJECTREPLICATIONBRIDGE(Verbose, TEXT("OnDetachInstanceFromRemote %s TearOff: %u BShouldDestroy: %u"), *Handle.ToString(), (uint32)bTearOff, (uint32)bShouldDestroyInstance);
+	UE_LOG_OBJECTREPLICATIONBRIDGE(Verbose, TEXT("OnDetachInstanceFromRemote %s TearOff: %u bShouldDestroyInstance: %u"), *Handle.ToString(), (uint32)bTearOff, (uint32)bShouldDestroyInstance);
 	UnregisterRemoteInstance(Handle, bTearOff, bShouldDestroyInstance);
 }
 
-void UObjectReplicationBridge::DetachInstance(FNetHandle Handle)
+void UObjectReplicationBridge::DetachInstance(FNetRefHandle Handle)
 {
-	UnregisterRemoteInstance(Handle, false, false);
+	constexpr bool bTearOff = false;
+	constexpr bool bShouldDestroyInstance = false;
+	UnregisterRemoteInstance(Handle, bTearOff, bShouldDestroyInstance);
 }
 
-void UObjectReplicationBridge::RegisterRemoteInstance(FNetHandle Handle, UObject* Instance, const UE::Net::FReplicationProtocol* Protocol, UE::Net::FReplicationInstanceProtocol* InstanceProtocol, const FCreationHeader* Header, uint32 ConnectionId)
+void UObjectReplicationBridge::RegisterRemoteInstance(FNetRefHandle Handle, UObject* Instance, const UE::Net::FReplicationProtocol* Protocol, UE::Net::FReplicationInstanceProtocol* InstanceProtocol, const FCreationHeader* Header, uint32 ConnectionId)
 {
 	// Attach the instance protocol and instance to the Handle
 	constexpr bool bBindInstanceProtocol = false;
@@ -479,7 +501,7 @@ void UObjectReplicationBridge::RegisterRemoteInstance(FNetHandle Handle, UObject
 	UE_LOG_OBJECTREPLICATIONBRIDGE(Verbose, TEXT("RegisterRemoteInstance %s %s with ProtocolId:0x%" UINT64_x_FMT), *Handle.ToString(), ToCStr(Instance->GetName()), Protocol->ProtocolIdentifier);
 }
 
-void UObjectReplicationBridge::UnregisterRemoteInstance(FNetHandle Handle, bool bTearOff, bool bShouldDestroyInstance)
+void UObjectReplicationBridge::UnregisterRemoteInstance(FNetRefHandle Handle, bool bTearOff, bool bShouldDestroyInstance)
 {
 	// Lookup the instance and remove it	
 	UObject* Instance = GetObjectFromReferenceHandle(Handle);
@@ -499,7 +521,7 @@ void UObjectReplicationBridge::UnregisterRemoteInstance(FNetHandle Handle, bool
 	// $TODO: Cleanup any pending creation data if we have not yet instantiated the instance
 }
 	
-UE::Net::FNetHandle UObjectReplicationBridge::CreateNetHandleFromRemote(FNetHandle SubObjectOwnerNetHandle, FNetHandle WantedNetHandle, FReplicationBridgeSerializationContext& Context)
+UE::Net::FNetRefHandle UObjectReplicationBridge::CreateNetRefHandleFromRemote(FNetRefHandle SubObjectOwnerNetHandle, FNetRefHandle WantedNetHandle, FReplicationBridgeSerializationContext& Context)
 {
 	LLM_SCOPE_BYTAG(IrisState);
 
@@ -514,7 +536,7 @@ UE::Net::FNetHandle UObjectReplicationBridge::CreateNetHandleFromRemote(FNetHand
 	TUniquePtr<FCreationHeader> Header(ReadCreationHeader(Context.SerializationContext));
 	if (Context.SerializationContext.HasErrorOrOverflow())
 	{
-		return FNetHandle();
+		return FNetRefHandle();
 	}
 
 	// Currently remote objects can only receive replicated data
@@ -532,7 +554,7 @@ UE::Net::FNetHandle UObjectReplicationBridge::CreateNetHandleFromRemote(FNetHand
 	UObject* InstancePtr = BeginInstantiateFromRemote(SubObjectOwnerNetHandle, Context.SerializationContext.GetInternalContext()->ResolveContext, Header.Get());
 	if (!ensureAlwaysMsgf(InstancePtr, TEXT("Failed to instantiate Handle: %s"), *WantedNetHandle.ToString()))
 	{
-		return FNetHandle();
+		return FNetRefHandle();
 	}
 
 	// Register all fragments
@@ -557,14 +579,14 @@ UE::Net::FNetHandle UObjectReplicationBridge::CreateNetHandleFromRemote(FNetHand
 	{
 		if (!ensureAlways(!bVerifyExistingProtocol || ProtocolManager->ValidateReplicationProtocol(ReplicationProtocol, RegisteredFragments)))
 		{
-			return FNetHandle();
+			return FNetRefHandle();
 		}
 	}
 
 	if (ReplicationProtocol)
 	{
 		// Create NetHandle
-		FNetHandle Handle = InternalCreateNetObjectFromRemote(WantedNetHandle, ReplicationProtocol);
+		FNetRefHandle Handle = InternalCreateNetObjectFromRemote(WantedNetHandle, ReplicationProtocol);
 		if (Handle.IsValid())
 		{
 			RegisterRemoteInstance(Handle, InstancePtr, ReplicationProtocol, InstanceProtocol.Get(), Header.Get(), Context.ConnectionId);
@@ -579,15 +601,15 @@ UE::Net::FNetHandle UObjectReplicationBridge::CreateNetHandleFromRemote(FNetHand
 		}
 	}
 
-	return FNetHandle();
+	return FNetRefHandle();
 }
 
-void UObjectReplicationBridge::PostApplyInitialState(FNetHandle Handle)
+void UObjectReplicationBridge::PostApplyInitialState(FNetRefHandle Handle)
 {
 	EndInstantiateFromRemote(Handle);
 }
 
-void UObjectReplicationBridge::PreSendUpdateSingleHandle(FNetHandle Handle)
+void UObjectReplicationBridge::PreSendUpdateSingleHandle(FNetRefHandle Handle)
 {
 	PreUpdateAndPollImpl(Handle);
 }
@@ -597,7 +619,7 @@ void UObjectReplicationBridge::PreSendUpdate()
 	IRIS_PROFILER_SCOPE(UObjectReplicationBridge_OnPreSendUpdate);
 
 	// Invalid Handle means update all objects
-	FNetHandle Handle;
+	FNetRefHandle Handle;
 	PreUpdateAndPollImpl(Handle);
 }
 
@@ -612,36 +634,36 @@ void UObjectReplicationBridge::PruneStaleObjects()
 	// Mark all objects with object references as potentially affected by GC
 	GarbageCollectionAffectedObjects = ObjectsWithObjectReferences;
 
-	FNetHandleManager& LocalNetHandleManager = GetReplicationSystem()->GetReplicationSystemInternal()->GetNetHandleManager();
-	const TArray<UObject*>& ReplicatedInstances = LocalNetHandleManager.GetReplicatedInstances();
+	FNetRefHandleManager& LocalNetRefHandleManager = GetReplicationSystem()->GetReplicationSystemInternal()->GetNetRefHandleManager();
+	const TArray<UObject*>& ReplicatedInstances = LocalNetRefHandleManager.GetReplicatedInstances();
 
-	TArray<FNetHandle> StaleObjects;
+	TArray<FNetRefHandle> StaleObjects;
 
 	// Detect stale references and try to kill/report them
-	auto DetectStaleObjectsFunc = [&LocalNetHandleManager, &StaleObjects, &ReplicatedInstances](uint32 InternalNetHandleIndex)
+	auto DetectStaleObjectsFunc = [&LocalNetRefHandleManager, &StaleObjects, &ReplicatedInstances](uint32 InternalNetHandleIndex)
 	{
 		if (ReplicatedInstances[InternalNetHandleIndex] == nullptr)
 		{
-			const FNetHandleManager::FReplicatedObjectData& ObjectData = LocalNetHandleManager.GetReplicatedObjectDataNoCheck(InternalNetHandleIndex);
+			const FNetRefHandleManager::FReplicatedObjectData& ObjectData = LocalNetRefHandleManager.GetReplicatedObjectDataNoCheck(InternalNetHandleIndex);
 			if (ObjectData.InstanceProtocol)
 			{
 				const FReplicationProtocol* Protocol = ObjectData.Protocol;
 				const FNetDebugName* DebugName = Protocol ? Protocol->DebugName : nullptr;
-				UE_LOG(LogIrisBridge, Warning, TEXT("UObjectReplicationBridge::PruneStaleObjects ObjectInstance replicated as: %s of Type named:%s has been destroyed without notifying the ReplicationSystem %u"), *ObjectData.Handle.ToString(), ToCStr(DebugName), ObjectData.Handle.GetReplicationSystemId());
+				UE_LOG(LogIrisBridge, Warning, TEXT("UObjectReplicationBridge::PruneStaleObjects ObjectInstance replicated as: %s of Type named:%s has been destroyed without notifying the ReplicationSystem %u"), *ObjectData.RefHandle.ToString(), ToCStr(DebugName), ObjectData.RefHandle.GetReplicationSystemId());
 
 				// If the instance protocol is bound, then this is an error and we cannot safely cleanup as unbinding abound instance protocol will modify bound states
-				checkf(!EnumHasAnyFlags(ObjectData.InstanceProtocol->InstanceTraits, EReplicationInstanceProtocolTraits::IsBound), TEXT("UObjectReplicationBridge::PruneStaleObjects Bound ObjectInstance replicated as: %s has been destroyed without notifying the ReplicationSystem."), *ObjectData.Handle.ToString(), ToCStr(DebugName), ObjectData.Handle.GetReplicationSystemId());
+				checkf(!EnumHasAnyFlags(ObjectData.InstanceProtocol->InstanceTraits, EReplicationInstanceProtocolTraits::IsBound), TEXT("UObjectReplicationBridge::PruneStaleObjects Bound ObjectInstance replicated as: %s has been destroyed without notifying the ReplicationSystem."), *ObjectData.RefHandle.ToString(), ToCStr(DebugName), ObjectData.RefHandle.GetReplicationSystemId());
 
-				StaleObjects.Push(ObjectData.Handle);
+				StaleObjects.Push(ObjectData.RefHandle);
 			}
 		}
 	};
 
 	// Iterate over assigned indices and detect if any of the replicated instances has been garbagecollected (excluding DestroyedStartupObjectInternalIndices) as they never have an instance
-	FNetBitArray::ForAllSetBits(LocalNetHandleManager.GetAssignedInternalIndices(), LocalNetHandleManager.GetDestroyedStartupObjectInternalIndices(), FNetBitArray::AndNotOp, DetectStaleObjectsFunc);
+	FNetBitArray::ForAllSetBits(LocalNetRefHandleManager.GetAssignedInternalIndices(), LocalNetRefHandleManager.GetDestroyedStartupObjectInternalIndices(), FNetBitArray::AndNotOp, DetectStaleObjectsFunc);
 
 	// EndReplication/detach stale instances
-	for (FNetHandle Handle : MakeArrayView(StaleObjects.GetData(), StaleObjects.Num()))
+	for (FNetRefHandle Handle : MakeArrayView(StaleObjects.GetData(), StaleObjects.Num()))
 	{
 		UReplicationBridge::EndReplication(Handle);
 	}
@@ -660,14 +682,14 @@ void UObjectReplicationBridge::SetInstanceGetWorldLocationFunction(FInstanceGetW
 	GetInstanceWorldLocationFunction = InGetWorldLocationFunction;
 }
 
-void UObjectReplicationBridge::PreUpdateAndPollImpl(FNetHandle Handle)
+void UObjectReplicationBridge::PreUpdateAndPollImpl(FNetRefHandle Handle)
 {
 	using namespace UE::Net;
 	using namespace UE::Net::Private;
 
 	FReplicationSystemInternal* ReplicationSystemInternal = GetReplicationSystem()->GetReplicationSystemInternal();
-	const FNetHandleManager& LocalNetHandleManager = ReplicationSystemInternal->GetNetHandleManager();
-	const TArray<UObject*>& ReplicatedInstances = LocalNetHandleManager.GetReplicatedInstances();
+	const FNetRefHandleManager& LocalNetRefHandleManager = ReplicationSystemInternal->GetNetRefHandleManager();
+	const TArray<UObject*>& ReplicatedInstances = LocalNetRefHandleManager.GetReplicatedInstances();
 	const bool bIsUsingPushModel = IsIrisPushModelEnabled();
 	const FNetBitArrayView DirtyObjects = ReplicationSystemInternal->GetDirtyNetObjectTracker().GetDirtyNetObjects();
 
@@ -680,9 +702,9 @@ void UObjectReplicationBridge::PreUpdateAndPollImpl(FNetHandle Handle)
 
 	FPreUpdateAndPollStats Stats;
 
-	auto UpdateAndPollFunction = [&ReplicatedInstances, this, &LocalNetHandleManager, &Stats](uint32 InternalObjectIndex)
+	auto UpdateAndPollFunction = [&ReplicatedInstances, this, &LocalNetRefHandleManager, &Stats](uint32 InternalObjectIndex)
 	{
-		const FNetHandleManager::FReplicatedObjectData& ObjectData = LocalNetHandleManager.GetReplicatedObjectDataNoCheck(InternalObjectIndex);
+		const FNetRefHandleManager::FReplicatedObjectData& ObjectData = LocalNetRefHandleManager.GetReplicatedObjectDataNoCheck(InternalObjectIndex);
 		if (ObjectData.InstanceProtocol)
 		{
 			// Call per-instance PreUpdate function
@@ -692,7 +714,7 @@ void UObjectReplicationBridge::PreUpdateAndPollImpl(FNetHandle Handle)
 				IRIS_PROFILER_SCOPE(PreReplicationUpdate);
 				IRIS_PROFILER_SCOPE_TEXT(ObjectData.Protocol->DebugName->Name);
 #endif
-				(*PreUpdateInstanceFunction)(ObjectData.Handle, ReplicatedInstances[InternalObjectIndex], this);
+				(*PreUpdateInstanceFunction)(ObjectData.RefHandle, ReplicatedInstances[InternalObjectIndex], this);
 				++Stats.PreUpdatedObjectCount;
 			}
 
@@ -717,11 +739,11 @@ void UObjectReplicationBridge::PreUpdateAndPollImpl(FNetHandle Handle)
 		}
 	};
 
-	const FNetBitArray& PrevScopableObjects = LocalNetHandleManager.GetPrevFrameScopableInternalIndices();
+	const FNetBitArray& PrevScopableObjects = LocalNetRefHandleManager.GetPrevFrameScopableInternalIndices();
 
-	auto PushModelUpdateAndPollFunction = [&ReplicatedInstances, this, &LocalNetHandleManager, &DirtyObjects, &PrevScopableObjects, &Stats](uint32 InternalObjectIndex)
+	auto PushModelUpdateAndPollFunction = [&ReplicatedInstances, this, &LocalNetRefHandleManager, &DirtyObjects, &PrevScopableObjects, &Stats](uint32 InternalObjectIndex)
 	{
-		const FNetHandleManager::FReplicatedObjectData& ObjectData = LocalNetHandleManager.GetReplicatedObjectDataNoCheck(InternalObjectIndex);
+		const FNetRefHandleManager::FReplicatedObjectData& ObjectData = LocalNetRefHandleManager.GetReplicatedObjectDataNoCheck(InternalObjectIndex);
 		if (const FReplicationInstanceProtocol* InstanceProtocol = ObjectData.InstanceProtocol)
 		{
 			const EReplicationInstanceProtocolTraits InstanceTraits = InstanceProtocol->InstanceTraits;
@@ -732,7 +754,7 @@ void UObjectReplicationBridge::PreUpdateAndPollImpl(FNetHandle Handle)
 				IRIS_PROFILER_SCOPE(PreReplicationUpdate);
 				IRIS_PROFILER_SCOPE_TEXT(ObjectData.Protocol->DebugName->Name);
 #endif
-				(*PreUpdateInstanceFunction)(ObjectData.Handle, ReplicatedInstances[InternalObjectIndex], this);
+				(*PreUpdateInstanceFunction)(ObjectData.RefHandle, ReplicatedInstances[InternalObjectIndex], this);
 				++Stats.PreUpdatedObjectCount;
 			}
 
@@ -800,12 +822,12 @@ void UObjectReplicationBridge::PreUpdateAndPollImpl(FNetHandle Handle)
 	// Update all objects
 	if (!Handle.IsValid())
 	{
-		const FNetBitArray& ScopableObjects = LocalNetHandleManager.GetScopableInternalIndices();
-		const FNetBitArray& WantToBeDormantObjects = LocalNetHandleManager.GetWantToBeDormantInternalIndices();
+		const FNetBitArray& ScopableObjects = LocalNetRefHandleManager.GetScopableInternalIndices();
+		const FNetBitArray& WantToBeDormantObjects = LocalNetRefHandleManager.GetWantToBeDormantInternalIndices();
 
 		// Filter the set of objects considered for pre-update and polling
 		// We always want to consider objects marked as dirty and their subobjects
-		FNetBitArray ObjectsConsideredForPolling(LocalNetHandleManager.GetMaxActiveObjectCount());
+		FNetBitArray ObjectsConsideredForPolling(LocalNetRefHandleManager.GetMaxActiveObjectCount());
 		if (bUseFrequencyBasedPolling)
 		{
 			FNetBitArrayView FrequencyBasedObjectsToPollView = MakeNetBitArrayView(ObjectsConsideredForPolling);
@@ -825,9 +847,9 @@ void UObjectReplicationBridge::PreUpdateAndPollImpl(FNetHandle Handle)
 			ObjectsConsideredForPolling.Combine(WantToBeDormantObjects, FNetBitArrayView::AndNotOp);
 
 			// Add objects that have requested to flush dormancy
-			for (FNetHandle HandlePendingFlush : MakeArrayView(DormantHandlesPendingFlush))
+			for (FNetRefHandle HandlePendingFlush : MakeArrayView(DormantHandlesPendingFlush))
 			{
-				if (const uint32 InternalObjectIndex = LocalNetHandleManager.GetInternalIndex(HandlePendingFlush))
+				if (const uint32 InternalObjectIndex = LocalNetRefHandleManager.GetInternalIndex(HandlePendingFlush))
 				{
 					if (ObjectsConsideredForPolling.GetBit(InternalObjectIndex))
 					{
@@ -835,7 +857,7 @@ void UObjectReplicationBridge::PreUpdateAndPollImpl(FNetHandle Handle)
 					}
 
 					ObjectsConsideredForPolling.SetBit(InternalObjectIndex);
-					for (const FInternalNetHandle SubObjectIndex : LocalNetHandleManager.GetSubObjects(InternalObjectIndex))
+					for (const FInternalNetRefIndex SubObjectIndex : LocalNetRefHandleManager.GetSubObjects(InternalObjectIndex))
 					{
 						ObjectsConsideredForPolling.SetBit(SubObjectIndex);
 					}
@@ -852,27 +874,27 @@ void UObjectReplicationBridge::PreUpdateAndPollImpl(FNetHandle Handle)
 		{
 			IRIS_PROFILER_SCOPE(UObjectReplicationBridge_PreUpdateAndPollImpl_PropagateDirtyness);
 
-			auto PropagateSubObjectDirtinessToOwner = [&LocalNetHandleManager, &ObjectsConsideredForPolling](uint32 InternalObjectIndex)
+			auto PropagateSubObjectDirtinessToOwner = [&LocalNetRefHandleManager, &ObjectsConsideredForPolling](uint32 InternalObjectIndex)
 			{
-				const FNetHandleManager::FReplicatedObjectData& ObjectData = LocalNetHandleManager.GetReplicatedObjectDataNoCheck(InternalObjectIndex);
+				const FNetRefHandleManager::FReplicatedObjectData& ObjectData = LocalNetRefHandleManager.GetReplicatedObjectDataNoCheck(InternalObjectIndex);
 				ObjectsConsideredForPolling.SetBit(ObjectData.SubObjectRootIndex);
 			};
 
-			auto PropagateOwnerDirtinessToSubObjectsAndDependentObjects = [&LocalNetHandleManager, &ObjectsConsideredForPolling](uint32 InternalObjectIndex)
+			auto PropagateOwnerDirtinessToSubObjectsAndDependentObjects = [&LocalNetRefHandleManager, &ObjectsConsideredForPolling](uint32 InternalObjectIndex)
 			{
-				const FNetHandleManager::FReplicatedObjectData& ObjectData = LocalNetHandleManager.GetReplicatedObjectDataNoCheck(InternalObjectIndex);
+				const FNetRefHandleManager::FReplicatedObjectData& ObjectData = LocalNetRefHandleManager.GetReplicatedObjectDataNoCheck(InternalObjectIndex);
 
 				// Mark owner as well as we might have masked it out in the dormancy pass
 				ObjectsConsideredForPolling.SetBit(InternalObjectIndex);
 
-				for (const FInternalNetHandle SubObjectInternalIndex : LocalNetHandleManager.GetSubObjects(InternalObjectIndex))
+				for (const FInternalNetRefIndex SubObjectInternalIndex : LocalNetRefHandleManager.GetSubObjects(InternalObjectIndex))
 				{
 					ObjectsConsideredForPolling.SetBit(SubObjectInternalIndex);
 				}
 			};
 
 			// Update subobjects' owner first and owners' subobjects second. It's the only way to properly mark all groups of objects in two passes.
-			const FNetBitArrayView SubObjects = MakeNetBitArrayView(LocalNetHandleManager.GetSubObjectInternalIndices());
+			const FNetBitArrayView SubObjects = MakeNetBitArrayView(LocalNetRefHandleManager.GetSubObjectInternalIndices());
 			FNetBitArrayView::ForAllSetBits(DirtyObjects, SubObjects, FNetBitArray::AndOp, PropagateSubObjectDirtinessToOwner);
 			FNetBitArrayView::ForAllSetBits(DirtyObjects, SubObjects, FNetBitArray::AndNotOp, PropagateOwnerDirtinessToSubObjectsAndDependentObjects);
 			
@@ -881,10 +903,10 @@ void UObjectReplicationBridge::PreUpdateAndPollImpl(FNetHandle Handle)
 				IRIS_PROFILER_SCOPE(UObjectReplicationBridge_PreUpdateAndPollImpl Patch in depedent objects);
 
 				FNetBitArray TempObjectsConsideredForPolling(ObjectsConsideredForPolling);
-				FNetBitArray::ForAllSetBits(TempObjectsConsideredForPolling, NetHandleManager->GetObjectsWithDependentObjectsInternalIndices(), FNetBitArray::AndOp,
-					[&LocalNetHandleManager, &ObjectsConsideredForPolling](FInternalNetHandle ObjectIndex) 
+				FNetBitArray::ForAllSetBits(TempObjectsConsideredForPolling, NetRefHandleManager->GetObjectsWithDependentObjectsInternalIndices(), FNetBitArray::AndOp,
+					[&LocalNetRefHandleManager, &ObjectsConsideredForPolling](FInternalNetRefIndex ObjectIndex) 
 					{
-						LocalNetHandleManager.ForAllDependentObjectsRecursive(ObjectIndex, [&ObjectsConsideredForPolling, ObjectIndex](FInternalNetHandle DependentObjectIndex) { ObjectsConsideredForPolling.SetBit(DependentObjectIndex); });
+						LocalNetRefHandleManager.ForAllDependentObjectsRecursive(ObjectIndex, [&ObjectsConsideredForPolling, ObjectIndex](FInternalNetRefIndex DependentObjectIndex) { ObjectsConsideredForPolling.SetBit(DependentObjectIndex); });
 					}
 				);
 			}
@@ -899,7 +921,7 @@ void UObjectReplicationBridge::PreUpdateAndPollImpl(FNetHandle Handle)
 			ObjectsConsideredForPolling.ForAllSetBits(UpdateAndPollFunction);
 		}
 	}
-	else if (uint32 InternalObjectIndex = LocalNetHandleManager.GetInternalIndex(Handle))
+	else if (uint32 InternalObjectIndex = LocalNetRefHandleManager.GetInternalIndex(Handle))
 	{
 		UpdateAndPollFunction(InternalObjectIndex);
 	}
@@ -921,18 +943,18 @@ void UObjectReplicationBridge::UpdateInstancesWorldLocation()
 	}
 
 	FReplicationSystemInternal* ReplicationSystemInternal = GetReplicationSystem()->GetReplicationSystemInternal();
-	const FNetHandleManager& LocalNetHandleManager = ReplicationSystemInternal->GetNetHandleManager();
+	const FNetRefHandleManager& LocalNetRefHandleManager = ReplicationSystemInternal->GetNetRefHandleManager();
 	FWorldLocations& WorldLocations = ReplicationSystemInternal->GetWorldLocations();
-	const TArray<UObject*>& ReplicatedInstances = LocalNetHandleManager.GetReplicatedInstances();
+	const TArray<UObject*>& ReplicatedInstances = LocalNetRefHandleManager.GetReplicatedInstances();
 
 
 	// Retrieve the world location for instances that supports it. Only dirty objects are considered.
-	auto UpdateInstanceWorldLocation = [this, &ReplicatedInstances, &LocalNetHandleManager, &WorldLocations](uint32 InternalObjectIndex)
+	auto UpdateInstanceWorldLocation = [this, &ReplicatedInstances, &LocalNetRefHandleManager, &WorldLocations](uint32 InternalObjectIndex)
 	{
-		const FNetHandleManager::FReplicatedObjectData& ObjectData = LocalNetHandleManager.GetReplicatedObjectDataNoCheck(InternalObjectIndex);
+		const FNetRefHandleManager::FReplicatedObjectData& ObjectData = LocalNetRefHandleManager.GetReplicatedObjectDataNoCheck(InternalObjectIndex);
 		if (ObjectData.InstanceProtocol && EnumHasAnyFlags(ObjectData.InstanceProtocol->InstanceTraits, EReplicationInstanceProtocolTraits::NeedsWorldLocationUpdate))
 		{
-			FVector WorldLocation = (*GetInstanceWorldLocationFunction)(ObjectData.Handle, ReplicatedInstances[InternalObjectIndex]);
+			FVector WorldLocation = (*GetInstanceWorldLocationFunction)(ObjectData.RefHandle, ReplicatedInstances[InternalObjectIndex]);
 			WorldLocations.SetWorldLocation(InternalObjectIndex, WorldLocation);
 		}
 	};
@@ -942,43 +964,43 @@ void UObjectReplicationBridge::UpdateInstancesWorldLocation()
 
 }
 
-void UObjectReplicationBridge::SetPollFramePeriod(UE::Net::Private::FInternalNetHandle InternalReplicationIndex, uint8 PollFramePeriod)
+void UObjectReplicationBridge::SetPollFramePeriod(UE::Net::Private::FInternalNetRefIndex InternalReplicationIndex, uint8 PollFramePeriod)
 {
 	PollFrequencyLimiter->SetPollFramePeriod(InternalReplicationIndex, PollFramePeriod);
 }
 
-void UObjectReplicationBridge::SetPollWithObject(FNetHandle ObjectToPollWithHandle, FNetHandle ObjectHandle)
+void UObjectReplicationBridge::SetPollWithObject(FNetRefHandle ObjectToPollWithHandle, FNetRefHandle ObjectHandle)
 {
-	const UE::Net::Private::FInternalNetHandle PollWithInternalReplicationIndex = NetHandleManager->GetInternalIndex(ObjectToPollWithHandle);
-	const UE::Net::Private::FInternalNetHandle InternalReplicationIndex = NetHandleManager->GetInternalIndex(ObjectHandle);
+	const UE::Net::Private::FInternalNetRefIndex PollWithInternalReplicationIndex = NetRefHandleManager->GetInternalIndex(ObjectToPollWithHandle);
+	const UE::Net::Private::FInternalNetRefIndex InternalReplicationIndex = NetRefHandleManager->GetInternalIndex(ObjectHandle);
 	PollFrequencyLimiter->SetPollWithObject(PollWithInternalReplicationIndex, InternalReplicationIndex);
 }
 
-bool UObjectReplicationBridge::GetObjectWantsToBeDormant(FNetHandle Handle) const
+bool UObjectReplicationBridge::GetObjectWantsToBeDormant(FNetRefHandle Handle) const
 {
 	using namespace UE::Net::Private;
 
 	const FReplicationSystemInternal* ReplicationSystemInternal = GetReplicationSystem()->GetReplicationSystemInternal();
-	const FNetHandleManager& LocalNetHandleManager = ReplicationSystemInternal->GetNetHandleManager();
+	const FNetRefHandleManager& LocalNetRefHandleManager = ReplicationSystemInternal->GetNetRefHandleManager();
 
-	if (const FInternalNetHandle InternalObjectIndex = LocalNetHandleManager.GetInternalIndex(Handle))
+	if (const FInternalNetRefIndex InternalObjectIndex = LocalNetRefHandleManager.GetInternalIndex(Handle))
 	{
-		return LocalNetHandleManager.GetWantToBeDormantInternalIndices().GetBit(InternalObjectIndex);
+		return LocalNetRefHandleManager.GetWantToBeDormantInternalIndices().GetBit(InternalObjectIndex);
 	}
 
 	return false;
 }
 
-void UObjectReplicationBridge::SetObjectWantsToBeDormant(FNetHandle Handle, bool bWantsToBeDormant)
+void UObjectReplicationBridge::SetObjectWantsToBeDormant(FNetRefHandle Handle, bool bWantsToBeDormant)
 {
 	using namespace UE::Net::Private;
 
 	FReplicationSystemInternal* ReplicationSystemInternal = GetReplicationSystem()->GetReplicationSystemInternal();
-	FNetHandleManager& LocalNetHandleManager = ReplicationSystemInternal->GetNetHandleManager();
+	FNetRefHandleManager& LocalNetRefHandleManager = ReplicationSystemInternal->GetNetRefHandleManager();
 
-	if (const FInternalNetHandle InternalObjectIndex = LocalNetHandleManager.GetInternalIndex(Handle))
+	if (const FInternalNetRefIndex InternalObjectIndex = LocalNetRefHandleManager.GetInternalIndex(Handle))
 	{
-		UE::Net::FNetBitArray& WantToBeDormantObjects = LocalNetHandleManager.GetWantToBeDormantInternalIndices();
+		UE::Net::FNetBitArray& WantToBeDormantObjects = LocalNetRefHandleManager.GetWantToBeDormantInternalIndices();
 
 		// Update pending dormancy status
 		WantToBeDormantObjects.SetBitValue(InternalObjectIndex, bWantsToBeDormant);
@@ -993,18 +1015,33 @@ void UObjectReplicationBridge::SetObjectWantsToBeDormant(FNetHandle Handle, bool
 
 		// Since we use this as a mask when updating objects we must include subobjects as well
 		// Subobjects added later will copy status from owner when they are added
-		for (const FInternalNetHandle SubObjectInternalIndex : LocalNetHandleManager.GetSubObjects(InternalObjectIndex))
+		for (const FInternalNetRefIndex SubObjectInternalIndex : LocalNetRefHandleManager.GetSubObjects(InternalObjectIndex))
 		{
 			WantToBeDormantObjects.SetBitValue(SubObjectInternalIndex, bWantsToBeDormant);
 		}
 	}
 }
 
-void UObjectReplicationBridge::ForceUpdateWantsToBeDormantObject(FNetHandle Handle)
+void UObjectReplicationBridge::ForceUpdateWantsToBeDormantObject(FNetRefHandle Handle)
 {
 	DormantHandlesPendingFlush.Add(Handle);
 }
 
+void UObjectReplicationBridge::SetNetPushIdOnInstance(UE::Net::FReplicationInstanceProtocol* InstanceProtocol, FNetHandle NetHandle)
+{
+#if WITH_PUSH_MODEL
+	using namespace UE::Net;
+
+	// Set push ID only if any state supports it. If no state supports it then we might crash if setting the ID.
+	if (EnumHasAnyFlags(InstanceProtocol->InstanceTraits, EReplicationInstanceProtocolTraits::HasPartialPushBasedDirtiness | EReplicationInstanceProtocolTraits::HasFullPushBasedDirtiness))
+	{
+		const Private::FNetPushObjectHandle PushHandle(NetHandle);
+		TArrayView<const FReplicationFragment* const> Fragments(InstanceProtocol->Fragments, InstanceProtocol->FragmentCount);
+		SetNetPushIdOnFragments(Fragments, PushHandle);
+	}
+#endif
+}
+
 bool UObjectReplicationBridge::FindPollInfo(const UClass* Class, uint8& OutPollPeriod)
 {
 	if (!(bAllowPollPeriodOverrides & bHasPollOverrides))
