Iris - Implemented globally unique NetHandles.

diff --git a/Engine/Source/Runtime/Experimental/Iris/Core/Private/Iris/ReplicationSystem/Prioritization/ReplicationPrioritization.cpp b/Engine/Source/Runtime/Experimental/Iris/Core/Private/Iris/ReplicationSystem/Prioritization/ReplicationPrioritization.cpp
--- a/Engine/Source/Runtime/Experimental/Iris/Core/Private/Iris/ReplicationSystem/Prioritization/ReplicationPrioritization.cpp
+++ b/Engine/Source/Runtime/Experimental/Iris/Core/Private/Iris/ReplicationSystem/Prioritization/ReplicationPrioritization.cpp
@@ -226,8 +226,8 @@ public:
 		FReplicationInstanceProtocol const** InstanceProtocols;
 	};
 
-	FUpdateDirtyObjectsBatchHelper(const FNetHandleManager* InNetHandleManager, uint32 PrioritizerCount)
-	: NetHandleManager(InNetHandleManager)
+	FUpdateDirtyObjectsBatchHelper(const FNetRefHandleManager* InNetRefHandleManager, uint32 PrioritizerCount)
+	: NetRefHandleManager(InNetRefHandleManager)
 	{
 		PerPrioritizerInfos.SetNumUninitialized(PrioritizerCount);
 		ObjectIndicesStorage.SetNumUninitialized(PrioritizerCount*MaxObjectCountPerBatch);
@@ -255,7 +255,7 @@ public:
 				continue;
 			}
 
-			if (const FReplicationInstanceProtocol* InstanceProtocol = NetHandleManager->GetReplicatedObjectDataNoCheck(ObjectIndex).InstanceProtocol)
+			if (const FReplicationInstanceProtocol* InstanceProtocol = NetRefHandleManager->GetReplicatedObjectDataNoCheck(ObjectIndex).InstanceProtocol)
 			{
 				FPerPrioritizerInfo& PerPrioritizerInfo = PerPrioritizerInfosData[PrioritizerIndex];
 				PerPrioritizerInfo.ObjectIndices[PerPrioritizerInfo.ObjectCount] = ObjectIndex;
@@ -278,7 +278,7 @@ private:
 
 	TArray<uint32> ObjectIndicesStorage;
 	TArray<const FReplicationInstanceProtocol*> InstanceProtocolsStorage;
-	const FNetHandleManager* NetHandleManager;
+	const FNetRefHandleManager* NetRefHandleManager;
 };
 
 FReplicationPrioritization::FReplicationPrioritization()
@@ -289,12 +289,12 @@ FReplicationPrioritization::FReplicationPrioritization()
 void FReplicationPrioritization::Init(FReplicationPrioritizationInitParams& Params)
 {
 	check(Params.Connections != nullptr);
-	check(Params.NetHandleManager != nullptr);
+	check(Params.NetRefHandleManager != nullptr);
 
 	ReplicationSystem = Params.ReplicationSystem;
 
 	Connections = Params.Connections;
-	NetHandleManager = Params.NetHandleManager;
+	NetRefHandleManager = Params.NetRefHandleManager;
 
 	MaxObjectCount = Params.MaxObjectCount;
 
@@ -383,10 +383,10 @@ bool FReplicationPrioritization::SetPrioritizer(uint32 ObjectIndex, FNetObjectPr
 
 	// Register object with new prioritizer
 	{
-		const FNetHandleManager::FReplicatedObjectData& ObjectData = NetHandleManager->GetReplicatedObjectDataNoCheck(ObjectIndex);
+		const FNetRefHandleManager::FReplicatedObjectData& ObjectData = NetRefHandleManager->GetReplicatedObjectDataNoCheck(ObjectIndex);
 
 		NetObjectPrioritizationInfo = FNetObjectPrioritizationInfo{};
-		FNetObjectPrioritizerAddObjectParams AddParams = {NetObjectPrioritizationInfo, ObjectData.InstanceProtocol, ObjectData.Protocol, NetHandleManager->GetReplicatedObjectStateBufferNoCheck(ObjectIndex)};
+		FNetObjectPrioritizerAddObjectParams AddParams = {NetObjectPrioritizationInfo, ObjectData.InstanceProtocol, ObjectData.Protocol, NetRefHandleManager->GetReplicatedObjectStateBufferNoCheck(ObjectIndex)};
 		FPrioritizerInfo& PrioritizerInfo = PrioritizerInfos[NewPrioritizer];
 		if (PrioritizerInfo.Prioritizer->AddObject(ObjectIndex, AddParams))
 		{
@@ -596,8 +596,8 @@ void FReplicationPrioritization::UpdatePrioritiesForNewAndDeletedObjects()
 {
 	IRIS_PROFILER_SCOPE(FReplicationPrioritization_UpdatePrioritiesForNewAndDeletedObjects);
 
-	const FNetBitArray& PrevScopedIndices = NetHandleManager->GetPrevFrameScopableInternalIndices();
-	const FNetBitArray& ScopedIndices = NetHandleManager->GetScopableInternalIndices();
+	const FNetBitArray& PrevScopedIndices = NetRefHandleManager->GetPrevFrameScopableInternalIndices();
+	const FNetBitArray& ScopedIndices = NetRefHandleManager->GetScopableInternalIndices();
 
 	auto ForEachRemovedObject = [this](uint32 ObjectIndex)
 	{
@@ -677,7 +677,7 @@ void FReplicationPrioritization::PrioritizeForConnection(uint32 ConnId, FPriorit
 	{
 		PrioParameters.Priorities = ConnInfo.Priorities.GetData();
 		PrioParameters.PrioritizationInfos = NetObjectPrioritizationInfos.GetData();
-		PrioParameters.StateBuffers = NetHandleManager->GetReplicatedObjectStateBuffers().GetData();
+		PrioParameters.StateBuffers = NetRefHandleManager->GetReplicatedObjectStateBuffers().GetData();
 		PrioParameters.ConnectionId = ConnId;
 		PrioParameters.View = Connections->GetReplicationView(ConnId);
 	}
@@ -741,7 +741,7 @@ void FReplicationPrioritization::PrioritizeForConnection(uint32 ConnId, FPriorit
  */
 void FReplicationPrioritization::NotifyPrioritizersOfDirtyObjects(const FNetBitArrayView& DirtyObjects)
 {
-	FUpdateDirtyObjectsBatchHelper BatchHelper(NetHandleManager, PrioritizerInfos.Num());
+	FUpdateDirtyObjectsBatchHelper BatchHelper(NetRefHandleManager, PrioritizerInfos.Num());
 
 	constexpr SIZE_T MaxBatchObjectCount = FUpdateDirtyObjectsBatchHelper::Constants::MaxObjectCountPerBatch;
 	uint32 ObjectIndices[MaxBatchObjectCount];
@@ -764,7 +764,7 @@ void FReplicationPrioritization::BatchNotifyPrioritizersOfDirtyObjects(FUpdateDi
 	BatchHelper.PrepareBatch(ObjectIndices, ObjectCount, ObjectIndexToPrioritizer.GetData());
 
 	FNetObjectPrioritizerUpdateParams UpdateParameters;
-	UpdateParameters.StateBuffers = NetHandleManager->GetReplicatedObjectStateBuffers().GetData();
+	UpdateParameters.StateBuffers = NetRefHandleManager->GetReplicatedObjectStateBuffers().GetData();
 	UpdateParameters.PrioritizationInfos = NetObjectPrioritizationInfos.GetData();
 
 	for (const FUpdateDirtyObjectsBatchHelper::FPerPrioritizerInfo& PerPrioritizerInfo : BatchHelper.PerPrioritizerInfos)
diff --git a/Engine/Source/Runtime/Experimental/Iris/Core/Private/Iris/ReplicationSystem/Prioritization/ReplicationPrioritization.h b/Engine/Source/Runtime/Experimental/Iris/Core/Private/Iris/ReplicationSystem/Prioritization/ReplicationPrioritization.h
--- a/Engine/Source/Runtime/Experimental/Iris/Core/Private/Iris/ReplicationSystem/Prioritization/ReplicationPrioritization.h
+++ b/Engine/Source/Runtime/Experimental/Iris/Core/Private/Iris/ReplicationSystem/Prioritization/ReplicationPrioritization.h
@@ -14,7 +14,7 @@ namespace UE::Net
 	class FNetBitArrayView;
 	namespace Private
 	{
-		class FNetHandleManager;
+		class FNetRefHandleManager;
 		class FReplicationConnections;
 	}
 }
@@ -25,7 +25,7 @@ namespace UE::Net::Private
 struct FReplicationPrioritizationInitParams
 {
 	TObjectPtr<const UReplicationSystem> ReplicationSystem;
-	const FNetHandleManager* NetHandleManager = nullptr;
+	const FNetRefHandleManager* NetRefHandleManager = nullptr;
 	FReplicationConnections* Connections = nullptr;
 	uint32 MaxObjectCount = 0;
 };
@@ -78,7 +78,7 @@ private:
 
 	TObjectPtr<const UReplicationSystem> ReplicationSystem;
 	FReplicationConnections* Connections = nullptr;
-	const FNetHandleManager* NetHandleManager = nullptr;
+	const FNetRefHandleManager* NetRefHandleManager = nullptr;
 	TStrongObjectPtr<UNetObjectPrioritizerDefinitions> PrioritizerDefinitions;
 	TArray<FNetObjectPrioritizationInfo> NetObjectPrioritizationInfos;
 	TArray<uint8> ObjectIndexToPrioritizer;
diff --git a/Engine/Source/Runtime/Experimental/Iris/Core/Private/Iris/ReplicationSystem/ReplicationBridge.cpp b/Engine/Source/Runtime/Experimental/Iris/Core/Private/Iris/ReplicationSystem/ReplicationBridge.cpp
--- a/Engine/Source/Runtime/Experimental/Iris/Core/Private/Iris/ReplicationSystem/ReplicationBridge.cpp
+++ b/Engine/Source/Runtime/Experimental/Iris/Core/Private/Iris/ReplicationSystem/ReplicationBridge.cpp
@@ -7,6 +7,7 @@
 #include "Iris/Core/IrisLog.h"
 #include "Iris/Core/IrisProfiler.h"
 #include "Net/Core/NetBitArray.h"
+#include "Net/Core/NetHandle/NetHandleManager.h"
 #include "Net/Core/Trace/NetTrace.h"
 #include "Net/Core/Trace/NetDebugName.h"
 #include "Iris/ReplicationState/ReplicationStateUtil.h"
@@ -28,35 +29,31 @@
 
 DEFINE_LOG_CATEGORY(LogIrisBridge)
 
-#if UE_NET_ALLOW_MULTIPLE_REPLICATION_SYSTEMS
 #define UE_LOG_REPLICATIONBRIDGE(Category, Format, ...)  UE_LOG(LogIrisBridge, Category, TEXT("ReplicationBridge(%u)::") Format, ReplicationSystem->GetId(), ##__VA_ARGS__)
-#else
-#define UE_LOG_REPLICATIONBRIDGE(Category, Format, ...)  UE_LOG(LogIrisBridge, Category, TEXT("ReplicationBridge::") Format, ##__VA_ARGS__)
-#endif
 
-/*
-* ReplicationBridge Implementation
-*/
+/**
+ * ReplicationBridge Implementation
+ */
 UReplicationBridge::UReplicationBridge()
 : ReplicationSystem(nullptr)
 , ReplicationProtocolManager(nullptr)
 , ReplicationStateDescriptorRegistry(nullptr)
-, NetHandleManager(nullptr)
+, NetRefHandleManager(nullptr)
 , DestructionInfoProtocol(nullptr)
 {
 }
 
-bool UReplicationBridge::WriteNetHandleCreationInfo(FReplicationBridgeSerializationContext& Context, FNetHandle Handle)
+bool UReplicationBridge::WriteNetRefHandleCreationInfo(FReplicationBridgeSerializationContext& Context, FNetRefHandle Handle)
 {
 	return true;
 }
 
-UE::Net::FNetHandle UReplicationBridge::CreateNetHandleFromRemote(FNetHandle SubObjectOwnerNetHandle, FNetHandle WantedNetHandle, FReplicationBridgeSerializationContext& Context)
+UE::Net::FNetRefHandle UReplicationBridge::CreateNetRefHandleFromRemote(FNetRefHandle SubObjectOwnerNetHandle, FNetRefHandle WantedNetHandle, FReplicationBridgeSerializationContext& Context)
 {
-	return FNetHandle();
+	return FNetRefHandle();
 };
 
-void UReplicationBridge::PreSendUpdateSingleHandle(FNetHandle Handle)
+void UReplicationBridge::PreSendUpdateSingleHandle(FNetRefHandle Handle)
 {
 }
 
@@ -68,15 +65,15 @@ void UReplicationBridge::UpdateInstancesWorldLocation()
 {
 }
 
-void UReplicationBridge::PostApplyInitialState(FNetHandle Handle)
+void UReplicationBridge::PostApplyInitialState(FNetRefHandle Handle)
 {
 }
 
-void UReplicationBridge::DetachInstanceFromRemote(FNetHandle Handle, bool bTearOff, bool bShouldDestroyInstance)
+void UReplicationBridge::DetachInstanceFromRemote(FNetRefHandle Handle, bool bTearOff, bool bShouldDestroyInstance)
 {
 }
 
-void UReplicationBridge::DetachInstance(FNetHandle Handle)
+void UReplicationBridge::DetachInstance(FNetRefHandle Handle)
 {
 }
 
@@ -88,12 +85,12 @@ void UReplicationBridge::SetNetDriver(UNetDriver* NetDriver)
 {
 }
 
-void UReplicationBridge::GetInitialDependencies(FNetHandle Handle, FNetDependencyInfoArray& OutDependencies) const
+void UReplicationBridge::GetInitialDependencies(FNetRefHandle Handle, FNetDependencyInfoArray& OutDependencies) const
 {
 	return;
 }
 
-void UReplicationBridge::CallGetInitialDependencies(FNetHandle Handle, FNetDependencyInfoArray& OutDependencies) const
+void UReplicationBridge::CallGetInitialDependencies(FNetRefHandle Handle, FNetDependencyInfoArray& OutDependencies) const
 {
 	using namespace UE::Net::Private;
 
@@ -109,17 +106,17 @@ void UReplicationBridge::CallGetInitialDependencies(FNetHandle Handle, FNetDepen
 	}
 }
 
-void UReplicationBridge::DetachSubObjectInstancesFromRemote(FNetHandle OwnerHandle, bool bTearOff, bool bShouldDestroyInstance)
+void UReplicationBridge::DetachSubObjectInstancesFromRemote(FNetRefHandle OwnerHandle, bool bTearOff, bool bShouldDestroyInstance)
 {
 	using namespace UE::Net::Private;
 
-	FInternalNetHandle OwnerInternalIndex = NetHandleManager->GetInternalIndex(OwnerHandle);
+	FInternalNetRefIndex OwnerInternalIndex = NetRefHandleManager->GetInternalIndex(OwnerHandle);
 	if (OwnerInternalIndex)
 	{
-		for (const FInternalNetHandle SubObjectInternalIndex : NetHandleManager->GetChildSubObjects(OwnerInternalIndex))
+		for (const FInternalNetRefIndex SubObjectInternalIndex : NetRefHandleManager->GetChildSubObjects(OwnerInternalIndex))
 		{
-			FNetHandleManager::FReplicatedObjectData& SubObjectData = NetHandleManager->GetReplicatedObjectDataNoCheck(SubObjectInternalIndex);
-			const FNetHandle SubObjectHandle = SubObjectData.Handle;
+			FNetRefHandleManager::FReplicatedObjectData& SubObjectData = NetRefHandleManager->GetReplicatedObjectDataNoCheck(SubObjectInternalIndex);
+			const FNetRefHandle SubObjectHandle = SubObjectData.RefHandle;
 			SubObjectData.bTearOff = bTearOff;
 
 			CallDetachInstanceFromRemote(SubObjectHandle, bTearOff, bShouldDestroyInstance);
@@ -127,14 +124,14 @@ void UReplicationBridge::DetachSubObjectInstancesFromRemote(FNetHandle OwnerHand
 	}
 }
 
-void UReplicationBridge::DestroyNetObjectFromRemote(FNetHandle Handle, bool bTearOff, bool bDestroyInstance)
+void UReplicationBridge::DestroyNetObjectFromRemote(FNetRefHandle Handle, bool bTearOff, bool bDestroyInstance)
 {
 	using namespace UE::Net::Private;
 
 	if (Handle.IsValid())
 	{
-		FInternalNetHandle OwnerInternalIndex = NetHandleManager->GetInternalIndex(Handle);
-		FNetHandleManager::FReplicatedObjectData& ObjectData = NetHandleManager->GetReplicatedObjectDataNoCheck(OwnerInternalIndex);
+		FInternalNetRefIndex OwnerInternalIndex = NetRefHandleManager->GetInternalIndex(Handle);
+		FNetRefHandleManager::FReplicatedObjectData& ObjectData = NetRefHandleManager->GetReplicatedObjectDataNoCheck(OwnerInternalIndex);
 		ObjectData.bTearOff = bTearOff;
 
 		// if the a subobject owner is to be destroyed we want to detach all subobjects before doing so to ensure we execute expected callbacks
@@ -145,10 +142,10 @@ void UReplicationBridge::DestroyNetObjectFromRemote(FNetHandle Handle, bool bTea
 		CallDetachInstanceFromRemote(Handle, bTearOff, bDestroyInstance);
 
 		// Detach instance protocol
-		InternalDetachInstanceFromNetHandle(Handle);
+		InternalDetachInstanceFromNetRefHandle(Handle);
 	
-		// Destroy the NetHandle
-		UReplicationBridge::InternalDestroyNetObject(Handle);
+		// Destroy the NetRefHandle
+		InternalDestroyNetObject(Handle);
 
 		UE_LOG_REPLICATIONBRIDGE(Verbose, TEXT("DestroyNetObjectFromRemote Remote %s"), *Handle.ToString());
 	}
@@ -162,7 +159,7 @@ void UReplicationBridge::CallPreSendUpdate(float DeltaSeconds)
 	PreSendUpdate();
 }
 
-void UReplicationBridge::CallPreSendUpdateSingleHandle(FNetHandle Handle)
+void UReplicationBridge::CallPreSendUpdateSingleHandle(FNetRefHandle Handle)
 {
 	PreSendUpdateSingleHandle(Handle);
 }
@@ -172,12 +169,12 @@ void UReplicationBridge::CallUpdateInstancesWorldLocation()
 	UpdateInstancesWorldLocation();
 }
 
-void UReplicationBridge::CallDetachInstance(FNetHandle Handle)
+void UReplicationBridge::CallDetachInstance(FNetRefHandle Handle)
 {
 	DetachInstance(Handle);
 }
 
-void UReplicationBridge::CallDetachInstanceFromRemote(FNetHandle Handle, bool bTearOff, bool bShouldDestroyInstance)
+void UReplicationBridge::CallDetachInstanceFromRemote(FNetRefHandle Handle, bool bTearOff, bool bShouldDestroyInstance)
 {
 	DetachInstanceFromRemote(Handle, bTearOff, bShouldDestroyInstance);
 }
@@ -187,7 +184,7 @@ void UReplicationBridge::CallPruneStaleObjects()
 	PruneStaleObjects();
 }
 
-bool UReplicationBridge::CallWriteNetHandleCreationInfo(FReplicationBridgeSerializationContext& Context, FNetHandle Handle)
+bool UReplicationBridge::CallWriteNetRefHandleCreationInfo(FReplicationBridgeSerializationContext& Context, FNetRefHandle Handle)
 {
 	using namespace UE::Net;
 
@@ -211,22 +208,22 @@ bool UReplicationBridge::CallWriteNetHandleCreationInfo(FReplicationBridgeSerial
 	}
 	else
 	{
-		return WriteNetHandleCreationInfo(Context, Handle);
+		return WriteNetRefHandleCreationInfo(Context, Handle);
 	}
 }
 
-UE::Net::FNetHandle UReplicationBridge::CallCreateNetHandleFromRemote(FNetHandle SubObjectOwnerHandle, FNetHandle WantedNetHandle, FReplicationBridgeSerializationContext& Context)
+UE::Net::FNetRefHandle UReplicationBridge::CallCreateNetRefHandleFromRemote(FNetRefHandle SubObjectOwnerHandle, FNetRefHandle WantedNetHandle, FReplicationBridgeSerializationContext& Context)
 {
 	using namespace UE::Net;
 
 	check(!Context.bIsDestructionInfo);
 
-	FNetHandle ResultingHandle = CreateNetHandleFromRemote(SubObjectOwnerHandle, WantedNetHandle, Context);
+	FNetRefHandle ResultingHandle = CreateNetRefHandleFromRemote(SubObjectOwnerHandle, WantedNetHandle, Context);
 
 	// Track subobjects on clients
 	if (ResultingHandle.IsValid() && SubObjectOwnerHandle.IsValid())
 	{
-		NetHandleManager->AddSubObject(SubObjectOwnerHandle, ResultingHandle);
+		NetRefHandleManager->AddSubObject(SubObjectOwnerHandle, ResultingHandle);
 	}
 
 	return ResultingHandle;
@@ -249,12 +246,13 @@ void UReplicationBridge::ReadAndExecuteDestructionInfoFromRemote(FReplicationBri
 	if (ObjectReferenceCache->ResolveObjectReference(ReferenceToDestroy, Context.SerializationContext.GetInternalContext()->ResolveContext) != nullptr)
 	{
 		// Need to forward this as we do not really know how to destroy the actual instance
-		const bool bTearOff = false;
-		CallDetachInstanceFromRemote(ReferenceToDestroy.GetRefHandle(), bTearOff, true);
+		constexpr bool bTearOff = false;
+		constexpr bool bShouldDestroyInstance = false;
+		CallDetachInstanceFromRemote(ReferenceToDestroy.GetRefHandle(), bTearOff, bShouldDestroyInstance);
 	}
 }
 
-void UReplicationBridge::CallPostApplyInitialState(FNetHandle Handle)
+void UReplicationBridge::CallPostApplyInitialState(FNetRefHandle Handle)
 {
 	PostApplyInitialState(Handle);
 }
@@ -272,7 +270,7 @@ void UReplicationBridge::Initialize(UReplicationSystem* InReplicationSystem)
 	ReplicationSystem = InReplicationSystem;
 	ReplicationProtocolManager = &ReplicationSystemInternal->GetReplicationProtocolManager();
 	ReplicationStateDescriptorRegistry = &ReplicationSystemInternal->GetReplicationStateDescriptorRegistry();
-	NetHandleManager = &ReplicationSystemInternal->GetNetHandleManager();
+	NetRefHandleManager = &ReplicationSystemInternal->GetNetRefHandleManager();
 	ObjectReferenceCache = &ReplicationSystemInternal->GetObjectReferenceCache();
 	Groups = &ReplicationSystemInternal->GetGroups();
 
@@ -291,17 +289,17 @@ void UReplicationBridge::Deinitialize()
 	ReplicationSystem = nullptr;
 	ReplicationProtocolManager = nullptr;
 	ReplicationStateDescriptorRegistry = nullptr;
-	NetHandleManager = nullptr;
+	NetRefHandleManager = nullptr;
 	ObjectReferenceCache = nullptr;
 	Groups = nullptr;
 }
 
-UE::Net::FNetHandle UReplicationBridge::InternalCreateNetObject(FNetHandle AllocatedHandle, const UE::Net::FReplicationProtocol* ReplicationProtocol)
+UE::Net::FNetRefHandle UReplicationBridge::InternalCreateNetObject(FNetRefHandle AllocatedHandle, FNetHandle GlobalHandle, const UE::Net::FReplicationProtocol* ReplicationProtocol)
 {
 	check(AllocatedHandle.IsValid() && AllocatedHandle.IsCompleteHandle());
 
-	FNetHandle Handle = NetHandleManager->CreateNetObject(AllocatedHandle, ReplicationProtocol);
-	
+	FNetRefHandle Handle = NetRefHandleManager->CreateNetObject(AllocatedHandle, GlobalHandle, ReplicationProtocol);
+
 	if (Handle.IsValid())
 	{
 		UE_NET_TRACE_NETHANDLE_CREATED(Handle, ReplicationProtocol->DebugName, ReplicationProtocol->ProtocolIdentifier, 0/*Local*/);
@@ -310,9 +308,14 @@ UE::Net::FNetHandle UReplicationBridge::InternalCreateNetObject(FNetHandle Alloc
 	return Handle;
 }
 
-UE::Net::FNetHandle UReplicationBridge::InternalCreateNetObjectFromRemote(FNetHandle WantedNetHandle, const UE::Net::FReplicationProtocol* ReplicationProtocol)
+UE::Net::FNetRefHandle UReplicationBridge::InternalCreateNetObject(FNetRefHandle AllocatedHandle, const UE::Net::FReplicationProtocol* ReplicationProtocol)
+{
+	return InternalCreateNetObject(AllocatedHandle, FNetHandle(), ReplicationProtocol);
+}
+
+UE::Net::FNetRefHandle UReplicationBridge::InternalCreateNetObjectFromRemote(FNetRefHandle WantedNetHandle, const UE::Net::FReplicationProtocol* ReplicationProtocol)
 {
-	FNetHandle Handle = NetHandleManager->CreateNetObjectFromRemote(WantedNetHandle, ReplicationProtocol);
+	FNetRefHandle Handle = NetRefHandleManager->CreateNetObjectFromRemote(WantedNetHandle, ReplicationProtocol);
 
 	if (Handle.IsValid())
 	{
@@ -322,69 +325,52 @@ UE::Net::FNetHandle UReplicationBridge::InternalCreateNetObjectFromRemote(FNetHa
 	return Handle;
 }
 
-void UReplicationBridge::InternalAttachInstanceToNetHandle(FNetHandle Handle, bool bBindInstanceProtocol, UE::Net::FReplicationInstanceProtocol* InstanceProtocol, UObject* Instance)
+void UReplicationBridge::InternalAttachInstanceToNetHandle(FNetRefHandle Handle, bool bBindInstanceProtocol, UE::Net::FReplicationInstanceProtocol* InstanceProtocol, UObject* Instance)
 {
 	using namespace UE::Net;
 	using namespace UE::Net::Private;
 
 	const uint32 ReplicationSystemId = Handle.GetReplicationSystemId();
-	const uint32 InternalReplicationIndex = NetHandleManager->GetInternalIndex(Handle);
+	const FInternalNetRefIndex InternalReplicationIndex = NetRefHandleManager->GetInternalIndex(Handle);
 
-	NetHandleManager->AttachInstanceProtocol(InternalReplicationIndex, InstanceProtocol, Instance);
+	NetRefHandleManager->AttachInstanceProtocol(InternalReplicationIndex, InstanceProtocol, Instance);
 	UE_LOG_REPLICATIONBRIDGE(Verbose, TEXT("InternalAttachInstanceToNetHandle Attached: %s %s to ( InternalIndex: %u )"), *Instance->GetName(), *Handle.ToString(), InternalReplicationIndex);
 
 	// Bind instance protocol to dirty state tracking
 	if (bBindInstanceProtocol)
 	{
-		// Set push ID only if any state supports it. If no state supports it then we might crash if setting the ID.
-#if WITH_PUSH_MODEL
-		if (EnumHasAnyFlags(InstanceProtocol->InstanceTraits, EReplicationInstanceProtocolTraits::HasPartialPushBasedDirtiness | EReplicationInstanceProtocolTraits::HasFullPushBasedDirtiness))
-		{
-			const FNetPushObjectHandle PushHandle(InternalReplicationIndex, ReplicationSystemId);
-			TArrayView<const FReplicationFragment* const> Fragments(InstanceProtocol->Fragments, InstanceProtocol->FragmentCount);
-			SetNetPushIdOnFragments(Fragments, PushHandle);
-		}
-#endif
-
-		FReplicationInstanceOperationsInternal::BindInstanceProtocol(ReplicationSystemId, InternalReplicationIndex, InstanceProtocol, NetHandleManager->GetReplicatedObjectDataNoCheck(InternalReplicationIndex).Protocol);
+		FReplicationInstanceOperationsInternal::BindInstanceProtocol(ReplicationSystemId, InternalReplicationIndex, InstanceProtocol, NetRefHandleManager->GetReplicatedObjectDataNoCheck(InternalReplicationIndex).Protocol);
 		MarkNetObjectStateDirty(Handle.GetReplicationSystemId(), InternalReplicationIndex);
 	}
 }
 
-void UReplicationBridge::InternalDetachInstanceFromNetHandle(FNetHandle Handle)
+void UReplicationBridge::InternalDetachInstanceFromNetRefHandle(FNetRefHandle Handle)
 {
 	using namespace UE::Net;
 	using namespace UE::Net::Private;
 
-	const uint32 InternalReplicationIndex = NetHandleManager->GetInternalIndex(Handle);
+	const FInternalNetRefIndex InternalReplicationIndex = NetRefHandleManager->GetInternalIndex(Handle);
 
-	if (FReplicationInstanceProtocol* InstanceProtocol = const_cast<FReplicationInstanceProtocol*>(NetHandleManager->DetachInstanceProtocol(InternalReplicationIndex)))
+	if (FReplicationInstanceProtocol* InstanceProtocol = const_cast<FReplicationInstanceProtocol*>(NetRefHandleManager->DetachInstanceProtocol(InternalReplicationIndex)))
 	{
 		if (EnumHasAnyFlags(InstanceProtocol->InstanceTraits, EReplicationInstanceProtocolTraits::IsBound))
 		{
-#if WITH_PUSH_MODEL
-			if (EnumHasAnyFlags(InstanceProtocol->InstanceTraits, EReplicationInstanceProtocolTraits::HasPartialPushBasedDirtiness | EReplicationInstanceProtocolTraits::HasFullPushBasedDirtiness))
-			{
-				TArrayView<const FReplicationFragment* const> Fragments(InstanceProtocol->Fragments, InstanceProtocol->FragmentCount);
-				ClearNetPushIdOnFragments(Fragments);
-			}
-#endif
-			FReplicationInstanceOperationsInternal::UnbindInstanceProtocol(InstanceProtocol, NetHandleManager->GetReplicatedObjectDataNoCheck(InternalReplicationIndex).Protocol);
+			FReplicationInstanceOperationsInternal::UnbindInstanceProtocol(InstanceProtocol, NetRefHandleManager->GetReplicatedObjectDataNoCheck(InternalReplicationIndex).Protocol);
 		}
 		ReplicationProtocolManager->DestroyInstanceProtocol(InstanceProtocol);
 	}
 }
 
-void UReplicationBridge::InternalDestroyNetObject(FNetHandle Handle)
+void UReplicationBridge::InternalDestroyNetObject(FNetRefHandle Handle)
 {
 	using namespace UE::Net;
 	using namespace UE::Net::Private;
 
-	if (const FInternalNetHandle ObjectInternalIndex = NetHandleManager->GetInternalIndex(Handle))
+	if (const FInternalNetRefIndex ObjectInternalIndex = NetRefHandleManager->GetInternalIndex(Handle))
 	{
-		const FNetHandleManager::FReplicatedObjectData& ObjectData = NetHandleManager->GetReplicatedObjectDataNoCheck(ObjectInternalIndex);
+		const FNetRefHandleManager::FReplicatedObjectData& ObjectData = NetRefHandleManager->GetReplicatedObjectDataNoCheck(ObjectInternalIndex);
 
-		if (ObjectData.bTearOff && NetHandleManager->GetNetObjectRefCount(ObjectInternalIndex) > 0U)
+		if (ObjectData.bTearOff && NetRefHandleManager->GetNetObjectRefCount(ObjectInternalIndex) > 0U)
 		{
 			// We need to explicitly notify all ReplicationWriters that we are destroying objects pending tearoff
 			// The handle will automatically be removed from HandlesPendingTearOff after the next update
@@ -402,13 +388,28 @@ void UReplicationBridge::InternalDestroyNetObject(FNetHandle Handle)
 		}
 	}
 
-	NetHandleManager->DestroyNetObject(Handle);
+	NetRefHandleManager->DestroyNetObject(Handle);
 }
 
-void UReplicationBridge::DestroyLocalNetHandle(FNetHandle Handle)
+void UReplicationBridge::DestroyLocalNetHandle(FNetRefHandle Handle, EEndReplicationFlags EndReplicationFlags)
 {
+	if (EnumHasAnyFlags(EndReplicationFlags, EEndReplicationFlags::DestroyNetHandle | EEndReplicationFlags::ClearNetPushId))
+	{
+		const UE::Net::Private::FInternalNetRefIndex InternalReplicationIndex = NetRefHandleManager->GetInternalIndex(Handle);
+
+		if (EnumHasAnyFlags(EndReplicationFlags, EEndReplicationFlags::DestroyNetHandle))
+		{
+			DestroyGlobalNetHandle(InternalReplicationIndex);
+		}
+
+		if (EnumHasAnyFlags(EndReplicationFlags, EEndReplicationFlags::ClearNetPushId))
+		{
+			ClearNetPushIds(InternalReplicationIndex);
+		}
+	}
+
 	// Detach instance protocol
-	InternalDetachInstanceFromNetHandle(Handle);
+	InternalDetachInstanceFromNetRefHandle(Handle);
 
 	// Allow derived bridges to cleanup any instance info they have stored
 	DetachInstance(Handle);
@@ -417,88 +418,89 @@ void UReplicationBridge::DestroyLocalNetHandle(FNetHandle Handle)
 	GetReplicationSystem()->RemoveFromAllGroups(Handle);
 
 	// If we have any attached SubObjects, tag them for destroy as well
-	InternalDestroySubObjects(Handle);
+	InternalDestroySubObjects(Handle, EndReplicationFlags);
 
-	// Tell ReplicationSystem to destroy the NetHandle
+	// Tell ReplicationSystem to destroy the handle
 	InternalDestroyNetObject(Handle);
 
 	UE_LOG_REPLICATIONBRIDGE(Verbose, TEXT("DestroyLocalNetHandle Local %s"), *Handle.ToString());
 }
 
-UE::Net::FNetHandle UReplicationBridge::InternalGetSubObjectOwner(FNetHandle SubObjectHandle) const
+UE::Net::FNetRefHandle UReplicationBridge::InternalGetSubObjectOwner(FNetRefHandle SubObjectHandle) const
 {
 	using namespace UE::Net::Private;
 
-	const FInternalNetHandle InternalHandle = NetHandleManager->GetInternalIndex(SubObjectHandle);
-	if (InternalHandle != FNetHandleManager::InvalidInternalIndex)
+	const FInternalNetRefIndex InternalNetRefIndex = NetRefHandleManager->GetInternalIndex(SubObjectHandle);
+	if (InternalNetRefIndex != FNetRefHandleManager::InvalidInternalIndex)
 	{
-		const FNetHandleManager::FReplicatedObjectData& ObjectData = NetHandleManager->GetReplicatedObjectDataNoCheck(InternalHandle);
-		const FInternalNetHandle OwnerInternalHandle = ObjectData.SubObjectRootIndex;
+		const FNetRefHandleManager::FReplicatedObjectData& ObjectData = NetRefHandleManager->GetReplicatedObjectDataNoCheck(InternalNetRefIndex);
+		const FInternalNetRefIndex OwnerInternalNetRefIndex = ObjectData.SubObjectRootIndex;
 
-		if (OwnerInternalHandle != FNetHandleManager::InvalidInternalIndex)
+		if (OwnerInternalNetRefIndex != FNetRefHandleManager::InvalidInternalIndex)
 		{
-			return NetHandleManager->GetReplicatedObjectDataNoCheck(OwnerInternalHandle).Handle;
+			return NetRefHandleManager->GetReplicatedObjectDataNoCheck(OwnerInternalNetRefIndex).RefHandle;
 		}
 	}
 
-	return FNetHandle();
+	return FNetRefHandle();
 }
 
-void UReplicationBridge::InternalAddSubObject(FNetHandle OwnerHandle, FNetHandle SubObjectHandle, FNetHandle InsertRelativeToSubObjectHandle, ESubObjectInsertionOrder InsertionOrder)
+void UReplicationBridge::InternalAddSubObject(FNetRefHandle OwnerHandle, FNetRefHandle SubObjectHandle, FNetRefHandle InsertRelativeToSubObjectHandle, ESubObjectInsertionOrder InsertionOrder)
 {
 	using namespace UE::Net::Private;
 
 	EAddSubObjectFlags AddSubObjectFlags = EAddSubObjectFlags::Default;
 	AddSubObjectFlags |= InsertionOrder == ESubObjectInsertionOrder::None ? EAddSubObjectFlags::None : EAddSubObjectFlags::ReplicateWithSubObject;
 
-	if (NetHandleManager->AddSubObject(OwnerHandle, SubObjectHandle, InsertRelativeToSubObjectHandle, AddSubObjectFlags))
+	if (NetRefHandleManager->AddSubObject(OwnerHandle, SubObjectHandle, InsertRelativeToSubObjectHandle, AddSubObjectFlags))
 	{
 		// If the subobject is not new we need to mark it as dirty to pick it up for replication with its new parent
-		MarkNetObjectStateDirty(ReplicationSystem->GetId(), NetHandleManager->GetInternalIndex(SubObjectHandle));
+		MarkNetObjectStateDirty(ReplicationSystem->GetId(), NetRefHandleManager->GetInternalIndex(SubObjectHandle));
 
 		// We set the priority of subobjects to be static as they will be prioritized with owner
-		ReplicationSystem->SetStaticPriority(SubObjectHandle, 1.f);
+		ReplicationSystem->SetStaticPriority(SubObjectHandle, 1.0f);
 	}
 }
 
-void UReplicationBridge::InternalDestroySubObjects(FNetHandle OwnerHandle)
+void UReplicationBridge::InternalDestroySubObjects(FNetRefHandle OwnerHandle, EEndReplicationFlags Flags)
 {
 	using namespace UE::Net::Private;
 
 	// Destroy SubObjects
-	FInternalNetHandle OwnerInternalIndex = NetHandleManager->GetInternalIndex(OwnerHandle);
+	FInternalNetRefIndex OwnerInternalIndex = NetRefHandleManager->GetInternalIndex(OwnerHandle);
 	if (OwnerInternalIndex)
 	{
-		for (FInternalNetHandle SubObjectInternalIndex : NetHandleManager->GetChildSubObjects(OwnerInternalIndex))
+		for (FInternalNetRefIndex SubObjectInternalIndex : NetRefHandleManager->GetChildSubObjects(OwnerInternalIndex))
 		{
-			const FNetHandleManager::FReplicatedObjectData& SubObjectData = NetHandleManager->GetReplicatedObjectData(SubObjectInternalIndex);
-			const FNetHandle SubObjectHandle = SubObjectData.Handle;
+			const FNetRefHandleManager::FReplicatedObjectData& SubObjectData = NetRefHandleManager->GetReplicatedObjectData(SubObjectInternalIndex);
+			const FNetRefHandle SubObjectHandle = SubObjectData.RefHandle;
 			const bool bDestroySubObjectWithOwner = SubObjectData.bDestroySubObjectWithOwner;
 				
-			// Tag subobject for destroy, the check against the scope is needed since the subobjects array might contain subobjects already pending destroy
-			if (bDestroySubObjectWithOwner && NetHandleManager->IsScopableIndex(SubObjectInternalIndex))
+			// Tag subobject for destroy. The check against the scope is needed since the subobjects array might contain subobjects already pending destroy.
+			if (bDestroySubObjectWithOwner && NetRefHandleManager->IsScopableIndex(SubObjectInternalIndex))
 			{
 				UE_LOG_REPLICATIONBRIDGE(Verbose, TEXT("InternalDestroySubObjects %s - SubObject %s"), *OwnerHandle.ToString(), *SubObjectHandle.ToString());
-				DestroyLocalNetHandle(SubObjectHandle);
+				DestroyLocalNetHandle(SubObjectHandle, Flags);
 			}
 		}
 	}
 }
 
-void UReplicationBridge::EndReplication(FNetHandle Handle, EEndReplicationFlags EndReplicationFlags, FEndReplicationParameters* Parameters)
+void UReplicationBridge::EndReplication(FNetRefHandle Handle, EEndReplicationFlags EndReplicationFlags, FEndReplicationParameters* Parameters)
 {
 	if (!IsReplicatedHandle(Handle))
 	{
 		return;
 	}
 
-	if (NetHandleManager->IsLocalNetHandle(Handle))
+	if (NetRefHandleManager->IsLocalNetRefHandle(Handle))
 	{
 		if (EnumHasAnyFlags(EndReplicationFlags, EEndReplicationFlags::TearOff))
 		{
 			// Add handle to list of objects pending tear off and indicate that it should be destroyed during next update
-			// We need to do this to cover the case where the torn off object not yet has been added to scope,
-			TearOff(Handle, true);
+			// We need to do this to cover the case where the torn off object not yet has been added to the scope.
+			constexpr bool bIsImmediate = true;
+			TearOff(Handle, EndReplicationFlags, bIsImmediate);
 
 			// We do however copy the final state data and mark object to stop propagating state changes
 			// Note: Current implementation keeps the mapping between instance and handle alive so that we can handle the case where we both start replication 
@@ -521,13 +523,13 @@ void UReplicationBridge::EndReplication(FNetHandle Handle, EEndReplicationFlags
 				InternalFlushStateData(Handle);
 			}
 
-			DestroyLocalNetHandle(Handle);	
+			DestroyLocalNetHandle(Handle, EndReplicationFlags);	
 		}
 	}
 	else
 	{
 		// If we get a call to end replication on the client, we need to detach the instance as it might be garbage collected
-		InternalDetachInstanceFromNetHandle(Handle);
+		InternalDetachInstanceFromNetRefHandle(Handle);
 	}
 }
 
@@ -545,27 +547,27 @@ void UReplicationBridge::RemoveDestructionInfosForGroup(UE::Net::FNetObjectGroup
 		FNetObjectGroup* Group = Groups->GetGroup(GroupHandle);
 		check(Group);
 
-		TArray<FNetHandle, TInlineAllocator<384>> ObjectsToRemove;	
-		TArray<FInternalNetHandle, TInlineAllocator<384>> ObjectIndicesToRemove;	
+		TArray<FNetRefHandle, TInlineAllocator<384>> ObjectsToRemove;	
+		TArray<FInternalNetRefIndex, TInlineAllocator<384>> ObjectIndicesToRemove;	
 		for (uint32 InternalObjectIndex : MakeArrayView(Group->Members))
 		{
-			if (NetHandleManager->GetIsDestroyedStartupObject(InternalObjectIndex))
+			if (NetRefHandleManager->GetIsDestroyedStartupObject(InternalObjectIndex))
 			{
-				const FNetHandleManager::FReplicatedObjectData& ObjectData = NetHandleManager->GetReplicatedObjectDataNoCheck(InternalObjectIndex);
-				if (StaticObjectsPendingDestroy.Remove(ObjectData.Handle))
+				const FNetRefHandleManager::FReplicatedObjectData& ObjectData = NetRefHandleManager->GetReplicatedObjectDataNoCheck(InternalObjectIndex);
+				if (StaticObjectsPendingDestroy.Remove(ObjectData.RefHandle))
 				{
-					ObjectsToRemove.Add(ObjectData.Handle);
+					ObjectsToRemove.Add(ObjectData.RefHandle);
 					ObjectIndicesToRemove.Add(InternalObjectIndex);
 				}
 			}			
 		}
 
-		for (FNetHandle Handle : MakeArrayView(ObjectsToRemove))
+		for (FNetRefHandle Handle : MakeArrayView(ObjectsToRemove))
 		{
-			NetHandleManager->DestroyNetObject(Handle);
+			NetRefHandleManager->DestroyNetObject(Handle);
 		}
 
-		for (FInternalNetHandle InternalReplicationIndex : MakeArrayView(ObjectIndicesToRemove))
+		for (FInternalNetRefIndex InternalReplicationIndex : MakeArrayView(ObjectIndicesToRemove))
 		{
 			constexpr bool bHasWorldLocation = false;
 			WorldLocations.SetHasWorldLocation(InternalReplicationIndex, bHasWorldLocation);
@@ -576,7 +578,7 @@ void UReplicationBridge::RemoveDestructionInfosForGroup(UE::Net::FNetObjectGroup
 		// We should remove all destruction infos and objects
 		for (const auto& It : StaticObjectsPendingDestroy)
 		{
-			NetHandleManager->DestroyNetObject(It.Key);
+			NetRefHandleManager->DestroyNetObject(It.Key);
 		}
 
 		// Remove from WorldLocations
@@ -604,24 +606,25 @@ void UReplicationBridge::UpdateHandlesPendingTearOff()
 	TArray<FTearOffInfo, TInlineAllocator<32>> ObjectsStillPendingTearOff;
 	for (FTearOffInfo Info : MakeArrayView(HandlesPendingTearOff))
 	{
-		if (uint32 ObjectInternalIndex = NetHandleManager->GetInternalIndex(Info.Handle))
+		if (uint32 ObjectInternalIndex = NetRefHandleManager->GetInternalIndex(Info.Handle))
 		{
 			// Immediate tear-off or object that no longer are referenced by any connections are destroyed
-			if (NetHandleManager->GetNetObjectRefCount(ObjectInternalIndex) == 0U || Info.bIsImmediate)
+			if (NetRefHandleManager->GetNetObjectRefCount(ObjectInternalIndex) == 0U || Info.bIsImmediate)
 			{
-				DestroyLocalNetHandle(Info.Handle);
+				DestroyLocalNetHandle(Info.Handle, Info.DestroyFlags);
 			}
 			else
 			{
 				// If the object is still in scope remove it from scope as objects being torn-off should not be added to new connections
-				if (NetHandleManager->IsScopableIndex(ObjectInternalIndex))
+				if (NetRefHandleManager->IsScopableIndex(ObjectInternalIndex))
 				{
 					// Mark object as no longer scopable, and that we should not propagate changed states
-					NetHandleManager->RemoveFromScope(ObjectInternalIndex);
+					NetRefHandleManager->RemoveFromScope(ObjectInternalIndex);
 				}
 			
 				// Keep object in the pending tear-off list until the object is no longer referenced by any ReplicationWriter
-				ObjectsStillPendingTearOff.Add(FTearOffInfo(Info.Handle, false));
+				constexpr bool bIsImmediate = false;
+				ObjectsStillPendingTearOff.Add(FTearOffInfo(Info.Handle, Info.DestroyFlags, bIsImmediate));
 			}
 		}
 	}
@@ -630,11 +633,11 @@ void UReplicationBridge::UpdateHandlesPendingTearOff()
 	HandlesPendingTearOff.Insert(ObjectsStillPendingTearOff.GetData(), ObjectsStillPendingTearOff.Num(), 0);
 }
 
-void UReplicationBridge::TearOff(FNetHandle Handle, bool bIsImmediate)
+void UReplicationBridge::TearOff(FNetRefHandle Handle, EEndReplicationFlags DestroyFlags, bool bIsImmediate)
 {
 	if (!HandlesPendingTearOff.FindByPredicate([&](const FTearOffInfo& Info){ return Info.Handle == Handle; }))
 	{
-		HandlesPendingTearOff.Add(FTearOffInfo(Handle, bIsImmediate));
+		HandlesPendingTearOff.Add(FTearOffInfo(Handle, DestroyFlags, bIsImmediate));
 	}	
 }
 
@@ -643,12 +646,12 @@ void UReplicationBridge::InternalFlushStateData(UE::Net::FNetSerializationContex
 	using namespace UE::Net;
 	using namespace UE::Net::Private;
 
-	if (InternalObjectIndex == FNetHandleManager::InvalidInternalIndex)
+	if (InternalObjectIndex == FNetRefHandleManager::InvalidInternalIndex)
 	{
 		return;
 	}
 
-	FNetHandleManager::FReplicatedObjectData& ObjectData = NetHandleManager->GetReplicatedObjectDataNoCheck(InternalObjectIndex);
+	FNetRefHandleManager::FReplicatedObjectData& ObjectData = NetRefHandleManager->GetReplicatedObjectDataNoCheck(InternalObjectIndex);
 
 	// Copy state data, if object already is torn off there is nothing to do
 	if (ObjectData.bTearOff)
@@ -656,33 +659,33 @@ void UReplicationBridge::InternalFlushStateData(UE::Net::FNetSerializationContex
 		return;
 	}
 
-	for (FInternalNetHandle SubObjectInternalIndex : NetHandleManager->GetChildSubObjects(InternalObjectIndex))
+	for (FInternalNetRefIndex SubObjectInternalIndex : NetRefHandleManager->GetChildSubObjects(InternalObjectIndex))
 	{
 		InternalFlushStateData(SerializationContext, ChangeMaskCache, ChangeMaskWriter, SubObjectInternalIndex);
 	}
 
 	if (EnumHasAnyFlags(ObjectData.InstanceProtocol->InstanceTraits, EReplicationInstanceProtocolTraits::NeedsPoll | EReplicationInstanceProtocolTraits::NeedsPreSendUpdate))
 	{
-		CallPreSendUpdateSingleHandle(ObjectData.Handle);
+		CallPreSendUpdateSingleHandle(ObjectData.RefHandle);
 	} 
 
-	FReplicationInstanceOperationsInternal::CopyObjectStateData(ChangeMaskWriter, ChangeMaskCache, *NetHandleManager, SerializationContext, InternalObjectIndex);
+	FReplicationInstanceOperationsInternal::CopyObjectStateData(ChangeMaskWriter, ChangeMaskCache, *NetRefHandleManager, SerializationContext, InternalObjectIndex);
 }
 
-void UReplicationBridge::InternalFlushStateData(FNetHandle Handle)
+void UReplicationBridge::InternalFlushStateData(FNetRefHandle Handle)
 {
 	using namespace UE::Net;
 	using namespace UE::Net::Private;
 
 	IRIS_PROFILER_SCOPE(InternalFlushStateData);
 
-	const uint32 InternalObjectIndex = NetHandleManager->GetInternalIndex(Handle);
-	if (InternalObjectIndex == FNetHandleManager::InvalidInternalIndex)
+	const uint32 InternalObjectIndex = NetRefHandleManager->GetInternalIndex(Handle);
+	if (InternalObjectIndex == FNetRefHandleManager::InvalidInternalIndex)
 	{
 		return;
 	}
 
-	FNetHandleManager::FReplicatedObjectData& ObjectData = NetHandleManager->GetReplicatedObjectDataNoCheck(InternalObjectIndex);
+	FNetRefHandleManager::FReplicatedObjectData& ObjectData = NetRefHandleManager->GetReplicatedObjectDataNoCheck(InternalObjectIndex);
 
 	FChangeMaskCache ChangeMaskCache;
 	FNetBitStreamWriter ChangeMaskWriter;
@@ -711,29 +714,29 @@ void UReplicationBridge::InternalFlushStateData(FNetHandle Handle)
 	}
 }
 
-void UReplicationBridge::InternalTearOff(FNetHandle Handle)
+void UReplicationBridge::InternalTearOff(FNetRefHandle Handle)
 {
 	using namespace UE::Net;
 	using namespace UE::Net::Private;
 
 	IRIS_PROFILER_SCOPE(InternalTearOff);
 
-	const uint32 InternalObjectIndex = NetHandleManager->GetInternalIndex(Handle);
+	const uint32 InternalObjectIndex = NetRefHandleManager->GetInternalIndex(Handle);
 
-	if (InternalObjectIndex == FNetHandleManager::InvalidInternalIndex)
+	if (InternalObjectIndex == FNetRefHandleManager::InvalidInternalIndex)
 	{
 		return;
 	}
 
-	FNetHandleManager::FReplicatedObjectData& ObjectData = NetHandleManager->GetReplicatedObjectDataNoCheck(InternalObjectIndex);
+	FNetRefHandleManager::FReplicatedObjectData& ObjectData = NetRefHandleManager->GetReplicatedObjectDataNoCheck(InternalObjectIndex);
 
 	// Copy state data and tear off now
 	if (!ObjectData.bTearOff)
 	{
 		// Mark subobjects for TearOff as well.
-		for (FInternalNetHandle SubObjectInternalIndex : NetHandleManager->GetChildSubObjects(InternalObjectIndex))
+		for (FInternalNetRefIndex SubObjectInternalIndex : NetRefHandleManager->GetChildSubObjects(InternalObjectIndex))
 		{
-			InternalTearOff(NetHandleManager->GetNetHandleFromInternalIndex(SubObjectInternalIndex));
+			InternalTearOff(NetRefHandleManager->GetNetRefHandleFromInternalIndex(SubObjectInternalIndex));
 		}	
 
 		UE_LOG_REPLICATIONBRIDGE(Verbose, TEXT("TearOff %s"), *Handle.ToString());
@@ -752,7 +755,7 @@ void UReplicationBridge::InternalTearOff(FNetHandle Handle)
 			CallPreSendUpdateSingleHandle(Handle);
 		} 
 
-		if (FReplicationInstanceOperationsInternal::CopyObjectStateData(ChangeMaskWriter, ChangeMaskCache, *NetHandleManager, SerializationContext, InternalObjectIndex))
+		if (FReplicationInstanceOperationsInternal::CopyObjectStateData(ChangeMaskWriter, ChangeMaskCache, *NetRefHandleManager, SerializationContext, InternalObjectIndex))
 		{
 			// Propagate changes to all connections that we currently have in scope
 			FReplicationConnections& Connections = ReplicationSystem->GetReplicationSystemInternal()->GetConnections();
@@ -774,19 +777,19 @@ void UReplicationBridge::InternalTearOff(FNetHandle Handle)
 	}
 }
 
-UE::Net::FNetHandle UReplicationBridge::InternalAddDestructionInfo(FNetHandle Handle, const FEndReplicationParameters& Parameters)
+UE::Net::FNetRefHandle UReplicationBridge::InternalAddDestructionInfo(FNetRefHandle Handle, const FEndReplicationParameters& Parameters)
 {
 	using namespace UE::Net;
 	using namespace UE::Net::Private;
 
 	FReplicationProtocolManager* ProtocolManager = GetReplicationProtocolManager();
 
-	// Create DestructionInfo NetHandle used to scope destruction infos
+	// Create handle for DestructionInfo to be able to scope destruction infos
 	const FNetObjectGroupHandle LevelGroupHandle = GetLevelGroup(Parameters.Level);
-	FNetHandle DestructionInfoHandle = NetHandleManager->CreateHandleForDestructionInfo(Handle, DestructionInfoProtocol);
+	FNetRefHandle DestructionInfoHandle = NetRefHandleManager->CreateHandleForDestructionInfo(Handle, DestructionInfoProtocol);
 	GetReplicationSystem()->AddToGroup(LevelGroupHandle, DestructionInfoHandle);
 
-	const FInternalNetHandle InternalReplicationIndex = NetHandleManager->GetInternalIndex(DestructionInfoHandle);
+	const FInternalNetRefIndex InternalReplicationIndex = NetRefHandleManager->GetInternalIndex(DestructionInfoHandle);
 
 	// We also need to store the actual data we send to destroy static references when they are scoped
 	FDestructionInfo PersistentDestructionInfo;
@@ -813,7 +816,7 @@ UE::Net::FNetHandle UReplicationBridge::InternalAddDestructionInfo(FNetHandle Ha
 	return DestructionInfoHandle;
 }
 
-bool UReplicationBridge::IsReplicatedHandle(FNetHandle Handle) const
+bool UReplicationBridge::IsReplicatedHandle(FNetRefHandle Handle) const
 {
 	return Handle.IsValid() && ReplicationSystem->IsValidHandle(Handle);
 }
@@ -907,3 +910,33 @@ UE::Net::FNetObjectGroupHandle UReplicationBridge::GetLevelGroup(const UObject*
 	const UE::Net::FNetObjectGroupHandle* LevelGroupHandle = LevelGroups.Find(FObjectKey(Level));
 	return (LevelGroupHandle != nullptr ? *LevelGroupHandle : UE::Net::InvalidNetObjectGroupHandle);
 }
+
+void UReplicationBridge::DestroyGlobalNetHandle(UE::Net::Private::FInternalNetRefIndex InternalReplicationIndex)
+{
+	using namespace UE::Net;
+	using namespace UE::Net::Private;
+
+	const FNetRefHandleManager::FReplicatedObjectData& ObjectData = NetRefHandleManager->GetReplicatedObjectDataNoCheck(InternalReplicationIndex);
+	if (ObjectData.NetHandle.IsValid())
+	{
+		FNetHandleDestroyer::DestroyNetHandle(ObjectData.NetHandle);
+	}
+}
+
+void UReplicationBridge::ClearNetPushIds(UE::Net::Private::FInternalNetRefIndex InternalReplicationIndex)
+{
+#if WITH_PUSH_MODEL
+	using namespace UE::Net;
+	using namespace UE::Net::Private;
+
+	const FNetRefHandleManager::FReplicatedObjectData& ObjectData = NetRefHandleManager->GetReplicatedObjectDataNoCheck(InternalReplicationIndex);
+	if (const FReplicationInstanceProtocol* InstanceProtocol = ObjectData.InstanceProtocol)
+	{
+		if (EnumHasAnyFlags(InstanceProtocol->InstanceTraits, EReplicationInstanceProtocolTraits::HasPartialPushBasedDirtiness | EReplicationInstanceProtocolTraits::HasFullPushBasedDirtiness))
+		{
+			TArrayView<const FReplicationFragment* const> Fragments(InstanceProtocol->Fragments, InstanceProtocol->FragmentCount);
+			ClearNetPushIdOnFragments(Fragments);
+		}
+	}
+#endif
+}
diff --git a/Engine/Source/Runtime/Experimental/Iris/Core/Private/Iris/ReplicationSystem/ReplicationOperationsInternal.cpp b/Engine/Source/Runtime/Experimental/Iris/Core/Private/Iris/ReplicationSystem/ReplicationOperationsInternal.cpp
--- a/Engine/Source/Runtime/Experimental/Iris/Core/Private/Iris/ReplicationSystem/ReplicationOperationsInternal.cpp
+++ b/Engine/Source/Runtime/Experimental/Iris/Core/Private/Iris/ReplicationSystem/ReplicationOperationsInternal.cpp
@@ -53,11 +53,11 @@ void FReplicationInstanceOperationsInternal::UnbindInstanceProtocol(FReplication
 	}
 }
 
-uint32 FReplicationInstanceOperationsInternal::CopyObjectStateData(FNetBitStreamWriter& ChangeMaskWriter, FChangeMaskCache& Cache, FNetHandleManager& NetHandleManager, FNetSerializationContext& SerializationContext, uint32 InternalIndex)
+uint32 FReplicationInstanceOperationsInternal::CopyObjectStateData(FNetBitStreamWriter& ChangeMaskWriter, FChangeMaskCache& Cache, FNetRefHandleManager& NetRefHandleManager, FNetSerializationContext& SerializationContext, uint32 InternalIndex)
 {
-	if (NetHandleManager.IsScopableIndex(InternalIndex))
+	if (NetRefHandleManager.IsScopableIndex(InternalIndex))
 	{
-		const FNetHandleManager::FReplicatedObjectData& Object = NetHandleManager.GetReplicatedObjectDataNoCheck(InternalIndex);
+		const FNetRefHandleManager::FReplicatedObjectData& Object = NetRefHandleManager.GetReplicatedObjectDataNoCheck(InternalIndex);
 
 		// We cannot copy state data for zero sized objects or objects that no longer has an instance protocol.
 		if (Object.InstanceProtocol && Object.Protocol->InternalTotalSize > 0U)
@@ -72,8 +72,8 @@ uint32 FReplicationInstanceOperationsInternal::CopyObjectStateData(FNetBitStream
 				ChangeMaskStorageType* ChangeMaskData = Cache.GetChangeMaskStorage(Info);
 				ChangeMaskWriter.InitBytes(ChangeMaskData, ChangeMaskByteCount);
 
-				//UE_LOG(LogIris, Log, TEXT("Copying state data for ( InternalIndex: %u ) NethHandleIndex: %u"), InternalIndex, Object.Handle.GetIndex());
-				FReplicationInstanceOperations::CopyAndQuantize(SerializationContext, NetHandleManager.GetReplicatedObjectStateBufferNoCheck(InternalIndex), &ChangeMaskWriter, Object.InstanceProtocol, Object.Protocol);
+				//UE_LOG(LogIris, Log, TEXT("Copying state data for ( InternalIndex: %u ) with NetRefHandle (Id=%u)"), InternalIndex, Object.RefHandle.GetId());
+				FReplicationInstanceOperations::CopyAndQuantize(SerializationContext, NetRefHandleManager.GetReplicatedObjectStateBufferNoCheck(InternalIndex), &ChangeMaskWriter, Object.InstanceProtocol, Object.Protocol);
 
 				Info.bHasDirtyChangeMask = MakeNetBitArrayView(ChangeMaskData, ChangeMaskByteCount * 8U).FindFirstOne() != FNetBitArrayView::InvalidIndex;
 			}
@@ -81,17 +81,17 @@ uint32 FReplicationInstanceOperationsInternal::CopyObjectStateData(FNetBitStream
 			// Mark subobject owner as dirty if this is a subobject
 			if (const uint32 SubObjectOwnerIndex = Object.SubObjectRootIndex)
 			{
-				const bool bIsOwnerScopable = NetHandleManager.IsScopableIndex(SubObjectOwnerIndex);
+				const bool bIsOwnerScopable = NetRefHandleManager.IsScopableIndex(SubObjectOwnerIndex);
 				// Dependent objects should not ensure if the owner isn't scopable.
-				ensureMsgf(bIsOwnerScopable || Object.IsDependentObject(), TEXT("SubObject ( InternaIndex: %u ) with NetHandleIndex: %u is trying to dirty parent ( InternalIndex: %u ) not in scope."), InternalIndex, Object.Handle.GetId(), SubObjectOwnerIndex);
+				ensureMsgf(bIsOwnerScopable || Object.IsDependentObject(), TEXT("SubObject ( InternaIndex: %u ) with NetRefHandle (Id=%u) is trying to dirty parent ( InternalIndex: %u ) not in scope."), InternalIndex, Object.RefHandle.GetId(), SubObjectOwnerIndex);
 				if (bIsOwnerScopable)
 				{
 					// Do we want to control this separately for subobjects? Or should they respect the setting on the owner?
 					// For now, we do and will not mark owner as dirty if owner should not propagate statechanges
-					bShouldPropagateChangedStates = bShouldPropagateChangedStates && NetHandleManager.GetReplicatedObjectDataNoCheck(SubObjectOwnerIndex).bShouldPropagateChangedStates;
+					bShouldPropagateChangedStates = bShouldPropagateChangedStates && NetRefHandleManager.GetReplicatedObjectDataNoCheck(SubObjectOwnerIndex).bShouldPropagateChangedStates;
 					if (bShouldPropagateChangedStates)
 					{
-						//UE_LOG(LogIris, Log, TEXT("Marking SubObjectOwner( InternalIndex: %u ) as dirty for ( InternalIndex: %u ) NethHandleIndex: %u"), InternalIndex, Object.Handle.GetIndex());
+						//UE_LOG(LogIris, Log, TEXT("Marking SubObjectOwner( InternalIndex: %u ) as dirty for ( InternalIndex: %u ) with NetRefHandle (Id=%u)"), SubObjectOwnerIndex, InternalIndex, Object.RefHandle.GetId());
 						Cache.AddSubObjectOwnerDirty(SubObjectOwnerIndex);
 					}
 				}
@@ -107,13 +107,13 @@ uint32 FReplicationInstanceOperationsInternal::CopyObjectStateData(FNetBitStream
 		}
 		else if (const uint32 SubObjectOwnerIndex = Object.SubObjectRootIndex)
 		{
-			if (Object.bShouldPropagateChangedStates && ensure(NetHandleManager.IsScopableIndex(SubObjectOwnerIndex)))
+			if (Object.bShouldPropagateChangedStates && ensure(NetRefHandleManager.IsScopableIndex(SubObjectOwnerIndex)))
 			{	
 				// Do we want to control this separately for subobjects? Or should they respect the setting on the owner?
 				// For now, we do and will not mark owner as dirty if owner should not propagate statechanges
-				if (NetHandleManager.GetReplicatedObjectDataNoCheck(SubObjectOwnerIndex).bShouldPropagateChangedStates)
+				if (NetRefHandleManager.GetReplicatedObjectDataNoCheck(SubObjectOwnerIndex).bShouldPropagateChangedStates)
 				{
-					//UE_LOG(LogIris, Log, TEXT("Marking SubObjectOwner( InternalIndex: %u ) as dirty for ( InternalIndex: %u ) NethHandleIndex: %u"), InternalIndex, Object.Handle.GetIndex());
+					//UE_LOG(LogIris, Log, TEXT("Marking SubObjectOwner( InternalIndex: %u ) as dirty for ( InternalIndex: %u ) with NetRefHandle (Id=%u)"), SubObjectOwnerIndex, InternalIndex, Object.Handle.GetIndex());
 					Cache.AddSubObjectOwnerDirty(SubObjectOwnerIndex);
 				}
 			}			
diff --git a/Engine/Source/Runtime/Experimental/Iris/Core/Private/Iris/ReplicationSystem/ReplicationOperationsInternal.h b/Engine/Source/Runtime/Experimental/Iris/Core/Private/Iris/ReplicationSystem/ReplicationOperationsInternal.h
--- a/Engine/Source/Runtime/Experimental/Iris/Core/Private/Iris/ReplicationSystem/ReplicationOperationsInternal.h
+++ b/Engine/Source/Runtime/Experimental/Iris/Core/Private/Iris/ReplicationSystem/ReplicationOperationsInternal.h
@@ -20,7 +20,7 @@ namespace UE::Net
 	namespace Private
 	{
 		struct FChangeMaskCache;
-		class FNetHandleManager;
+		class FNetRefHandleManager;
 	}
 }
 
@@ -36,7 +36,7 @@ struct FReplicationInstanceOperationsInternal
 	static IRISCORE_API void UnbindInstanceProtocol(FReplicationInstanceProtocol* InstanceProtocol, const FReplicationProtocol* Protocol);
 
 	/** Copy state data for a single instance */
-	static IRISCORE_API uint32 CopyObjectStateData(FNetBitStreamWriter& ChangeMaskWriter, FChangeMaskCache& Cache, FNetHandleManager& NetHandleManager, FNetSerializationContext& SerializationContext, uint32 InternalIndex);
+	static IRISCORE_API uint32 CopyObjectStateData(FNetBitStreamWriter& ChangeMaskWriter, FChangeMaskCache& Cache, FNetRefHandleManager& NetRefHandleManager, FNetSerializationContext& SerializationContext, uint32 InternalIndex);
 };
 
 struct FReplicationStateOperationsInternal
diff --git a/Engine/Source/Runtime/Experimental/Iris/Core/Private/Iris/ReplicationSystem/ReplicationReader.cpp b/Engine/Source/Runtime/Experimental/Iris/Core/Private/Iris/ReplicationSystem/ReplicationReader.cpp
--- a/Engine/Source/Runtime/Experimental/Iris/Core/Private/Iris/ReplicationSystem/ReplicationReader.cpp
+++ b/Engine/Source/Runtime/Experimental/Iris/Core/Private/Iris/ReplicationSystem/ReplicationReader.cpp
@@ -108,7 +108,7 @@ private:
 };
 
 FReplicationReader::FReplicatedObjectInfo::FReplicatedObjectInfo()
-: InternalIndex(FNetHandleManager::InvalidInternalIndex)
+: InternalIndex(FNetRefHandleManager::InvalidInternalIndex)
 , Value(0U)
 {
 	FMemory::Memzero(StoredBaselines);
@@ -120,7 +120,7 @@ FReplicationReader::FReplicationReader()
 : TempLinearAllocator()
 , TempChangeMaskAllocator(&TempLinearAllocator)
 , ReplicationSystemInternal(nullptr)
-, NetHandleManager(nullptr)
+, NetRefHandleManager(nullptr)
 , StateStorage(nullptr)
 , ObjectsToDispatch(nullptr)
 , ObjectsToDispatchCount(0U)
@@ -145,7 +145,7 @@ void FReplicationReader::Init(const FReplicationParameters& InParameters)
 
 	// Cache internal systems
 	ReplicationSystemInternal = Parameters.ReplicationSystem->GetReplicationSystemInternal();
-	NetHandleManager = &ReplicationSystemInternal->GetNetHandleManager();
+	NetRefHandleManager = &ReplicationSystemInternal->GetNetRefHandleManager();
 	StateStorage = &ReplicationSystemInternal->GetReplicationStateStorage();
 	NetBlobHandlerManager = &ReplicationSystemInternal->GetNetBlobHandlerManager();
 	ObjectReferenceCache = &ReplicationSystemInternal->GetObjectReferenceCache();
@@ -179,10 +179,10 @@ void FReplicationReader::Deinit()
 }
 
 // Read incomplete handle
-FNetHandle FReplicationReader::ReadNetHandleId(FNetBitStreamReader& Reader) const
+FNetRefHandle FReplicationReader::ReadNetRefHandleId(FNetBitStreamReader& Reader) const
 {
-	const uint32 NetId = Reader.ReadBits(FNetHandle::IdBits);
-	return FNetHandleManager::MakeNetHandleFromId(NetId);
+	const uint32 NetId = Reader.ReadBits(FNetRefHandle::IdBits);
+	return FNetRefHandleManager::MakeNetRefHandleFromId(NetId);
 }
 	
 uint16 FReplicationReader::ReadObjectsPendingDestroy(FNetSerializationContext& Context)
@@ -198,14 +198,14 @@ uint16 FReplicationReader::ReadObjectsPendingDestroy(FNetSerializationContext& C
 	{
 		for (uint32 It = 0; It < ObjectsToRead; ++It)
 		{
-			UE_NET_TRACE_NAMED_OBJECT_SCOPE(DestroyedObjectScope, FNetHandle(), Reader, Context.GetTraceCollector(), ENetTraceVerbosity::Trace);
+			UE_NET_TRACE_NAMED_OBJECT_SCOPE(DestroyedObjectScope, FNetRefHandle(), Reader, Context.GetTraceCollector(), ENetTraceVerbosity::Trace);
 
-			FNetHandle IncompleteHandle = ReadNetHandleId(Reader);
+			FNetRefHandle IncompleteHandle = ReadNetRefHandleId(Reader);
 			const bool bShouldDestroyInstance = Reader.ReadBool();
 			if (!Reader.IsOverflown())
 			{
 				// Resolve handle and destroy using bridge
-				const uint32 InternalIndex = NetHandleManager->GetInternalIndex(IncompleteHandle);
+				const uint32 InternalIndex = NetRefHandleManager->GetInternalIndex(IncompleteHandle);
 				if (InternalIndex)
 				{
 					UE_NET_TRACE_SET_SCOPE_OBJECTID(DestroyedObjectScope, IncompleteHandle);
@@ -258,7 +258,7 @@ FReplicationReader::FReplicatedObjectInfo& FReplicationReader::StartReplication(
 	// Allocate changemask (if needed)
 	if (InternalIndex != 0U)
 	{		
-		const FNetHandleManager::FReplicatedObjectData& Data = NetHandleManager->GetReplicatedObjectDataNoCheck(InternalIndex);
+		const FNetRefHandleManager::FReplicatedObjectData& Data = NetRefHandleManager->GetReplicatedObjectDataNoCheck(InternalIndex);
 		ObjectInfo.ChangeMaskBitCount = Data.Protocol->ChangeMaskBitCount;
 
 		// Alloc and init changemask
@@ -296,11 +296,11 @@ void FReplicationReader::EndReplication(uint32 InternalIndex, bool bTearOff, boo
 {
 	if (FReplicatedObjectInfo* ObjectInfo = ReplicatedObjects.Find(InternalIndex))
 	{
-		const FNetHandleManager::FReplicatedObjectData& Data = NetHandleManager->GetReplicatedObjectDataNoCheck(InternalIndex);
+		const FNetRefHandleManager::FReplicatedObjectData& Data = NetRefHandleManager->GetReplicatedObjectDataNoCheck(InternalIndex);
 
 		CleanupReferenceTracking(ObjectInfo);
 		Attachments.DropAllAttachments(ENetObjectAttachmentType::Normal, InternalIndex);
-		ReplicationBridge->DestroyNetObjectFromRemote(Data.Handle, bTearOff, bDestroyInstance);
+		ReplicationBridge->DestroyNetObjectFromRemote(Data.RefHandle, bTearOff, bDestroyInstance);
 
 		CleanupObjectData(*ObjectInfo);
 
@@ -308,7 +308,7 @@ void FReplicationReader::EndReplication(uint32 InternalIndex, bool bTearOff, boo
 	}
 }
 
-void FReplicationReader::DeserializeObjectStateDelta(FNetSerializationContext& Context, uint32 InternalIndex, FDispatchObjectInfo& Info, FReplicatedObjectInfo& ObjectInfo, const FNetHandleManager::FReplicatedObjectData& ObjectData, uint32& OutNewBaselineIndex)
+void FReplicationReader::DeserializeObjectStateDelta(FNetSerializationContext& Context, uint32 InternalIndex, FDispatchObjectInfo& Info, FReplicatedObjectInfo& ObjectInfo, const FNetRefHandleManager::FReplicatedObjectData& ObjectData, uint32& OutNewBaselineIndex)
 {
 	FNetBitStreamReader& Reader = *Context.GetBitStreamReader();
 
@@ -376,7 +376,7 @@ void FReplicationReader::ReadObject(FNetSerializationContext& Context)
 	ReadAndVerifySentinelBits(&Reader, TEXT("ReadObject"), 8);
 #endif
 
-	const FNetHandle IncompleteHandle = ReadNetHandleId(Reader);
+	const FNetRefHandle IncompleteHandle = ReadNetRefHandleId(Reader);
 	
 	// Read replicated destroy header if necessary
 	const bool bReadReplicatedDestroyHeader = !IsObjectIndexForOOBAttachment(IncompleteHandle.GetId());
@@ -414,13 +414,13 @@ void FReplicationReader::ReadObject(FNetSerializationContext& Context)
 #endif
 
 		// SubObject data
-		FNetHandle SubObjectOwnerHandle;
+		FNetRefHandle SubObjectOwnerHandle;
 		if (Reader.ReadBool())
 		{
-			const FNetHandle IncompleteOwnerHandle = ReadNetHandleId(Reader);
+			const FNetRefHandle IncompleteOwnerHandle = ReadNetRefHandleId(Reader);
 				
-			FInternalNetHandle SubObjectOwnerInternalIndex = NetHandleManager->GetInternalIndex(IncompleteOwnerHandle);
-			if (Reader.IsOverflown() || SubObjectOwnerInternalIndex == FNetHandleManager::InvalidInternalIndex)
+			FInternalNetRefIndex SubObjectOwnerInternalIndex = NetRefHandleManager->GetInternalIndex(IncompleteOwnerHandle);
+			if (Reader.IsOverflown() || SubObjectOwnerInternalIndex == FNetRefHandleManager::InvalidInternalIndex)
 			{
 				UE_LOG_REPLICATIONREADER_ERROR(TEXT("FReplicationReader::ReadObject Invalid subobjectowner handle. %s"), ToCStr(IncompleteOwnerHandle.ToString()));
 				const FName& NetError = (Reader.IsOverflown() ? GNetError_BitStreamOverflow : GNetError_InvalidNetHandle);
@@ -428,7 +428,7 @@ void FReplicationReader::ReadObject(FNetSerializationContext& Context)
 				return;			
 			}
 
-			SubObjectOwnerHandle = NetHandleManager->GetReplicatedObjectDataNoCheck(SubObjectOwnerInternalIndex).Handle;
+			SubObjectOwnerHandle = NetRefHandleManager->GetReplicatedObjectDataNoCheck(SubObjectOwnerInternalIndex).RefHandle;
 		}
 
 		const bool bIsDeltaCompressed = Reader.ReadBool();
@@ -450,13 +450,13 @@ void FReplicationReader::ReadObject(FNetSerializationContext& Context)
 		// Get Bridge
 		FReplicationBridgeSerializationContext BridgeContext(Context, Parameters.ConnectionId);
 
-		FNetHandle NetHandle = ReplicationBridge->CallCreateNetHandleFromRemote(SubObjectOwnerHandle, IncompleteHandle, BridgeContext);
+		FNetRefHandle NetRefHandle = ReplicationBridge->CallCreateNetRefHandleFromRemote(SubObjectOwnerHandle, IncompleteHandle, BridgeContext);
 
-		if (!NetHandle.IsValid())
+		if (!NetRefHandle.IsValid())
 		{
 #if UE_NET_REPLICATION_SUPPORT_SKIP_INITIAL_STATE
 			// If we support skipping
-			UE_LOG_REPLICATIONREADER_WARNING(TEXT("FReplicationReader::ReadObject Failed to instantiate %s, skipping over it assuming that object was streamed out"), *NetHandle.ToString());
+			UE_LOG_REPLICATIONREADER_WARNING(TEXT("FReplicationReader::ReadObject Failed to instantiate %s, skipping over it assuming that object was streamed out"), *NetRefHandle.ToString());
 			Reader.Seek(SkipSeekPos);
 			return;
 #endif
@@ -467,7 +467,7 @@ void FReplicationReader::ReadObject(FNetSerializationContext& Context)
 			goto ErrorHandling;
 		}
 
-		InternalIndex = NetHandleManager->GetInternalIndex(NetHandle);		 
+		InternalIndex = NetRefHandleManager->GetInternalIndex(NetRefHandle);		 
 		FReplicatedObjectInfo& ObjectInfo = StartReplication(InternalIndex);
 
 		ObjectInfo.bIsDeltaCompressionEnabled = bIsDeltaCompressed;
@@ -482,12 +482,12 @@ void FReplicationReader::ReadObject(FNetSerializationContext& Context)
 		else
 		{
 			// If we get back an invalid internal index then either the object has been deleted or there's bitstream corruption.
-			InternalIndex = NetHandleManager->GetInternalIndex(IncompleteHandle);
+			InternalIndex = NetRefHandleManager->GetInternalIndex(IncompleteHandle);
 
 			// If this is a subobject that is being destroyed this was no error as we send destroy info for unconfirmed objects
 			if (!!(ReplicatedDestroyHeaderFlags & ReplicatedDestroyHeaderFlags_EndReplication))
 			{
-				bHasErrors = InternalIndex == FNetHandleManager::InvalidInternalIndex;
+				bHasErrors = InternalIndex == FNetRefHandleManager::InvalidInternalIndex;
 			}
 		}
 	}
@@ -500,7 +500,7 @@ void FReplicationReader::ReadObject(FNetSerializationContext& Context)
 
 	// Read state data and attachments
 	{
-		const FNetHandleManager::FReplicatedObjectData& ObjectData = NetHandleManager->GetReplicatedObjectDataNoCheck(InternalIndex);
+		const FNetRefHandleManager::FReplicatedObjectData& ObjectData = NetRefHandleManager->GetReplicatedObjectDataNoCheck(InternalIndex);
 
 		// Add entry in our received data as we postpone state application until we have received all data in order to be able to properly resolve references
 		FDispatchObjectInfo& Info = ObjectsToDispatch[ObjectsToDispatchCount];
@@ -549,7 +549,7 @@ void FReplicationReader::ReadObject(FNetSerializationContext& Context)
 			if (NewBaselineIndex != FDeltaCompressionBaselineManager::InvalidBaselineIndex)
 			{
 				// This object uses delta compression, store the last received state as a baseline with the specified index
-				UE_LOG_REPLICATIONREADER(TEXT("Storing new baselineindex: %u for (:%u) Handle %s"), NewBaselineIndex, InternalIndex, *ObjectData.Handle.ToString());
+				UE_LOG_REPLICATIONREADER(TEXT("Storing new baselineindex: %u for (:%u) Handle %s"), NewBaselineIndex, InternalIndex, *ObjectData.RefHandle.ToString());
 
 				check(NewBaselineIndex < FDeltaCompressionBaselineManager::MaxBaselineCount);
 				if (ObjectInfo->StoredBaselines[NewBaselineIndex])
@@ -594,7 +594,7 @@ void FReplicationReader::ReadObject(FNetSerializationContext& Context)
 				}
 			}
 
-			Attachments.Deserialize(Context, AttachmentType, InternalIndex, ObjectData.Handle);
+			Attachments.Deserialize(Context, AttachmentType, InternalIndex, ObjectData.RefHandle);
 		}
 
 		bHasErrors = bHasErrors || Context.HasErrorOrOverflow();
@@ -637,7 +637,7 @@ void FReplicationReader::UpdateObjectReferenceTracking(FReplicatedObjectInfo* Re
 		// Try to avoid dynamic allocations during the update of the UnresolvedObjectReferences.
 		ReplicationInfo->UnresolvedObjectReferences.Reserve(ReplicationInfo->UnresolvedObjectReferences.Num() + NewUnresolvedReferences.Num());
 
-		TSet<FNetHandle> OldUnresolvedSet;
+		TSet<FNetRefHandle> OldUnresolvedSet;
 		OldUnresolvedSet.Reserve(ReplicationInfo->UnresolvedObjectReferences.Num());
 		for (const FObjectReferenceTracker::ElementType& Element : ReplicationInfo->UnresolvedObjectReferences)
 		{
@@ -655,8 +655,8 @@ void FReplicationReader::UpdateObjectReferenceTracking(FReplicatedObjectInfo* Re
 
 			for (auto It = NewUnresolvedReferences.CreateConstKeyIterator(ChangeBit); It; ++It)
 			{
-				const FNetHandle NetHandle = It.Value();
-				UnresolvedObjectReferences.AddByHash(KeyHash, ChangeBit, NetHandle);
+				const FNetRefHandle RefHandle = It.Value();
+				UnresolvedObjectReferences.AddByHash(KeyHash, ChangeBit, RefHandle);
 			}
 		};
 
@@ -667,7 +667,7 @@ void FReplicationReader::UpdateObjectReferenceTracking(FReplicatedObjectInfo* Re
 		}
 
 		// The unresolved set is now updated with the current status of unresolved references.
-		TSet<FNetHandle> NewUnresolvedSet;
+		TSet<FNetRefHandle> NewUnresolvedSet;
 		NewUnresolvedSet.Reserve(NewUnresolvedReferences.Num());
 		for (const FObjectReferenceTracker::ElementType& Element : ReplicationInfo->UnresolvedObjectReferences)
 		{
@@ -680,24 +680,24 @@ void FReplicationReader::UpdateObjectReferenceTracking(FReplicatedObjectInfo* Re
 
 		// Remove resolved or no longer existing references
 		const uint32 OwnerInternalIndex = ReplicationInfo->InternalIndex;
-		for (FNetHandle Handle : OldUnresolvedSet)
+		for (FNetRefHandle Handle : OldUnresolvedSet)
 		{
 			if (!NewUnresolvedSet.Contains(Handle))
 			{
 				// Remove from tracking
 				UnresolvedHandleToDependents.RemoveSingle(Handle, OwnerInternalIndex);
-				UE_LOG(LogIris, Verbose, TEXT("FReplicationReader::UpdateObjectReferenceTracking Removing unresolved reference %s for %s"), ToCStr(Handle.ToString()), ToCStr(NetHandleManager->GetNetHandleFromInternalIndex(OwnerInternalIndex).ToString()));
+				UE_LOG(LogIris, Verbose, TEXT("FReplicationReader::UpdateObjectReferenceTracking Removing unresolved reference %s for %s"), ToCStr(Handle.ToString()), ToCStr(NetRefHandleManager->GetNetRefHandleFromInternalIndex(OwnerInternalIndex).ToString()));
 			}
 		}
 
 		// Add new unresolved references
-		for (FNetHandle Handle : NewUnresolvedSet)
+		for (FNetRefHandle Handle : NewUnresolvedSet)
 		{
 			if (!OldUnresolvedSet.Contains(Handle))
 			{
 				// Add to tracking
 				UnresolvedHandleToDependents.Add(Handle, OwnerInternalIndex);
-				UE_LOG(LogIris, Verbose, TEXT("FReplicationReader::UpdateObjectReferenceTracking Adding unresolved reference %s for %s"), ToCStr(Handle.ToString()), ToCStr(NetHandleManager->GetNetHandleFromInternalIndex(OwnerInternalIndex).ToString()));
+				UE_LOG(LogIris, Verbose, TEXT("FReplicationReader::UpdateObjectReferenceTracking Adding unresolved reference %s for %s"), ToCStr(Handle.ToString()), ToCStr(NetRefHandleManager->GetNetRefHandleFromInternalIndex(OwnerInternalIndex).ToString()));
 			}
 		}
 	}
@@ -708,7 +708,7 @@ void FReplicationReader::UpdateObjectReferenceTracking(FReplicatedObjectInfo* Re
 		// Try to avoid dynamic allocations during the update of the ResolvedDynamicObjectReferences.
 		ReplicationInfo->ResolvedDynamicObjectReferences.Reserve(ReplicationInfo->ResolvedDynamicObjectReferences.Num() + NewMappedDynamicReferences.Num());
 
-		TSet<FNetHandle> OldResolvedSet;
+		TSet<FNetRefHandle> OldResolvedSet;
 		OldResolvedSet.Reserve(ReplicationInfo->ResolvedDynamicObjectReferences.Num());
 		for (const FObjectReferenceTracker::ElementType& Element : ReplicationInfo->ResolvedDynamicObjectReferences)
 		{
@@ -725,8 +725,8 @@ void FReplicationReader::UpdateObjectReferenceTracking(FReplicatedObjectInfo* Re
 
 			for (auto It = NewMappedDynamicReferences.CreateConstKeyIterator(ChangeBit); It; ++It)
 			{
-				const FNetHandle NetHandle = It.Value();
-				ResolvedDynamicObjectReferences.AddByHash(KeyHash, ChangeBit, NetHandle);
+				const FNetRefHandle RefHandle = It.Value();
+				ResolvedDynamicObjectReferences.AddByHash(KeyHash, ChangeBit, RefHandle);
 			}
 		};
 
@@ -735,7 +735,7 @@ void FReplicationReader::UpdateObjectReferenceTracking(FReplicatedObjectInfo* Re
 		// the initial state has already been applied.
 
 		// The resolved set is now updated with the current status of resolved references.
-		TSet<FNetHandle> NewResolvedSet;
+		TSet<FNetRefHandle> NewResolvedSet;
 		NewResolvedSet.Reserve(ReplicationInfo->ResolvedDynamicObjectReferences.Num());
 		for (const FObjectReferenceTracker::ElementType& Element : ReplicationInfo->ResolvedDynamicObjectReferences)
 		{
@@ -744,55 +744,55 @@ void FReplicationReader::UpdateObjectReferenceTracking(FReplicatedObjectInfo* Re
 
 		// Remove now unresolved or no longer existing references
 		const uint32 OwnerInternalIndex = ReplicationInfo->InternalIndex;
-		for (FNetHandle Handle : OldResolvedSet)
+		for (FNetRefHandle Handle : OldResolvedSet)
 		{
 			if (!NewResolvedSet.Contains(Handle))
 			{
 				// Remove from tracking
 				ResolvedDynamicHandleToDependents.RemoveSingle(Handle, OwnerInternalIndex);
-				UE_LOG(LogIris, Verbose, TEXT("FReplicationReader::UpdateObjectReferenceTracking Removing resolved dynamic reference %s for %s"), ToCStr(Handle.ToString()), ToCStr(NetHandleManager->GetNetHandleFromInternalIndex(OwnerInternalIndex).ToString()));
+				UE_LOG(LogIris, Verbose, TEXT("FReplicationReader::UpdateObjectReferenceTracking Removing resolved dynamic reference %s for %s"), ToCStr(Handle.ToString()), ToCStr(NetRefHandleManager->GetNetHandleFromInternalIndex(OwnerInternalIndex).ToString()));
 			}
 		}
 
 		// Add new resolved dynamic references
-		for (FNetHandle Handle : NewResolvedSet)
+		for (FNetRefHandle Handle : NewResolvedSet)
 		{
 			if (Handle.IsDynamic() && !OldResolvedSet.Contains(Handle))
 			{
 				// Add to tracking
 				ResolvedDynamicHandleToDependents.Add(Handle, OwnerInternalIndex);
-				UE_LOG(LogIris, Verbose, TEXT("FReplicationReader::UpdateObjectReferenceTracking Adding resolved dynamic reference %s for %s"), ToCStr(Handle.ToString()), ToCStr(NetHandleManager->GetNetHandleFromInternalIndex(OwnerInternalIndex).ToString()));
+				UE_LOG(LogIris, Verbose, TEXT("FReplicationReader::UpdateObjectReferenceTracking Adding resolved dynamic reference %s for %s"), ToCStr(Handle.ToString()), ToCStr(NetRefHandleManager->GetNetHandleFromInternalIndex(OwnerInternalIndex).ToString()));
 			}
 		}
 	}
 #endif
 }
 
-void FReplicationReader::RemoveUnresolvedObjectReferenceInReplicationInfo(FReplicatedObjectInfo* ReplicationInfo, FNetHandle Handle)
+void FReplicationReader::RemoveUnresolvedObjectReferenceInReplicationInfo(FReplicatedObjectInfo* ReplicationInfo, FNetRefHandle Handle)
 {
 	for (FObjectReferenceTracker::TIterator It = ReplicationInfo->UnresolvedObjectReferences.CreateIterator(); It; ++It)
 	{
-		const FNetHandle NetHandle = It->Value;
-		if (NetHandle == Handle)
+		const FNetRefHandle RefHandle = It->Value;
+		if (RefHandle == Handle)
 		{
 			It.RemoveCurrent();
 		}
 	}
 }
 
-void FReplicationReader::RemoveResolvedObjectReferenceInReplicationInfo(FReplicatedObjectInfo* ReplicationInfo, FNetHandle Handle)
+void FReplicationReader::RemoveResolvedObjectReferenceInReplicationInfo(FReplicatedObjectInfo* ReplicationInfo, FNetRefHandle Handle)
 {
 	for (FObjectReferenceTracker::TIterator It = ReplicationInfo->ResolvedDynamicObjectReferences.CreateIterator(); It; ++It)
 	{
-		const FNetHandle NetHandle = It->Value;
-		if (NetHandle == Handle)
+		const FNetRefHandle RefHandle = It->Value;
+		if (RefHandle == Handle)
 		{
 			It.RemoveCurrent();
 		}
 	}
 }
 
-bool FReplicationReader::MoveResolvedObjectReferenceToUnresolvedInReplicationInfo(FReplicatedObjectInfo* ReplicationInfo, FNetHandle UnresolvableHandle)
+bool FReplicationReader::MoveResolvedObjectReferenceToUnresolvedInReplicationInfo(FReplicatedObjectInfo* ReplicationInfo, FNetRefHandle UnresolvableHandle)
 {
 	bool bFoundHandle = false;
 	bool bHasUnresolvedReferences = ReplicationInfo->bHasUnresolvedReferences;
@@ -801,8 +801,8 @@ bool FReplicationReader::MoveResolvedObjectReferenceToUnresolvedInReplicationInf
 	FObjectReferenceTracker& UnresolvedObjectReferences = ReplicationInfo->UnresolvedObjectReferences;
 	for (FObjectReferenceTracker::TIterator It = ReplicationInfo->ResolvedDynamicObjectReferences.CreateIterator(); It; ++It)
 	{
-		const FNetHandle NetHandle = It->Value;
-		if (NetHandle == UnresolvableHandle)
+		const FNetRefHandle RefHandle = It->Value;
+		if (RefHandle == UnresolvableHandle)
 		{
 			bFoundHandle = true;
 
@@ -824,7 +824,7 @@ bool FReplicationReader::MoveResolvedObjectReferenceToUnresolvedInReplicationInf
 			// already exist in the UnresolvedObjectReferences map, so no need to call AddUnique.
 			UnresolvedObjectReferences.Add(ChangemaskOffset, UnresolvableHandle);
 
-			UE_LOG(LogIris, Verbose, TEXT("FReplicationReader::MoveResolvedObjectReferenceToUnresolvedInReplicationInfo Moving from resolved to unresolved reference %s for %s"), ToCStr(UnresolvableHandle.ToString()), ToCStr(NetHandleManager->GetNetHandleFromInternalIndex(ReplicationInfo->InternalIndex).ToString()));
+			UE_LOG(LogIris, Verbose, TEXT("FReplicationReader::MoveResolvedObjectReferenceToUnresolvedInReplicationInfo Moving from resolved to unresolved reference %s for %s"), ToCStr(UnresolvableHandle.ToString()), ToCStr(NetRefHandleManager->GetNetRefHandleFromInternalIndex(ReplicationInfo->InternalIndex).ToString()));
 		}
 	}
 
@@ -843,9 +843,9 @@ void FReplicationReader::CleanupReferenceTracking(FReplicatedObjectInfo* ObjectI
 	for (FObjectReferenceTracker::ElementType Element : ObjectInfo->UnresolvedObjectReferences)
 	{
 		// Remove from tracking
-		FNetHandle Handle = Element.Value;
+		FNetRefHandle Handle = Element.Value;
 		UnresolvedHandleToDependents.RemoveSingle(Handle, ObjectIndex);
-		UE_LOG(LogIris, Verbose, TEXT("FReplicationReader::CleanupReferenceTracking Removing unresolved reference %s for %s"), *Handle.ToString(), *(NetHandleManager->GetNetHandleFromInternalIndex(ObjectIndex).ToString()));
+		UE_LOG(LogIris, Verbose, TEXT("FReplicationReader::CleanupReferenceTracking Removing unresolved reference %s for %s"), *Handle.ToString(), *(NetRefHandleManager->GetNetRefHandleFromInternalIndex(ObjectIndex).ToString()));
 	}
 	ObjectInfo->UnresolvedObjectReferences.Reset();
 
@@ -853,9 +853,9 @@ void FReplicationReader::CleanupReferenceTracking(FReplicatedObjectInfo* ObjectI
 	for (FObjectReferenceTracker::ElementType Element : ObjectInfo->ResolvedDynamicObjectReferences)
 	{
 		// Remove from tracking
-		const FNetHandle Handle = Element.Value;
+		const FNetRefHandle Handle = Element.Value;
 		ResolvedDynamicHandleToDependents.RemoveSingle(Handle, ObjectIndex);
-		UE_LOG(LogIris, Verbose, TEXT("FReplicationReader::CleanupReferenceTracking Removing resolved dynamic reference %s for %s"), *Handle.ToString(), *(NetHandleManager->GetNetHandleFromInternalIndex(ObjectIndex).ToString()));
+		UE_LOG(LogIris, Verbose, TEXT("FReplicationReader::CleanupReferenceTracking Removing resolved dynamic reference %s for %s"), *Handle.ToString(), *(NetRefHandleManager->GetNetRefHandleFromInternalIndex(ObjectIndex).ToString()));
 	}
 	ObjectInfo->ResolvedDynamicObjectReferences.Reset();
 
@@ -911,7 +911,7 @@ void FReplicationReader::ResolveAndDispatchUnresolvedReferencesForObject(FNetSer
 	ENetObjectAttachmentDispatchFlags AttachmentDispatchedFlags = ENetObjectAttachmentDispatchFlags::None;
 	if (bObjectHasReferences)
 	{
-		const FNetHandleManager::FReplicatedObjectData& ObjectData = NetHandleManager->GetReplicatedObjectDataNoCheck(ReplicationInfo->InternalIndex);
+		const FNetRefHandleManager::FReplicatedObjectData& ObjectData = NetRefHandleManager->GetReplicatedObjectDataNoCheck(ReplicationInfo->InternalIndex);
 		const uint32 ChangeMaskBitCount = ReplicationInfo->ChangeMaskBitCount;
 		
 		// Try to resolve references and collect unresolved references
@@ -980,7 +980,7 @@ void FReplicationReader::ResolveAndDispatchUnresolvedReferencesForObject(FNetSer
 		{
 			// $IRIS: $TODO: Figure out how to handle this, currently we do not crash but we probably want to
 			// handle this properly by accumulating changemask for later instantiation
-			UE_LOG_REPLICATIONREADER_WARNING(TEXT("Cannot dispatch state data for not instantiated %s"), *ObjectData.Handle.ToString());
+			UE_LOG_REPLICATIONREADER_WARNING(TEXT("Cannot dispatch state data for not instantiated %s"), *ObjectData.RefHandle.ToString());
 		}
 	}
 
@@ -1044,8 +1044,8 @@ void FReplicationReader::DispatchStateData(FNetSerializationContext& Context)
 			// If object was created this frame it`s initial state is now applied
 			if (Info.bHasState && Info.bIsInitialState)
 			{
-				const FNetHandleManager::FReplicatedObjectData& ObjectData = NetHandleManager->GetReplicatedObjectDataNoCheck(Info.InternalIndex);
-				ReplicationBridge->CallPostApplyInitialState(ObjectData.Handle);
+				const FNetRefHandleManager::FReplicatedObjectData& ObjectData = NetRefHandleManager->GetReplicatedObjectDataNoCheck(Info.InternalIndex);
+				ReplicationBridge->CallPostApplyInitialState(ObjectData.RefHandle);
 			}
 
 			// Dispatch attachment and enqueue for later resolving
@@ -1067,7 +1067,7 @@ void FReplicationReader::DispatchStateData(FNetSerializationContext& Context)
 	};
 
 	// In order to properly execute legacy callbacks we need to batch apply state data for owner/subobjects
-	FInternalNetHandle LastDispatchedRootInternalIndex = 0U;
+	FInternalNetRefIndex LastDispatchedRootInternalIndex = 0U;
 	
 	// Dispatch and apply received state data
 	for (FDispatchObjectInfo& Info : MakeArrayView(ObjectsToDispatch, ObjectsToDispatchCount))
@@ -1083,10 +1083,10 @@ void FReplicationReader::DispatchStateData(FNetSerializationContext& Context)
 		// If we have any object references we want to update any unresolved ones, including previously unresolved references
 		if (Info.bHasState)
 		{
-			const FNetHandleManager::FReplicatedObjectData& ObjectData = NetHandleManager->GetReplicatedObjectDataNoCheck(Info.InternalIndex);
+			const FNetRefHandleManager::FReplicatedObjectData& ObjectData = NetRefHandleManager->GetReplicatedObjectDataNoCheck(Info.InternalIndex);
 
 			// We only need to flush if we are switching to a new root object with state data
-			const FInternalNetHandle RootInternalIndex = ObjectData.SubObjectRootIndex == FNetHandleManager::InvalidInternalIndex ? Info.InternalIndex : ObjectData.SubObjectRootIndex;
+			const FInternalNetRefIndex RootInternalIndex = ObjectData.SubObjectRootIndex == FNetRefHandleManager::InvalidInternalIndex ? Info.InternalIndex : ObjectData.SubObjectRootIndex;
 			if (RootInternalIndex != LastDispatchedRootInternalIndex)
 			{
 				FlushPostDispatchForBatch();
@@ -1141,7 +1141,7 @@ void FReplicationReader::DispatchStateData(FNetSerializationContext& Context)
 			{
 				// $IRIS: $TODO: Figure out how to handle this, currently we do not crash but we probably want to
 				// handle this properly by accumulating changemask for later instantiation
-				UE_LOG_REPLICATIONREADER_WARNING(TEXT("Cannot dispatch state data for not instantiated %s"), *(ObjectData.Handle.ToString()));
+				UE_LOG_REPLICATIONREADER_WARNING(TEXT("Cannot dispatch state data for not instantiated %s"), *(ObjectData.RefHandle.ToString()));
 			}
 		}
 
@@ -1170,14 +1170,14 @@ void FReplicationReader::ResolveAndDispatchUnresolvedReferences()
 	Context.SetInternalContext(&InternalContext);
 
 	// Currently we brute force this by iterating over all handles pending resolve and update all objects pending resolve
-	TArray<FNetHandle> UpdatedHandles;
+	TArray<FNetRefHandle> UpdatedHandles;
 	UpdatedHandles.Reserve(128);
 	UnresolvedHandleToDependents.GenerateKeyArray(UpdatedHandles);
 	
 	TSet<uint32> InternalObjectsToResolve;
 	InternalObjectsToResolve.Reserve(UnresolvedHandleToDependents.Num());
 
-	for (FNetHandle Handle : UpdatedHandles)
+	for (FNetRefHandle Handle : UpdatedHandles)
 	{
 		// Only make sense to update dependant objects if handle is resolvable
 		if (ObjectReferenceCache->ResolveObjectReferenceHandle(Handle, ResolveContext) != nullptr)
@@ -1212,11 +1212,11 @@ void FReplicationReader::UpdateUnresolvableReferenceTracking()
 	TArray<uint32, TInlineAllocator<AssumedMaxDependentCount>> Dependents;
 
 	// Naively go through every object pending destroy, see if it's dynamic and update dependent's unresolved tracking
-	TArrayView<const FInternalNetHandle> ObjectsPendingDestroy = NetHandleManager->GetObjectsPendingDestroy();
-	for (const FInternalNetHandle InternalIndex : ObjectsPendingDestroy)
+	TArrayView<const FInternalNetRefIndex> ObjectsPendingDestroy = NetRefHandleManager->GetObjectsPendingDestroy();
+	for (const FInternalNetRefIndex InternalIndex : ObjectsPendingDestroy)
 	{
-		const FNetHandleManager::FReplicatedObjectData& ObjectData = NetHandleManager->GetReplicatedObjectDataNoCheck(InternalIndex);
-		const FNetHandle DestroyedHandle = ObjectData.Handle;
+		const FNetRefHandleManager::FReplicatedObjectData& ObjectData = NetRefHandleManager->GetReplicatedObjectDataNoCheck(InternalIndex);
+		const FNetRefHandle DestroyedHandle = ObjectData.RefHandle;
 		if (!DestroyedHandle.IsDynamic())
 		{
 			continue;
diff --git a/Engine/Source/Runtime/Experimental/Iris/Core/Private/Iris/ReplicationSystem/ReplicationReader.h b/Engine/Source/Runtime/Experimental/Iris/Core/Private/Iris/ReplicationSystem/ReplicationReader.h
--- a/Engine/Source/Runtime/Experimental/Iris/Core/Private/Iris/ReplicationSystem/ReplicationReader.h
+++ b/Engine/Source/Runtime/Experimental/Iris/Core/Private/Iris/ReplicationSystem/ReplicationReader.h
@@ -23,7 +23,7 @@ namespace UE::Net
 		class FReplicationSystemInternal;
 		class FResolveAndCollectUnresolvedAndResolvedReferenceCollector;
 		class FNetBlobHandlerManager;
-		class FNetHandleManager;
+		class FNetRefHandleManager;
 	}
 }
 
@@ -53,7 +53,7 @@ public:
 	void Read(FNetSerializationContext& Context);
 
 	// Read index part of handle
-	FNetHandle ReadNetHandleId(FNetBitStreamReader& Reader) const;
+	FNetRefHandle ReadNetRefHandleId(FNetBitStreamReader& Reader) const;
 
 	void SetRemoteNetTokenStoreState(FNetTokenStoreState* RemoteTokenStoreState);
 	
@@ -61,12 +61,12 @@ public:
 	void UpdateUnresolvableReferenceTracking();
 
 private:
-	// ChangeMaskOffset -> FNetHandle
+	// ChangeMaskOffset -> FNetRefHandle
 	enum EConstants : uint32
 	{
 		FakeInitChangeMaskOffset = 0xFFFFFFFFU,
 	};
-	typedef TMultiMap<uint32, FNetHandle> FObjectReferenceTracker;
+	typedef TMultiMap<uint32, FNetRefHandle> FObjectReferenceTracker;
 
 	struct FReplicatedObjectInfo
 	{
@@ -179,14 +179,14 @@ private:
 	// Update ReplicationInfo and OutUnresolvedChangeMask based on data collected by the Collector
 	void BuildUnresolvedChangeMaskAndUpdateObjectReferenceTracking(const FResolveAndCollectUnresolvedAndResolvedReferenceCollector& Collector, FNetBitArrayView CollectorChangeMask, FReplicatedObjectInfo* ReplicationInfo, FNetBitArrayView& OutUnresolvedChangeMask);
 
-	void RemoveUnresolvedObjectReferenceInReplicationInfo(FReplicatedObjectInfo* ReplicationInfo, FNetHandle Handle);
-	void RemoveResolvedObjectReferenceInReplicationInfo(FReplicatedObjectInfo* ReplicationInfo, FNetHandle Handle);
+	void RemoveUnresolvedObjectReferenceInReplicationInfo(FReplicatedObjectInfo* ReplicationInfo, FNetRefHandle Handle);
+	void RemoveResolvedObjectReferenceInReplicationInfo(FReplicatedObjectInfo* ReplicationInfo, FNetRefHandle Handle);
 
 	// A previously resolved dynamic reference is now unresolvable. The ReplicationInfo needs to be updated to reflect this.
 	// Returns true if the reference was found.
-	bool MoveResolvedObjectReferenceToUnresolvedInReplicationInfo(FReplicatedObjectInfo* ReplicationInfo, FNetHandle UnresolvableHandle);
+	bool MoveResolvedObjectReferenceToUnresolvedInReplicationInfo(FReplicatedObjectInfo* ReplicationInfo, FNetRefHandle UnresolvableHandle);
 
-	void DeserializeObjectStateDelta(FNetSerializationContext& Context, uint32 InternalIndex, FDispatchObjectInfo& Info, FReplicatedObjectInfo& ObjectInfo, const FNetHandleManager::FReplicatedObjectData& ObjectData, uint32& OutNewBaselineIndex);
+	void DeserializeObjectStateDelta(FNetSerializationContext& Context, uint32 InternalIndex, FDispatchObjectInfo& Info, FReplicatedObjectInfo& ObjectInfo, const FNetRefHandleManager::FReplicatedObjectData& ObjectData, uint32& OutNewBaselineIndex);
 
 private:
 	FReplicationParameters Parameters;
@@ -197,7 +197,7 @@ private:
 	FGlobalChangeMaskAllocator PersistentChangeMaskAllocator;
 
 	FReplicationSystemInternal* ReplicationSystemInternal;
-	FNetHandleManager* NetHandleManager;
+	FNetRefHandleManager* NetRefHandleManager;
 	FReplicationStateStorage* StateStorage;
 	UReplicationBridge* ReplicationBridge;
 
@@ -215,10 +215,10 @@ private:
 	TArray<uint32> ObjectsWithAttachmentPendingResolve;
 
 	// Track all objects waiting for this handle to be resolvable
-	TMultiMap<FNetHandle, uint32> UnresolvedHandleToDependents;
+	TMultiMap<FNetRefHandle, uint32> UnresolvedHandleToDependents;
 	
 	// Track all objects with a dynamic handle in case it becomes unresolvable
-	TMultiMap<FNetHandle, uint32> ResolvedDynamicHandleToDependents;
+	TMultiMap<FNetRefHandle, uint32> ResolvedDynamicHandleToDependents;
 
 	FNetBlobHandlerManager* NetBlobHandlerManager;
 	FNetBlobType NetObjectBlobType;
