Iris - Implemented globally unique NetHandles.

diff --git a/Engine/Source/Runtime/Experimental/Iris/Core/Private/Iris/ReplicationSystem/ReplicationSystem.cpp b/Engine/Source/Runtime/Experimental/Iris/Core/Private/Iris/ReplicationSystem/ReplicationSystem.cpp
--- a/Engine/Source/Runtime/Experimental/Iris/Core/Private/Iris/ReplicationSystem/ReplicationSystem.cpp
+++ b/Engine/Source/Runtime/Experimental/Iris/Core/Private/Iris/ReplicationSystem/ReplicationSystem.cpp
@@ -77,20 +77,25 @@ public:
 
 		const uint32 ReplicationSystemId = ReplicationSystem->GetId();
 
-		// $IRIS TODO: Need object ID range. Currently abusing hardcoded values from FNetHandleManager
-		FDirtyNetObjectTrackerInitParams DirtyNetObjectTrackerInitParams = {};
-		DirtyNetObjectTrackerInitParams.ReplicationSystemId = ReplicationSystemId;
-		DirtyNetObjectTrackerInitParams.NetObjectIndexRangeStart = 1;
-		DirtyNetObjectTrackerInitParams.NetObjectIndexRangeEnd = ReplicationSystemInternal.GetNetHandleManager().GetMaxActiveObjectCount() - 1U;
+		FNetRefHandleManager& NetRefHandleManager = ReplicationSystemInternal.GetNetRefHandleManager();
+
+		// $IRIS TODO: Need object ID range. Currently abusing hardcoded values from FNetRefHandleManager
 		FDirtyNetObjectTracker& DirtyNetObjectTracker = ReplicationSystemInternal.GetDirtyNetObjectTracker();
+		FDirtyNetObjectTrackerInitParams DirtyNetObjectTrackerInitParams;
+		{
+			DirtyNetObjectTrackerInitParams.NetRefHandleManager = &NetRefHandleManager;
+			DirtyNetObjectTrackerInitParams.ReplicationSystemId = ReplicationSystemId;
+			DirtyNetObjectTrackerInitParams.NetObjectIndexRangeStart = 1;
+			DirtyNetObjectTrackerInitParams.NetObjectIndexRangeEnd = NetRefHandleManager.GetMaxActiveObjectCount() - 1U;
 
-		DirtyNetObjectTracker.Init(DirtyNetObjectTrackerInitParams);
+			DirtyNetObjectTracker.Init(DirtyNetObjectTrackerInitParams);
+		}
 
 		FReplicationStateStorage& StateStorage = ReplicationSystemInternal.GetReplicationStateStorage();
 		{
 			FReplicationStateStorageInitParams InitParams;
 			InitParams.ReplicationSystem = ReplicationSystem;
-			InitParams.NetHandleManager = &ReplicationSystemInternal.GetNetHandleManager();
+			InitParams.NetRefHandleManager = &NetRefHandleManager;
 			InitParams.MaxObjectCount = DirtyNetObjectTrackerInitParams.NetObjectIndexRangeEnd + 1U;
 			InitParams.MaxConnectionCount = ReplicationSystemInternal.GetConnections().GetMaxConnectionCount();
 			InitParams.MaxDeltaCompressedObjectCount = Params.MaxDeltaCompressedObjectCount;
@@ -133,7 +138,7 @@ public:
 			FDeltaCompressionBaselineManagerInitParams InitParams;
 			InitParams.BaselineInvalidationTracker = &DeltaCompressionBaselineInvalidationTracker;
 			InitParams.Connections = &ReplicationSystemInternal.GetConnections();
-			InitParams.NetHandleManager = &ReplicationSystemInternal.GetNetHandleManager();
+			InitParams.NetRefHandleManager = &NetRefHandleManager;
 			InitParams.ReplicationStateStorage = &StateStorage;
 			InitParams.MaxObjectCount = DirtyNetObjectTrackerInitParams.NetObjectIndexRangeEnd + 1U;
 			InitParams.MaxDeltaCompressedObjectCount = Params.MaxDeltaCompressedObjectCount;
@@ -146,7 +151,7 @@ public:
 			FReplicationFilteringInitParams InitParams;
 			InitParams.ReplicationSystem = ReplicationSystem;
 			InitParams.Connections = &ReplicationSystemInternal.GetConnections();
-			InitParams.NetHandleManager = &ReplicationSystemInternal.GetNetHandleManager();
+			InitParams.NetRefHandleManager = &NetRefHandleManager;
 			InitParams.Groups = &Groups;
 			InitParams.BaselineInvalidationTracker = &ReplicationSystemInternal.GetDeltaCompressionBaselineInvalidationTracker();
 			InitParams.MaxObjectCount = DirtyNetObjectTrackerInitParams.NetObjectIndexRangeEnd + 1U;
@@ -159,7 +164,7 @@ public:
 		FReplicationConditionals& ReplicationConditionals = ReplicationSystemInternal.GetConditionals();
 		{
 			FReplicationConditionalsInitParams InitParams = {};
-			InitParams.NetHandleManager = &ReplicationSystemInternal.GetNetHandleManager();
+			InitParams.NetRefHandleManager = &NetRefHandleManager;
 			InitParams.ReplicationConnections = &ReplicationSystemInternal.GetConnections();
 			InitParams.ReplicationFiltering = &ReplicationFiltering;
 			InitParams.NetObjectGroups = &Groups;
@@ -174,7 +179,7 @@ public:
 			FReplicationPrioritizationInitParams InitParams;
 			InitParams.ReplicationSystem = ReplicationSystem;
 			InitParams.Connections = &ReplicationSystemInternal.GetConnections();
-			InitParams.NetHandleManager = &ReplicationSystemInternal.GetNetHandleManager();
+			InitParams.NetRefHandleManager = &NetRefHandleManager;
 			InitParams.MaxObjectCount = DirtyNetObjectTrackerInitParams.NetObjectIndexRangeEnd + 1U;
 			ReplicationPrioritization.Init(InitParams);
 		}
@@ -203,14 +208,14 @@ public:
 		ReplicationSystemInternal.GetReplicationBridge()->Deinitialize();
 	}
 
-	/**
-	 * Due to objects having been marked as dirty and later removed we must make sure
-	 * that all dirty objects are still in scope.
-	 */
 	void UpdateDirtyObjectList()
 	{
-		FNetBitArrayView DirtyObjects = ReplicationSystemInternal.GetDirtyNetObjectTracker().GetDirtyNetObjects();
-		const FNetBitArrayView ScopableInternalIndices = MakeNetBitArrayView(ReplicationSystemInternal.GetNetHandleManager().GetScopableInternalIndices());
+		FDirtyNetObjectTracker& DirtyNetObjectTracker = ReplicationSystemInternal.GetDirtyNetObjectTracker();
+		DirtyNetObjectTracker.UpdateDirtyNetObjects();
+		FNetBitArrayView DirtyObjects = DirtyNetObjectTracker.GetDirtyNetObjects();
+
+		// Due to objects having been marked as dirty and later removed we must make sure that all dirty objects are still in scope.
+		const FNetBitArrayView ScopableInternalIndices = MakeNetBitArrayView(ReplicationSystemInternal.GetNetRefHandleManager().GetScopableInternalIndices());
 		DirtyObjects.Combine(ScopableInternalIndices, FNetBitArrayView::AndOp);
 	}
 
@@ -282,7 +287,7 @@ public:
 		IRIS_PROFILER_SCOPE(FReplicationSystem_CopyDirtyStateData);
 		LLM_SCOPE_BYTAG(IrisState);
 
-		FNetHandleManager& NetHandleManager = ReplicationSystemInternal.GetNetHandleManager();
+		FNetRefHandleManager& NetRefHandleManager = ReplicationSystemInternal.GetNetRefHandleManager();
 		FChangeMaskCache& Cache = ReplicationSystemInternal.GetChangeMaskCache();
 
 		const uint32 ReplicationSystemId = ReplicationSystem->GetId();
@@ -302,9 +307,9 @@ public:
 
 		SerializationContext.SetInternalContext(&InternalContext);
 
-		auto CopyFunction = [&ChangeMaskWriter, &Cache, &NetHandleManager, &CopiedObjectCount, &SerializationContext, ReplicationSystemId](uint32 DirtyIndex)
+		auto CopyFunction = [&ChangeMaskWriter, &Cache, &NetRefHandleManager, &CopiedObjectCount, &SerializationContext, ReplicationSystemId](uint32 DirtyIndex)
 		{
-			CopiedObjectCount += FReplicationInstanceOperationsInternal::CopyObjectStateData(ChangeMaskWriter, Cache, NetHandleManager, SerializationContext, DirtyIndex);
+			CopiedObjectCount += FReplicationInstanceOperationsInternal::CopyObjectStateData(ChangeMaskWriter, Cache, NetRefHandleManager, SerializationContext, DirtyIndex);
 		};
 
 		FDirtyNetObjectTracker& DirtyNetObjectTracker = ReplicationSystemInternal.GetDirtyNetObjectTracker();
@@ -351,7 +356,7 @@ public:
 			Params.ReplicationSystem = ReplicationSystem;
 			Params.PacketSendWindowSize = 256;
 			Params.ConnectionId = ConnectionId;
-			Params.MaxActiveReplicatedObjectCount = ReplicationSystemInternal.GetNetHandleManager().GetMaxActiveObjectCount();
+			Params.MaxActiveReplicatedObjectCount = ReplicationSystemInternal.GetNetRefHandleManager().GetMaxActiveObjectCount();
 			/** 
 			  * Currently we expect all objects to be replicated from server to client.
 			  * That means we will have to support sending attachments such as RPCs from
@@ -515,7 +520,7 @@ void UReplicationSystem::PreSendUpdate(float DeltaSeconds)
 
 	if (bAllowObjectReplication)
 	{
-		UE_NET_TRACE_FRAME_STATSCOUNTER(GetId(), ReplicationSystem.ReplicatedObjectCount, InternalSys.GetNetHandleManager().GetActiveObjectCount(), ENetTraceVerbosity::Verbose);
+		UE_NET_TRACE_FRAME_STATSCOUNTER(GetId(), ReplicationSystem.ReplicatedObjectCount, InternalSys.GetNetRefHandleManager().GetActiveObjectCount(), ENetTraceVerbosity::Verbose);
 
 		// Refresh dirty object list
 		Impl->UpdateDirtyObjectList();
@@ -566,7 +571,7 @@ void UReplicationSystem::PreSendUpdate(float DeltaSeconds)
 	// Destroy objects pending destroy
 	{
 		Impl->UpdateUnresolvableReferenceTracking();
-		InternalSys.GetNetHandleManager().DestroyObjectsPendingDestroy();
+		InternalSys.GetNetRefHandleManager().DestroyObjectsPendingDestroy();
 	}
 }
 
@@ -588,7 +593,7 @@ void UReplicationSystem::PostSendUpdate()
 	InternalSys.GetChangeMaskCache().ResetCache();
 
 	// Store the state of the previous frames scopable objects
-	InternalSys.GetNetHandleManager().SetPrevFrameScopableInternalIndicesToCurrent();
+	InternalSys.GetNetRefHandleManager().SetPrevFrameScopableInternalIndicesToCurrent();
 
 	// Update handles pending tear-off
 	InternalSys.GetReplicationBridge()->UpdateHandlesPendingTearOff();
@@ -668,26 +673,34 @@ void UReplicationSystem::SetReplicationView(uint32 ConnectionId, const UE::Net::
 	Connections.SetReplicationView(ConnectionId, View);
 }
 
-void UReplicationSystem::SetStaticPriority(FNetHandle Handle, float Priority)
+void UReplicationSystem::SetStaticPriority(FNetRefHandle Handle, float Priority)
 {
-	const uint32 ObjectIndex = Impl->ReplicationSystemInternal.GetNetHandleManager().GetInternalIndex(Handle);
-	if (ObjectIndex == UE::Net::Private::FNetHandleManager::InvalidInternalIndex)
+	const UE::Net::Private::FInternalNetRefIndex ObjectInternalIndex = Impl->ReplicationSystemInternal.GetNetRefHandleManager().GetInternalIndex(Handle);
+	if (ObjectInternalIndex == UE::Net::Private::FNetRefHandleManager::InvalidInternalIndex)
 	{
 		return;
 	}
 
-	return Impl->ReplicationSystemInternal.GetPrioritization().SetStaticPriority(ObjectIndex, Priority);
+	return Impl->ReplicationSystemInternal.GetPrioritization().SetStaticPriority(ObjectInternalIndex, Priority);
 }
 
-bool UReplicationSystem::SetPrioritizer(FNetHandle Handle, FNetObjectPrioritizerHandle Prioritizer)
+bool UReplicationSystem::SetPrioritizer(FNetRefHandle Handle, FNetObjectPrioritizerHandle Prioritizer)
 {
-	const uint32 ObjectIndex = Impl->ReplicationSystemInternal.GetNetHandleManager().GetInternalIndex(Handle);
-	if (ObjectIndex == UE::Net::Private::FNetHandleManager::InvalidInternalIndex)
+	using namespace UE::Net::Private;
+
+	if (!Handle.IsValid())
+	{
+		return false;
+	}
+
+	FReplicationSystemInternal& ReplicationSystemInternal = Impl->ReplicationSystemInternal;
+	const FInternalNetRefIndex ObjectInternalIndex = ReplicationSystemInternal.GetNetRefHandleManager().GetInternalIndex(Handle);
+	if (ObjectInternalIndex == FNetRefHandleManager::InvalidInternalIndex)
 	{
 		return false;
 	}
 
-	return Impl->ReplicationSystemInternal.GetPrioritization().SetPrioritizer(ObjectIndex, Prioritizer);
+	return ReplicationSystemInternal.GetPrioritization().SetPrioritizer(ObjectInternalIndex, Prioritizer);
 }
 
 FNetObjectPrioritizerHandle UReplicationSystem::GetPrioritizerHandle(const FName PrioritizerName) const
@@ -700,13 +713,6 @@ UNetObjectPrioritizer* UReplicationSystem::GetPrioritizer(const FName Prioritize
 	return Impl->ReplicationSystemInternal.GetPrioritization().GetPrioritizer(PrioritizerName);
 }
 
-#if 0
-UNetObjectPrioritizer* UReplicationSystem::GetPrioritizer(const FName PrioritizerName)
-{
-	return Impl->ReplicationSystemInternal.GetPrioritization().GetPrioritizer(PrioritizerName);
-}
-#endif
-
 const UE::Net::FStringTokenStore* UReplicationSystem::GetStringTokenStore() const
 {
 	return &Impl->ReplicationSystemInternal.GetStringTokenStore();
@@ -777,33 +783,33 @@ UReplicationBridge* UReplicationSystem::GetReplicationBridge() const
 	return Impl->ReplicationSystemInternal.GetReplicationBridge();
 }
 
-bool UReplicationSystem::IsValidHandle(FNetHandle Handle) const
+bool UReplicationSystem::IsValidHandle(FNetRefHandle Handle) const
 {
-	return Handle.IsValid() && Impl->ReplicationSystemInternal.GetNetHandleManager().IsValidNetHandle(Handle);
+	return Handle.IsValid() && Impl->ReplicationSystemInternal.GetNetRefHandleManager().IsValidNetRefHandle(Handle);
 }
 
-const UE::Net::FReplicationProtocol* UReplicationSystem::GetReplicationProtocol(FNetHandle Handle) const
+const UE::Net::FReplicationProtocol* UReplicationSystem::GetReplicationProtocol(FNetRefHandle Handle) const
 {
 	using namespace UE::Net::Private;
 
-	FNetHandleManager& NetHandleManager = Impl->ReplicationSystemInternal.GetNetHandleManager();
+	FNetRefHandleManager& NetRefHandleManager = Impl->ReplicationSystemInternal.GetNetRefHandleManager();
 
-	const uint32 ObjectInternalIndex = NetHandleManager.GetInternalIndex(Handle);
-	if (ObjectInternalIndex == FNetHandleManager::InvalidInternalIndex)
+	const FInternalNetRefIndex ObjectInternalIndex = NetRefHandleManager.GetInternalIndex(Handle);
+	if (ObjectInternalIndex == FNetRefHandleManager::InvalidInternalIndex)
 	{
 		return nullptr;
 	}
 
-	return NetHandleManager.GetReplicatedObjectDataNoCheck(ObjectInternalIndex).Protocol;
+	return NetRefHandleManager.GetReplicatedObjectDataNoCheck(ObjectInternalIndex).Protocol;
 }
 
-void UReplicationSystem::SetOwningNetConnection(FNetHandle Handle, uint32 ConnectionId)
+void UReplicationSystem::SetOwningNetConnection(FNetRefHandle Handle, uint32 ConnectionId)
 {
 	using namespace UE::Net::Private;
 
-	FNetHandleManager& NetHandleManager = Impl->ReplicationSystemInternal.GetNetHandleManager();
-	const uint32 ObjectInternalIndex = NetHandleManager.GetInternalIndex(Handle);
-	if (ObjectInternalIndex == FNetHandleManager::InvalidInternalIndex)
+	FNetRefHandleManager& NetRefHandleManager = Impl->ReplicationSystemInternal.GetNetRefHandleManager();
+	const FInternalNetRefIndex ObjectInternalIndex = NetRefHandleManager.GetInternalIndex(Handle);
+	if (ObjectInternalIndex == FNetRefHandleManager::InvalidInternalIndex)
 	{
 		return;
 	}
@@ -812,28 +818,28 @@ void UReplicationSystem::SetOwningNetConnection(FNetHandle Handle, uint32 Connec
 	Filtering.SetOwningConnection(ObjectInternalIndex, ConnectionId);
 }
 
-uint32 UReplicationSystem::GetOwningNetConnection(FNetHandle Handle) const
+uint32 UReplicationSystem::GetOwningNetConnection(FNetRefHandle Handle) const
 {
 	using namespace UE::Net::Private;
 
-	FNetHandleManager& NetHandleManager = Impl->ReplicationSystemInternal.GetNetHandleManager();
-	const uint32 ObjectInternalIndex = NetHandleManager.GetInternalIndex(Handle);
-	if (ObjectInternalIndex == FNetHandleManager::InvalidInternalIndex)
+	FNetRefHandleManager& NetRefHandleManager = Impl->ReplicationSystemInternal.GetNetRefHandleManager();
+	const FInternalNetRefIndex ObjectInternalIndex = NetRefHandleManager.GetInternalIndex(Handle);
+	if (ObjectInternalIndex == FNetRefHandleManager::InvalidInternalIndex)
 	{
-		return 0U;
+		return UE::Net::InvalidConnectionId;
 	}
 
 	const FReplicationFiltering& Filtering = Impl->ReplicationSystemInternal.GetFiltering();
 	return Filtering.GetOwningConnection(ObjectInternalIndex);
 }
 
-bool UReplicationSystem::SetFilter(FNetHandle Handle, UE::Net::FNetObjectFilterHandle Filter)
+bool UReplicationSystem::SetFilter(FNetRefHandle Handle, UE::Net::FNetObjectFilterHandle Filter)
 {
 	using namespace UE::Net::Private;
 
-	FNetHandleManager& NetHandleManager = Impl->ReplicationSystemInternal.GetNetHandleManager();
-	const uint32 ObjectInternalIndex = NetHandleManager.GetInternalIndex(Handle);
-	if (ObjectInternalIndex == FNetHandleManager::InvalidInternalIndex)
+	FNetRefHandleManager& NetRefHandleManager = Impl->ReplicationSystemInternal.GetNetRefHandleManager();
+	const FInternalNetRefIndex ObjectInternalIndex = NetRefHandleManager.GetInternalIndex(Handle);
+	if (ObjectInternalIndex == FNetRefHandleManager::InvalidInternalIndex)
 	{
 		return false;
 	}
@@ -852,14 +858,14 @@ UNetObjectFilter* UReplicationSystem::GetFilter(const FName FilterName) const
 	return Impl->ReplicationSystemInternal.GetFiltering().GetFilter(FilterName);
 }
 
-bool UReplicationSystem::SetConnectionFilter(FNetHandle Handle, const TBitArray<>& Connections, UE::Net::ENetFilterStatus ReplicationStatus)
+bool UReplicationSystem::SetConnectionFilter(FNetRefHandle Handle, const TBitArray<>& Connections, UE::Net::ENetFilterStatus ReplicationStatus)
 {
 	using namespace UE::Net;
 	using namespace UE::Net::Private;
 
-	FNetHandleManager& NetHandleManager = Impl->ReplicationSystemInternal.GetNetHandleManager();
-	const uint32 ObjectInternalIndex = NetHandleManager.GetInternalIndex(Handle);
-	if (ObjectInternalIndex == FNetHandleManager::InvalidInternalIndex)
+	FNetRefHandleManager& NetRefHandleManager = Impl->ReplicationSystemInternal.GetNetRefHandleManager();
+	const FInternalNetRefIndex ObjectInternalIndex = NetRefHandleManager.GetInternalIndex(Handle);
+	if (ObjectInternalIndex == FNetRefHandleManager::InvalidInternalIndex)
 	{
 		return false;
 	}
@@ -952,17 +958,17 @@ UE::Net::FNetObjectGroupHandle UReplicationSystem::CreateGroup()
 	return Impl->ReplicationSystemInternal.GetGroups().CreateGroup();
 }
 
-void UReplicationSystem::AddToGroup(FNetObjectGroupHandle GroupHandle, FNetHandle Handle)
+void UReplicationSystem::AddToGroup(FNetObjectGroupHandle GroupHandle, FNetRefHandle Handle)
 {
 	LLM_SCOPE_BYTAG(Iris);
 
 	using namespace UE::Net::Private;
 
 	FNetObjectGroups& Groups = Impl->ReplicationSystemInternal.GetGroups();
-	FNetHandleManager& NetHandleManager = Impl->ReplicationSystemInternal.GetNetHandleManager();
+	FNetRefHandleManager& NetRefHandleManager = Impl->ReplicationSystemInternal.GetNetRefHandleManager();
 	FReplicationFiltering& Filtering = Impl->ReplicationSystemInternal.GetFiltering();	
 
-	const uint32 ObjectInternalIndex = NetHandleManager.GetInternalIndex(Handle);
+	const FInternalNetRefIndex ObjectInternalIndex = NetRefHandleManager.GetInternalIndex(Handle);
 
 	if (GroupHandle && ObjectInternalIndex)
 	{
@@ -971,13 +977,13 @@ void UReplicationSystem::AddToGroup(FNetObjectGroupHandle GroupHandle, FNetHandl
 	}
 }
 
-void UReplicationSystem::RemoveFromGroup(FNetObjectGroupHandle GroupHandle, FNetHandle Handle)
+void UReplicationSystem::RemoveFromGroup(FNetObjectGroupHandle GroupHandle, FNetRefHandle Handle)
 {	
 	using namespace UE::Net::Private;
 
-	FNetHandleManager& NetHandleManager = Impl->ReplicationSystemInternal.GetNetHandleManager();
+	FNetRefHandleManager& NetRefHandleManager = Impl->ReplicationSystemInternal.GetNetRefHandleManager();
 
-	const uint32 ObjectInternalIndex = NetHandleManager.GetInternalIndex(Handle);
+	const FInternalNetRefIndex ObjectInternalIndex = NetRefHandleManager.GetInternalIndex(Handle);
 	if (GroupHandle && ObjectInternalIndex)
 	{
 		FNetObjectGroups& Groups = Impl->ReplicationSystemInternal.GetGroups();
@@ -988,16 +994,16 @@ void UReplicationSystem::RemoveFromGroup(FNetObjectGroupHandle GroupHandle, FNet
 	}
 }
 
-void UReplicationSystem::RemoveFromAllGroups(FNetHandle Handle)
+void UReplicationSystem::RemoveFromAllGroups(FNetRefHandle Handle)
 {
 	using namespace UE::Net::Private;
 
-	FNetHandleManager& NetHandleManager = Impl->ReplicationSystemInternal.GetNetHandleManager();
+	FNetRefHandleManager& NetRefHandleManager = Impl->ReplicationSystemInternal.GetNetRefHandleManager();
 	FNetObjectGroups& Groups = Impl->ReplicationSystemInternal.GetGroups();
 	FReplicationFiltering& Filtering = Impl->ReplicationSystemInternal.GetFiltering();
 
 	uint32 NumGroupMemberShips;
-	const uint32 ObjectInternalIndex = NetHandleManager.GetInternalIndex(Handle);
+	const FInternalNetRefIndex ObjectInternalIndex = NetRefHandleManager.GetInternalIndex(Handle);
 	if (const FNetObjectGroupHandle* GroupHandles = Groups.GetGroupMemberships(ObjectInternalIndex, NumGroupMemberShips))
 	{
 		// We copy the membership array as it is modified during removal
@@ -1010,12 +1016,14 @@ void UReplicationSystem::RemoveFromAllGroups(FNetHandle Handle)
 	}	
 }
 
-bool UReplicationSystem::IsInGroup(FNetObjectGroupHandle GroupHandle, FNetHandle Handle) const
+bool UReplicationSystem::IsInGroup(FNetObjectGroupHandle GroupHandle, FNetRefHandle Handle) const
 {
-	const UE::Net::Private::FNetObjectGroups& Groups = Impl->ReplicationSystemInternal.GetGroups();
-	const UE::Net::Private::FNetHandleManager& NetHandleManager = Impl->ReplicationSystemInternal.GetNetHandleManager();
+	using namespace UE::Net::Private;
+
+	const FNetObjectGroups& Groups = Impl->ReplicationSystemInternal.GetGroups();
+	const FNetRefHandleManager& NetRefHandleManager = Impl->ReplicationSystemInternal.GetNetRefHandleManager();
 
-	const uint32 ObjectInternalIndex = NetHandleManager.GetInternalIndex(Handle);
+	const FInternalNetRefIndex ObjectInternalIndex = NetRefHandleManager.GetInternalIndex(Handle);
 
 	return Groups.Contains(GroupHandle, ObjectInternalIndex);
 }
@@ -1078,13 +1086,13 @@ void UReplicationSystem::SetGroupFilterStatus(FNetObjectGroupHandle GroupHandle,
 	Filtering.SetGroupFilterStatus(GroupHandle, ReplicationStatus);
 }
 
-bool UReplicationSystem::SetReplicationConditionConnectionFilter(FNetHandle Handle, UE::Net::EReplicationCondition Condition, uint32 ConnectionId, bool bEnable)
+bool UReplicationSystem::SetReplicationConditionConnectionFilter(FNetRefHandle Handle, UE::Net::EReplicationCondition Condition, uint32 ConnectionId, bool bEnable)
 {
 	using namespace UE::Net::Private;
 
-	FNetHandleManager& NetHandleManager = Impl->ReplicationSystemInternal.GetNetHandleManager();
-	const uint32 ObjectInternalIndex = NetHandleManager.GetInternalIndex(Handle);
-	if (ObjectInternalIndex == FNetHandleManager::InvalidInternalIndex)
+	FNetRefHandleManager& NetRefHandleManager = Impl->ReplicationSystemInternal.GetNetRefHandleManager();
+	const FInternalNetRefIndex ObjectInternalIndex = NetRefHandleManager.GetInternalIndex(Handle);
+	if (ObjectInternalIndex == FNetRefHandleManager::InvalidInternalIndex)
 	{
 		return false;
 	}
@@ -1093,13 +1101,13 @@ bool UReplicationSystem::SetReplicationConditionConnectionFilter(FNetHandle Hand
 	return Conditionals.SetConditionConnectionFilter(ObjectInternalIndex, Condition, ConnectionId, bEnable);
 }
 
-bool UReplicationSystem::SetReplicationCondition(FNetHandle Handle, UE::Net::EReplicationCondition Condition, bool bEnable)
+bool UReplicationSystem::SetReplicationCondition(FNetRefHandle Handle, UE::Net::EReplicationCondition Condition, bool bEnable)
 {
 	using namespace UE::Net::Private;
 
-	FNetHandleManager& NetHandleManager = Impl->ReplicationSystemInternal.GetNetHandleManager();
-	const uint32 ObjectInternalIndex = NetHandleManager.GetInternalIndex(Handle);
-	if (ObjectInternalIndex == FNetHandleManager::InvalidInternalIndex)
+	FNetRefHandleManager& NetRefHandleManager = Impl->ReplicationSystemInternal.GetNetRefHandleManager();
+	const FInternalNetRefIndex ObjectInternalIndex = NetRefHandleManager.GetInternalIndex(Handle);
+	if (ObjectInternalIndex == FNetRefHandleManager::InvalidInternalIndex)
 	{
 		return false;
 	}
@@ -1108,14 +1116,14 @@ bool UReplicationSystem::SetReplicationCondition(FNetHandle Handle, UE::Net::ERe
 	return Conditionals.SetCondition(ObjectInternalIndex, Condition, bEnable);
 }
 
-void UReplicationSystem::SetDeltaCompressionStatus(FNetHandle Handle, UE::Net::ENetObjectDeltaCompressionStatus Status)
+void UReplicationSystem::SetDeltaCompressionStatus(FNetRefHandle Handle, UE::Net::ENetObjectDeltaCompressionStatus Status)
 {
 	using namespace UE::Net::Private;
 
 	FReplicationSystemInternal& InternalSys = Impl->ReplicationSystemInternal;
-	FNetHandleManager& NetHandleManager = InternalSys.GetNetHandleManager();
-	const uint32 ObjectInternalIndex = NetHandleManager.GetInternalIndex(Handle);
-	if (ObjectInternalIndex == FNetHandleManager::InvalidInternalIndex)
+	FNetRefHandleManager& NetRefHandleManager = InternalSys.GetNetRefHandleManager();
+	const FInternalNetRefIndex ObjectInternalIndex = NetRefHandleManager.GetInternalIndex(Handle);
+	if (ObjectInternalIndex == FNetRefHandleManager::InvalidInternalIndex)
 	{
 		return;
 	}
@@ -1124,24 +1132,25 @@ void UReplicationSystem::SetDeltaCompressionStatus(FNetHandle Handle, UE::Net::E
 	return DC.SetDeltaCompressionStatus(ObjectInternalIndex, Status);
 }
 
-void UReplicationSystem::SetIsNetTemporary(FNetHandle Handle)
+void UReplicationSystem::SetIsNetTemporary(FNetRefHandle Handle)
 {
-	UE::Net::Private::FNetHandleManager& NetHandleManager = Impl->ReplicationSystemInternal.GetNetHandleManager();
-	if (ensure(NetHandleManager.IsLocalNetHandle(Handle)))
+	UE::Net::Private::FNetRefHandleManager& NetRefHandleManager = Impl->ReplicationSystemInternal.GetNetRefHandleManager();
+	if (ensure(NetRefHandleManager.IsLocalNetRefHandle(Handle)))
 	{
 		// Set the object to not propagate changed states
-		NetHandleManager.SetShouldPropagateChangedStates(Handle, false);
+		NetRefHandleManager.SetShouldPropagateChangedStates(Handle, false);
 	}
 }
 
-void UReplicationSystem::TearOffNextUpdate(FNetHandle Handle)
+void UReplicationSystem::TearOffNextUpdate(FNetRefHandle Handle)
 {
-	Impl->ReplicationSystemInternal.GetReplicationBridge()->TearOff(Handle, false);
+	constexpr EEndReplicationFlags DestroyFlags = EEndReplicationFlags::DestroyNetHandle | EEndReplicationFlags::ClearNetPushId;
+	Impl->ReplicationSystemInternal.GetReplicationBridge()->TearOff(Handle, DestroyFlags, false);
 }
 
-void UReplicationSystem::MarkDirty(FNetHandle Handle)
+void UReplicationSystem::MarkDirty(FNetRefHandle Handle)
 {
-	if (const uint32 InternalObjectIndex = Impl->ReplicationSystemInternal.GetNetHandleManager().GetInternalIndex(Handle))
+	if (const uint32 InternalObjectIndex = Impl->ReplicationSystemInternal.GetNetRefHandleManager().GetInternalIndex(Handle))
 	{
 		 UE::Net::Private::MarkNetObjectStateDirty(GetId(), InternalObjectIndex);
 	}
@@ -1157,7 +1166,7 @@ void UReplicationSystem::AddReferencedObjects(UObject* InThis, FReferenceCollect
 	UReplicationSystem* This = CastChecked<UReplicationSystem>(InThis);
 	if (This->Impl.IsValid())
 	{
-		This->Impl->ReplicationSystemInternal.GetNetHandleManager().AddReferencedObjects(Collector);
+		This->Impl->ReplicationSystemInternal.GetNetRefHandleManager().AddReferencedObjects(Collector);
 	}
 	Super::AddReferencedObjects(InThis, Collector);
 }
@@ -1170,10 +1179,6 @@ const UE::Net::FWorldLocations& UReplicationSystem::GetWorldLocations() const
 namespace UE::Net
 {
 
-#if !UE_NET_ALLOW_MULTIPLE_REPLICATION_SYSTEMS
-	UReplicationSystem* FReplicationSystemFactory::ReplicationSystem = nullptr;
-#endif
-
 FReplicationSystemCreatedDelegate& FReplicationSystemFactory::GetReplicationSystemCreatedDelegate()
 {
 	static FReplicationSystemCreatedDelegate Delegate;
@@ -1187,20 +1192,17 @@ FReplicationSystemDestroyedDelegate& FReplicationSystemFactory::GetReplicationSy
 }
 
 // FReplicationSystemFactory
-#if UE_NET_ALLOW_MULTIPLE_REPLICATION_SYSTEMS
 UReplicationSystem* FReplicationSystemFactory::ReplicationSystems[MaxReplicationSystemCount];
-#endif
+uint32 FReplicationSystemFactory::MaxReplicationSystemId = 0;
 
 UReplicationSystem* FReplicationSystemFactory::CreateReplicationSystem(const UReplicationSystem::FReplicationSystemParams& Params)
 {
 	if (!Params.ReplicationBridge)
 	{
 		UE_LOG(LogIris, Error, TEXT("Cannot create ReplicationSystem without a ReplicationBridge"));
-
 		return nullptr;
 	}
 
-#if UE_NET_ALLOW_MULTIPLE_REPLICATION_SYSTEMS
 	for (uint32 It = 0, EndIt = MaxReplicationSystemCount; It != EndIt; ++It)
 	{
 		if (ReplicationSystems[It] != nullptr)
@@ -1210,9 +1212,15 @@ UReplicationSystem* FReplicationSystemFactory::CreateReplicationSystem(const URe
 
 		UReplicationSystem* ReplicationSystem = NewObject<UReplicationSystem>();
 		ReplicationSystems[It] = ReplicationSystem;
-
 		ReplicationSystem->AddToRoot();
-		ReplicationSystem->Init(It, Params);
+
+		const uint32 ReplicationSystemId = It;
+		if (ReplicationSystemId > MaxReplicationSystemId)
+		{
+			MaxReplicationSystemId = ReplicationSystemId;
+		}
+
+		ReplicationSystem->Init(ReplicationSystemId, Params);
 
 		if (GetReplicationSystemCreatedDelegate().IsBound())
 		{
@@ -1224,42 +1232,30 @@ UReplicationSystem* FReplicationSystemFactory::CreateReplicationSystem(const URe
 
 	LowLevelFatalError(TEXT("Too many ReplicationSystems have already been created (%u)"), MaxReplicationSystemCount);
 	return nullptr;
-#else
-	if (ReplicationSystem != nullptr)
-	{
-		LowLevelFatalError(TEXT("There can only be one ReplicationSystem"));
-	}
-
-	ReplicationSystem = NewObject<UReplicationSystem>();
-	ReplicationSystem->AddToRoot();
-	ReplicationSystem->Init(0, Params);
-
-	if (GetReplicationSystemCreatedDelegate().IsBound())
-	{
-		GetReplicationSystemCreatedDelegate().Broadcast(ReplicationSystem);
-	}
-	
-	return ReplicationSystem;
-#endif
 }
 
 void FReplicationSystemFactory::DestroyReplicationSystem(UReplicationSystem* System)
 {
 	if (System == nullptr)
+	{
 		return;
+	}
 
-#if UE_NET_ALLOW_MULTIPLE_REPLICATION_SYSTEMS
 	const uint32 Id = System->GetId();
 	if (Id < MaxReplicationSystemCount)
 	{
 		ReplicationSystems[Id] = nullptr;
+
+		uint32 NewMaxReplicationSystemId = 0;
+		for (uint32 It = 0, EndIt = MaxReplicationSystemCount; It != EndIt; ++It)
+		{
+			if (ReplicationSystems[It] != nullptr)
+			{
+				NewMaxReplicationSystemId = Id;
+			}
+		}
+		MaxReplicationSystemId = NewMaxReplicationSystemId;
 	}
-#else
-	if (System == ReplicationSystem)
-	{
-		ReplicationSystem = nullptr;
-	}
-#endif
 
 	if (GetReplicationSystemDestroyedDelegate().IsBound())
 	{
@@ -1271,4 +1267,9 @@ void FReplicationSystemFactory::DestroyReplicationSystem(UReplicationSystem* Sys
 	System->MarkAsGarbage();
 }
 
+TArrayView<UReplicationSystem*> FReplicationSystemFactory::GetAllReplicationSystems()
+{
+	return MakeArrayView(ReplicationSystems, MaxReplicationSystemId + 1);
+}
+
 }
diff --git a/Engine/Source/Runtime/Experimental/Iris/Core/Private/Iris/ReplicationSystem/ReplicationSystemInternal.h b/Engine/Source/Runtime/Experimental/Iris/Core/Private/Iris/ReplicationSystem/ReplicationSystemInternal.h
--- a/Engine/Source/Runtime/Experimental/Iris/Core/Private/Iris/ReplicationSystem/ReplicationSystemInternal.h
+++ b/Engine/Source/Runtime/Experimental/Iris/Core/Private/Iris/ReplicationSystem/ReplicationSystemInternal.h
@@ -38,7 +38,7 @@ class FReplicationSystemInternal
 {
 public:
 	explicit FReplicationSystemInternal(const FReplicationSystemInternalInitParams& Params)
-	: NetHandleManager(ReplicationProtocolManager, Params.ReplicationSystemId, Params.MaxReplicatedObjectCount)
+	: NetRefHandleManager(ReplicationProtocolManager, Params.ReplicationSystemId, Params.MaxReplicatedObjectCount)
 	, DirtyNetObjectTracker()
 	, StringTokenStore(NetTokenStore)
 	, Id(Params.ReplicationSystemId)
@@ -46,8 +46,8 @@ public:
 
 	FReplicationProtocolManager& GetReplicationProtocolManager() { return ReplicationProtocolManager; }
 
-	FNetHandleManager& GetNetHandleManager() { return NetHandleManager; }
-	const FNetHandleManager& GetNetHandleManager() const { return NetHandleManager; }
+	FNetRefHandleManager& GetNetRefHandleManager() { return NetRefHandleManager; }
+	const FNetRefHandleManager& GetNetRefHandleManager() const { return NetRefHandleManager; }
 
 	FDirtyNetObjectTracker& GetDirtyNetObjectTracker() { return DirtyNetObjectTracker; }
 
@@ -59,7 +59,7 @@ public:
 
 	void SetReplicationBridge(UReplicationBridge* InReplicationBridge) { ReplicationBridge = InReplicationBridge; }
 	UReplicationBridge* GetReplicationBridge() const { return ReplicationBridge; }
-	UReplicationBridge* GetReplicationBridge(FNetHandle Handle) const { return ReplicationBridge; }
+	UReplicationBridge* GetReplicationBridge(FNetRefHandle Handle) const { return ReplicationBridge; }
 
 	FChangeMaskCache& GetChangeMaskCache() { return ChangeMaskCache; }
 
@@ -95,7 +95,7 @@ public:
 
 private:
 	FReplicationProtocolManager ReplicationProtocolManager;
-	FNetHandleManager NetHandleManager;
+	FNetRefHandleManager NetRefHandleManager;
 	FDirtyNetObjectTracker DirtyNetObjectTracker;
 	FReplicationStateStorage ReplicationStateStorage;
 	FReplicationStateDescriptorRegistry ReplicationStateDescriptorRegistry;
diff --git a/Engine/Source/Runtime/Experimental/Iris/Core/Private/Iris/ReplicationSystem/ReplicationWriter.cpp b/Engine/Source/Runtime/Experimental/Iris/Core/Private/Iris/ReplicationSystem/ReplicationWriter.cpp
--- a/Engine/Source/Runtime/Experimental/Iris/Core/Private/Iris/ReplicationSystem/ReplicationWriter.cpp
+++ b/Engine/Source/Runtime/Experimental/Iris/Core/Private/Iris/ReplicationSystem/ReplicationWriter.cpp
@@ -299,7 +299,7 @@ bool FReplicationWriter::IsReplicationEnabled() const
 
 // $IRIS TODO : May need to introduce queue and send behaviors. For example one may want to send only with object.
 // One may not want to send unless the object is replicated very soon etc.
-void FReplicationWriter::QueueNetObjectAttachments(FInternalNetHandle OwnerInternalIndex, FInternalNetHandle SubObjectInternalIndex, TArrayView<const TRefCountPtr<FNetBlob>> InAttachments)
+void FReplicationWriter::QueueNetObjectAttachments(FInternalNetRefIndex OwnerInternalIndex, FInternalNetRefIndex SubObjectInternalIndex, TArrayView<const TRefCountPtr<FNetBlob>> InAttachments)
 {
 	if (InAttachments.Num() <= 0)
 	{
@@ -314,7 +314,7 @@ void FReplicationWriter::QueueNetObjectAttachments(FInternalNetHandle OwnerInter
 		return;
 	}
 
-	const uint32 TargetIndex = bObjectInScope ? (SubObjectInternalIndex != FNetHandleManager::InvalidInternalIndex ? SubObjectInternalIndex : OwnerInternalIndex) : ObjectIndexForOOBAttachment;
+	const uint32 TargetIndex = bObjectInScope ? (SubObjectInternalIndex != FNetRefHandleManager::InvalidInternalIndex ? SubObjectInternalIndex : OwnerInternalIndex) : ObjectIndexForOOBAttachment;
 	ENetObjectAttachmentType AttachmentType = (bObjectInScope ? ENetObjectAttachmentType::Normal : ENetObjectAttachmentType::OutOfBand);
 	if (!Attachments.Enqueue(AttachmentType, TargetIndex, InAttachments))
 	{
@@ -354,7 +354,7 @@ void FReplicationWriter::Init(const FReplicationParameters& InParameters)
 
 	// Cache internal systems
 	ReplicationSystemInternal = Parameters.ReplicationSystem->GetReplicationSystemInternal();
-	NetHandleManager = &ReplicationSystemInternal->GetNetHandleManager();
+	NetRefHandleManager = &ReplicationSystemInternal->GetNetRefHandleManager();
 	ReplicationBridge = Parameters.ReplicationSystem->GetReplicationBridge();
 	BaselineManager = &ReplicationSystemInternal->GetDeltaCompressionBaselineManager();
 	ObjectReferenceCache = &ReplicationSystemInternal->GetObjectReferenceCache();
@@ -399,11 +399,11 @@ void FReplicationWriter::StartReplication(uint32 InternalIndex)
 	Info.LastAckedBaselineIndex = FDeltaCompressionBaselineManager::InvalidBaselineIndex;
 	Info.PendingBaselineIndex = FDeltaCompressionBaselineManager::InvalidBaselineIndex;
 
-	const uint32 bIsDestructionInfo = NetHandleManager->GetIsDestroyedStartupObject(InternalIndex) ? 1U : 0U;	
+	const uint32 bIsDestructionInfo = NetRefHandleManager->GetIsDestroyedStartupObject(InternalIndex) ? 1U : 0U;	
 	if (bIsDestructionInfo)
 	{
 		// Check status of original object about to be destroyed, if it has been confirmed as created we do not replicate the destruction info object at all
-		if (const uint32 OriginalInternalIndex = NetHandleManager->GetOriginalDestroyedStartupObjectIndex(InternalIndex))
+		if (const uint32 OriginalInternalIndex = NetRefHandleManager->GetOriginalDestroyedStartupObjectIndex(InternalIndex))
 		{
 			const FReplicationInfo& OriginalInfo = GetReplicationInfo(OriginalInternalIndex);
 			if (OriginalInfo.GetState() != EReplicatedObjectState::Invalid && OriginalInfo.IsCreationConfirmed)
@@ -414,7 +414,7 @@ void FReplicationWriter::StartReplication(uint32 InternalIndex)
 				Info.IsDestructionInfo = bIsDestructionInfo;
 				Info.IsCreationConfirmed = 1U;
 
-				NetHandleManager->AddNetObjectRef(InternalIndex);
+				NetRefHandleManager->AddNetObjectRef(InternalIndex);
 
 				return;
 			}
@@ -424,12 +424,12 @@ void FReplicationWriter::StartReplication(uint32 InternalIndex)
 	// Pending create
 	SetState(InternalIndex, EReplicatedObjectState::PendingCreate);
 
-	const FNetHandleManager::FReplicatedObjectData& Data = NetHandleManager->GetReplicatedObjectDataNoCheck(InternalIndex);
-	NetHandleManager->AddNetObjectRef(InternalIndex);
+	const FNetRefHandleManager::FReplicatedObjectData& Data = NetRefHandleManager->GetReplicatedObjectDataNoCheck(InternalIndex);
+	NetRefHandleManager->AddNetObjectRef(InternalIndex);
 
 	Info.ChangeMaskBitCount = Data.Protocol->ChangeMaskBitCount;
 	Info.HasDirtySubObjects = 1U;
-	Info.IsSubObject = NetHandleManager->GetSubObjectInternalIndices().GetBit(InternalIndex);
+	Info.IsSubObject = NetRefHandleManager->GetSubObjectInternalIndices().GetBit(InternalIndex);
 	Info.HasDirtyChangeMask = 1U;
 	Info.HasAttachments = 0U;
 	Info.HasChangemaskFilter = EnumHasAnyFlags(Data.Protocol->ProtocolTraits, EReplicationProtocolTraits::HasConditionalChangeMask);
@@ -483,7 +483,7 @@ void FReplicationWriter::StopReplication(uint32 InternalIndex)
 	ObjectsInScope.ClearBit(InternalIndex);
 
 	UE_LOG_REPLICATIONWRITER_CONN(TEXT("ReplicationWriter.StopReplication for ( InternalIndex: %u )"), InternalIndex);
-	NetHandleManager->ReleaseNetObjectRef(InternalIndex);
+	NetRefHandleManager->ReleaseNetObjectRef(InternalIndex);
 
 	Attachments.DropAllAttachments(ENetObjectAttachmentType::Normal, InternalIndex);
 
@@ -510,9 +510,9 @@ const FReplicationWriter::FReplicationInfo& FReplicationWriter::GetReplicationIn
 	return ReplicatedObjects[InternalIndex];
 }
 
-void FReplicationWriter::WriteNetHandleId(FNetBitStreamWriter& Writer, FNetHandle Handle)
+void FReplicationWriter::WriteNetRefHandleId(FNetBitStreamWriter& Writer, FNetRefHandle Handle)
 {
-	Writer.WriteBits(Handle.GetId(), FNetHandle::IdBits);
+	Writer.WriteBits(Handle.GetId(), FNetRefHandle::IdBits);
 }
 
 uint32 FReplicationWriter::GetDefaultFlushFlags() const
@@ -543,7 +543,7 @@ uint32 FReplicationWriter::GetFlushStatus(uint32 InternalIndex, const FReplicati
 	if (!Info.IsSubObject && FlushFlags != FlushFlags_All)
 	{
 		// Check status of SubObjects as well.
-		for (uint32 SubObjectIndex : NetHandleManager->GetSubObjects(InternalIndex))
+		for (uint32 SubObjectIndex : NetRefHandleManager->GetSubObjects(InternalIndex))
 		{
 			const FReplicationInfo& SubObjectInfo = GetReplicationInfo(SubObjectIndex);
 			FlushFlags |= GetFlushStatus(SubObjectIndex, SubObjectInfo, InFlushFlags);
@@ -563,7 +563,7 @@ void FReplicationWriter::SetPendingDestroyOrSubObjectPendingDestroyState(uint32
 	if (Info.IsSubObject)
 	{
 		// Subobject destroyed before its owner is explicitly replicated as state data.
-		FNetHandleManager::FReplicatedObjectData& ObjectData = NetHandleManager->GetReplicatedObjectDataNoCheck(InternalIndex);
+		FNetRefHandleManager::FReplicatedObjectData& ObjectData = NetRefHandleManager->GetReplicatedObjectDataNoCheck(InternalIndex);
 		if (ObjectData.IsSubObject())
 		{
 			// If owner is not pending destroy we mark the state of the SubObject to SubObjectPendingDestroy and mark owner as having dirty subobjects which will 
@@ -584,7 +584,7 @@ void FReplicationWriter::SetPendingDestroyOrSubObjectPendingDestroyState(uint32
 	else if (Info.HasDirtySubObjects)
 	{
 		// If the owner is destroyed, all subobjects in the EReplicatedObjectState::SubObjectPendingDestroy state must also be marked as PendingDestroy as owner no longer will be replicated
-		for (uint32 SubObjectIndex : NetHandleManager->GetSubObjects(InternalIndex))
+		for (uint32 SubObjectIndex : NetRefHandleManager->GetSubObjects(InternalIndex))
 		{
 			FReplicationInfo& SubObjectInfo = GetReplicationInfo(SubObjectIndex);
 			if (SubObjectInfo.GetState() == EReplicatedObjectState::SubObjectPendingDestroy)
@@ -647,7 +647,7 @@ void FReplicationWriter::UpdateScope(const FNetBitArrayView& UpdatedScope)
 			if (State == EReplicatedObjectState::SubObjectPendingDestroy)
 			{
 				// If owner is not pending destroy we mark it as dirty as appropriate
-				FNetHandleManager::FReplicatedObjectData& ObjectData = NetHandleManager->GetReplicatedObjectDataNoCheck(Index);
+				FNetRefHandleManager::FReplicatedObjectData& ObjectData = NetRefHandleManager->GetReplicatedObjectDataNoCheck(Index);
 				FReplicationInfo& OwnerInfo = ReplicatedObjects[ObjectData.SubObjectRootIndex];
 				if (OwnerInfo.GetState() < EReplicatedObjectState::PendingDestroy)
 				{
@@ -773,7 +773,7 @@ void FReplicationWriter::InternalUpdateDirtyChangeMasks(const FChangeMaskCache&
 	//UE_LOG_REPLICATIONWRITER(TEXT("FReplicationWriter::UpdateDirtyChangeMasks() Updated %u Objects for ConnectionId:%u, ReplicationSystemId: %u."), CachedChangeMasks.Indices.Num(), Parameters.ConnectionId, Parameters.ReplicationSystem->GetId());	
 }
 
-void FReplicationWriter::NotifyDestroyedObjectPendingTearOff(FInternalNetHandle ObjectInternalIndex)
+void FReplicationWriter::NotifyDestroyedObjectPendingTearOff(FInternalNetRefIndex ObjectInternalIndex)
 {
 	const FReplicationInfo& ReplicationInfo = GetReplicationInfo(ObjectInternalIndex);
 	if (ReplicationInfo.GetState() == EReplicatedObjectState::PendingCreate)
@@ -812,7 +812,7 @@ uint32 FReplicationWriter::ScheduleObjects(FScheduleObjectInfo* OutScheduledObje
 	ObjectsWithDirtyChanges.ClearBit(ObjectIndexForOOBAttachment);
 
 	const FNetBitArray& UpdatedObjects = ObjectsWithDirtyChanges;
-	const FNetBitArray& SubObjects = NetHandleManager->GetSubObjectInternalIndices();
+	const FNetBitArray& SubObjects = NetRefHandleManager->GetSubObjectInternalIndices();
 
 	// Fill IndexList with all objects with dirty state that have positive priority excluding subobjects
 	auto FillIndexListFunc = [&LocalPriorities, &ScheduledObjectIndices, &ScheduledObjectCount](uint32 Index)
@@ -990,7 +990,7 @@ void FReplicationWriter::HandleDeliveredRecord(const FReplicationRecord::FRecord
 		if (!bStillPendingFlush && !Info.IsSubObject)
 		{
 			// Check status of SubObjects as well.
-			for (uint32 SubObjectIndex : NetHandleManager->GetSubObjects(InternalIndex))
+			for (uint32 SubObjectIndex : NetRefHandleManager->GetSubObjects(InternalIndex))
 			{
 				const FReplicationInfo& SubObjectInfo = GetReplicationInfo(SubObjectIndex);
 				if (GetFlushStatus(SubObjectIndex, SubObjectInfo, Info.FlushFlags) != FlushFlags_None)
@@ -1011,7 +1011,7 @@ void FReplicationWriter::HandleDeliveredRecord(const FReplicationRecord::FRecord
 				// Must also mark owner dirty to make sure that we send the tearoff
 				if (Info.IsSubObject)
 				{
-					FNetHandleManager::FReplicatedObjectData& ObjectData = NetHandleManager->GetReplicatedObjectDataNoCheck(InternalIndex);
+					FNetRefHandleManager::FReplicatedObjectData& ObjectData = NetRefHandleManager->GetReplicatedObjectDataNoCheck(InternalIndex);
 					if (ObjectData.IsSubObject())
 					{
 						FReplicationInfo& OwnerInfo = ReplicatedObjects[ObjectData.SubObjectRootIndex];
@@ -1162,7 +1162,7 @@ void FReplicationWriter::HandleDroppedRecord<FReplicationWriter::EReplicatedObje
 			if (Info.IsSubObject)
 			{
 				// Mark owner dirty as well as subobjects only are scheduled together with owner
-				const FNetHandleManager::FReplicatedObjectData& ObjectData = NetHandleManager->GetReplicatedObjectData(InternalIndex);
+				const FNetRefHandleManager::FReplicatedObjectData& ObjectData = NetRefHandleManager->GetReplicatedObjectData(InternalIndex);
 				uint32 SubObjectOwnerInternalIndex = ObjectData.SubObjectRootIndex;
 
 				FReplicationInfo& SubObjectOwnerReplicationInfo = GetReplicationInfo(SubObjectOwnerInternalIndex);
@@ -1193,12 +1193,12 @@ void FReplicationWriter::HandleDroppedRecord<FReplicationWriter::EReplicatedObje
 	// If we want to cancel the destroy and lost the destroy packet we can resume normal replication.
 	if (CurrentState == EReplicatedObjectState::CancelPendingDestroy)
 	{
-		checkfSlow(!RecordInfo.WroteTearOff, TEXT("Torn off objects can't cancel destroy. ( InternalIndex: %u ) %s"), InternalIndex, ToCStr(NetHandleManager->GetReplicatedObjectData(InternalIndex).Handle.ToString()));
+		checkfSlow(!RecordInfo.WroteTearOff, TEXT("Torn off objects can't cancel destroy. ( InternalIndex: %u ) %s"), InternalIndex, ToCStr(NetRefHandleManager->GetReplicatedObjectData(InternalIndex).RefHandle.ToString()));
 
 		if (RecordInfo.WroteDestroySubObject && Info.SubObjectPendingDestroy)
 		{
 			// If the subobject owner still is replicated and valid
-			FNetHandleManager::FReplicatedObjectData& ObjectData = NetHandleManager->GetReplicatedObjectDataNoCheck(InternalIndex);
+			FNetRefHandleManager::FReplicatedObjectData& ObjectData = NetRefHandleManager->GetReplicatedObjectDataNoCheck(InternalIndex);
 			check(ObjectData.IsSubObject());
 
 			// If owner is not pending destroy we mark it as dirty as appropriate.
@@ -1263,7 +1263,7 @@ void FReplicationWriter::HandleDroppedRecord<FReplicationWriter::EReplicatedObje
 		else if (RecordInfo.WroteDestroySubObject && Info.SubObjectPendingDestroy)
 		{
 			// If the subobject owner still is replicated and valid
-			FNetHandleManager::FReplicatedObjectData& ObjectData = NetHandleManager->GetReplicatedObjectDataNoCheck(InternalIndex);
+			FNetRefHandleManager::FReplicatedObjectData& ObjectData = NetRefHandleManager->GetReplicatedObjectDataNoCheck(InternalIndex);
 			check(ObjectData.IsSubObject());
 
 			// If owner is not pending destroy we mark it as dirty so that we can replicate subobject destruction properly
@@ -1298,7 +1298,7 @@ void FReplicationWriter::HandleDroppedRecord(const FReplicationRecord::FRecordIn
 
 	check(CurrentState != EReplicatedObjectState::Invalid);
 
-	UE_LOG_REPLICATIONWRITER_CONN(TEXT("Handle dropped data for ( InternalIndex: %u ) %s, LostState %s, CurrentState is %s"), InternalIndex, *(NetHandleManager->GetReplicatedObjectData(InternalIndex).Handle.ToString()), LexToString(LostObjectState), LexToString(CurrentState));
+	UE_LOG_REPLICATIONWRITER_CONN(TEXT("Handle dropped data for ( InternalIndex: %u ) %s, LostState %s, CurrentState is %s"), InternalIndex, *(NetRefHandleManager->GetReplicatedObjectData(InternalIndex).Handle.ToString()), LexToString(LostObjectState), LexToString(CurrentState));
 
 	// If we loose a baseline we must notify the BaselineManager and invalidate our pendingbaselineindex
 	if (RecordInfo.NewBaselineIndex != FDeltaCompressionBaselineManager::InvalidBaselineIndex)
@@ -1475,10 +1475,10 @@ uint32 FReplicationWriter::WriteObjectsPendingDestroy(FNetSerializationContext&
 
 	bool bWroteAllDestroyedObjects = true;
 
-	for (uint32 InternalIndex = 0U; (InternalIndex = ObjectsPendingDestroy.FindFirstOne(InternalIndex + 1U)) != ~0U; )
+	for (uint32 InternalIndex = 0U; (InternalIndex = ObjectsPendingDestroy.FindFirstOne(InternalIndex + 1U)) != FNetBitArray::InvalidIndex; )
 	{
 		FReplicationInfo& Info = GetReplicationInfo(InternalIndex);
-		const FNetHandleManager::FReplicatedObjectData& ObjectData = NetHandleManager->GetReplicatedObjectDataNoCheck(InternalIndex);
+		const FNetRefHandleManager::FReplicatedObjectData& ObjectData = NetRefHandleManager->GetReplicatedObjectDataNoCheck(InternalIndex);
 
 		// Already waiting on destroy confirmation
 		if (Info.GetState() == EReplicatedObjectState::WaitOnDestroyConfirmation)
@@ -1504,7 +1504,7 @@ uint32 FReplicationWriter::WriteObjectsPendingDestroy(FNetSerializationContext&
 				// If the owner is pending destroy then we need to replicate the destruction for the subobject here.
 				if (ObjectsPendingDestroy.GetBit(ObjectData.SubObjectRootIndex))
 				{
-					//			UE_LOG(LogTemp, Warning, TEXT("Object %s was a subobject but owner was destroyed so we need to destroy object normally"), ToCStr(ObjectData.Handle.ToString()));
+					//			UE_LOG(LogTemp, Warning, TEXT("Object %s was a subobject but owner was destroyed so we need to destroy object normally"), ToCStr(ObjectData.RefHandle.ToString()));
 					SetState(InternalIndex, EReplicatedObjectState::PendingDestroy);
 					Info.SubObjectPendingDestroy = 0U;
 				}
@@ -1522,7 +1522,7 @@ uint32 FReplicationWriter::WriteObjectsPendingDestroy(FNetSerializationContext&
 					FReplicationInfo& OwnerInfo = GetReplicationInfo(ObjectData.SubObjectRootIndex);
 					if (OwnerInfo.GetState() != EReplicatedObjectState::Invalid && OwnerInfo.GetState() < EReplicatedObjectState::PendingDestroy)
 					{
-						//			UE_LOG(LogTemp, Warning, TEXT("SubObject %s Should be destroyed by replicating owner"), ToCStr(ObjectData.Handle.ToString()));
+						//			UE_LOG(LogTemp, Warning, TEXT("SubObject %s Should be destroyed by replicating owner"), ToCStr(ObjectData.RefHandle.ToString()));
 
 						SetState(InternalIndex, EReplicatedObjectState::SubObjectPendingDestroy);
 						Info.SubObjectPendingDestroy = 1U;
@@ -1546,15 +1546,15 @@ uint32 FReplicationWriter::WriteObjectsPendingDestroy(FNetSerializationContext&
 			continue;
 		}
 
-		UE_NET_TRACE_OBJECT_SCOPE(ObjectData.Handle,  Writer, Context.GetTraceCollector(), ENetTraceVerbosity::Trace);
+		UE_NET_TRACE_OBJECT_SCOPE(ObjectData.RefHandle,  Writer, Context.GetTraceCollector(), ENetTraceVerbosity::Trace);
 
 		FNetBitStreamRollbackScope RollbackScope(Writer);
 
-		// Write NetObjectHandle with the needed bitCount
-		WriteNetHandleId(Writer, ObjectData.Handle);
+		// Write handle with the needed bitCount
+		WriteNetRefHandleId(Writer, ObjectData.RefHandle);
 
 		// Write bit indicating if the static instance should be destroyed or not (could skip the bit for dynamic objects)
-		const bool bShouldDestroyInstance = ObjectData.Handle.IsDynamic() || NetHandleManager->GetIsDestroyedStartupObject(InternalIndex);
+		const bool bShouldDestroyInstance = ObjectData.RefHandle.IsDynamic() || NetRefHandleManager->GetIsDestroyedStartupObject(InternalIndex);
 		Writer.WriteBool(bShouldDestroyInstance);
 
 		if (!Writer.IsOverflown())
@@ -1616,7 +1616,7 @@ bool FReplicationWriter::CanSendObject(uint32 InternalIndex) const
 
 	if (Info.HasDirtySubObjects)
 	{
-		for (FInternalNetHandle SubObjectInternalIndex : NetHandleManager->GetSubObjects(InternalIndex))
+		for (FInternalNetRefIndex SubObjectInternalIndex : NetRefHandleManager->GetSubObjects(InternalIndex))
 		{
 			if (!CanSendObject(SubObjectInternalIndex))
 			{
@@ -1626,7 +1626,7 @@ bool FReplicationWriter::CanSendObject(uint32 InternalIndex) const
 	}
 
 	// Currently we need to enforce a strict dependency on the initial dependencies
-	for (FInternalNetHandle DependentInternalIndex : NetHandleManager->GetDependentObjects(InternalIndex))
+	for (FInternalNetRefIndex DependentInternalIndex : NetRefHandleManager->GetDependentObjects(InternalIndex))
 	{
 		const FReplicationInfo& DependentInfo = GetReplicationInfo(DependentInternalIndex);
 		const bool bIsDependentObjectInitialState = IsInitialState(DependentInfo.GetState());
@@ -1641,7 +1641,7 @@ bool FReplicationWriter::CanSendObject(uint32 InternalIndex) const
 	return true;
 }
 
-void FReplicationWriter::SerializeObjectStateDelta(FNetSerializationContext& Context, uint32 InternalIndex, const FReplicationInfo& Info, const FNetHandleManager::FReplicatedObjectData& ObjectData, const uint8* ReplicatedObjectStateBuffer, FDeltaCompressionBaseline& CurrentBaseline, uint32 CreatedBaselineIndex)
+void FReplicationWriter::SerializeObjectStateDelta(FNetSerializationContext& Context, uint32 InternalIndex, const FReplicationInfo& Info, const FNetRefHandleManager::FReplicatedObjectData& ObjectData, const uint8* ReplicatedObjectStateBuffer, FDeltaCompressionBaseline& CurrentBaseline, uint32 CreatedBaselineIndex)
 {
 	FNetBitStreamWriter& Writer = *Context.GetBitStreamWriter();
 
@@ -1710,8 +1710,8 @@ FReplicationWriter::EWriteObjectStatus FReplicationWriter::WriteObjectInBatch(FN
 	};
 
 	// Only write data for the object if we have data to write
-	const FNetHandleManager::FReplicatedObjectData& ObjectData = NetHandleManager->GetReplicatedObjectDataNoCheck(InternalIndex);
-	uint8* ReplicatedObjectStateBuffer = NetHandleManager->GetReplicatedObjectStateBufferNoCheck(InternalIndex);
+	const FNetRefHandleManager::FReplicatedObjectData& ObjectData = NetRefHandleManager->GetReplicatedObjectDataNoCheck(InternalIndex);
+	uint8* ReplicatedObjectStateBuffer = NetRefHandleManager->GetReplicatedObjectStateBufferNoCheck(InternalIndex);
 
 	// Filter out changemasks that are not supposed to be replicated to this connection
 	const bool bNeedToFilterChangeMask = (bIsInitialState || Info.HasDirtyChangeMask) && Info.HasChangemaskFilter;
@@ -1736,10 +1736,10 @@ FReplicationWriter::EWriteObjectStatus FReplicationWriter::WriteObjectInBatch(FN
 
 	if (bHasState | bWriteAttachments | bSentTearOff | Info.SubObjectPendingDestroy)
 	{
-		const FNetHandle NetHandle = ObjectData.Handle;
+		const FNetRefHandle NetRefHandle = ObjectData.RefHandle;
 #if UE_NET_TRACE_ENABLED
-		const FNetHandle NetHandleForTraceScope = (WriteObjectFlags & EWriteObjectFlag::WriteObjectFlag_HugeObject ? NetHandleManager->GetReplicatedObjectDataNoCheck(HugeObjectContext.InternalIndex).Handle : NetHandle);
-		UE_NET_TRACE_OBJECT_SCOPE(NetHandleForTraceScope, Writer, Context.GetTraceCollector(), ENetTraceVerbosity::Trace);
+		const FNetRefHandle NetRefHandleForTraceScope = (WriteObjectFlags & EWriteObjectFlag::WriteObjectFlag_HugeObject ? NetRefHandleManager->GetReplicatedObjectDataNoCheck(HugeObjectContext.InternalIndex).RefHandle : NetRefHandle);
+		UE_NET_TRACE_OBJECT_SCOPE(NetRefHandleForTraceScope, Writer, Context.GetTraceCollector(), ENetTraceVerbosity::Trace);
 #endif
 		FNetBitStreamRollbackScope ObjectRollbackScope(Writer);
 
@@ -1754,7 +1754,7 @@ FReplicationWriter::EWriteObjectStatus FReplicationWriter::WriteObjectInBatch(FN
 		// $IRIS: $TODO: consider sending the internal index instead to save bits and only send handle when we create the object, https://jira.it.epicgames.com/browse/UE-127373
 		{
 			UE_NET_TRACE_SCOPE(Handle, Writer, Context.GetTraceCollector(), ENetTraceVerbosity::Trace);
-			WriteNetHandleId(Writer, NetHandle);
+			WriteNetRefHandleId(Writer, NetRefHandle);
 		}
 
 		// Store position of header bits
@@ -1845,18 +1845,18 @@ FReplicationWriter::EWriteObjectStatus FReplicationWriter::WriteObjectInBatch(FN
 					// Warn if we cannot replicate this object
 					if (!ObjectData.InstanceProtocol)
 					{
-						UE_LOG_REPLICATIONWRITER_WARNING(TEXT("Failed to replicate ( InternalIndex: %u ) %s, ProtocolName: %s, Currently we do not support creating a remote instance when the instance has been detached."), InternalIndex, *NetHandle.ToString(), ToCStr(ObjectData.Protocol->DebugName));
+						UE_LOG_REPLICATIONWRITER_WARNING(TEXT("Failed to replicate ( InternalIndex: %u ) %s, ProtocolName: %s, Currently we do not support creating a remote instance when the instance has been detached."), InternalIndex, *NetRefHandle.ToString(), ToCStr(ObjectData.Protocol->DebugName));
 						return EWriteObjectStatus::NoInstanceProtocol;
 					}
 
 					// Real SubObjects needs to be tracked on client as well
 					if (Writer.WriteBool(ObjectData.IsSubObject()))
 					{
-						const FNetHandle& OwnerHandle = NetHandleManager->GetReplicatedObjectDataNoCheck(ObjectData.SubObjectRootIndex).Handle;
-						WriteNetHandleId(Writer, OwnerHandle);
-						if (!OwnerHandle.IsValid() || Writer.IsOverflown())
+						const FNetRefHandle OwnerRefHandle = NetRefHandleManager->GetReplicatedObjectDataNoCheck(ObjectData.SubObjectRootIndex).RefHandle;
+						WriteNetRefHandleId(Writer, OwnerRefHandle);
+						if (!OwnerRefHandle.IsValid() || Writer.IsOverflown())
 						{
-							return OwnerHandle.IsValid() ? EWriteObjectStatus::BitStreamOverflow : EWriteObjectStatus::InvalidOwner;
+							return OwnerRefHandle.IsValid() ? EWriteObjectStatus::BitStreamOverflow : EWriteObjectStatus::InvalidOwner;
 						}
 					}
 
@@ -1869,7 +1869,7 @@ FReplicationWriter::EWriteObjectStatus FReplicationWriter::WriteObjectInBatch(FN
 					FReplicationBridgeSerializationContext BridgeContext(Context, Parameters.ConnectionId, Info.IsDestructionInfo == 1U);
 
 					// We need to send creation info, if we fail, we skip this object for now
-					if (!ReplicationBridge->CallWriteNetHandleCreationInfo(BridgeContext, NetHandle))
+					if (!ReplicationBridge->CallWriteNetRefHandleCreationInfo(BridgeContext, NetRefHandle))
 					{
 						return BridgeContext.SerializationContext.HasError() ? EWriteObjectStatus::Error : EWriteObjectStatus::BitStreamOverflow;
 					}
@@ -1910,7 +1910,7 @@ FReplicationWriter::EWriteObjectStatus FReplicationWriter::WriteObjectInBatch(FN
 					AttachmentWriter.WriteBool(BatchEntry.AttachmentType == ENetObjectAttachmentType::HugeObject);
 				}
 
-				const EAttachmentWriteStatus AttachmentWriteStatus = Attachments.Serialize(AttachmentContext, BatchEntry.AttachmentType, InternalIndex, NetHandle, BatchEntry.AttachmentRecord, BatchEntry.bHasUnsentAttachments);
+				const EAttachmentWriteStatus AttachmentWriteStatus = Attachments.Serialize(AttachmentContext, BatchEntry.AttachmentType, InternalIndex, NetRefHandle, BatchEntry.AttachmentRecord, BatchEntry.bHasUnsentAttachments);
 				if (BatchEntry.AttachmentType == ENetObjectAttachmentType::HugeObject)
 				{
 					if (AttachmentWriteStatus == EAttachmentWriteStatus::ReliableWindowFull)
@@ -1944,7 +1944,7 @@ FReplicationWriter::EWriteObjectStatus FReplicationWriter::WriteObjectInBatch(FN
 					const bool bFallbackToHugeObjectPath = BatchEntry.AttachmentType != ENetObjectAttachmentType::HugeObject && BatchEntry.bHasUnsentAttachments && (BitsThatWasAvailableForAttachements >= SplitThreshold);
 					if (bFallbackToHugeObjectPath)
 					{
-						UE_LOG(LogIris, Verbose, TEXT("Failed to write huge attachment for object %s ( InternalIndex: %u ), forcing fallback on hugeobject for attachments"), *NetHandle.ToString(), InternalIndex);
+						UE_LOG(LogIris, Verbose, TEXT("Failed to write huge attachment for object %s ( InternalIndex: %u ), forcing fallback on hugeobject for attachments"), *NetRefHandle.ToString(), InternalIndex);
 						Writer.DoOverflow();
 					}
 					else if (!BatchEntry.bSentState)
@@ -1984,7 +1984,7 @@ FReplicationWriter::EWriteObjectStatus FReplicationWriter::WriteObjectInBatch(FN
 				if (Info.SubObjectPendingDestroy)
 				{
 					ReplicatedDestroyHeaderFlags |= ReplicatedDestroyHeaderFlags_EndReplication;
-					const bool bShouldDestroyInstance = ObjectData.Handle.IsDynamic() || NetHandleManager->GetIsDestroyedStartupObject(InternalIndex);					
+					const bool bShouldDestroyInstance = ObjectData.RefHandle.IsDynamic() || NetRefHandleManager->GetIsDestroyedStartupObject(InternalIndex);					
 					ReplicatedDestroyHeaderFlags |= bShouldDestroyInstance ? ReplicatedDestroyHeaderFlags_DestroyInstance : ReplicatedDestroyHeaderFlags_None;
 				}
 
@@ -2019,7 +2019,7 @@ FReplicationWriter::EWriteObjectStatus FReplicationWriter::WriteObjectInBatch(FN
 		FReplicationConditionals::FSubObjectsToReplicateArray SubObjectsToReplicate;
 		ReplicationConditionals->GetSubObjectsToReplicate(Parameters.ConnectionId, InternalIndex, SubObjectsToReplicate);		
 
-		for (FInternalNetHandle SubObjectInternalIndex : SubObjectsToReplicate)
+		for (FInternalNetRefIndex SubObjectInternalIndex : SubObjectsToReplicate)
 		{
 			const int BatchObjectInfoCount = OutBatchInfo.ObjectInfos.Num();
 			EWriteObjectStatus SubObjectWriteStatus = WriteObjectInBatch(Context, SubObjectInternalIndex, WriteObjectFlags, OutBatchInfo);
@@ -2047,7 +2047,7 @@ FReplicationWriter::EWriteObjectStatus FReplicationWriter::WriteObjectInBatch(FN
 
 	// If we have any dependent objects pending creation we currently include them in the batch in order to guarantee that they get sent in the same packet
 	// When we implement support for proper dependencies this can be changed.
-	for (FInternalNetHandle DependentInternalIndex : NetHandleManager->GetDependentObjects(InternalIndex))
+	for (FInternalNetRefIndex DependentInternalIndex : NetRefHandleManager->GetDependentObjects(InternalIndex))
 	{
 		const bool bIsDependentInitialState = IsInitialState(GetReplicationInfo(DependentInternalIndex).GetState());
 		if (bIsDependentInitialState && !WriteContext.ObjectsWrittenThisTick.GetBit(DependentInternalIndex))
@@ -2174,7 +2174,7 @@ int FReplicationWriter::PrepareAndSendHugeObjectPayload(FNetSerializationContext
 		FNetBitStreamRollbackScope RollbackScope(Writer);
 		
 		FBatchInfo HugeObjectBatchInfo;
-		HugeObjectBatchInfo.ParentInternalIndex = FNetHandleManager::InvalidInternalIndex;
+		HugeObjectBatchInfo.ParentInternalIndex = FNetRefHandleManager::InvalidInternalIndex;
 		const uint32 WriteHugeObjectFlags = WriteObjectFlag_Attachments | WriteObjectFlag_HugeObject;
 		const EWriteObjectStatus HugeObjectStatus = WriteObjectInBatch(Context, ObjectIndexForOOBAttachment, WriteHugeObjectFlags, HugeObjectBatchInfo);
 		if (!IsWriteObjectSuccess(HugeObjectStatus))
@@ -2299,14 +2299,14 @@ int FReplicationWriter::WriteDestructionInfo(FNetSerializationContext& Context,
 	FNetExportRollbackScope ExportRollbackScope(Context);
 
 	// Only write data for the object if we have data to write
-	const FNetHandleManager::FReplicatedObjectData& ObjectData = NetHandleManager->GetReplicatedObjectDataNoCheck(InternalIndex);
+	const FNetRefHandleManager::FReplicatedObjectData& ObjectData = NetRefHandleManager->GetReplicatedObjectDataNoCheck(InternalIndex);
 
 	// Special case for static objects that should be destroyed on the client but we have not replicated
 	Writer.WriteBool(true);
 
 	FReplicationBridgeSerializationContext BridgeContext(Context, Parameters.ConnectionId, true);
 
-	if (!ReplicationBridge->CallWriteNetHandleCreationInfo(BridgeContext, ObjectData.Handle))
+	if (!ReplicationBridge->CallWriteNetRefHandleCreationInfo(BridgeContext, ObjectData.RefHandle))
 	{
 		// Trigger Rollback
 		Writer.DoOverflow();
@@ -2395,7 +2395,7 @@ uint32 FReplicationWriter::WriteObjects(FNetSerializationContext& Context)
 
 	bool bContinue = WriteContext.bHasUpdatedObjectsToSend;
 
-	auto SendObjectFunction = [this, &Context, &WrittenObjectCount](FInternalNetHandle InternalIndex)
+	auto SendObjectFunction = [this, &Context, &WrittenObjectCount](FInternalNetRefIndex InternalIndex)
 	{
 		if (!this->WriteContext.ObjectsWrittenThisTick.GetBit(InternalIndex) && this->CanSendObject(InternalIndex))
 		{
@@ -2486,7 +2486,7 @@ int FReplicationWriter::HandleObjectBatchSuccess(const FBatchInfo& BatchInfo, FR
 	OutRecord.ObjectReplicationRecords.Reserve(BatchInfo.ObjectInfos.Num());
 	for (const FBatchObjectInfo& BatchObjectInfo : BatchInfo.ObjectInfos)
 	{
-		//UE_LOG_REPLICATIONWRITER(TEXT("FReplicationWriter::Wrote Object with %s InitialState: %u"), *NetHandleManager->NetHandle.ToString(), bIsInitialState ? 1u : 0u);
+		//UE_LOG_REPLICATIONWRITER(TEXT("FReplicationWriter::Wrote Object with %s InitialState: %u"), *NetRefHandleManager->NetHandle.ToString(), bIsInitialState ? 1u : 0u);
 		FReplicationInfo& Info = GetReplicationInfo(BatchObjectInfo.InternalIndex);
 
 		// We did write the initial state, change the state to WaitOnCreateConfirmation
@@ -2534,7 +2534,7 @@ int FReplicationWriter::HandleObjectBatchSuccess(const FBatchInfo& BatchInfo, FR
 		}
 
 		// Schedule rest of dependent objects for replication, note there is no guarantee that they will replicate in same packet
-		for (FInternalNetHandle DependentInternalIndex : NetHandleManager->GetDependentObjects(BatchObjectInfo.InternalIndex))
+		for (FInternalNetRefIndex DependentInternalIndex : NetRefHandleManager->GetDependentObjects(BatchObjectInfo.InternalIndex))
 		{
 			if (ObjectsWithDirtyChanges.GetBit(DependentInternalIndex) && !WriteContext.ObjectsWrittenThisTick.GetBit(DependentInternalIndex))
 			{
@@ -3038,7 +3038,7 @@ void FReplicationWriter::StopAllReplication()
 
 		// Release object reference
 		const uint32 InternalIndex = &Info - FirstInfo;
-		NetHandleManager->ReleaseNetObjectRef(InternalIndex);
+		NetRefHandleManager->ReleaseNetObjectRef(InternalIndex);
 	}
 }
 
diff --git a/Engine/Source/Runtime/Experimental/Iris/Core/Private/Iris/ReplicationSystem/ReplicationWriter.h b/Engine/Source/Runtime/Experimental/Iris/Core/Private/Iris/ReplicationSystem/ReplicationWriter.h
--- a/Engine/Source/Runtime/Experimental/Iris/Core/Private/Iris/ReplicationSystem/ReplicationWriter.h
+++ b/Engine/Source/Runtime/Experimental/Iris/Core/Private/Iris/ReplicationSystem/ReplicationWriter.h
@@ -34,7 +34,7 @@ namespace UE::Net
 	namespace Private
 	{
 		struct FChangeMaskCache;
-		class FNetHandleManager;
+		class FNetRefHandleManager;
 		class FReliableNetBlobQueue;
 		class FReplicationConditionals;
 		class FReplicationFiltering;
@@ -155,7 +155,7 @@ public:
 	void ForceUpdateDirtyChangeMasks(const FChangeMaskCache& CachedChangeMasks, EFlushFlags ExtraFlushFlags, bool bMarkForTearOff) { InternalUpdateDirtyChangeMasks(CachedChangeMasks, ExtraFlushFlags, bMarkForTearOff); }
 
 	// Called if an object first being teared off and then explicitly destroyed before it has been removed from scope
-	void NotifyDestroyedObjectPendingTearOff(FInternalNetHandle ObjectInternalIndex);
+	void NotifyDestroyedObjectPendingTearOff(FInternalNetRefIndex ObjectInternalIndex);
 
 	// Propagate dirty changemasks
 	void UpdateDirtyChangeMasks(const FChangeMaskCache& CachedChangeMasks) { InternalUpdateDirtyChangeMasks(CachedChangeMasks, EFlushFlags::FlushFlags_None, false); }
@@ -182,7 +182,7 @@ public:
 	void SetNetExports(FNetExports& InNetExports);
 
 	// Attachments
-	void QueueNetObjectAttachments(FInternalNetHandle OwnerInternalIndex, FInternalNetHandle SubObjectInternalIndex, TArrayView<const TRefCountPtr<FNetBlob>> Attachments);
+	void QueueNetObjectAttachments(FInternalNetRefIndex OwnerInternalIndex, FInternalNetRefIndex SubObjectInternalIndex, TArrayView<const TRefCountPtr<FNetBlob>> Attachments);
 
 private:
 	// Various types
@@ -236,7 +236,7 @@ private:
 
 	struct FBatchObjectInfo
 	{
-		FNetHandle Handle;
+		FNetRefHandle Handle;
 		uint32 InternalIndex;
 		FNetObjectAttachmentsWriter::ReplicationRecord AttachmentRecord;
 		ENetObjectAttachmentType AttachmentType;
@@ -359,7 +359,7 @@ private:
 	void SetState(uint32 InternalIndex, EReplicatedObjectState NewState);
 
 	// Write index part of handle
-	void WriteNetHandleId(FNetBitStreamWriter& Writer, FNetHandle Handle);
+	void WriteNetRefHandleId(FNetBitStreamWriter& Writer, FNetRefHandle RefHandle);
 		
 	// Create new ObjectRecord
 	// Note: be aware that it will allocate a copy of the ChangeMask that needs to be handled if the record is not Committed
@@ -448,7 +448,7 @@ private:
 
 	bool IsWriteObjectSuccess(EWriteObjectStatus Status) const;
 
-	void SerializeObjectStateDelta(FNetSerializationContext& Context, uint32 InternalIndex, const FReplicationInfo& Info, const FNetHandleManager::FReplicatedObjectData& ObjectData, const uint8* RESTRICT ReplicatedObjectStateBuffer, FDeltaCompressionBaseline& CurrentBaseline, uint32 CreatedBaselineIndex);
+	void SerializeObjectStateDelta(FNetSerializationContext& Context, uint32 InternalIndex, const FReplicationInfo& Info, const FNetRefHandleManager::FReplicatedObjectData& ObjectData, const uint8* RESTRICT ReplicatedObjectStateBuffer, FDeltaCompressionBaseline& CurrentBaseline, uint32 CreatedBaselineIndex);
 
 	void DiscardAllRecords();
 	void StopAllReplication();
@@ -483,7 +483,7 @@ private:
 
 	// Cached internal systems
 	FReplicationSystemInternal* ReplicationSystemInternal = nullptr;
-	FNetHandleManager* NetHandleManager = nullptr;
+	FNetRefHandleManager* NetRefHandleManager = nullptr;
 	UReplicationBridge* ReplicationBridge = nullptr;
 	FDeltaCompressionBaselineManager* BaselineManager = nullptr;
 	FObjectReferenceCache* ObjectReferenceCache = nullptr;
diff --git a/Engine/Source/Runtime/Experimental/Iris/Core/Private/Iris/Serialization/NetExportContext.h b/Engine/Source/Runtime/Experimental/Iris/Core/Private/Iris/Serialization/NetExportContext.h
--- a/Engine/Source/Runtime/Experimental/Iris/Core/Private/Iris/Serialization/NetExportContext.h
+++ b/Engine/Source/Runtime/Experimental/Iris/Core/Private/Iris/Serialization/NetExportContext.h
@@ -15,12 +15,12 @@ namespace UE::Net::Private
 class FNetExportContext
 {
 public:
-	typedef TArray<FNetHandle, TInlineAllocator<32>> FExportsArray;
+	typedef TArray<FNetRefHandle, TInlineAllocator<32>> FExportsArray;
 	typedef TArray<FNetToken, TInlineAllocator<32>> FNetTokenExportsArray;
 
 	struct FAcknowledgedExports
 	{
-		TSet<FNetHandle> AcknowledgedExportedHandles;
+		TSet<FNetRefHandle> AcknowledgedExportedHandles;
 		TSet<FNetToken> AcknowledgedExportedNetTokens;
 	};
 
@@ -42,13 +42,13 @@ public:
 	FNetExportContext(const FAcknowledgedExports& InAcknowledgedExports, FBatchExports& BatchExports);
 
 	// Returns true if the Handle is acknowledged as delivered or if it is exported in the current batch
-	bool IsExported(FNetHandle Handle) const;
+	bool IsExported(FNetRefHandle Handle) const;
 
 	// Returns true if the Handle is acknowledged as delivered or if it is exported in the current batch
 	bool IsExported(FNetToken Token) const;
 
 	// Add a Handle to the current export batch
-	void AddExported(FNetHandle Handle);
+	void AddExported(FNetRefHandle Handle);
 
 	// Add a Handle to the current export batch
 	void AddExported(FNetToken Token);
@@ -114,12 +114,12 @@ inline FNetExportContext::FNetExportContext(const FAcknowledgedExports& InAcknow
 {
 }
 
-inline bool FNetExportContext::IsExported(FNetHandle Handle) const
+inline bool FNetExportContext::IsExported(FNetRefHandle Handle) const
 {
 	return AcknowledgedExports.AcknowledgedExportedHandles.Contains(Handle) || (BatchExports.HandlesExportedInCurrentBatch.Find(Handle) != INDEX_NONE);
 }
 
-inline void FNetExportContext::AddExported(FNetHandle Handle)
+inline void FNetExportContext::AddExported(FNetRefHandle Handle)
 {
 	BatchExports.HandlesExportedInCurrentBatch.Add(Handle);
 }
diff --git a/Engine/Source/Runtime/Experimental/Iris/Core/Private/Iris/Serialization/ObjectNetSerializer.cpp b/Engine/Source/Runtime/Experimental/Iris/Core/Private/Iris/Serialization/ObjectNetSerializer.cpp
--- a/Engine/Source/Runtime/Experimental/Iris/Core/Private/Iris/Serialization/ObjectNetSerializer.cpp
+++ b/Engine/Source/Runtime/Experimental/Iris/Core/Private/Iris/Serialization/ObjectNetSerializer.cpp
@@ -15,7 +15,7 @@
 namespace UE::Net
 {
 
-void WriteNetHandle(FNetSerializationContext& Context, const FNetHandle Handle)
+void WriteNetRefHandle(FNetSerializationContext& Context, const FNetRefHandle Handle)
 {
 	FNetBitStreamWriter* Writer = Context.GetBitStreamWriter();
 
@@ -30,7 +30,7 @@ void WriteNetHandle(FNetSerializationContext& Context, const FNetHandle Handle)
 	}
 }
 
-FNetHandle ReadNetHandle(FNetSerializationContext& Context)
+FNetRefHandle ReadNetRefHandle(FNetSerializationContext& Context)
 {
 	FNetBitStreamReader* Reader = Context.GetBitStreamReader();
 
@@ -39,10 +39,10 @@ FNetHandle ReadNetHandle(FNetSerializationContext& Context)
 		const uint32 NetId = ReadPackedUint32(Reader);
 		if (!Reader->IsOverflown())
 		{
-			return Private::FNetHandleManager::MakeNetHandleFromId(NetId);
+			return Private::FNetRefHandleManager::MakeNetRefHandleFromId(NetId);
 		}
 	}
-	return FNetHandle();
+	return FNetRefHandle();
 }
 
 void ReadFullNetObjectReference(FNetSerializationContext& Context, FNetObjectReference& Reference)
@@ -180,7 +180,7 @@ template<typename T>
 void FObjectNetSerializerBase<T>::Quantize(FNetSerializationContext& Context, const FNetQuantizeArgs& Args)
 {
 	QuantizedType& Target = *reinterpret_cast<QuantizedType*>(Args.Target);
-	const UObject* Source = GetValue(*reinterpret_cast<const T*>(Args.Source));
+	UObject* Source = GetValue(*reinterpret_cast<T*>(Args.Source));
 
 	const FInternalNetSerializationContext* InternalContext = Context.GetInternalContext();
 	Target = InternalContext->ObjectReferenceCache ? InternalContext->ObjectReferenceCache->GetOrCreateObjectReference(Source) : QuantizedType();
diff --git a/Engine/Source/Runtime/Experimental/Iris/Core/Public/Iris/Core/IrisDebugging.h b/Engine/Source/Runtime/Experimental/Iris/Core/Public/Iris/Core/IrisDebugging.h
--- a/Engine/Source/Runtime/Experimental/Iris/Core/Public/Iris/Core/IrisDebugging.h
+++ b/Engine/Source/Runtime/Experimental/Iris/Core/Public/Iris/Core/IrisDebugging.h
@@ -4,6 +4,7 @@
 
 #include "HAL/Platform.h"
 #include "Iris/ReplicationSystem/NetHandle.h"
+#include "Net/Core/NetHandle/NetHandle.h"
 
 class UReplicationSystem;
 namespace UE::Net
@@ -23,14 +24,14 @@ IRISCORE_API uint64 Init();
 /** Trigger a breakpoint and return true if the object contains the current debug name */
 bool BreakOnObjectName(UObject* Object);
 
-/** Trigger a breakpoint and return true if the NetHandle is the current debug nethandle */
-bool BreakOnNetHandle(FNetHandle NetHandle);
+/** Trigger a breakpoint and return true if the NetRefHandle is the current debug NetRefHandle */
+bool BreakOnNetRefHandle(FNetRefHandle NetRefHandle);
 
 /** Trigger a breakpoint and return true if the name contains the debug RPC string */
 bool BreakOnRPCName(FName RPCName);
 
 /** Output state data to StringBuilder for the specified Handle */
-void NetObjectStateToString(FStringBuilderBase& StringBuilder, FNetHandle NetHandle);
+void NetObjectStateToString(FStringBuilderBase& StringBuilder, FNetRefHandle RefHandle);
 
 /** Find all handles references registered for a protocol and output to StringBuilder */
 void NetObjectProtocolReferencesToString(FStringBuilderBase& StringBuilder, uint64 ProtocolId, uint32 ReplicationSystemId);
@@ -40,14 +41,14 @@ void NetObjectProtocolReferencesToString(FStringBuilderBase& StringBuilder, uint
 /** Get the ReplicationSystem with the Id */
 extern "C" IRISCORE_API UReplicationSystem* GetReplicationSystemForDebug(uint32 Id);
 
-/** DebugOutputObject state for Handle specified by NetHandleId and ReplicationSystemId and */
-extern "C" IRISCORE_API void DebugOutputNetObjectState(uint32 NetHandleId, uint32 ReplicationSystemId);
+/** DebugOutputObject state for Handle specified by NetRefHandleId and ReplicationSystemId. */
+extern "C" IRISCORE_API void DebugOutputNetObjectState(uint32 NetRefHandleId, uint32 ReplicationSystemId);
 
 /**
-	Variant of NetObjectStateToString that can be used from breakpoints and in watch window to print the current state of a NetHandle
-	NOTE: Use only for debugging as this variant uses a static buffer which is not thread safe	
-*/
-extern "C" IRISCORE_API const TCHAR* DebugNetObjectStateToString(uint32 NetHandleId, uint32 ReplicationSystemId);
+ * Variant of NetObjectStateToString that can be used from breakpoints and in watch window to print the current state of a NetRefHandle.
+ * NOTE: Use only for debugging as this variant uses a static buffer which is not thread safe.
+ */
+extern "C" IRISCORE_API const TCHAR* DebugNetObjectStateToString(uint32 NetRefHandleId, uint32 ReplicationSystemId);
 
 /** Find all handles references registered for a protocol and output to DebugOutput in debugger */
 extern "C" IRISCORE_API void DebugOutputNetObjectProtocolReferences(uint64 ProtocolId, uint32 ReplicationSystemId);
@@ -55,24 +56,24 @@ extern "C" IRISCORE_API void DebugOutputNetObjectProtocolReferences(uint64 Proto
 /** Get info about replicated object */
 struct FNetReplicatedObjectDebugInfo
 {
-	const FNetHandle* Handle;
-	uint32 InternalNetHandleIndex;
+	const FNetRefHandle* RefHandle;
+	uint32 InternalNetRefIndex;
 	const UReplicationSystem* ReplicationSystem;
 	const FReplicationProtocol* Protocol;
 	const FReplicationInstanceProtocol* InstanceProtocol;
 };
 
-/** Look up replicated NetHandle from Instance pointer and return debug information, this variant searches all active replication systems */
+/** Look up replicated handle from Instance pointer and return debug information. This variant searches all active replication systems and returns information for the first one replicating the Instance. */
 extern "C" IRISCORE_API FNetReplicatedObjectDebugInfo DebugNetObject(UObject* Instance);
 
-/** Look up replicated NetHandle from Instance pointer and return debug information, this variant only searches the replicatioSystem with the specified id */
+/** Look up replicated handle from Instance pointer and return debug information. This variant only searches the ReplicationSystem with the specified id. */
 extern "C" IRISCORE_API FNetReplicatedObjectDebugInfo DebugNetObjectById(UObject* Instance, uint32 ReplicationSystemId);
 
 /** Look up replicated handle and return debug information */
-extern "C" IRISCORE_API FNetReplicatedObjectDebugInfo DebugNetHandle(FNetHandle Handle);
+extern "C" IRISCORE_API FNetReplicatedObjectDebugInfo DebugNetRefHandle(FNetRefHandle Handle);
 
 /** Look up replicated handle specified by handle id and replicationsystem id and return debug information */
-extern "C" IRISCORE_API FNetReplicatedObjectDebugInfo DebugNetHandleById(uint32 NetHandleId, uint32 ReplicationSystemId);
+extern "C" IRISCORE_API FNetReplicatedObjectDebugInfo DebugNetRefHandleById(uint32 NetRefHandleId, uint32 ReplicationSystemId);
 
 /** Variant of DebugOutputNetObjectProtocolReferences that can be used from breakpoints or in watch window to find all handles references registered for a protocol and output to DebugOutput in debugger
  * NOTE: Use only for debugging as this variant uses a static buffer which is not thread safe	
diff --git a/Engine/Source/Runtime/Experimental/Iris/Core/Public/Iris/Core/NetObjectReference.h b/Engine/Source/Runtime/Experimental/Iris/Core/Public/Iris/Core/NetObjectReference.h
--- a/Engine/Source/Runtime/Experimental/Iris/Core/Public/Iris/Core/NetObjectReference.h
+++ b/Engine/Source/Runtime/Experimental/Iris/Core/Public/Iris/Core/NetObjectReference.h
@@ -45,9 +45,9 @@ public:
 	inline bool operator==(const FNetObjectReference& Other) const { return RefHandle == Other.GetRefHandle() && PathToken == Other.PathToken; }
 	/** Returns whether this reference doesn't point to the same network addressable objects as some other reference. */
 	inline bool operator!=(const FNetObjectReference& Other) const { return !(*this == Other); }
-
-	/** Returns the NetHandle part of the object reference. */
-	FNetHandle GetRefHandle() const { return RefHandle; }
+		
+	/** Returns the NetRefHandle part of the object reference. */
+	FNetRefHandle GetRefHandle() const { return RefHandle; }
 	/** Returns the NetToken part of the object reference. */
 	FNetToken GetPathToken() const { return PathToken; }
 	/** Returns whether the object reference can be exported. */
@@ -78,19 +78,19 @@ public:
 private:
 	friend class Private::FObjectReferenceCache;
 
-	inline FNetObjectReference(FNetHandle InHandle, FNetToken InPathToken, ENetObjectReferenceTraits InTraits)
+	inline FNetObjectReference(FNetRefHandle InHandle, FNetToken InPathToken, ENetObjectReferenceTraits InTraits)
 		: RefHandle(InHandle)
 		, PathToken(InPathToken)
 		, Traits(InTraits)
 	{
 	}
 
-	inline explicit FNetObjectReference(FNetHandle Handle)
+	inline explicit FNetObjectReference(FNetRefHandle Handle)
 	: FNetObjectReference(Handle, FNetToken(), ENetObjectReferenceTraits::None)
 	{
 	}
 
-	FNetHandle RefHandle;
+	FNetRefHandle RefHandle;
 	FNetToken PathToken;
 	ENetObjectReferenceTraits Traits;
 };
diff --git a/Engine/Source/Runtime/Experimental/Iris/Core/Public/Iris/ReplicationSystem/NetBlob/NetBlob.h b/Engine/Source/Runtime/Experimental/Iris/Core/Public/Iris/ReplicationSystem/NetBlob/NetBlob.h
--- a/Engine/Source/Runtime/Experimental/Iris/Core/Public/Iris/ReplicationSystem/NetBlob/NetBlob.h
+++ b/Engine/Source/Runtime/Experimental/Iris/Core/Public/Iris/ReplicationSystem/NetBlob/NetBlob.h
@@ -76,10 +76,10 @@ public:
 	IRISCORE_API static void DeserializeCreationInfo(FNetSerializationContext& Context, FNetBlobCreationInfo& OutCreationInfo);
 
 	/** Serialize the blob together with/targeting a specific object knowing the NetHandle has already been serialized. */
-	IRISCORE_API virtual void SerializeWithObject(FNetSerializationContext& Context, FNetHandle NetHandle) const;
+	IRISCORE_API virtual void SerializeWithObject(FNetSerializationContext& Context, FNetRefHandle RefHandle) const;
 
 	/** Deserialize a blob that was serialized with SerializeWithObject. */
-	IRISCORE_API virtual void DeserializeWithObject(FNetSerializationContext& Context, FNetHandle NetHandle);
+	IRISCORE_API virtual void DeserializeWithObject(FNetSerializationContext& Context, FNetRefHandle RefHandle);
 
 	/** Serialize the blob. */
 	IRISCORE_API virtual void Serialize(FNetSerializationContext& Context) const;
@@ -155,10 +155,10 @@ protected:
 	IRISCORE_API void DeserializeObjectReference(FNetSerializationContext& Context);
 
 	/** Serializes only the subobject reference and assumes the passed NetHandle is the owner. */
-	IRISCORE_API void SerializeSubObjectReference(FNetSerializationContext& Context, FNetHandle NetHandle) const;
+	IRISCORE_API void SerializeSubObjectReference(FNetSerializationContext& Context, FNetRefHandle RefHandle) const;
 
 	/** Deserializes a subobject reference that was serialized using SerializeSubObjectReference with the same NetHandle. */
-	IRISCORE_API void DeserializeSubObjectReference(FNetSerializationContext& Context, FNetHandle NetHandle);
+	IRISCORE_API void DeserializeSubObjectReference(FNetSerializationContext& Context, FNetRefHandle RefHandle);
 
 	/** Set the owner and subobject references. */
 	void SetNetObjectReference(const FNetObjectReference& InQueueOwnerReference, const FNetObjectReference& InTargetObjectReference);
diff --git a/Engine/Source/Runtime/Experimental/Iris/Core/Public/Iris/ReplicationSystem/NetBlob/NetBlobAssembler.h b/Engine/Source/Runtime/Experimental/Iris/Core/Public/Iris/ReplicationSystem/NetBlob/NetBlobAssembler.h
--- a/Engine/Source/Runtime/Experimental/Iris/Core/Public/Iris/ReplicationSystem/NetBlob/NetBlobAssembler.h
+++ b/Engine/Source/Runtime/Experimental/Iris/Core/Public/Iris/ReplicationSystem/NetBlob/NetBlobAssembler.h
@@ -31,7 +31,7 @@ public:
 	 * @param Context FNetSerializationContext for error reporting. Call HasError() on it afterwards to check whether something went wrong.
 	 * @param NetHandle The NetHandle needs to remain constant for every call until the original blob has been assembled or the first part of a new blob is added.
 	 */
-	IRISCORE_API void AddPartialNetBlob(FNetSerializationContext& Context, FNetHandle NetHandle, const TRefCountPtr<FPartialNetBlob>& PartialNetBlob);
+	IRISCORE_API void AddPartialNetBlob(FNetSerializationContext& Context, FNetRefHandle RefHandle, const TRefCountPtr<FPartialNetBlob>& PartialNetBlob);
 
 	/** Returns true if all parts of the split blob have been added and is ready to be assembled. */
 	bool IsReadyToAssemble() const { return bIsReadyToAssemble; }
@@ -49,7 +49,7 @@ private:
 	FNetBlobCreationInfo NetBlobCreationInfo;
 	TArray<uint32> Payload;
 	FNetBitStreamWriter BitWriter;
-	FNetHandle NetHandle;
+	FNetRefHandle RefHandle;
 	uint32 NextPartIndex = 0;
 	uint32 PartCount = 0;
 	uint32 FirstPayloadBitCount = 0;
diff --git a/Engine/Source/Runtime/Experimental/Iris/Core/Public/Iris/ReplicationSystem/NetBlob/RawDataNetBlob.h b/Engine/Source/Runtime/Experimental/Iris/Core/Public/Iris/ReplicationSystem/NetBlob/RawDataNetBlob.h
--- a/Engine/Source/Runtime/Experimental/Iris/Core/Public/Iris/ReplicationSystem/NetBlob/RawDataNetBlob.h
+++ b/Engine/Source/Runtime/Experimental/Iris/Core/Public/Iris/ReplicationSystem/NetBlob/RawDataNetBlob.h
@@ -9,7 +9,7 @@ namespace UE::Net
 
 /**
  * Helper class for stateless data, such as when arbitrary data has been serialized to a bitstream.
- * The serialization will simply serialize the raw data regardless of whether a NetHandle is provided or not.
+ * The serialization will simply serialize the raw data regardless of whether a NetRefHandle is provided or not.
  * Things like splitting and assembling have optimized code paths for this type of blob.
  * You can inherit from this blob type but you cannot override the serialization functions.
  * @note Sending huge blobs that require splitting and assembling is strongly discouraged. 
@@ -51,8 +51,8 @@ protected:
 	IRISCORE_API void InternalDeserialize(FNetSerializationContext& Context);
 
 private:
-	virtual void SerializeWithObject(FNetSerializationContext& Context, FNetHandle NetHandle) const override final;
-	virtual void DeserializeWithObject(FNetSerializationContext& Context, FNetHandle NetHandle) override final;
+	virtual void SerializeWithObject(FNetSerializationContext& Context, FNetRefHandle RefHandle) const override final;
+	virtual void DeserializeWithObject(FNetSerializationContext& Context, FNetRefHandle RefHandle) override final;
 
 	virtual void Serialize(FNetSerializationContext& Context) const override final;
 	virtual void Deserialize(FNetSerializationContext& Context) override final;
diff --git a/Engine/Source/Runtime/Experimental/Iris/Core/Public/Iris/ReplicationSystem/NetBlob/ReliableNetBlobQueue.h b/Engine/Source/Runtime/Experimental/Iris/Core/Public/Iris/ReplicationSystem/NetBlob/ReliableNetBlobQueue.h
--- a/Engine/Source/Runtime/Experimental/Iris/Core/Public/Iris/ReplicationSystem/NetBlob/ReliableNetBlobQueue.h
+++ b/Engine/Source/Runtime/Experimental/Iris/Core/Public/Iris/ReplicationSystem/NetBlob/ReliableNetBlobQueue.h
@@ -57,22 +57,22 @@ public:
 	IRISCORE_API void Pop();
 
 	/**
-	 * Serializes as many blobs as possible using their respective SerializeWithObject() method. It is assumed the NetHandle will be 
+	 * Serializes as many blobs as possible using their respective SerializeWithObject() method. It is assumed the NetRefHandle will be 
 	 * reconstructed somehow on the receiving end and passed to DeserializeWithObject().
-	 * This provides an opportunity fot FNetObjectAttachments, such as FNetRPCs, to avoid serializing the same NetHandle redundantly.
+	 * This provides an opportunity fot FNetObjectAttachments, such as FNetRPCs, to avoid serializing the same NetRefHandle redundantly.
 	 * @param Context A FNetSerializationContext.
-	 * @param NetHandle The handle for the blobs' target object.
+	 * @param RefHandle The handle for the blobs' target object.
 	 * @param OutRecord The record to pass to CommitReplicationRecord() if a packet containing the serialized data was sent.
 	 * @return The number of blobs that were serialized.
 	 */
-	IRISCORE_API uint32 SerializeWithObject(FNetSerializationContext& Context, FNetHandle NetHandle, ReplicationRecord& OutRecord);
+	IRISCORE_API uint32 SerializeWithObject(FNetSerializationContext& Context, FNetRefHandle RefHandle, ReplicationRecord& OutRecord);
 
 	/**
 	 * Deserializes blobs with object using their respective DeserializeWithObject() method.
 	 * @return The number of blobs that were deserialized.
 	 * @see SerializeWithObject
 	 */
-	IRISCORE_API uint32 DeserializeWithObject(FNetSerializationContext& Context, FNetHandle NetHandle);
+	IRISCORE_API uint32 DeserializeWithObject(FNetSerializationContext& Context, FNetRefHandle RefHandle);
 
 	/**
 	 * Serializes as many blobs as possible using their respective Serialize() method.
@@ -129,8 +129,8 @@ private:
 
 	bool IsValidReceiveSequence(uint32 Seq) const;
 
-	uint32 SerializeInternal(FNetSerializationContext& Context, FNetHandle NetHandle, FReliableNetBlobQueue::ReplicationRecord& OutRecord, const bool bSerializeWithObject);
-	uint32 DeserializeInternal(FNetSerializationContext& Context, FNetHandle NetHandle, const bool bSerializeWithObject);
+	uint32 SerializeInternal(FNetSerializationContext& Context, FNetRefHandle RefHandle, FReliableNetBlobQueue::ReplicationRecord& OutRecord, const bool bSerializeWithObject);
+	uint32 DeserializeInternal(FNetSerializationContext& Context, FNetRefHandle RefHandle, const bool bSerializeWithObject);
 
 	void OnPacketDelivered(ReplicationRecord Record);
 	void OnPacketDropped(ReplicationRecord Record);
diff --git a/Engine/Source/Runtime/Experimental/Iris/Core/Public/Iris/ReplicationSystem/NetBlob/ShrinkWrapNetBlob.h b/Engine/Source/Runtime/Experimental/Iris/Core/Public/Iris/ReplicationSystem/NetBlob/ShrinkWrapNetBlob.h
--- a/Engine/Source/Runtime/Experimental/Iris/Core/Public/Iris/ReplicationSystem/NetBlob/ShrinkWrapNetBlob.h
+++ b/Engine/Source/Runtime/Experimental/Iris/Core/Public/Iris/ReplicationSystem/NetBlob/ShrinkWrapNetBlob.h
@@ -23,8 +23,8 @@ public:
 
 private:
 	virtual TArrayView<const FNetObjectReference> GetExports() const override final;
-	virtual void SerializeWithObject(FNetSerializationContext& Context, FNetHandle NetHandle) const override final;
-	virtual void DeserializeWithObject(FNetSerializationContext& Context, FNetHandle NetHandle) override final;
+	virtual void SerializeWithObject(FNetSerializationContext& Context, FNetRefHandle RefHandle) const override final;
+	virtual void DeserializeWithObject(FNetSerializationContext& Context, FNetRefHandle RefHandle) override final;
 
 	virtual void Serialize(FNetSerializationContext& Context) const override final;
 	virtual void Deserialize(FNetSerializationContext& Context) override final;
diff --git a/Engine/Source/Runtime/Experimental/Iris/Core/Public/Iris/ReplicationSystem/NetHandle.h b/Engine/Source/Runtime/Experimental/Iris/Core/Public/Iris/ReplicationSystem/NetHandle.h
--- a/Engine/Source/Runtime/Experimental/Iris/Core/Public/Iris/ReplicationSystem/NetHandle.h
+++ b/Engine/Source/Runtime/Experimental/Iris/Core/Public/Iris/ReplicationSystem/NetHandle.h
@@ -12,17 +12,16 @@ class FString;
 
 namespace UE::Net::Private
 {
-	class FNetHandleManager;
+	class FNetRefHandleManager;
 }
 
 namespace UE::Net
 {
 
 /**
-* FNetHandle
-*/
-
-class FNetHandle
+ * FNetRefHandle
+ */
+class FNetRefHandle
 {
 public:
 	enum { Invalid = 0 };
@@ -30,7 +29,7 @@ public:
 	enum { ReplicationSystemIdBits = 4 };
 
 public:
-	FNetHandle() : Value(Invalid) {}
+	FNetRefHandle() : Value(Invalid) {}
 
 	uint32 GetId() const { return Id; }
 	uint32 GetReplicationSystemId() const { check(ReplicationSystemId != 0U); return ReplicationSystemId - 1U; }
@@ -41,21 +40,17 @@ public:
 	bool IsStatic() const { return Id & StaticIdMask; }
 	bool IsDynamic() const { return IsValid() && !IsStatic(); }
 
-	bool operator==(const FNetHandle& Other)const { return Id == Other.Id; }
-	bool operator<(const FNetHandle& Other)const { return Id < Other.Id; }
-	bool operator!=(const FNetHandle& Other)const { return Id != Other.Id; }
+	bool operator==(const FNetRefHandle& Other)const { return Id == Other.Id; }
+	bool operator<(const FNetRefHandle& Other)const { return Id < Other.Id; }
+	bool operator!=(const FNetRefHandle& Other)const { return Id != Other.Id; }
 
 	IRISCORE_API FString ToString() const;
 
-	static bool FullCompare(FNetHandle A, FNetHandle B) { return A.Value == B.Value; }
+	static bool FullCompare(FNetRefHandle A, FNetRefHandle B) { return A.Value == B.Value; }
 	
 private:
-	static uint32 MakeNetHandleId(uint32 Seed, bool bIsStatic);
-	static FNetHandle MakeNetHandle(uint32 Id, uint32 ReplicationSystemId);
-	static FNetHandle MakeNetHandleFromId(uint32 Id);
-
-	friend uint32 GetTypeHash(const FNetHandle& Handle);
-	friend Private::FNetHandleManager;
+	friend uint32 GetTypeHash(const FNetRefHandle& Handle);
+	friend Private::FNetRefHandleManager;
 
 	static constexpr uint32 StaticIdMask = 1U;
 	static constexpr uint32 IdMask = ~0U;
@@ -72,18 +67,17 @@ private:
 	};
 };
 
-FORCEINLINE uint32 GetTypeHash(const FNetHandle& Handle)
+FORCEINLINE uint32 GetTypeHash(const FNetRefHandle& Handle)
 {
-	// Id and type is relevant
 	return ::GetTypeHash(Handle.GetId());
 }
 
-FORCEINLINE uint32 GetObjectIdForNetTrace(const FNetHandle& Handle)
+FORCEINLINE uint32 GetObjectIdForNetTrace(const FNetRefHandle& Handle)
 {
 	return Handle.GetId();
 }
 
 }
 
-FStringBuilderBase& operator<<(FStringBuilderBase& Builder, const UE::Net::FNetHandle& NetHandle);
-FAnsiStringBuilderBase& operator<<(FAnsiStringBuilderBase& Builder, const UE::Net::FNetHandle& NetHandle);
+IRISCORE_API FStringBuilderBase& operator<<(FStringBuilderBase& Builder, const UE::Net::FNetRefHandle& NetHandle);
+IRISCORE_API FAnsiStringBuilderBase& operator<<(FAnsiStringBuilderBase& Builder, const UE::Net::FNetRefHandle& NetHandle);
