Iris - Implemented globally unique NetHandles.

diff --git a/Engine/Source/Runtime/Experimental/Iris/Core/Public/Iris/ReplicationSystem/ObjectReplicationBridge.h b/Engine/Source/Runtime/Experimental/Iris/Core/Public/Iris/ReplicationSystem/ObjectReplicationBridge.h
--- a/Engine/Source/Runtime/Experimental/Iris/Core/Public/Iris/ReplicationSystem/ObjectReplicationBridge.h
+++ b/Engine/Source/Runtime/Experimental/Iris/Core/Public/Iris/ReplicationSystem/ObjectReplicationBridge.h
@@ -2,8 +2,9 @@
 
 #pragma once
 
-#include "CoreTypes.h"
+#include "HAL/Platform.h"
 #include "Net/Core/NetBitArray.h"
+#include "Net/Core/NetHandle/NetHandle.h"
 #include "Iris/ReplicationSystem/ReplicationBridge.h"
 #include "Delegates/IDelegateInstance.h"
 
@@ -18,7 +19,7 @@ namespace UE::Net
 	class FNetObjectReference;
 	namespace Private
 	{
-		typedef uint32 FInternalNetHandle;
+		typedef uint32 FInternalNetRefIndex;
 		class FObjectPollFrequencyLimiter;
 	}
 }
@@ -35,7 +36,7 @@ class UObjectReplicationBridge : public UReplicationBridge
 public:
 	using UReplicationBridge::EndReplication;
 
-	struct FCreateNetHandleParams
+	struct FCreateNetRefHandleParams
 	{
 		uint32 bCanReceive : 1;
 		uint32 bNeedsPreUpdate : 1;
@@ -54,17 +55,20 @@ public:
 		  */
 		uint8 PollFramePeriod;
 	};
-	IRISCORE_API static FCreateNetHandleParams DefaultCreateNetHandleParams;
+	IRISCORE_API static FCreateNetRefHandleParams DefaultCreateNetRefHandleParams;
 
 	IRISCORE_API UObjectReplicationBridge();
 
 	/** Get the Object from a replicated handle, if the handle is invalid or not is a replicated handle the function will return nullptr */
-	IRISCORE_API UObject* GetReplicatedObject(FNetHandle Handle) const;
+	IRISCORE_API UObject* GetReplicatedObject(FNetRefHandle Handle) const;
 
-	/** Get NetHandle from a Replicated UObject */
-	IRISCORE_API FNetHandle GetReplicatedHandle(const UObject* Object) const;
+	/** Get NetRefHandle from a replicated UObject. */
+	IRISCORE_API FNetRefHandle GetReplicatedRefHandle(const UObject* Object) const;
 
-	/** Try to resolve UObject from NetObjectReference, this function tries to resolve the object by loading if necessary.*/
+	/** Get NetRefHandle from a NetHandle. */
+	IRISCORE_API FNetRefHandle GetReplicatedRefHandle(FNetHandle Handle) const;
+
+	/** Try to resolve UObject from NetObjectReference, this function tries to resolve the object by loading if necessary. */
 	IRISCORE_API UObject* ResolveObjectReference(const UE::Net::FNetObjectReference& ObjectRef, const UE::Net::FNetObjectResolveContext& ResolveContext);
 
 	/** Get or create NetObjectReference for object instance. */
@@ -73,25 +77,26 @@ public:
 	/** Get or create NetObjectReference for object identified by path relative to outer. */
 	IRISCORE_API UE::Net::FNetObjectReference GetOrCreateObjectReference(const FString& Path, const UObject* Outer) const;
 
-	/** Create NetHandle and start replicating the Instance */
-	IRISCORE_API FNetHandle BeginReplication(UObject* Instance, const FCreateNetHandleParams& Params = DefaultCreateNetHandleParams);
+	/** Begin replicating the Instance and return a valid NetRefHandle for the Instance if successful. */
+	IRISCORE_API FNetRefHandle BeginReplication(UObject* Instance, const FCreateNetRefHandleParams& Params = DefaultCreateNetRefHandleParams);
 
 	/**
-	 * Create NetHandle and start replicating the Instance as a SubObject of the OwnerHandle. If InsertRelativeSubObjectHandle is valid
+	 * Begin replicating the Instance as a subobject of the OwnerHandle. If InsertRelativeSubObjectHandle is valid
 	 * the new subobject will be inserted in the subobject replication list next to the specified handle and the wanted insertion order.
-	 * Default behavior is to always add new subobjects at the end of the list.
+	 * Default behavior is to always add new subobjects at the end of the list. Returns a valid NetRefHandle for the Instance if successful.
 	 */
-	IRISCORE_API FNetHandle BeginReplication(FNetHandle OwnerHandle, UObject* Instance, FNetHandle InsertRelativeToSubObjectHandle, ESubObjectInsertionOrder InsertionOrder = UReplicationBridge::ESubObjectInsertionOrder::None, const FCreateNetHandleParams& Params = DefaultCreateNetHandleParams);
+	IRISCORE_API FNetRefHandle BeginReplication(FNetRefHandle OwnerHandle, UObject* Instance, FNetRefHandle InsertRelativeToSubObjectHandle, ESubObjectInsertionOrder InsertionOrder = UReplicationBridge::ESubObjectInsertionOrder::None, const FCreateNetRefHandleParams& Params = DefaultCreateNetRefHandleParams);
 
-	/** Create NetHandle and start replicating the Instance as a SubObject of the OwnerHandle */
-	FNetHandle BeginReplication(FNetHandle OwnerHandle, UObject* Instance, const FCreateNetHandleParams& Params = DefaultCreateNetHandleParams) { return BeginReplication(OwnerHandle, Instance, FNetHandle(), ESubObjectInsertionOrder::None, Params); };
+	/** Create handle and start replicating the Instance as a SubObject of the OwnerHandle. */
+	/** Begin replicating the Instance as a subobject of the OwnerHandle and return a valid NetRefHandle for the Instance if successful. */
+	FNetRefHandle BeginReplication(FNetRefHandle OwnerHandle, UObject* Instance, const FCreateNetRefHandleParams& Params = DefaultCreateNetRefHandleParams);
 
 	/** 
-	 * Set NetCondition for a subobject, the condition is used to determine if the SubObject should replicated or not.
-	 * @note: As the filtering is done at the serialization level it is typically more efficient to use a separate NetObject for connection 
+	 * Set NetCondition for a subobject, the condition is used to determine if the SubObject should replicate or not.
+	 * @note As the filtering is done at the serialization level it is typically more efficient to use a separate NetObject for connection 
 	 * specific data as filtering can then be done at a higher level.
 	 */
-	IRISCORE_API void SetSubObjectNetCondition(FNetHandle SubObjectHandle, ELifetimeCondition Condition);
+	IRISCORE_API void SetSubObjectNetCondition(FNetRefHandle SubObjectHandle, ELifetimeCondition Condition);
 
 	/** Stop replicating the object. */
 	IRISCORE_API void EndReplication(UObject* Instance, EEndReplicationFlags EndReplicationFlags = EEndReplicationFlags::Destroy, FEndReplicationParameters* Parameters = nullptr);
@@ -123,21 +128,21 @@ public:
 	 * Dependent objects cannot be filtered out by dynamic filtering unless the parent is also filtered out.
 	 * @note: There is no guarantee that the data will end up in the same packet so it is a very loose form of dependency.
 	 */
-	IRISCORE_API void AddDependentObject(FNetHandle Parent, FNetHandle DependentObject, EDependentWarnFlags WarnFlags = EDependentWarnFlags::WarnAlreadyDependent);
+	IRISCORE_API void AddDependentObject(FNetRefHandle Parent, FNetRefHandle DependentObject, EDependentWarnFlags WarnFlags = EDependentWarnFlags::WarnAlreadyDependent);
 
 	/** Remove dependent object from parent. The dependent object will function as a standard standalone replicated object. */
-	IRISCORE_API void RemoveDependentObject(FNetHandle Parent, FNetHandle DependentObject);
+	IRISCORE_API void RemoveDependentObject(FNetRefHandle Parent, FNetRefHandle DependentObject);
 
 	// Dormancy support
 
 	/** Set whether object should go dormant. If dormancy is enabled any dirty state will be replicated first. */
-	IRISCORE_API void SetObjectWantsToBeDormant(FNetHandle Handle, bool bWantsToBeDormant);
+	IRISCORE_API void SetObjectWantsToBeDormant(FNetRefHandle Handle, bool bWantsToBeDormant);
 
 	/** Returns whether the object wants to be dormant. */
-	IRISCORE_API bool GetObjectWantsToBeDormant(FNetHandle Handle) const;
+	IRISCORE_API bool GetObjectWantsToBeDormant(FNetRefHandle Handle) const;
 
 	/** Trigger replication of dirty state for object wanting to be dormant. */
-	IRISCORE_API void ForceUpdateWantsToBeDormantObject(FNetHandle Handle);	
+	IRISCORE_API void ForceUpdateWantsToBeDormantObject(FNetRefHandle Handle);	
 
 protected:
 	IRISCORE_API virtual ~UObjectReplicationBridge();
@@ -146,15 +151,15 @@ protected:
 
 	IRISCORE_API virtual void Initialize(UReplicationSystem* InReplicationSystem) override;
 	IRISCORE_API virtual void Deinitialize() override;
-	IRISCORE_API virtual void PreSendUpdateSingleHandle(FNetHandle Handle) override;	
+	IRISCORE_API virtual void PreSendUpdateSingleHandle(FNetRefHandle Handle) override;	
 	IRISCORE_API virtual void PreSendUpdate() override;	
 	IRISCORE_API virtual void UpdateInstancesWorldLocation() override;
 	IRISCORE_API virtual void PruneStaleObjects() override;	
-	IRISCORE_API virtual bool WriteNetHandleCreationInfo(FReplicationBridgeSerializationContext& Context, FNetHandle Handle) override;
-	IRISCORE_API virtual FNetHandle CreateNetHandleFromRemote(FNetHandle SubObjectOwnerNetHandle, FNetHandle WantedNetHandle, FReplicationBridgeSerializationContext& Context) override;
-	IRISCORE_API virtual void PostApplyInitialState(FNetHandle Handle) override;
-	IRISCORE_API virtual void DetachInstanceFromRemote(FNetHandle Handle, bool bTearOff, bool bShouldDestroyInstance) override;
-	IRISCORE_API virtual void DetachInstance(FNetHandle Handle) override;
+	IRISCORE_API virtual bool WriteNetRefHandleCreationInfo(FReplicationBridgeSerializationContext& Context, FNetRefHandle Handle) override;
+	IRISCORE_API virtual FNetRefHandle CreateNetRefHandleFromRemote(FNetRefHandle SubObjectOwnerNetHandle, FNetRefHandle WantedNetHandle, FReplicationBridgeSerializationContext& Context) override;
+	IRISCORE_API virtual void PostApplyInitialState(FNetRefHandle Handle) override;
+	IRISCORE_API virtual void DetachInstanceFromRemote(FNetRefHandle Handle, bool bTearOff, bool bShouldDestroyInstance) override;
+	IRISCORE_API virtual void DetachInstance(FNetRefHandle Handle) override;
 
 protected:
 	/**
@@ -165,35 +170,35 @@ protected:
 	*/
 
 	/** Write data required to instantiate the Handle. */
-	virtual bool WriteCreationHeader(UE::Net::FNetSerializationContext& Context, FNetHandle Handle) { return false; };
+	virtual bool WriteCreationHeader(UE::Net::FNetSerializationContext& Context, FNetRefHandle Handle) { return false; };
 
 	/** Create Header and read data required to instantiate the Handle. */
 	virtual FCreationHeader* ReadCreationHeader(UE::Net::FNetSerializationContext& Context) { return nullptr; };
 
 	/** Called when we instantiate/find object instance requested by remote. */
-	virtual UObject* BeginInstantiateFromRemote(FNetHandle SubObjectOwnerHandle, const UE::Net::FNetObjectResolveContext& ResolveContext, const FCreationHeader* Header) { return nullptr; };
+	virtual UObject* BeginInstantiateFromRemote(FNetRefHandle SubObjectOwnerHandle, const UE::Net::FNetObjectResolveContext& ResolveContext, const FCreationHeader* Header) { return nullptr; };
 
 	/** Invoked before we start applying state data to instance on remote end. */
 	virtual bool OnInstantiatedFromRemote(UObject* Instance, const FCreationHeader* InHeader, uint32 ConnectionId) const { return true; }
 
 	/** Invoked after remote NetHandle has been created and initial state is applied. */
-	virtual void EndInstantiateFromRemote(FNetHandle Handle) {};
+	virtual void EndInstantiateFromRemote(FNetRefHandle Handle) {};
 
 	/** Destroy or tear-off the game instance on request from remote. */
 	virtual void DestroyInstanceFromRemote(UObject* Instance, bool bTearOff) {};
 
 protected:
 	/** Lookup the UObject associated with the provided Handle. This function will not try to resolve the reference. */
-	IRISCORE_API UObject* GetObjectFromReferenceHandle(FNetHandle NetHandle) const;
+	IRISCORE_API UObject* GetObjectFromReferenceHandle(FNetRefHandle RefHandle) const;
 
 	/** Helper method that calls provided PreUpdateFunction and polls state data for all replicated instances with the NeedsPoll trait. */
-	using FInstancePreUpdateFunction = void(*)(FNetHandle, UObject*, const UReplicationBridge*);
+	using FInstancePreUpdateFunction = void(*)(FNetRefHandle, UObject*, const UReplicationBridge*);
 
 	/** Set the function that we should call before copying state data. */
 	IRISCORE_API void SetInstancePreUpdateFunction(FInstancePreUpdateFunction InPreUpdateFunction);
 	
 	/** Helper method to get the world location for replicated instances with the HasWorldLocation trait. */
-	using FInstanceGetWorldLocationFunction = FVector(*)(FNetHandle, const UObject*);
+	using FInstanceGetWorldLocationFunction = FVector(*)(FNetRefHandle, const UObject*);
 
 	/** Set the function that we should call to get the world location of an object. */
 	IRISCORE_API void SetInstanceGetWorldLocationFunction(FInstanceGetWorldLocationFunction InGetWorldLocationFunction);
@@ -201,10 +206,10 @@ protected:
 	// Poll frequency support
 
 	/** Set poll frame period for a specific object. Zero based period where 0 means every frame and 255 means every 256th frame. */
-	IRISCORE_API void SetPollFramePeriod(UE::Net::Private::FInternalNetHandle InternalReplicationIndex, uint8 FramePeriod);
+	IRISCORE_API void SetPollFramePeriod(UE::Net::Private::FInternalNetRefIndex InternalReplicationIndex, uint8 FramePeriod);
 
-	/** Force polling of Object when ObjectToPoll is polled. */
-	IRISCORE_API void SetPollWithObject(FNetHandle ObjectToPollWith, FNetHandle Object);
+	/** Force polling of Object when ObjectToPollWith is polled. */
+	IRISCORE_API void SetPollWithObject(FNetRefHandle ObjectToPollWith, FNetRefHandle Object);
 
 	/** Re-initialize config-driven parameters found in ObjectReplicationBridgeConfig. */
 	IRISCORE_API void LoadConfig();
@@ -220,11 +225,13 @@ protected:
 
 private:
 
-	void PreUpdateAndPollImpl(FNetHandle Handle);
+	void PreUpdateAndPollImpl(FNetRefHandle Handle);
+
+	void RegisterRemoteInstance(FNetRefHandle Handle, UObject* InstancePtr, const UE::Net::FReplicationProtocol* Protocol, UE::Net::FReplicationInstanceProtocol* InstanceProtocol, const FCreationHeader* Header, uint32 ConnectionId);
 
-	void RegisterRemoteInstance(FNetHandle Handle, UObject* InstancePtr, const UE::Net::FReplicationProtocol* Protocol, UE::Net::FReplicationInstanceProtocol* InstanceProtocol, const FCreationHeader* Header, uint32 ConnectionId);
+	void UnregisterRemoteInstance(FNetRefHandle Handle, bool bTearOff, bool bShouldDestroyInstance);
 
-	void UnregisterRemoteInstance(FNetHandle Handle, bool bTearOff, bool bShouldDestroyInstance);
+	void SetNetPushIdOnInstance(UE::Net::FReplicationInstanceProtocol* InstanceProtocol, FNetHandle NetHandle);
 
 	/** Finds a poll period override if it exists. OutPollPeriod will only be modified if this method returns true. */
 	bool FindPollInfo(const UClass* Class, uint8& OutPollPeriod);
@@ -277,7 +284,7 @@ private:
 	TMap<FName, bool> ClassesWithDeltaCompression;
 
 	// Array of dormant objects that has requested a flush
-	TArray<FNetHandle> DormantHandlesPendingFlush;
+	TArray<FNetRefHandle> DormantHandlesPendingFlush;
 
 	// Objects which has object references and could be affected by garbage collection.
 	UE::Net::FNetBitArray ObjectsWithObjectReferences;
@@ -291,3 +298,9 @@ private:
 	bool bHasPollOverrides = false;
 	bool bHasDirtyClassesInPollPeriodOverrides = false;
 };
+
+
+inline UE::Net::FNetRefHandle UObjectReplicationBridge::BeginReplication(UE::Net::FNetRefHandle OwnerHandle, UObject* Instance, const FCreateNetRefHandleParams& Params)
+{
+	return BeginReplication(OwnerHandle, Instance, FNetRefHandle(), ESubObjectInsertionOrder::None, Params);
+}
diff --git a/Engine/Source/Runtime/Experimental/Iris/Core/Public/Iris/ReplicationSystem/ReplicationBridge.h b/Engine/Source/Runtime/Experimental/Iris/Core/Public/Iris/ReplicationSystem/ReplicationBridge.h
--- a/Engine/Source/Runtime/Experimental/Iris/Core/Public/Iris/ReplicationSystem/ReplicationBridge.h
+++ b/Engine/Source/Runtime/Experimental/Iris/Core/Public/Iris/ReplicationSystem/ReplicationBridge.h
@@ -2,11 +2,13 @@
 
 #pragma once
 
-#include "CoreTypes.h"
+#include "HAL/Platform.h"
 #include "Containers/ContainersFwd.h"
 #include "Containers/Map.h"
 #include "Iris/Core/NetObjectReference.h"
 #include "Iris/ReplicationSystem/NetHandle.h"
+#include "Misc/EnumClassFlags.h"
+#include "Net/Core/NetHandle/NetHandle.h"
 #include "UObject/ObjectKey.h"
 #include "UObject/ObjectMacros.h"
 
@@ -27,10 +29,10 @@ namespace UE::Net
 	class FNetSerializationContext;
 	namespace Private
 	{
-		typedef uint32 FInternalNetHandle;
-		class FNetHandleManager;
+		typedef uint32 FInternalNetRefIndex;
+		class FNetRefHandleManager;
 		class FNetObjectGroups;
-		struct FNetPushObjectHandle;
+		class FNetPushObjectHandle;
 		class FObjectReferenceCache;
 		class FReplicationProtocolManager;
 		class FReplicationReader;
@@ -68,6 +70,10 @@ enum class EEndReplicationFlags : uint32
 	TearOff								= Destroy << 1U,
 	/** Complete replication of pending state to all clients before ending replication. */
 	Flush								= TearOff << 1U,
+	/** Destroy NetHandle if one is associated with the replicated object. This should only be done if the object should not be replicated by any other replication system. */
+	DestroyNetHandle					= Flush << 1U,
+	/** Clear net push ID to prevent this object and its subobjects from being marked as dirty in the networking system. This should only be done if the object should not be replicated by any other replication system. */
+	ClearNetPushId						= DestroyNetHandle << 1U,
 };
 ENUM_CLASS_FLAGS(EEndReplicationFlags);
 
@@ -78,6 +84,7 @@ class UReplicationBridge : public UObject
 
 protected:
 	using FNetHandle = UE::Net::FNetHandle;
+	using FNetRefHandle = UE::Net::FNetRefHandle;
 	using FNetDependencyInfoArray = UE::Net::FNetDependencyInfoArray;
 
 public:
@@ -109,10 +116,10 @@ public:
 	 * If EEndReplication::Destroy is set the remote instance will be destroyed, if this is set for a static instance and the EndReplicationParameters are set a permanent destruction info will be added
 	 * Dynamic instances are always destroyed unless the TearOff flag is set.
 	 */
-	IRISCORE_API void EndReplication(FNetHandle Handle, EEndReplicationFlags EndReplicationFlags = EEndReplicationFlags::Destroy, FEndReplicationParameters* Parameters = nullptr);
+	IRISCORE_API void EndReplication(FNetRefHandle Handle, EEndReplicationFlags EndReplicationFlags = EEndReplicationFlags::Destroy, FEndReplicationParameters* Parameters = nullptr);
 
 	/** Returns true if the handle is replicated. */
-	IRISCORE_API bool IsReplicatedHandle(FNetHandle Handle) const;
+	IRISCORE_API bool IsReplicatedHandle(FNetRefHandle Handle) const;
 
 	/** Set the NetDriver used by the bridge. Called during creation if the NetDriver is recreated. */
 	IRISCORE_API virtual void SetNetDriver(UNetDriver* NetDriver);
@@ -131,7 +138,7 @@ protected:
 	IRISCORE_API virtual void PreSendUpdate();
 
 	/** Invoked before ReplicationSystem copies dirty state data for a single replicated object. */
-	IRISCORE_API virtual void PreSendUpdateSingleHandle(FNetHandle Handle);
+	IRISCORE_API virtual void PreSendUpdateSingleHandle(FNetRefHandle Handle);
 
 	/** Update world locations in FWorldLocations for objects that support it. */
 	IRISCORE_API virtual void UpdateInstancesWorldLocation();
@@ -139,13 +146,13 @@ protected:
 	// Remote interface, invoked from Replication code during serialization
 	
 	/** Write data required to instantiate NetObject remotely to bitstream. */
-	IRISCORE_API virtual bool WriteNetHandleCreationInfo(FReplicationBridgeSerializationContext& Context, FNetHandle Handle);
+	IRISCORE_API virtual bool WriteNetRefHandleCreationInfo(FReplicationBridgeSerializationContext& Context, FNetRefHandle Handle);
 
 	/** Read data required to instantiate NetObject from bitstream. */
-	IRISCORE_API virtual FNetHandle CreateNetHandleFromRemote(FNetHandle SubObjectOwnerNetHandle, FNetHandle WantedNetHandle, FReplicationBridgeSerializationContext& Context);
+	IRISCORE_API virtual FNetRefHandle CreateNetRefHandleFromRemote(FNetRefHandle SubObjectOwnerNetHandle, FNetRefHandle WantedNetHandle, FReplicationBridgeSerializationContext& Context);
 
 	/** Invoke after we have applied the initial state for an object.*/
-	IRISCORE_API virtual void PostApplyInitialState(FNetHandle Handle);
+	IRISCORE_API virtual void PostApplyInitialState(FNetRefHandle Handle);
 
 	/**
 	 * Called when the instance is detached from the protocol on request by the remote. 
@@ -153,43 +160,46 @@ protected:
 	 * @param bTearOff Whether the object should be torn off, i.e. not destroyed. 
 	 * @param bShouldDestroyInstance Whether the object should be destroyed.
 	 */
-	IRISCORE_API virtual void DetachInstanceFromRemote(FNetHandle Handle, bool bTearOff, bool bShouldDestroyInstance);
+	IRISCORE_API virtual void DetachInstanceFromRemote(FNetRefHandle Handle, bool bTearOff, bool bShouldDestroyInstance);
 
 	/** Called when we detach instance protocol from the local instance */
-	IRISCORE_API virtual void DetachInstance(FNetHandle Handle);
+	IRISCORE_API virtual void DetachInstance(FNetRefHandle Handle);
 
 	/** Invoked post garbage collect to allow us to detect stale objects */
 	IRISCORE_API virtual void PruneStaleObjects();
 
 	/** Invoked when we start to replicate an object for a specific connection to fill in any initial dependencies */
-	IRISCORE_API virtual void GetInitialDependencies(FNetHandle Handle, FNetDependencyInfoArray& OutDependencies) const;
+	IRISCORE_API virtual void GetInitialDependencies(FNetRefHandle Handle, FNetDependencyInfoArray& OutDependencies) const;
 
 protected:
 	// Forward calls to internal operations that we allow replication bridges to access
 
-	/** Create a local NetHandle / NetObject using the ReplicationProtocol. */
-	IRISCORE_API FNetHandle InternalCreateNetObject(FNetHandle AllocatedHandle, const UE::Net::FReplicationProtocol* ReplicationProtocol);
+	/** Create a local NetRefHandle / NetObject using the ReplicationProtocol. */
+	IRISCORE_API FNetRefHandle InternalCreateNetObject(FNetRefHandle AllocatedHandle, FNetHandle GlobalHandle, const UE::Net::FReplicationProtocol* ReplicationProtocol);
 
-	/** Create a NetHandle / NetObject on request from the authoritative end. */
-	IRISCORE_API FNetHandle InternalCreateNetObjectFromRemote(FNetHandle WantedNetHandle, const UE::Net::FReplicationProtocol* ReplicationProtocol);
+	/** Create a local NetRefHandle / NetObject using the ReplicationProtocol. */
+	IRISCORE_API FNetRefHandle InternalCreateNetObject(FNetRefHandle AllocatedHandle, const UE::Net::FReplicationProtocol* ReplicationProtocol);
 
-	/** Attach instance to NetHandle. */
-	IRISCORE_API void InternalAttachInstanceToNetHandle(FNetHandle Handle, bool bBindInstanceProtocol, UE::Net::FReplicationInstanceProtocol* InstanceProtocol, UObject* Instance);
+	/** Create a NetRefHandle / NetObject on request from the authoritative end. */
+	IRISCORE_API FNetRefHandle InternalCreateNetObjectFromRemote(FNetRefHandle WantedNetHandle, const UE::Net::FReplicationProtocol* ReplicationProtocol);
 
-	// Detach instance from NetObject/Handle and destroy the instance protocol
-	IRISCORE_API void InternalDetachInstanceFromNetHandle(FNetHandle Handle);
+	/** Attach instance to NetRefHandle. */
+	IRISCORE_API void InternalAttachInstanceToNetHandle(FNetRefHandle Handle, bool bBindInstanceProtocol, UE::Net::FReplicationInstanceProtocol* InstanceProtocol, UObject* Instance);
+
+	/** Detach instance from NetRefHandle and destroy the instance protocol. */
+	IRISCORE_API void InternalDetachInstanceFromNetRefHandle(FNetRefHandle Handle);
 
 	/** Destroy the handle and all internal book keeping associated with it. */
-	IRISCORE_API void InternalDestroyNetObject(FNetHandle Handle);
+	IRISCORE_API void InternalDestroyNetObject(FNetRefHandle Handle);
 	
 	/** Get the owner handle of a subobject handle. */
-	IRISCORE_API FNetHandle InternalGetSubObjectOwner(FNetHandle SubObjectHandle) const;
+	IRISCORE_API FNetRefHandle InternalGetSubObjectOwner(FNetRefHandle SubObjectHandle) const;
 
 	/** Add SubObjectHandle as SubObject to OwnerHandle. */
-	IRISCORE_API void InternalAddSubObject(FNetHandle OwnerHandle, FNetHandle SubObjectHandle, FNetHandle InsertRelativeToSubObjectHandle, ESubObjectInsertionOrder InsertionOrder);
+	IRISCORE_API void InternalAddSubObject(FNetRefHandle OwnerHandle, FNetRefHandle SubObjectHandle, FNetRefHandle InsertRelativeToSubObjectHandle, ESubObjectInsertionOrder InsertionOrder);
 
 	/** Add destruction info for the referenced object. */
-	FNetHandle InternalAddDestructionInfo(FNetHandle Handle, const FEndReplicationParameters& Parameters);
+	FNetRefHandle InternalAddDestructionInfo(FNetRefHandle Handle, const FEndReplicationParameters& Parameters);
 
 	inline UE::Net::Private::FReplicationProtocolManager* GetReplicationProtocolManager() const { return ReplicationProtocolManager; }
 	inline UReplicationSystem* GetReplicationSystem() const { return ReplicationSystem; }
@@ -206,36 +216,36 @@ private:
 
 	// Internal operations invoked by ReplicationSystem/ReplicationWriter
 	void ReadAndExecuteDestructionInfoFromRemote(FReplicationBridgeSerializationContext& Context);
-	void DetachSubObjectInstancesFromRemote(FNetHandle Handle, bool bTearOff, bool bShouldDestroyInstance);
-	void DestroyNetObjectFromRemote(FNetHandle Handle, bool bTearOff, bool bDestroyInstance);
+	void DetachSubObjectInstancesFromRemote(FNetRefHandle Handle, bool bTearOff, bool bShouldDestroyInstance);
+	void DestroyNetObjectFromRemote(FNetRefHandle Handle, bool bTearOff, bool bDestroyInstance);
 
 	// Adds the Handle to the list of handles pending tear-off, if bIsImmediate is true the object will be destroyed after the next update, otherwise
 	// it will be kept around until EndReplication is called.
-	void TearOff(FNetHandle Handle, bool bIsImmediate);
+	void TearOff(FNetRefHandle Handle, EEndReplicationFlags DestroyFlags, bool bIsImmediate);
 
-	FNetHandle CallCreateNetHandleFromRemote(FNetHandle SubObjectOwnerHandle, FNetHandle WantedNetHandle, FReplicationBridgeSerializationContext& Context);
+	FNetRefHandle CallCreateNetRefHandleFromRemote(FNetRefHandle SubObjectOwnerHandle, FNetRefHandle WantedNetHandle, FReplicationBridgeSerializationContext& Context);
 	void CallPreSendUpdate(float DeltaSeconds);	
-	void CallPreSendUpdateSingleHandle(FNetHandle Handle);
+	void CallPreSendUpdateSingleHandle(FNetRefHandle Handle);
 	void CallUpdateInstancesWorldLocation();
-	bool CallWriteNetHandleCreationInfo(FReplicationBridgeSerializationContext& Context, FNetHandle Handle);
-	void CallPostApplyInitialState(FNetHandle Handle);
+	bool CallWriteNetRefHandleCreationInfo(FReplicationBridgeSerializationContext& Context, FNetRefHandle Handle);
+	void CallPostApplyInitialState(FNetRefHandle Handle);
 	void CallPruneStaleObjects();
-	void CallGetInitialDependencies(FNetHandle Handle, FNetDependencyInfoArray& OutDependencies) const;
-	void CallDetachInstance(FNetHandle Handle);
-	void CallDetachInstanceFromRemote(FNetHandle Handle, bool bTearOff, bool bShouldDestroyInstance);
+	void CallGetInitialDependencies(FNetRefHandle Handle, FNetDependencyInfoArray& OutDependencies) const;
+	void CallDetachInstance(FNetRefHandle Handle);
+	void CallDetachInstanceFromRemote(FNetRefHandle Handle, bool bTearOff, bool bShouldDestroyInstance);
 
 private:
 
 	void InternalFlushStateData(UE::Net::FNetSerializationContext& SerializationContext, UE::Net::Private::FChangeMaskCache& ChangeMaskCache, UE::Net::FNetBitStreamWriter& ChangeMaskWriter, uint32 InternalObjectIndex);
 
 	// Internal method to copy state data for Handle
-	void InternalFlushStateData(FNetHandle Handle);
+	void InternalFlushStateData(FNetRefHandle Handle);
 
 	// Internal method to copy state data for Handle and any SubObjects and mark them as being torn-off
-	void InternalTearOff(FNetHandle OwnerHandle);
+	void InternalTearOff(FNetRefHandle OwnerHandle);
 
 	// Destroy all SubObjects owned by provided handle
-	void InternalDestroySubObjects(FNetHandle OwnerHandle);
+	void InternalDestroySubObjects(FNetRefHandle OwnerHandle, EEndReplicationFlags Flags);
 
 	/**
 	 * Called from ReplicationSystem when a streaming level is about to unload.
@@ -249,7 +259,7 @@ private:
 	 */
 	void RemoveDestructionInfosForGroup(UE::Net::FNetObjectGroupHandle GroupHandle);
 
-	void DestroyLocalNetHandle(FNetHandle Handle);
+	void DestroyLocalNetHandle(FNetRefHandle Handle, EEndReplicationFlags Flags);
 
 	// Tear-off all handles in the PendingTearOff list that has not yet been torn-off
 	void TearOffHandlesPendingTearOff();
@@ -260,6 +270,9 @@ private:
 	void SetNetPushIdOnFragments(const TArrayView<const UE::Net::FReplicationFragment*const>& Fragments, const UE::Net::Private::FNetPushObjectHandle& PushHandle);
 	void ClearNetPushIdOnFragments(const TArrayView<const UE::Net::FReplicationFragment*const>& Fragments);
 
+	void DestroyGlobalNetHandle(UE::Net::Private::FInternalNetRefIndex InternalReplicationIndex);
+	void ClearNetPushIds(UE::Net::Private::FInternalNetRefIndex InternalReplicationIndex);
+
 	friend UReplicationSystem;
 	friend UE::Net::Private::FReplicationSystemImpl;
 	friend UE::Net::Private::FReplicationSystemInternal;
@@ -270,7 +283,7 @@ private:
 	UReplicationSystem* ReplicationSystem;
 	UE::Net::Private::FReplicationProtocolManager* ReplicationProtocolManager;
 	UE::Net::Private::FReplicationStateDescriptorRegistry* ReplicationStateDescriptorRegistry;
-	UE::Net::Private::FNetHandleManager* NetHandleManager;
+	UE::Net::Private::FNetRefHandleManager* NetRefHandleManager;
 	UE::Net::Private::FObjectReferenceCache* ObjectReferenceCache;
 	UE::Net::Private::FNetObjectGroups* Groups;
 
@@ -283,21 +296,22 @@ private:
 	{
 		UE::Net::FNetObjectReference StaticRef;
 		UE::Net::FNetObjectGroupHandle LevelGroupHandle;
-		UE::Net::Private::FInternalNetHandle InternalReplicationIndex;
+		UE::Net::Private::FInternalNetRefIndex InternalReplicationIndex;
 	};
 
 	const UE::Net::FReplicationProtocol* DestructionInfoProtocol;
 	
 	// Need to track the objects with destruction infos so that we can clean them up properly
 	// We use this to be able ask remote to destroy static objects
-	TMap<FNetHandle, FDestructionInfo> StaticObjectsPendingDestroy;
+	TMap<FNetRefHandle, FDestructionInfo> StaticObjectsPendingDestroy;
 
 	struct FTearOffInfo
 	{
-		FTearOffInfo(FNetHandle InHandle, bool bInIsImmediate) : Handle(InHandle), bIsImmediate(bInIsImmediate) {}
+		FTearOffInfo(FNetRefHandle InHandle, EEndReplicationFlags InDestroyFlags, bool bInIsImmediate) : Handle(InHandle), DestroyFlags(InDestroyFlags), bIsImmediate(bInIsImmediate) {}
 
-		FNetHandle Handle;
-		bool	bIsImmediate;
+		FNetRefHandle Handle;
+		EEndReplicationFlags DestroyFlags;
+		bool bIsImmediate;
 	};
 	TArray<FTearOffInfo> HandlesPendingTearOff;
 };
diff --git a/Engine/Source/Runtime/Experimental/Iris/Core/Public/Iris/ReplicationSystem/ReplicationOperations.h b/Engine/Source/Runtime/Experimental/Iris/Core/Public/Iris/ReplicationSystem/ReplicationOperations.h
--- a/Engine/Source/Runtime/Experimental/Iris/Core/Public/Iris/ReplicationSystem/ReplicationOperations.h
+++ b/Engine/Source/Runtime/Experimental/Iris/Core/Public/Iris/ReplicationSystem/ReplicationOperations.h
@@ -77,19 +77,19 @@ struct FReplicationInstanceOperations
 	/** Update object references in fragments that has object references and additional required traits. */
 	static IRISCORE_API void PollAndRefreshCachedObjectReferences(const FReplicationInstanceProtocol* InstanceProtocol, EReplicationFragmentTraits RequiredTraits);
 
-	/** Quantize the state for a replicated object with a given NetHandle using the ReplicationProtocol, DstObjectStateBuffer does need to be cleared */
+	/** Quantize the state for a replicated object with a given InstanceProtocol using the ReplicationProtocol. DstObjectStateBuffer needs to be cleared before calling this function. */
 	static IRISCORE_API void CopyAndQuantize(FNetSerializationContext& Context, uint8* DstObjectStateBuffer, FNetBitStreamWriter* OutChangeMaskWriter, const FReplicationInstanceProtocol* InstanceProtocol, const FReplicationProtocol* Protocol);
 
-	/** Dequantize the state for a replicated object with a given NetHandle using the ReplicationProtocol, Data will be pushed out by using the ReplicationFragments */
+	/** Dequantize the state for a replicated object with a given protocol. Data will be pushed out by using the ReplicationFragments. */
 	static IRISCORE_API void DequantizeAndApply(FNetSerializationContext& Context, const FDequantizeAndApplyParameters& Parameters);
 
-	/** Dequantize the state for a replicated object with a given NetHandle using the ReplicationProtocol, Data will be pushed out by using the ReplicationFragments */
+	/** Dequantize the state for a replicated object with a given protocol. Data will be pushed out by using the ReplicationFragments. */
 	static IRISCORE_API void DequantizeAndApply(FNetSerializationContext& Context, FMemStackBase& InAllocator, const uint32* ChangeMaskData, const FReplicationInstanceProtocol* InstanceProtocol, const uint8* SrcObjectStateBuffer, const FReplicationProtocol* Protocol);
 
-	/** Dequantize the state for a replicated object with a given NetHandle using the ReplicationProtocol and output state to string */
+	/** Dequantize the state for a replicated object with a given protocol and output the state to string. */
 	static IRISCORE_API void OutputInternalStateToString(FNetSerializationContext& Context, FStringBuilderBase& StringBuilder, const uint32* ChangeMaskData, const uint8* SrcInternalObjectStateBuffer, const FReplicationInstanceProtocol* InstanceProtocol, const FReplicationProtocol* Protocol);
 
-	/** Dequantize the state for a replicated object with a given NetHandle using the ReplicationProtocol and output state to string */
+	/** Dequantize the default state for a replicated object with a given protocol and output the state to string. */
 	static IRISCORE_API void OutputInternalDefaultStateToString(FNetSerializationContext& Context, FStringBuilderBase& StringBuilder, const FReplicationInstanceProtocol* InstanceProtocol, const FReplicationProtocol* Protocol);
 };
 
@@ -104,28 +104,28 @@ struct FReplicationProtocolOperations
 	/** Serialize the state of all dirty members a NetObject and changemask to Bitstream */
 	static IRISCORE_API void SerializeWithMask(FNetSerializationContext& Context, const uint32* ChangeMaskData, const uint8* RESTRICT SrcObjectStateBuffer, const FReplicationProtocol* Protocol);
 
-	/** Deserialize changemask and the state of a NetObjectfrom from a BitStream to a ObjectStateBuffer large enough to fit all data */
+	/** Deserialize changemask and the state of an object from a BitStream to a ObjectStateBuffer large enough to fit all data */
 	static IRISCORE_API void DeserializeWithMask(FNetSerializationContext& Context, uint32* DstChangeMaskData, uint8* RESTRICT DstObjectStateBuffer, const FReplicationProtocol* Protocol);
 
-	/** Compare two quantized states return false if they are different */
+	/** Compare two quantized states and return whether they're equal or not. */
 	static IRISCORE_API bool IsEqualQuantizedState(FNetSerializationContext& Context, const uint8* RESTRICT Source0, const uint8* RESTRICT Source1, const FReplicationProtocol* Protocol);
 
-	/** Free dynamic state for a full NetObject */
+	/** Free dynamic state for the entire protocol. */
 	static IRISCORE_API void FreeDynamicState(FNetSerializationContext& Context, uint8* RESTRICT SrcObjectStateBuffer, const FReplicationProtocol* Protocol);
 
 	/** Initialize from default state */
 	static IRISCORE_API void InitializeFromDefaultState(FNetSerializationContext& Context, uint8* RESTRICT StateBuffer, const FReplicationProtocol* Protocol);
 
-	/** Serialize the initial state of all dirty members a NetObject and changemask to Bitstream compressed against the default state */
+	/** Serialize the initial state of all dirty members a NetObject and changemask to a Bitstream, delta compressed against the default state. */
 	static IRISCORE_API void SerializeInitialStateWithMask(FNetSerializationContext& Context, const uint32* ChangeMaskData, const uint8* RESTRICT SrcObjectStateBuffer, const FReplicationProtocol* Protocol);
 
-	/** Deserialize changemask and the state of a NetObjectfrom from a BitStream to a ObjectStateBuffer large enough to fit all data compressed against default state */
+	/** Deserialize changemask and the state of an object from a BitStream to an ObjectStateBuffer large enough to fit all data delta compressed against default state. */
 	static IRISCORE_API void DeserializeInitialStateWithMask(FNetSerializationContext& Context, uint32* DstChangeMaskData, uint8* RESTRICT DstObjectStateBuffer, const FReplicationProtocol* Protocol);
 
-	/** Serialize the state of all dirty members a NetObject and changemask to Bitstream, delta compressed against the PrevObjectStateBuffer */
+	/** Serialize the state of all dirty members of an object and changemask to a Bitstream, delta compressed against the PrevObjectStateBuffer. */
 	static IRISCORE_API void SerializeWithMaskDelta(FNetSerializationContext& Context, const uint32* ChangeMaskData, const uint8* RESTRICT SrcObjectStateBuffer, const uint8* RESTRICT PrevObjectStateBuffer, const FReplicationProtocol* Protocol);
 
-	/** Deserialize changemask and the state of a NetObjectfrom from a BitStream to a ObjectStateBuffer large enough to fit all data, delta compressed against the PrevObjectStateBuffer */
+	/** Deserialize changemask and the state of an object from a BitStream to an ObjectStateBuffer large enough to fit all data, delta compressed against the PrevObjectStateBuffer. */
 	static IRISCORE_API void DeserializeWithMaskDelta(FNetSerializationContext& Context, uint32* DstChangeMaskData, uint8* RESTRICT DstObjectStateBuffer, const uint8* RESTRICT PrevObjectStateBuffer, const FReplicationProtocol* Protocol);
 };
 
diff --git a/Engine/Source/Runtime/Experimental/Iris/Core/Public/Iris/ReplicationSystem/ReplicationSystem.h b/Engine/Source/Runtime/Experimental/Iris/Core/Public/Iris/ReplicationSystem/ReplicationSystem.h
--- a/Engine/Source/Runtime/Experimental/Iris/Core/Public/Iris/ReplicationSystem/ReplicationSystem.h
+++ b/Engine/Source/Runtime/Experimental/Iris/Core/Public/Iris/ReplicationSystem/ReplicationSystem.h
@@ -8,6 +8,7 @@
 #include "Templates/RefCounting.h"
 #include "Iris/IrisConfig.h"
 #include "Iris/ReplicationSystem/NetHandle.h"
+#include "Net/Core/NetHandle/NetHandle.h"
 
 #include "ReplicationSystem.generated.h"
 
@@ -48,8 +49,9 @@ class UReplicationSystem : public UObject
 {
 	GENERATED_BODY()
 
-	using FNetHandle = UE::Net::FNetHandle;
+	using FNetRefHandle = UE::Net::FNetRefHandle;
 	using FNetObjectGroupHandle = UE::Net::FNetObjectGroupHandle;
+	using FNetHandle = UE::Net::FNetHandle;
 
 public:
 
@@ -143,7 +145,7 @@ public:
 	 * @param Priority A value >= 0. 1.0 means the object will be considered for replication every frame, if it has updated replicated properties.
 	 * @see SetPrioritizer
 	 */
-	IRISCORE_API void SetStaticPriority(FNetHandle Handle, float Priority);
+	IRISCORE_API void SetStaticPriority(FNetRefHandle Handle, float Priority);
 
 	/**
 	 * Sets a prioritizer for a replicated object which will be used until the next call to SetPrioritizer or SetStaticPriority.
@@ -156,7 +158,7 @@ public:
 	 * @see GetPrioritizerHandle
 	 * @see SetStaticPriority
 	 */
-	IRISCORE_API bool SetPrioritizer(FNetHandle Handle, UE::Net::FNetObjectPrioritizerHandle PrioritizerHandle);
+	IRISCORE_API bool SetPrioritizer(FNetRefHandle Handle, UE::Net::FNetObjectPrioritizerHandle PrioritizerHandle);
 
 
 	/**
@@ -239,18 +241,18 @@ public:
 	IRISCORE_API UE::Net::FStringTokenStore* GetStringTokenStore();
 
 	/**
-	 * Check whether a FNetHandle is still associated with a replicated object.
+	 * Check whether a FNetRefHandle is still associated with a replicated object.
 	 * @param Handle The handle check.
 	 * @return true if the handle is still valid, false if not. 
 	 */
-	IRISCORE_API bool IsValidHandle(FNetHandle Handle) const;
+	IRISCORE_API bool IsValidHandle(FNetRefHandle Handle) const;
 
 	/**
 	 * Get the ReplicationProtocol for a handle.
 	 * @param Handle The handle to retrieve the protocol for.
 	 * @return A valid pointer to the protocol if the handle is valid, nullptr if not.
 	 */
-	IRISCORE_API const UE::Net::FReplicationProtocol* GetReplicationProtocol(FNetHandle Handle) const;
+	IRISCORE_API const UE::Net::FReplicationProtocol* GetReplicationProtocol(FNetRefHandle Handle) const;
 
 	// Groups
 
@@ -273,20 +275,20 @@ public:
 	 * @param GroupHandle A valid group handle.
 	 * @param Handle A valid handle to an object.
 	 */
-	IRISCORE_API void AddToGroup(FNetObjectGroupHandle GroupHandle, FNetHandle Handle);
+	IRISCORE_API void AddToGroup(FNetObjectGroupHandle GroupHandle, FNetRefHandle Handle);
 
 	/**
 	 * Removes an object from a group.
 	 * @param GroupHandle A valid group handle.
 	 * @param Handle A valid handle to an object.
 	 */
-	IRISCORE_API void RemoveFromGroup(FNetObjectGroupHandle GroupHandle, FNetHandle Handle);
+	IRISCORE_API void RemoveFromGroup(FNetObjectGroupHandle GroupHandle, FNetRefHandle Handle);
 
 	/**
 	 * Removes an object from all groups it's part of.
 	 * @param Handle A valid handle to an object.
 	 */
-	IRISCORE_API void RemoveFromAllGroups(FNetHandle Handle);
+	IRISCORE_API void RemoveFromAllGroups(FNetRefHandle Handle);
 
 	/**
 	 * Check whether an objects belongs to a particular group or not.
@@ -294,7 +296,7 @@ public:
 	 * @param Handle A valid handle to an object.
 	 * @return true if both handles are valid and the object belongs to the group, false otherwise.
 	 */
-	IRISCORE_API bool IsInGroup(FNetObjectGroupHandle GroupHandle, FNetHandle Handle) const;
+	IRISCORE_API bool IsInGroup(FNetObjectGroupHandle GroupHandle, FNetRefHandle Handle) const;
 
 	/**
 	 * Check if a group handle is valid.
@@ -311,13 +313,13 @@ public:
 	 * @param Handle A valid handle to an object.
 	 * @param ConnectionId A valid connection ID to be set as the owner.
 	 */
-	IRISCORE_API void SetOwningNetConnection(FNetHandle Handle, uint32 ConnectionId);
+	IRISCORE_API void SetOwningNetConnection(FNetRefHandle Handle, uint32 ConnectionId);
 
 	/** 
 	 * Get the owning net connection for an object.
 	 * @see SetOwningNetConnection
 	 */
-	IRISCORE_API uint32 GetOwningNetConnection(FNetHandle Handle) const;
+	IRISCORE_API uint32 GetOwningNetConnection(FNetRefHandle Handle) const;
 
 	/**
 	 * Sets a filter for a replicated object which will be used until the next call to SetFilter or SetConnectionFilter.
@@ -332,7 +334,7 @@ public:
 	 * @see GetFilterHandle
 	 * @see SetConnectionFilter
 	 */
-	IRISCORE_API bool SetFilter(FNetHandle Handle, UE::Net::FNetObjectFilterHandle FilterHandle);
+	IRISCORE_API bool SetFilter(FNetRefHandle Handle, UE::Net::FNetObjectFilterHandle FilterHandle);
 
 	/**
 	 * Gets the handle for a filter with a given name. The handle can be used in subsequent calls to SetFilter.
@@ -400,7 +402,7 @@ public:
 	 * @return true if the connection filter was properly set, false if not.
 	 * @see SetFilter
 	 */
-	IRISCORE_API bool SetConnectionFilter(FNetHandle Handle, const TBitArray<>& Connections, UE::Net::ENetFilterStatus ReplicationStatus);
+	IRISCORE_API bool SetConnectionFilter(FNetRefHandle Handle, const TBitArray<>& Connections, UE::Net::ENetFilterStatus ReplicationStatus);
 
 	/**
 	 * Enable or disable a replication condition for a single connection and do the inverse for all other connections. This will affect
@@ -416,7 +418,7 @@ public:
 	 * @return true if the condition was successfully set, false if not.
 	 * @see SetReplicationCondition
 	 */
-	IRISCORE_API bool SetReplicationConditionConnectionFilter(FNetHandle Handle, UE::Net::EReplicationCondition Condition, uint32 ConnectionId, bool bEnable);
+	IRISCORE_API bool SetReplicationConditionConnectionFilter(FNetRefHandle Handle, UE::Net::EReplicationCondition Condition, uint32 ConnectionId, bool bEnable);
 
 	/**
 	 * Enable or disable a replication condition for all connections. This will affect the replication of properties with conditions.
@@ -427,7 +429,7 @@ public:
 	 * @return true if the condition was successfully set, false if not.
 	 * @see SetReplicationCondition, SetReplicationConditionConnectionFilter
 	 */
-	IRISCORE_API bool SetReplicationCondition(FNetHandle Handle, UE::Net::EReplicationCondition Condition, bool bEnable);
+	IRISCORE_API bool SetReplicationCondition(FNetRefHandle Handle, UE::Net::EReplicationCondition Condition, bool bEnable);
 
 
 	/**
@@ -436,19 +438,19 @@ public:
 	 * whether the delta compression feature is enabled or not and other reasons.
 	 * @param Handle A valid handle to a replicated object.
 	 */
-	IRISCORE_API void SetDeltaCompressionStatus(FNetHandle Handle, UE::Net::ENetObjectDeltaCompressionStatus Status);
+	IRISCORE_API void SetDeltaCompressionStatus(FNetRefHandle Handle, UE::Net::ENetObjectDeltaCompressionStatus Status);
 
 	 /**
 	  * Mark an object as a net temporary. Such objects only replicate its
 	  * initial state and ignores future state changes.
 	  */
-	IRISCORE_API void SetIsNetTemporary(FNetHandle Handle);
+	IRISCORE_API void SetIsNetTemporary(FNetRefHandle Handle);
 
 	/**
 	 * Mark an object to be torn off next update.
 	 * @param Handle A valid handle to an object.
 	 */
-	IRISCORE_API void TearOffNextUpdate(FNetHandle Handle);
+	IRISCORE_API void TearOffNextUpdate(FNetRefHandle Handle);
 
 	/**
 	 * Explicitly mark object as dirty. This will make it considered for replication this frame
@@ -457,7 +459,7 @@ public:
 	 * @param Handle A valid handle to an object.
 	 * @see FObjectReplicationBridgePollConfig
 	 */
-	IRISCORE_API void MarkDirty(FNetHandle Handle);
+	IRISCORE_API void MarkDirty(FNetRefHandle Handle);
 
 	/**
 	 * Retrieve the WorldLocations instance which holds world locations for all objects that support it. 
@@ -557,6 +559,9 @@ public:
 	 */
 	IRISCORE_API static void DestroyReplicationSystem(UReplicationSystem* System);
 
+	/** Returns all replication systems. Entries may be null. */
+	IRISCORE_API static TArrayView<UReplicationSystem*> GetAllReplicationSystems();
+
 	/** Static delegate that is triggered just after creating and initializing a new replication system. */
 	IRISCORE_API static FReplicationSystemCreatedDelegate& GetReplicationSystemCreatedDelegate();
 
@@ -565,31 +570,25 @@ public:
 
 	enum ReplicationSystemConstants : uint32
 	{
+#if UE_NET_ALLOW_MULTIPLE_REPLICATION_SYSTEMS
 		MaxReplicationSystemCount = 16,
+#else
+		MaxReplicationSystemCount = 1,
+#endif
 	};
 
 private:
 
-	friend UReplicationSystem* GetReplicationSystem(uint32);
+	friend UReplicationSystem* GetReplicationSystem(uint32 Id);
 
-#if UE_NET_ALLOW_MULTIPLE_REPLICATION_SYSTEMS
 	IRISCORE_API static UReplicationSystem* ReplicationSystems[];
-#else
-	IRISCORE_API static UReplicationSystem* ReplicationSystem;
-#endif
+	static uint32 MaxReplicationSystemId;
 };
 
-#if UE_NET_ALLOW_MULTIPLE_REPLICATION_SYSTEMS
 inline UReplicationSystem* GetReplicationSystem(uint32 Id)
 {
 	return Id >= FReplicationSystemFactory::MaxReplicationSystemCount ? nullptr : FReplicationSystemFactory::ReplicationSystems[Id];
 }
-#else
-inline UReplicationSystem* GetReplicationSystem(uint32)
-{ 
-	return FReplicationSystemFactory::ReplicationSystem;
-}
-#endif
 
 }
 
diff --git a/Engine/Source/Runtime/Experimental/Iris/Core/Public/Iris/Serialization/ObjectNetSerializer.h b/Engine/Source/Runtime/Experimental/Iris/Core/Public/Iris/Serialization/ObjectNetSerializer.h
--- a/Engine/Source/Runtime/Experimental/Iris/Core/Public/Iris/Serialization/ObjectNetSerializer.h
+++ b/Engine/Source/Runtime/Experimental/Iris/Core/Public/Iris/Serialization/ObjectNetSerializer.h
@@ -37,8 +37,8 @@ struct FScriptInterfaceNetSerializerConfig : public FNetSerializerConfig
 namespace UE::Net
 {
 
-IRISCORE_API FNetHandle ReadNetHandle(FNetSerializationContext& Context);
-IRISCORE_API void WriteNetHandle(FNetSerializationContext& Context, const FNetHandle Handle);
+IRISCORE_API FNetRefHandle ReadNetRefHandle(FNetSerializationContext& Context);
+IRISCORE_API void WriteNetRefHandle(FNetSerializationContext& Context, const FNetRefHandle Handle);
 
 IRISCORE_API void ReadFullNetObjectReference(FNetSerializationContext& Context, FNetObjectReference& Reference);
 IRISCORE_API void WriteFullNetObjectReference(FNetSerializationContext& Context, const FNetObjectReference& Reference);
diff --git a/Engine/Source/Runtime/Experimental/Iris/Stub/Public/Iris/ReplicationSystem/ReplicationBridge.h b/Engine/Source/Runtime/Experimental/Iris/Stub/Public/Iris/ReplicationSystem/ReplicationBridge.h
--- a/Engine/Source/Runtime/Experimental/Iris/Stub/Public/Iris/ReplicationSystem/ReplicationBridge.h
+++ b/Engine/Source/Runtime/Experimental/Iris/Stub/Public/Iris/ReplicationSystem/ReplicationBridge.h
@@ -2,7 +2,7 @@
 
 #pragma once
 
-#include "CoreTypes.h"
+#include "HAL/Platform.h"
 #include "UObject/ObjectMacros.h"
 #include "ReplicationBridge.generated.h"
 
diff --git a/Engine/Source/Runtime/Net/Core/Private/Net/Core/DirtyNetObjectTracker/GlobalDirtyNetObjectTracker.cpp b/Engine/Source/Runtime/Net/Core/Private/Net/Core/DirtyNetObjectTracker/GlobalDirtyNetObjectTracker.cpp
new file mode 100644
--- /dev/null
+++ b/Engine/Source/Runtime/Net/Core/Private/Net/Core/DirtyNetObjectTracker/GlobalDirtyNetObjectTracker.cpp
@@ -0,0 +1,168 @@
+// Copyright Epic Games, Inc. All Rights Reserved.
+
+#include "Net/Core/DirtyNetObjectTracker/GlobalDirtyNetObjectTracker.h"
+#include "Containers/Set.h"
+#include "Misc/CoreDelegates.h"
+
+namespace UE::Net
+{
+
+class FGlobalDirtyNetObjectTracker::FPimpl
+{
+private:
+	FPimpl();
+	~FPimpl();
+
+	void ResetDirtyNetObjects();
+
+private:
+	friend FGlobalDirtyNetObjectTracker;
+
+	static TSet<FNetHandle> EmptyDirtyObjects;
+
+#if WITH_PUSH_MODEL
+	TSet<FNetHandle> DirtyObjects;
+	FNetBitArray AssignedHandleIndices;
+	FNetBitArray Pollers;
+	uint32 PollerCount = 0;
+#endif
+};
+
+TSet<FNetHandle> FGlobalDirtyNetObjectTracker::FPimpl::EmptyDirtyObjects;
+
+#if WITH_PUSH_MODEL
+FGlobalDirtyNetObjectTracker::FPimpl* FGlobalDirtyNetObjectTracker::Instance = nullptr;
+#endif
+
+void FGlobalDirtyNetObjectTracker::MarkNetObjectStateDirty(FNetHandle NetHandle)
+{
+#if WITH_PUSH_MODEL
+	if (Instance && Instance->PollerCount > 0)
+	{
+		Instance->DirtyObjects.Add(NetHandle);
+	}
+#endif
+}
+
+FGlobalDirtyNetObjectTracker::FPollHandle FGlobalDirtyNetObjectTracker::CreatePoller()
+{
+#if WITH_PUSH_MODEL
+	if (Instance)
+	{
+		if (Instance->PollerCount >= Instance->AssignedHandleIndices.GetNumBits())
+		{
+			Instance->AssignedHandleIndices.SetNumBits(Instance->PollerCount + 1U);
+			Instance->Pollers.SetNumBits(Instance->PollerCount + 1U);
+		}
+
+		const uint32 HandleIndex = Instance->AssignedHandleIndices.FindFirstZero();
+		if (!ensure(HandleIndex != FNetBitArrayBase::InvalidIndex))
+		{
+			return FPollHandle();
+		}
+
+		Instance->AssignedHandleIndices.SetBit(HandleIndex);
+		++Instance->PollerCount;
+		return FPollHandle(HandleIndex);
+	}
+#endif
+
+	return FPollHandle();
+}
+
+void FGlobalDirtyNetObjectTracker::DestroyPoller(uint32 HandleIndex)
+{
+#if WITH_PUSH_MODEL
+	if (HandleIndex == FPollHandle::InvalidIndex)
+	{
+		return;
+	}
+
+	if (ensureAlwaysMsgf((HandleIndex < Instance->AssignedHandleIndices.GetNumBits()) && Instance->AssignedHandleIndices.GetBit(HandleIndex), TEXT("Destroying unknown poller with handle index %u"), HandleIndex))
+	{
+		Instance->AssignedHandleIndices.ClearBit(HandleIndex);
+
+		const uint32 PollerCalled = Instance->Pollers.GetBit(HandleIndex);
+		ensureAlwaysMsgf(PollerCalled == 0U, TEXT("Destroying poller that called GetDirtyNetObjects() but not ResetDirtyNetObjects()"));
+		Instance->Pollers.ClearBit(HandleIndex);
+
+		--Instance->PollerCount;
+		if (Instance->PollerCount <= 0)
+		{
+			Instance->DirtyObjects.Reset();
+		}
+	}
+#endif
+}
+
+const TSet<FNetHandle>& FGlobalDirtyNetObjectTracker::GetDirtyNetObjects(const FPollHandle& Handle)
+{
+#if WITH_PUSH_MODEL
+	if (Instance && Handle.IsValid())
+	{
+		Instance->Pollers.SetBit(Handle.Index);
+		return Instance->DirtyObjects;
+	}
+#endif
+
+	return FPimpl::EmptyDirtyObjects;
+}
+
+void FGlobalDirtyNetObjectTracker::ResetDirtyNetObjects(const FPollHandle& Handle)
+{
+#if WITH_PUSH_MODEL
+	if (Instance && Handle.IsValid())
+	{
+		Instance->Pollers.ClearBit(Handle.Index);
+		if (Instance->Pollers.IsNoBitSet())
+		{
+			Instance->DirtyObjects.Reset();
+		}
+	}
+#endif
+}
+
+void FGlobalDirtyNetObjectTracker::Init()
+{
+#if WITH_PUSH_MODEL
+	checkf(Instance == nullptr, TEXT("%s"), TEXT("Only one FNetHandleManager instance may exist."));
+	Instance = new FGlobalDirtyNetObjectTracker::FPimpl();
+#endif
+}
+
+void FGlobalDirtyNetObjectTracker::Deinit()
+{
+#if WITH_PUSH_MODEL
+	delete Instance;
+	Instance = nullptr;
+#endif
+}
+
+FGlobalDirtyNetObjectTracker::FPimpl::FPimpl()
+{
+#if WITH_ENGINE
+	FCoreDelegates::OnEndFrame.AddRaw(this, &FPimpl::ResetDirtyNetObjects);
+#endif
+}
+
+FGlobalDirtyNetObjectTracker::FPimpl::~FPimpl()
+{
+#if WITH_ENGINE
+	FCoreDelegates::OnEndFrame.RemoveAll(this);
+#endif
+}
+
+void FGlobalDirtyNetObjectTracker::FPimpl::ResetDirtyNetObjects()
+{
+#if WITH_PUSH_MODEL
+	if (!ensureMsgf(Pollers.IsNoBitSet(), TEXT("FGlobalDirtyNetObjectTracker poller %u forgot to call ResetDirtNetObjects."), Pollers.FindFirstOne()))
+	{
+		Pollers.Reset();
+
+		// DirtyObjects should already be reset if there are no pollers.
+		DirtyObjects.Reset();
+	}
+#endif
+}
+
+}
diff --git a/Engine/Source/Runtime/Net/Core/Private/Net/Core/NetCoreModule.cpp b/Engine/Source/Runtime/Net/Core/Private/Net/Core/NetCoreModule.cpp
--- a/Engine/Source/Runtime/Net/Core/Private/Net/Core/NetCoreModule.cpp
+++ b/Engine/Source/Runtime/Net/Core/Private/Net/Core/NetCoreModule.cpp
@@ -6,6 +6,8 @@
 #include "Net/Core/Misc/NetCoreLog.h"
 
 #include "HAL/IConsoleManager.h"
+#include "Net/Core/DirtyNetObjectTracker/GlobalDirtyNetObjectTracker.h"
+#include "Net/Core/NetHandle/NetHandleManager.h"
 #include "Net/Serialization/FastArraySerializer.h"
 
 TAutoConsoleVariable<int32> CVarNetEnableDetailedScopeCounters(TEXT("net.EnableDetailedScopeCounters"), 1, TEXT("Enables detailed networking scope cycle counters. There are often lots of these which can negatively impact performance."));
@@ -20,10 +22,14 @@ class FNetCoreModule : public IModuleInterface
 public:
 	virtual void StartupModule() override
 	{
+		UE::Net::FNetHandleManager::Init();
+		UE::Net::FGlobalDirtyNetObjectTracker::Init();
 	}
 
 	virtual void ShutdownModule() override
 	{
+		UE::Net::FGlobalDirtyNetObjectTracker::Deinit();
+		UE::Net::FNetHandleManager::Deinit();
 	}
 
 	virtual bool SupportsDynamicReloading() override
