Chaos : fix geometry collection caching not properly record and replaying break events

diff --git a/Engine/Source/Runtime/Experimental/Chaos/Private/PBDRigidsSolver.cpp b/Engine/Source/Runtime/Experimental/Chaos/Private/PBDRigidsSolver.cpp
--- a/Engine/Source/Runtime/Experimental/Chaos/Private/PBDRigidsSolver.cpp
+++ b/Engine/Source/Runtime/Experimental/Chaos/Private/PBDRigidsSolver.cpp
@@ -426,8 +426,6 @@ namespace Chaos
 					bool ResetData = (MSubStepInfo.Step == 0);
 					MSolver->GetEventManager()->FillProducerData(MSolver, ResetData);
 					MSolver->GetEvolution()->ResetAllRemovals();
-					MSolver->GetEvolution()->GetRigidClustering().ResetAllClusterBreakings();
-					MSolver->GetEvolution()->GetRigidClustering().ResetAllClusterCrumblings();
 				}
 
 				// flip on last sub-step of frame
@@ -458,6 +456,12 @@ namespace Chaos
 				MSolver->CompleteSceneSimulation();
 			}
 
+			// reset all clustering event needs to be after CompleteSceneSimulation to make sure the cache recording gets them before they get removed
+			// they cannot be right after the presolve callback ( as they were before ) because they will cause geometry collection replicated clients to miss them
+			// Todo(chaos) we should probably move all of the solver event reset here in the future
+			MSolver->GetEvolution()->GetRigidClustering().ResetAllClusterBreakings();
+			MSolver->GetEvolution()->GetRigidClustering().ResetAllClusterCrumblings();
+
 			if (FRewindData* RewindData = MSolver->GetRewindData())
 			{
 				RewindData->FinishFrame();
