Iris - Removed unused and incorrectly implemented IsAssignedIndex method.

#rnx

diff --git a/Engine/Source/Runtime/Experimental/Iris/Core/Private/Iris/ReplicationSystem/NetHandleManager.h b/Engine/Source/Runtime/Experimental/Iris/Core/Private/Iris/ReplicationSystem/NetHandleManager.h
--- a/Engine/Source/Runtime/Experimental/Iris/Core/Private/Iris/ReplicationSystem/NetHandleManager.h
+++ b/Engine/Source/Runtime/Experimental/Iris/Core/Private/Iris/ReplicationSystem/NetHandleManager.h
@@ -118,9 +118,6 @@ public:
 	// Return true if this is a scopable index
 	bool IsScopableIndex(FInternalNetRefIndex InternalIndex) const { return ScopableInternalIndices.GetBit(InternalIndex); }
 
-	// Return true if this index is assigned
-	bool IsAssignedIndex(FInternalNetRefIndex InternalIndex) const { return InternalIndex && ScopableInternalIndices.GetBit(InternalIndex); }
-
 	static FNetRefHandle MakeNetRefHandle(uint32 Id, uint32 ReplicationSystemId);
 	static FNetRefHandle MakeNetRefHandleFromId(uint32 Id);
 
