GeometryCollection : Facade API cleanup

diff --git a/Engine/Plugins/Experimental/ChaosFlesh/Source/ChaosFlesh/Private/ChaosFlesh/FleshCollection.cpp b/Engine/Plugins/Experimental/ChaosFlesh/Source/ChaosFlesh/Private/ChaosFlesh/FleshCollection.cpp
--- a/Engine/Plugins/Experimental/ChaosFlesh/Source/ChaosFlesh/Private/ChaosFlesh/FleshCollection.cpp
+++ b/Engine/Plugins/Experimental/ChaosFlesh/Source/ChaosFlesh/Private/ChaosFlesh/FleshCollection.cpp
@@ -35,7 +35,7 @@ void FFleshCollection::Construct()
 	for (float& m : Mass) { m = (float)1; }
 
 	// Kinematic Bindings
-	Chaos::Facades::FKinematicBindingFacade(*this).Init();
+	GeometryCollection::Facades::FKinematicBindingFacade(*this).DefineSchema();
 }
 
 
diff --git a/Engine/Plugins/Experimental/ChaosFlesh/Source/ChaosFleshEngine/Private/ChaosFlesh/FleshComponent.cpp b/Engine/Plugins/Experimental/ChaosFlesh/Source/ChaosFleshEngine/Private/ChaosFlesh/FleshComponent.cpp
--- a/Engine/Plugins/Experimental/ChaosFlesh/Source/ChaosFleshEngine/Private/ChaosFlesh/FleshComponent.cpp
+++ b/Engine/Plugins/Experimental/ChaosFlesh/Source/ChaosFleshEngine/Private/ChaosFlesh/FleshComponent.cpp
@@ -117,9 +117,9 @@ UDeformablePhysicsComponent::FDataMapValue UFleshComponent::NewDeformableData()
 		{
 			if (const FFleshCollection* Rest = FleshAsset->GetCollection())
 			{
-				if (FTransformSource::HasFacade(Rest))
+				FTransformSource TransformSource(*Rest);
+				if (TransformSource.IsValid())
 				{
-
 					TArray<USkeletalMeshComponent*> SkeletalMeshComponents;
 					GetOwner()->GetComponents<USkeletalMeshComponent>(SkeletalMeshComponents);
 
@@ -134,7 +134,7 @@ UDeformablePhysicsComponent::FDataMapValue UFleshComponent::NewDeformableData()
 							{
 								if (const USkeleton* Skeleton = SkeletalMesh->GetSkeleton())
 								{
-									TSet<int32> Roots = FTransformSource::GetTransformSource(Rest, Skeleton->GetName(), Skeleton->GetGuid().ToString());
+									TSet<int32> Roots = TransformSource.GetTransformSource(Skeleton->GetName(), Skeleton->GetGuid().ToString());
 									if (!Roots.IsEmpty() && ensureMsgf(Roots.Num() == 1, TEXT("Error: Only supports a single root per skeleton.(%s)"), *Skeleton->GetName()))
 									{
 										TArray<FTransform> ComponentLocalPose;
diff --git a/Engine/Plugins/Experimental/ChaosFlesh/Source/ChaosFleshNodes/Private/Dataflow/ChaosFleshKinematicInitializationNodes.cpp b/Engine/Plugins/Experimental/ChaosFlesh/Source/ChaosFleshNodes/Private/Dataflow/ChaosFleshKinematicInitializationNodes.cpp
--- a/Engine/Plugins/Experimental/ChaosFlesh/Source/ChaosFleshNodes/Private/Dataflow/ChaosFleshKinematicInitializationNodes.cpp
+++ b/Engine/Plugins/Experimental/ChaosFlesh/Source/ChaosFleshNodes/Private/Dataflow/ChaosFleshKinematicInitializationNodes.cpp
@@ -104,23 +104,27 @@ void FKinematicTetrahedralBindingsDataflowNode::Evaluate(Dataflow::FContext& Con
 						BoundWeights.Init(float(1), BoundVerts.Num());
 						
 						//get local coords of bound verts
-						typedef Chaos::Facades::FKinematicBindingFacade Kinematics;
-						Kinematics::FBindingKey Binding = Kinematics::SetBoneBindings(&InCollection, b, BoundVerts, BoundWeights);
-						TManagedArray<TArray<FVector3f>>& LocalPos = InCollection.AddAttribute<TArray<FVector3f>>("LocalPosition", Binding.GroupName);
-						Kinematics::AddKinematicBinding(&InCollection, Binding);
+						typedef GeometryCollection::Facades::FKinematicBindingFacade FKinematics;
+						FKinematics Kinematics(InCollection); Kinematics.DefineSchema();
+						if (Kinematics.IsValid())
+						{
+							FKinematics::FBindingKey Binding = Kinematics.SetBoneBindings(b, BoundVerts, BoundWeights);
+							TManagedArray<TArray<FVector3f>>& LocalPos = InCollection.AddAttribute<TArray<FVector3f>>("LocalPosition", Binding.GroupName);
+							Kinematics.AddKinematicBinding(Binding);
 
-						auto FloatVert = [](FVector3d V) { return FVector3f(V.X, V.Y, V.Z); };
-						auto DoubleVert = [](FVector3f V) { return FVector3d(V.X, V.Y, V.Z); };
-						LocalPos[Binding.Index].SetNum(BoundVerts.Num());
-						for (int32 i = 0; i < BoundVerts.Num(); i++)
-						{	
-							FVector3f Temp = (*Vertex)[BoundVerts[i]];
-							LocalPos[Binding.Index][i] = FloatVert(ComponentPose[b].InverseTransformPosition(DoubleVert(Temp)));
+							auto FloatVert = [](FVector3d V) { return FVector3f(V.X, V.Y, V.Z); };
+							auto DoubleVert = [](FVector3f V) { return FVector3d(V.X, V.Y, V.Z); };
+							LocalPos[Binding.Index].SetNum(BoundVerts.Num());
+							for (int32 i = 0; i < BoundVerts.Num(); i++)
+							{
+								FVector3f Temp = (*Vertex)[BoundVerts[i]];
+								LocalPos[Binding.Index][i] = FloatVert(ComponentPose[b].InverseTransformPosition(DoubleVert(Temp)));
+							}
 						}
 					}
 				}
 			}
-			GeometryCollection::Facades::FVertexBoneWeightsFacade::AddBoneWeightsFromKinematicBindings(&InCollection);
+			GeometryCollection::Facades::FVertexBoneWeightsFacade(InCollection).AddBoneWeightsFromKinematicBindings();
 		}
 
 		SetValue<DataType>(Context, InCollection, &Collection);
@@ -173,8 +177,8 @@ void FKinematicInitializationDataflowNode::Evaluate(Dataflow::FContext& Context,
 							}
 						}
 					}
-					typedef Chaos::Facades::FKinematicBindingFacade Kinematics;
-					Kinematics::AddKinematicBinding(&InCollection, Kinematics::SetBoneBindings(&InCollection, Index, BoundVerts, BoundWeights));
+					GeometryCollection::Facades::FKinematicBindingFacade Kinematics(InCollection);
+					Kinematics.AddKinematicBinding(Kinematics.SetBoneBindings(Index, BoundVerts, BoundWeights));
 				}
 			}
 		}
@@ -226,8 +230,8 @@ void FSetVerticesKinematicDataflowNode::Evaluate(Dataflow::FContext& Context, co
 			}
 			if (BoundVerts.Num() > 0)
 			{
-				typedef Chaos::Facades::FKinematicBindingFacade Kinematics;
-				Kinematics::AddKinematicBinding(&InCollection, Kinematics::SetBoneBindings(&InCollection, INDEX_NONE, BoundVerts, BoundWeights));
+				GeometryCollection::Facades::FKinematicBindingFacade Kinematics(InCollection);
+				Kinematics.AddKinematicBinding(Kinematics.SetBoneBindings(INDEX_NONE, BoundVerts, BoundWeights));
 			}	
 		}
 		SetValue<DataType>(Context, InCollection, &Collection);
diff --git a/Engine/Plugins/Experimental/Dataflow/Source/DataflowEditor/Private/Dataflow/DataflowEditorViewport.cpp b/Engine/Plugins/Experimental/Dataflow/Source/DataflowEditor/Private/Dataflow/DataflowEditorViewport.cpp
--- a/Engine/Plugins/Experimental/Dataflow/Source/DataflowEditor/Private/Dataflow/DataflowEditorViewport.cpp
+++ b/Engine/Plugins/Experimental/Dataflow/Source/DataflowEditor/Private/Dataflow/DataflowEditorViewport.cpp
@@ -109,7 +109,7 @@ void FDataflowEditorViewportClient::Draw(const FSceneView* View, FPrimitiveDrawI
 	Super::Draw(View, PDI);
 	if (bEnableDataflowClientDrawing)
 	{
-		GeometryCollection::Facades::FRenderingFacade Facade(&RenderCollection);
+		GeometryCollection::Facades::FRenderingFacade Facade(RenderCollection);
 		if (Facade.IsValid())
 		{
 			if (bRenderMesh)
@@ -169,7 +169,7 @@ void FDataflowEditorViewportClient::ReleaseRenderStructures()
 
 void FDataflowEditorViewportClient::RenderIntoStructures()
 {
-	GeometryCollection::Facades::FRenderingFacade Facade(&RenderCollection);
+	GeometryCollection::Facades::FRenderingFacade Facade(RenderCollection);
 	if (bEnableDataflowClientDrawing)
 	{
 		if (Facade.IsValid())
@@ -184,8 +184,8 @@ void FDataflowEditorViewportClient::RenderIntoStructures()
 			//        we still need to rely on Component rendering through
 			//        a scene proxy as the primary rendering path. 
 			{
-				const TManagedArray< FIntVector >& Indices = *Facade.GetIndices();
-				const TManagedArray< FVector3f >& Vertices = *Facade.GetVertices();
+				const TManagedArray< FIntVector >& Indices = Facade.GetIndices();
+				const TManagedArray< FVector3f >& Vertices = Facade.GetVertices();
 				if (Vertices.Num() && Indices.Num())
 				{
 					IndexBuffer.SetNum(Indices.Num() * 3);
@@ -237,7 +237,7 @@ void FDataflowEditorViewportClient::Tick(float DeltaSeconds)
 
 						if (Dataflow->GetRenderTargets().Num())
 						{
-							GeometryCollection::Facades::FRenderingFacade Facade(&RenderCollection);
+							GeometryCollection::Facades::FRenderingFacade Facade(RenderCollection);
 							for (const UDataflowEdNode* Target : Dataflow->GetRenderTargets())
 							{
 								Target->Render(Facade, Context);
diff --git a/Engine/Plugins/Experimental/Dataflow/Source/DataflowEnginePlugin/Private/Dataflow/DataflowRenderingComponent.cpp b/Engine/Plugins/Experimental/Dataflow/Source/DataflowEnginePlugin/Private/Dataflow/DataflowRenderingComponent.cpp
--- a/Engine/Plugins/Experimental/Dataflow/Source/DataflowEnginePlugin/Private/Dataflow/DataflowRenderingComponent.cpp
+++ b/Engine/Plugins/Experimental/Dataflow/Source/DataflowEnginePlugin/Private/Dataflow/DataflowRenderingComponent.cpp
@@ -54,7 +54,7 @@ void UDataflowRenderingComponent::UpdateLocalBounds()
 {
 	if (bBoundsNeedsUpdate)
 	{
-		GeometryCollection::Facades::FBoundsFacade::UpdateBoundingBox(RenderCollection);
+		GeometryCollection::Facades::FBoundsFacade(RenderCollection).UpdateBoundingBox();
 		bBoundsNeedsUpdate = false;
 	}
 }
@@ -95,7 +95,8 @@ void UDataflowRenderingComponent::RenderProceduralMesh()
 	if (bUpdateRender)
 	{
 		RenderCollection = FManagedArrayCollection();
-		GeometryCollection::Facades::FRenderingFacade Facade(&RenderCollection);
+		GeometryCollection::Facades::FRenderingFacade Facade(RenderCollection);
+		Facade.DefineSchema();
 
 		bool bNeedsRefresh = false;
 		if (Context && Dataflow)
@@ -107,9 +108,10 @@ void UDataflowRenderingComponent::RenderProceduralMesh()
 		}
 		bUpdateRender = false;
 	}
-		
+	
+
 	bool bCanRender = false;	
-	if (GeometryCollection::Facades::FRenderingFacade::IsValid(&RenderCollection))
+	if (GeometryCollection::Facades::FRenderingFacade(RenderCollection).IsValid())
 	{
 		int32 NumVertices = RenderCollection.NumElements(FGeometryCollection::VerticesGroup);
 		int32 NumFaces = RenderCollection.NumElements(FGeometryCollection::FacesGroup);
diff --git a/Engine/Plugins/Experimental/GeometryCollectionPlugin/Source/GeometryCollectionNodes/Private/Dataflow/GeometryCollectionNodes.cpp b/Engine/Plugins/Experimental/GeometryCollectionPlugin/Source/GeometryCollectionNodes/Private/Dataflow/GeometryCollectionNodes.cpp
--- a/Engine/Plugins/Experimental/GeometryCollectionPlugin/Source/GeometryCollectionNodes/Private/Dataflow/GeometryCollectionNodes.cpp
+++ b/Engine/Plugins/Experimental/GeometryCollectionPlugin/Source/GeometryCollectionNodes/Private/Dataflow/GeometryCollectionNodes.cpp
@@ -1068,9 +1068,9 @@ void FGetBoundingBoxesDataflowNode::Evaluate(Dataflow::FContext& Context, const
 		FManagedArrayCollection InCollection = GetValue<FManagedArrayCollection>(Context, &Collection);
 		FDataflowTransformSelection InTransformSelection = GetValue<FDataflowTransformSelection>(Context, &TransformSelection);
 
-		const TManagedArray<FBox>* BoundingBoxesArr = GeometryCollection::Facades::FBoundsFacade::GetBoundingBoxes(InCollection);
+		const TManagedArray<FBox>& BoundingBoxesArr = GeometryCollection::Facades::FBoundsFacade(InCollection).GetBoundingBoxes();
 
-		SetValue<TArray<FBox>>(Context, BoundingBoxesArr->GetConstArray(), &BoundingBoxes);
+		SetValue<TArray<FBox>>(Context, BoundingBoxesArr.GetConstArray(), &BoundingBoxes);
 	}
 }
 
@@ -1081,12 +1081,12 @@ void FGetCentroidsDataflowNode::Evaluate(Dataflow::FContext& Context, const FDat
 		FManagedArrayCollection InCollection = GetValue<FManagedArrayCollection>(Context, &Collection);
 		FDataflowTransformSelection InTransformSelection = GetValue<FDataflowTransformSelection>(Context, &TransformSelection);
 
-		const TManagedArray<FBox>* BoundingBoxes = GeometryCollection::Facades::FBoundsFacade::GetBoundingBoxes(InCollection);
+		const TManagedArray<FBox>& BoundingBoxes = GeometryCollection::Facades::FBoundsFacade(InCollection).GetBoundingBoxes();
 
 		TArray<FVector> CentroidsArr;
-		for (int32 Idx = 0; Idx < BoundingBoxes->Num(); ++Idx)
+		for (int32 Idx = 0; Idx < BoundingBoxes.Num(); ++Idx)
 		{
-			const FBox& BoundingBox = (*BoundingBoxes)[Idx];
+			const FBox& BoundingBox = BoundingBoxes[Idx];
 			if (BoundingBox.IsValid)
 			{
 				if (InTransformSelection.Num() > 0)
diff --git a/Engine/Plugins/Experimental/GeometryCollectionPlugin/Source/GeometryCollectionNodes/Public/Dataflow/GeometryCollectionNodes.h b/Engine/Plugins/Experimental/GeometryCollectionPlugin/Source/GeometryCollectionNodes/Public/Dataflow/GeometryCollectionNodes.h
--- a/Engine/Plugins/Experimental/GeometryCollectionPlugin/Source/GeometryCollectionNodes/Public/Dataflow/GeometryCollectionNodes.h
+++ b/Engine/Plugins/Experimental/GeometryCollectionPlugin/Source/GeometryCollectionNodes/Public/Dataflow/GeometryCollectionNodes.h
@@ -99,7 +99,7 @@ struct FAppendCollectionAssetsDataflowNode : public FDataflowNode
 public:
 	typedef FManagedArrayCollection DataType;
 
-	UPROPERTY(meta = (DataflowInput, DataflowOutput, DisplayName = "Collection", DataflowPassthrough = "Collection"))
+	UPROPERTY(meta = (DataflowInput, DataflowOutput, DisplayName = "Collection", DataflowPassthrough = "Collection1"))
 	FManagedArrayCollection Collection1;
 
 	UPROPERTY(meta = (DataflowInput, DisplayName = "Collection"))
@@ -109,9 +109,9 @@ public:
 	FAppendCollectionAssetsDataflowNode(const Dataflow::FNodeParameters& InParam, FGuid InGuid = FGuid::NewGuid())
 		: FDataflowNode(InParam, InGuid)
 	{
-		RegisterOutputConnection(&Collection1, &Collection1);
 		RegisterInputConnection(&Collection1);
 		RegisterInputConnection(&Collection2);
+		RegisterOutputConnection(&Collection1, &Collection1);
 	}
 
 	virtual void Evaluate(Dataflow::FContext& Context, const FDataflowOutput* Out) const override;
diff --git a/Engine/Plugins/Experimental/PlanarCutPlugin/Source/PlanarCut/Private/PlanarCut.cpp b/Engine/Plugins/Experimental/PlanarCutPlugin/Source/PlanarCut/Private/PlanarCut.cpp
--- a/Engine/Plugins/Experimental/PlanarCutPlugin/Source/PlanarCut/Private/PlanarCut.cpp
+++ b/Engine/Plugins/Experimental/PlanarCutPlugin/Source/PlanarCut/Private/PlanarCut.cpp
@@ -104,7 +104,7 @@ int32 FInternalSurfaceMaterials::GetDefaultMaterialIDForGeometry(const FGeometry
 	return InternalMaterialID;
 }
 
-void FInternalSurfaceMaterials::SetUVScaleFromCollection(const FGeometryCollectionMeshFacade& CollectionMesh, int32 GeometryIdx)
+void FInternalSurfaceMaterials::SetUVScaleFromCollection(const GeometryCollection::Facades::FMeshFacade& CollectionMesh, int32 GeometryIdx)
 {
 	const auto& VertexArray = CollectionMesh.Vertex.Get();
 	const auto& UVsArray = CollectionMesh.UVs.Get();
@@ -756,7 +756,7 @@ int32 CutMultipleWithMultiplePlanes(
 	int32 OrigNumGeom = Collection.FaceCount.Num();
 	int32 CurNumGeom = OrigNumGeom;
 
-	FGeometryCollectionMeshFacade CollectionMesh(Collection);
+	GeometryCollection::Facades::FMeshFacade CollectionMesh(Collection);
 	if (!CollectionMesh.IsValid())
 	{
 		return -1;
diff --git a/Engine/Plugins/Experimental/PlanarCutPlugin/Source/PlanarCut/Public/PlanarCut.h b/Engine/Plugins/Experimental/PlanarCutPlugin/Source/PlanarCut/Public/PlanarCut.h
--- a/Engine/Plugins/Experimental/PlanarCutPlugin/Source/PlanarCut/Public/PlanarCut.h
+++ b/Engine/Plugins/Experimental/PlanarCutPlugin/Source/PlanarCut/Public/PlanarCut.h
@@ -6,6 +6,7 @@
 
 #include "Voronoi/Voronoi.h"
 #include "GeometryCollection/GeometryCollection.h"
+#include "GeometryCollection/Facades/CollectionMeshFacade.h"
 #include "MeshDescription.h"
 #include "DynamicMesh/DynamicMesh3.h"
 
@@ -81,7 +82,7 @@ struct PLANARCUT_API FInternalSurfaceMaterials
 	 * @param Collection	Reference collection to use for setting UV scale
 	 * @param GeometryIdx	Reference geometry inside collection; if -1, use all geometries in collection
 	 */
-	void SetUVScaleFromCollection(const FGeometryCollectionMeshFacade& CollectionMesh, int32 GeometryIdx = -1);
+	void SetUVScaleFromCollection(const GeometryCollection::Facades::FMeshFacade& CollectionMesh, int32 GeometryIdx = -1);
 
 	
 	int32 GetDefaultMaterialIDForGeometry(const FGeometryCollection& Collection, int32 GeometryIdx = -1) const;
diff --git a/Engine/Source/Runtime/Experimental/Chaos/Private/Chaos/Deformable/ChaosDeformableSolver.cpp b/Engine/Source/Runtime/Experimental/Chaos/Private/Chaos/Deformable/ChaosDeformableSolver.cpp
--- a/Engine/Source/Runtime/Experimental/Chaos/Private/Chaos/Deformable/ChaosDeformableSolver.cpp
+++ b/Engine/Source/Runtime/Experimental/Chaos/Private/Chaos/Deformable/ChaosDeformableSolver.cpp
@@ -189,17 +189,18 @@ namespace Chaos::Softs
 
 		if (Property.bEnableKinematics)
 		{
-			typedef Chaos::Facades::FKinematicBindingFacade Kinematics;
+			typedef GeometryCollection::Facades::FKinematicBindingFacade FKinematics;
+			FKinematics Kinematics(Rest);
 
 			// Add Kinematics Node
-			for (int i = Kinematics::NumKinematicBindings(&Rest) - 1; i >= 0; i--)
+			for (int i = Kinematics.NumKinematicBindings() - 1; i >= 0; i--)
 			{
-				Kinematics::FBindingKey Key = Kinematics::GetKinematicBindingKey(&Rest, i);
+				FKinematics::FBindingKey Key = Kinematics.GetKinematicBindingKey(i);
 
 				int32 BoneIndex = INDEX_NONE;
 				TArray<int32> BoundVerts;
 				TArray<float> BoundWeights;
-				Kinematics::GetBoneBindings(&Rest, Key, BoneIndex, BoundVerts, BoundWeights);
+				Kinematics.GetBoneBindings(Key, BoneIndex, BoundVerts, BoundWeights);
 
 				for (int32 vdx : BoundVerts)
 				{
@@ -326,8 +327,9 @@ namespace Chaos::Softs
 							GlobalTransform = FleshInputBuffer->GlobalTransform;
 						}
 
-						typedef GeometryCollection::Facades::FVertexBoneWeightsFacade WeightsFacade;
-						if (WeightsFacade::HasFacade(&Rest))
+						typedef GeometryCollection::Facades::FVertexBoneWeightsFacade FWeightsFacade;
+						FWeightsFacade WeightsFacade(Rest);
+						if (WeightsFacade.IsValid())
 						{
 							int32 NumObjectVertices = Rest.NumElements(FGeometryCollection::VerticesGroup);
 							int32 ObjectVertexIndex = SolverParticleToObjectVertexIndex(Index);
@@ -335,8 +337,8 @@ namespace Chaos::Softs
 							{
 								if(FleshInputBuffer)
 								{
-									TArray<int32> BoneIndices = (*WeightsFacade::GetBoneIndices(&Rest))[ObjectVertexIndex];
-									TArray<float> BoneWeights = (*WeightsFacade::GetBoneWeights(&Rest))[ObjectVertexIndex];
+									TArray<int32> BoneIndices = WeightsFacade.GetBoneIndices()[ObjectVertexIndex];
+									TArray<float> BoneWeights = WeightsFacade.GetBoneWeights()[ObjectVertexIndex];
 
 									FFleshThreadingProxy::FFleshInputBuffer* PreviousFleshBuffer = nullptr;
 									if (this->PreviousInputPackage && this->PreviousInputPackage->ObjectMap.Contains(Owner))
diff --git a/Engine/Source/Runtime/Experimental/Chaos/Private/GeometryCollection/Facades/CollectionAnchoringFacade.cpp b/Engine/Source/Runtime/Experimental/Chaos/Private/GeometryCollection/Facades/CollectionAnchoringFacade.cpp
--- a/Engine/Source/Runtime/Experimental/Chaos/Private/GeometryCollection/Facades/CollectionAnchoringFacade.cpp
+++ b/Engine/Source/Runtime/Experimental/Chaos/Private/GeometryCollection/Facades/CollectionAnchoringFacade.cpp
@@ -10,6 +10,11 @@ namespace Chaos::Facades
 	    , AnchoredAttribute(InCollection, "Anchored", FTransformCollection::TransformGroup)
 	{}
 
+	FCollectionAnchoringFacade::FCollectionAnchoringFacade(const FManagedArrayCollection& InCollection)
+		: InitialDynamicStateAttribute(InCollection, "InitialDynamicState", FTransformCollection::TransformGroup)
+		, AnchoredAttribute(InCollection, "Anchored", FTransformCollection::TransformGroup)
+	{}
+
 	bool FCollectionAnchoringFacade::HasInitialDynamicStateAttribute() const
 	{
 		return InitialDynamicStateAttribute.IsValid();
@@ -22,11 +27,13 @@ namespace Chaos::Facades
 
 	void FCollectionAnchoringFacade::AddAnchoredAttribute()
 	{
+		check(!IsConst());
 		return AnchoredAttribute.AddAndFill(false);
 	}
 
 	void FCollectionAnchoringFacade::CopyAnchoredAttribute(const FCollectionAnchoringFacade& From)
 	{
+		check(!IsConst());
 		if (From.HasAnchoredAttribute())
 		{
 			AnchoredAttribute.Copy(From.AnchoredAttribute);
@@ -41,6 +48,7 @@ namespace Chaos::Facades
 	
 	void FCollectionAnchoringFacade::SetInitialDynamicState(int32 TransformIndex, EObjectStateType State)
 	{
+		check(!IsConst());
 		// we don't really support sleeping as a dynamic state
 		if (ensure(State != EObjectStateType::Sleeping && State < EObjectStateType::Count))
 		{
@@ -50,6 +58,7 @@ namespace Chaos::Facades
 
 	void FCollectionAnchoringFacade::SetInitialDynamicState(const TArray<int32>& TransformIndices, EObjectStateType State)
 	{
+		check(!IsConst());
 		// we don't really support sleeping as a dynamic state
 		if (ensure(State != EObjectStateType::Sleeping && State < EObjectStateType::Count))
 		{
@@ -68,11 +77,13 @@ namespace Chaos::Facades
 
 	void FCollectionAnchoringFacade::SetAnchored(int32 TransformIndex, bool bValue)
 	{
+		check(!IsConst());
 		AnchoredAttribute.Modify()[TransformIndex] = bValue;
 	}
 
 	void FCollectionAnchoringFacade::SetAnchored(const TArray<int32>& TransformIndices, bool bValue)
 	{
+		check(!IsConst());
 		TManagedArray<bool>& Anchored = AnchoredAttribute.Modify();
 		for (const int32 TransformIndex: TransformIndices)
 		{
diff --git a/Engine/Source/Runtime/Experimental/Chaos/Private/GeometryCollection/Facades/CollectionBoundsFacade.cpp b/Engine/Source/Runtime/Experimental/Chaos/Private/GeometryCollection/Facades/CollectionBoundsFacade.cpp
--- a/Engine/Source/Runtime/Experimental/Chaos/Private/GeometryCollection/Facades/CollectionBoundsFacade.cpp
+++ b/Engine/Source/Runtime/Experimental/Chaos/Private/GeometryCollection/Facades/CollectionBoundsFacade.cpp
@@ -10,69 +10,53 @@
 namespace GeometryCollection::Facades
 {
 
-	// Attributes
-	const FName FBoundsFacade::AAAAttribute = "AAA";
-
 	FBoundsFacade::FBoundsFacade(FManagedArrayCollection& InCollection)
-		: Self(InCollection)
-		, BoundingBox(InCollection,"BoundingBox", FGeometryCollection::GeometryGroup)
-		, Vertex(InCollection,"Vertex", FGeometryCollection::VerticesGroup)
-		, BoneMap(InCollection,"BoneMap", FGeometryCollection::VerticesGroup)
-		, TransformToGeometryIndex(InCollection,"TransformToGeometryIndex", FTransformCollection::TransformGroup)
-	{
-		DefineSchema(Self);
-	}
+		: BoundingBoxAttribute(InCollection,"BoundingBox", FGeometryCollection::GeometryGroup)
+		, VertexAttribute(InCollection,"Vertex", FGeometryCollection::VerticesGroup, FGeometryCollection::VerticesGroup)
+		, BoneMapAttribute(InCollection,"BoneMap", FGeometryCollection::VerticesGroup)
+		, TransformToGeometryIndexAttribute(InCollection,"TransformToGeometryIndex", FTransformCollection::TransformGroup)
+	{}
+
+	FBoundsFacade::FBoundsFacade(const FManagedArrayCollection& InCollection)
+		: BoundingBoxAttribute(InCollection, "BoundingBox", FGeometryCollection::GeometryGroup)
+		, VertexAttribute(InCollection, "Vertex", FGeometryCollection::VerticesGroup, FGeometryCollection::VerticesGroup)
+		, BoneMapAttribute(InCollection, "BoneMap", FGeometryCollection::VerticesGroup)
+		, TransformToGeometryIndexAttribute(InCollection, "TransformToGeometryIndex", FTransformCollection::TransformGroup)
+	{}
 
 	//
 	//  Initialization
 	//
 
-		//
-	//  Initialization
-	//
-
-	void FBoundsFacade::DefineSchema(FManagedArrayCollection& InCollection)
+	void FBoundsFacade::DefineSchema()
 	{
-		FManagedArrayCollection::FConstructionParameters VertexDependency(FGeometryCollection::VerticesGroup);
-
-		// surface rendering attributes
-		if(!InCollection.HasGroup(FGeometryCollection::VerticesGroup))
-		{
-			InCollection.AddGroup(FGeometryCollection::VerticesGroup);
-		}
-		if( !InCollection.HasGroup(FGeometryCollection::FacesGroup) )
-		{
-			InCollection.AddGroup(FGeometryCollection::FacesGroup);
-		}
-
-		InCollection.AddAttribute<FVector3f>("Vertex", FGeometryCollection::VerticesGroup);
-		InCollection.AddAttribute<FIntVector>("Indices", FGeometryCollection::FacesGroup, VertexDependency);
-
-		ensure(InCollection.FindAttributeTyped<FVector3f>("Vertex", FGeometryCollection::VerticesGroup) != nullptr);
-		ensure(InCollection.FindAttributeTyped<FIntVector>("Indices", FGeometryCollection::FacesGroup) != nullptr);
+		check(!IsConst());
+		BoundingBoxAttribute.Add();
+		VertexAttribute.Add();
+		BoneMapAttribute.Add();
+		TransformToGeometryIndexAttribute.Add();
 	}
 
-	bool FBoundsFacade::IsValid(const FManagedArrayCollection& InCollection)
+	bool FBoundsFacade::IsValid() const
 	{
-		return InCollection.HasAttribute("BoundingBox", FGeometryCollection::GeometryGroup)
-			&& InCollection.HasAttribute("Vertex", FGeometryCollection::VerticesGroup)
-			&& InCollection.HasAttribute("BoneMap", FGeometryCollection::VerticesGroup)
-			&& InCollection.HasAttribute("TransformToGeometryIndex", FTransformCollection::TransformGroup);
+		return BoundingBoxAttribute.IsValid() && VertexAttribute.IsValid() 
+			&& BoneMapAttribute.IsValid() && TransformToGeometryIndexAttribute.IsValid();
 	}
 
 
-
-	void FBoundsFacade::UpdateBoundingBox(FManagedArrayCollection& InCollection, bool bSkipCheck)
+	void FBoundsFacade::UpdateBoundingBox(bool bSkipCheck)
 	{
-		if (!bSkipCheck || !IsValid(InCollection))
+		check(!IsConst());
+
+		if (!bSkipCheck || !IsValid())
 		{
 			return;
 		}
 
-		TManagedArray<FBox>& BoundingBox = InCollection.ModifyAttribute<FBox>("BoundingBox", FGeometryCollection::GeometryGroup);
-		const TManagedArray<FVector3f>& Vertex = InCollection.GetAttribute<FVector3f>("Vertex", FGeometryCollection::VerticesGroup);
-		const TManagedArray<int32>& BoneMap = InCollection.GetAttribute<int32>("BoneMap", FGeometryCollection::VerticesGroup);
-		const TManagedArray<int32>& TransformToGeometryIndex = InCollection.GetAttribute<int32>("TransformToGeometryIndex", FTransformCollection::TransformGroup);
+		TManagedArray<FBox>& BoundingBox = BoundingBoxAttribute.Modify();
+		const TManagedArray<FVector3f>& Vertex = VertexAttribute.Get();
+		const TManagedArray<int32>& BoneMap = BoneMapAttribute.Get();
+		const TManagedArray<int32>& TransformToGeometryIndex = TransformToGeometryIndexAttribute.Get();
 
 		if (BoundingBox.Num())
 		{
@@ -90,17 +74,6 @@ namespace GeometryCollection::Facades
 			}
 		}
 	}
-
-	//
-	//  GetAttributes
-	//
-
-	const TManagedArray< FBox >* FBoundsFacade::GetBoundingBoxes(const FManagedArrayCollection& InCollection)
-	{
-		return InCollection.FindAttributeTyped <FBox>("BoundingBox", FGeometryCollection::GeometryGroup);
-	}
-
-
 }; // GeometryCollection::Facades
 
 
diff --git a/Engine/Source/Runtime/Experimental/Chaos/Private/GeometryCollection/Facades/CollectionConnectionGraphFacade.cpp b/Engine/Source/Runtime/Experimental/Chaos/Private/GeometryCollection/Facades/CollectionConnectionGraphFacade.cpp
--- a/Engine/Source/Runtime/Experimental/Chaos/Private/GeometryCollection/Facades/CollectionConnectionGraphFacade.cpp
+++ b/Engine/Source/Runtime/Experimental/Chaos/Private/GeometryCollection/Facades/CollectionConnectionGraphFacade.cpp
@@ -15,23 +15,37 @@ namespace GeometryCollection::Facades
 	{
 	}
 
+	FCollectionConnectionGraphFacade::FCollectionConnectionGraphFacade(const FManagedArrayCollection& InCollection)
+		: ConnectionsAttribute(InCollection, "Connections", FTransformCollection::TransformGroup)
+#if UE_BUILD_DEBUG
+		, ParentAttribute(InCollection, "Parent", FTransformCollection::TransformGroup)
+#endif
+	{
+	}
+
 	bool FCollectionConnectionGraphFacade::IsValid() const
 	{
 		return ConnectionsAttribute.IsValid();
 	}
 
-	void FCollectionConnectionGraphFacade::AddAttributes()
+	void FCollectionConnectionGraphFacade::DefineSchema()
 	{
+		check(!IsConst());
+
 		ConnectionsAttribute.Add();
 	}
 	
 	void FCollectionConnectionGraphFacade::ClearAttributes()
 	{
+		check(!IsConst());
+
 		ConnectionsAttribute.Remove();
 	}
 
 	void FCollectionConnectionGraphFacade::Connect(int32 BoneA, int32 BoneB)
 	{
+		check(!IsConst());
+
 #if UE_BUILD_DEBUG
 		// Expect to always have a parent array and connected bones should always have the same parent
 		checkSlow(ParentAttribute.IsValid() && ParentAttribute.Get()[BoneA] == ParentAttribute.Get()[BoneB]);
diff --git a/Engine/Source/Runtime/Experimental/Chaos/Private/GeometryCollection/Facades/CollectionHierarchyFacade.cpp b/Engine/Source/Runtime/Experimental/Chaos/Private/GeometryCollection/Facades/CollectionHierarchyFacade.cpp
--- a/Engine/Source/Runtime/Experimental/Chaos/Private/GeometryCollection/Facades/CollectionHierarchyFacade.cpp
+++ b/Engine/Source/Runtime/Experimental/Chaos/Private/GeometryCollection/Facades/CollectionHierarchyFacade.cpp
@@ -11,6 +11,12 @@ namespace Chaos::Facades
 		, LevelAttribute(InCollection, "Level", FTransformCollection::TransformGroup)
 	{}
 
+	FCollectionHierarchyFacade::FCollectionHierarchyFacade(const FManagedArrayCollection& InCollection)
+		: ParentAttribute(InCollection, "Parent", FTransformCollection::TransformGroup)
+		, ChildrenAttribute(InCollection, "Children", FTransformCollection::TransformGroup)
+		, LevelAttribute(InCollection, "Level", FTransformCollection::TransformGroup)
+	{}
+
 	bool FCollectionHierarchyFacade::IsValid() const
 	{
 		return ParentAttribute.IsValid() && ChildrenAttribute.IsValid();
@@ -46,50 +52,52 @@ namespace Chaos::Facades
 
 	void FCollectionHierarchyFacade::GenerateLevelAttribute()
 	{
-		check(IsValid());
-
-		// level attribute must actually be persistent 
-		LevelAttribute.Add(ManageArrayAccessor::EPersistencePolicy::MakePersistent);
+		check(!IsConst());
+		if (IsValid())
+		{
+			// level attribute must actually be persistent 
+			LevelAttribute.Add(ManageArrayAccessor::EPersistencePolicy::MakePersistent);
 
-		TManagedArray<int32>& Levels = LevelAttribute.Modify();
-		const TManagedArray<int32>& Parents = ParentAttribute.Get();
+			TManagedArray<int32>& Levels = LevelAttribute.Modify();
+			const TManagedArray<int32>& Parents = ParentAttribute.Get();
 
-		// In place level computation
-		LevelAttribute.Fill(INDEX_NONE);
-		for (int32 TransformIndex = 0; TransformIndex < Levels.Num(); TransformIndex++)
-		{
-			// parent propagation may have already set the value for this transform
-			if (Levels[TransformIndex] == INDEX_NONE)
+			// In place level computation
+			LevelAttribute.Fill(INDEX_NONE);
+			for (int32 TransformIndex = 0; TransformIndex < Levels.Num(); TransformIndex++)
 			{
-				// first compute the level by walking the parents
-				Levels[TransformIndex] = 0;
-				int32 ParentIndex = Parents[TransformIndex];
-				while (ParentIndex != INDEX_NONE)
+				// parent propagation may have already set the value for this transform
+				if (Levels[TransformIndex] == INDEX_NONE)
 				{
-					Levels[TransformIndex]++;
-
-					// do we need to go further ? 
-					const int32 ParentLevel = Levels[ParentIndex];
-					if (ParentLevel != INDEX_NONE)
+					// first compute the level by walking the parents
+					Levels[TransformIndex] = 0;
+					int32 ParentIndex = Parents[TransformIndex];
+					while (ParentIndex != INDEX_NONE)
 					{
-						Levels[TransformIndex] += ParentLevel;
-						break; // early out as the parent level has been computed already
-					}
+						Levels[TransformIndex]++;
 
-					ParentIndex = Parents[ParentIndex];
-				}
+						// do we need to go further ? 
+						const int32 ParentLevel = Levels[ParentIndex];
+						if (ParentLevel != INDEX_NONE)
+						{
+							Levels[TransformIndex] += ParentLevel;
+							break; // early out as the parent level has been computed already
+						}
 
-				// second propagate to parents if they are still unitialized
-				if (Levels[TransformIndex] > 0)
-				{
-					int32 ParentLevel = Levels[TransformIndex] - 1;
-					ParentIndex = Parents[TransformIndex];
-					while (ParentIndex != INDEX_NONE && Levels[ParentIndex] == INDEX_NONE)
-					{
-						Levels[ParentIndex] = ParentLevel;
-						ParentLevel--;
 						ParentIndex = Parents[ParentIndex];
 					}
+
+					// second propagate to parents if they are still unitialized
+					if (Levels[TransformIndex] > 0)
+					{
+						int32 ParentLevel = Levels[TransformIndex] - 1;
+						ParentIndex = Parents[TransformIndex];
+						while (ParentIndex != INDEX_NONE && Levels[ParentIndex] == INDEX_NONE)
+						{
+							Levels[ParentIndex] = ParentLevel;
+							ParentLevel--;
+							ParentIndex = Parents[ParentIndex];
+						}
+					}
 				}
 			}
 		}
diff --git a/Engine/Source/Runtime/Experimental/Chaos/Private/GeometryCollection/Facades/CollectionKinematicBindingFacade.cpp b/Engine/Source/Runtime/Experimental/Chaos/Private/GeometryCollection/Facades/CollectionKinematicBindingFacade.cpp
--- a/Engine/Source/Runtime/Experimental/Chaos/Private/GeometryCollection/Facades/CollectionKinematicBindingFacade.cpp
+++ b/Engine/Source/Runtime/Experimental/Chaos/Private/GeometryCollection/Facades/CollectionKinematicBindingFacade.cpp
@@ -7,68 +7,85 @@
 #include "GeometryCollection/Facades/CollectionKinematicBindingFacade.h"
 #include "GeometryCollection/GeometryCollection.h"
 
-namespace Chaos::Facades
+namespace GeometryCollection::Facades
 {
 	// Attributes
 	const FName FKinematicBindingFacade::KinematicGroup("Kinematics");
 	const FName FKinematicBindingFacade::KinematicBoneBindingIndex("Binding");
 	const FName FKinematicBindingFacade::KinematicBoneBindingToGroup("BindingToGroup");
 
+	FKinematicBindingFacade::FKinematicBindingFacade(FManagedArrayCollection & InCollection)
+		: ConstCollection(InCollection)
+		, Collection(&InCollection)
+		, KinemaitcBoneBindingAttribute(InCollection, KinematicBoneBindingIndex, KinematicGroup)
+		, KinemaitcBoneBindingToGroupAttribute(InCollection, KinematicBoneBindingToGroup, KinematicGroup)
+	{}
 
-	void FKinematicBindingFacade::Init()
+	FKinematicBindingFacade::FKinematicBindingFacade(const FManagedArrayCollection& InCollection)
+		: ConstCollection(InCollection)
+		, Collection(nullptr)
+		, KinemaitcBoneBindingAttribute(InCollection, KinematicBoneBindingIndex, KinematicGroup)
+		, KinemaitcBoneBindingToGroupAttribute(InCollection, KinematicBoneBindingToGroup, KinematicGroup)
+	{}
+
+	bool FKinematicBindingFacade::IsValid() const
 	{
-		FName SelectionGroup1 = FName(GeometryCollection::Facades::FSelectionFacade::BoundGroup.ToString() + "_" + FGeometryCollection::VerticesGroup.ToString());
-		FManagedArrayCollection::FConstructionParameters BindingGroupDependency(SelectionGroup1);
+		return KinemaitcBoneBindingAttribute.IsValid() && KinemaitcBoneBindingToGroupAttribute.IsValid();
+	}
 
-		// Kinematic Bindings
-		Collection.AddAttribute<int32>(FKinematicBindingFacade::KinematicBoneBindingIndex, FKinematicBindingFacade::KinematicGroup, BindingGroupDependency);
-		Collection.AddAttribute<FString>(FKinematicBindingFacade::KinematicBoneBindingToGroup, FKinematicBindingFacade::KinematicGroup);
+	void FKinematicBindingFacade::DefineSchema()
+	{
+		check(!IsConst());
+		FName SelectionGroup1 = FName(GeometryCollection::Facades::FSelectionFacade::BoundGroup.ToString() + "_" + FGeometryCollection::VerticesGroup.ToString());
+		KinemaitcBoneBindingAttribute.Add(ManageArrayAccessor::EPersistencePolicy::MakePersistent,SelectionGroup1);
+		KinemaitcBoneBindingToGroupAttribute.Add();
 	}
 
-	FKinematicBindingFacade::FBindingKey FKinematicBindingFacade::SetBoneBindings(FManagedArrayCollection* InCollection, const int32 InBoneIndex, const TArray<int32>& InBoneVerts, const TArray<float>& InBoneWeights)
+	FKinematicBindingFacade::FBindingKey 
+	FKinematicBindingFacade::SetBoneBindings(const int32 InBoneIndex, const TArray<int32>& InBoneVerts, const TArray<float>& InBoneWeights)
 	{
-		return GeometryCollection::Facades::FSelectionFacade::AddSelection(InCollection, InBoneIndex, InBoneVerts, InBoneWeights, FGeometryCollection::VerticesGroup);
+		check(!IsConst());
+		return GeometryCollection::Facades::FSelectionFacade(*Collection).AddSelection(InBoneIndex, InBoneVerts, InBoneWeights, FGeometryCollection::VerticesGroup);
 	}
 
-	void FKinematicBindingFacade::GetBoneBindings(const FManagedArrayCollection* InCollection, const FKinematicBindingFacade::FBindingKey& Key, int32& OutBoneIndex, TArray<int32>& OutBoneVerts, TArray<float>& OutBoneWeights)
+	void FKinematicBindingFacade::GetBoneBindings(const FKinematicBindingFacade::FBindingKey& Key, int32& OutBoneIndex, TArray<int32>& OutBoneVerts, TArray<float>& OutBoneWeights) const
 	{
-		GeometryCollection::Facades::FSelectionFacade::GetSelection(InCollection, Key, OutBoneIndex, OutBoneVerts, OutBoneWeights);
+		GeometryCollection::Facades::FSelectionFacade(ConstCollection).GetSelection(Key, OutBoneIndex, OutBoneVerts, OutBoneWeights);
 	}
 
 
-	int32 FKinematicBindingFacade::AddKinematicBinding(FManagedArrayCollection* InCollection, const FKinematicBindingFacade::FBindingKey& Key)
+	int32 FKinematicBindingFacade::AddKinematicBinding(const FKinematicBindingFacade::FBindingKey& Key)
 	{
-		if (InCollection && InCollection->HasGroup(KinematicGroup))
+		check(!IsConst());
+		if (IsValid())
 		{
-			int32 NewIndex = InCollection->AddElements(1, KinematicGroup);
-			TManagedArray<int32>* KinemaitcBoneBinding = InCollection->FindAttribute<int32>(KinematicBoneBindingIndex, KinematicGroup);
-			TManagedArray<FString>* KinemaitcBoneBindingToGroup = InCollection->FindAttribute<FString>(KinematicBoneBindingToGroup, KinematicGroup);
-
-			(*KinemaitcBoneBinding)[NewIndex] = Key.Index;
-			(*KinemaitcBoneBindingToGroup)[NewIndex] = Key.GroupName.ToString();
+			int32 NewIndex = KinemaitcBoneBindingAttribute.AddElements(1);
+			TManagedArray<int32>& KinemaitcBoneBinding = KinemaitcBoneBindingAttribute.Modify();
+			TManagedArray<FString>& KinemaitcBoneBindingToGroup = KinemaitcBoneBindingToGroupAttribute.Modify();
+			
+			KinemaitcBoneBinding[NewIndex] = Key.Index;
+			KinemaitcBoneBindingToGroup[NewIndex] = Key.GroupName.ToString();
 
 			return NewIndex;
 		}
 		return INDEX_NONE;
 	}
 
-	FKinematicBindingFacade::FBindingKey FKinematicBindingFacade::GetKinematicBindingKey(const FManagedArrayCollection* InCollection, int Index)
+	FKinematicBindingFacade::FBindingKey FKinematicBindingFacade::GetKinematicBindingKey(int Index) const
 	{
 		FKinematicBindingFacade::FBindingKey Key;
-		if (InCollection && InCollection->HasGroup(KinematicGroup))
+		if (IsValid())
 		{
-			const TManagedArray<int32>* KinemaitcBoneBinding = InCollection->FindAttribute<int32>(KinematicBoneBindingIndex, KinematicGroup);
-			const TManagedArray<FString>* KinemaitcBoneBindingToGroup = InCollection->FindAttribute<FString>(KinematicBoneBindingToGroup, KinematicGroup);
-			if (KinemaitcBoneBinding && KinemaitcBoneBindingToGroup)
+			const TManagedArray<int32>& KinemaitcBoneBinding = KinemaitcBoneBindingAttribute.Get();
+			const TManagedArray<FString>& KinemaitcBoneBindingToGroup = KinemaitcBoneBindingToGroupAttribute.Get();
+
+			if (0 <= Index && Index < ConstCollection.NumElements(KinematicGroup))
 			{
-				if (0 <= Index && Index < InCollection->NumElements(KinematicGroup))
-				{
-					Key.GroupName = FName((*KinemaitcBoneBindingToGroup)[Index]);
-					Key.Index = (*KinemaitcBoneBinding)[Index];
-				}
+				Key.GroupName = FName(KinemaitcBoneBindingToGroup[Index]);
+				Key.Index = KinemaitcBoneBinding[Index];
 			}
+
 		}
 		return Key;
 	}
-
 }
diff --git a/Engine/Source/Runtime/Experimental/Chaos/Private/GeometryCollection/Facades/CollectionMeshFacade.cpp b/Engine/Source/Runtime/Experimental/Chaos/Private/GeometryCollection/Facades/CollectionMeshFacade.cpp
new file mode 100644
--- /dev/null
+++ b/Engine/Source/Runtime/Experimental/Chaos/Private/GeometryCollection/Facades/CollectionMeshFacade.cpp
@@ -0,0 +1,72 @@
+// Copyright Epic Games, Inc. All Rights Reserved.
+
+/*=============================================================================
+	GeometryCollection.cpp: FGeometryCollection methods.
+=============================================================================*/
+
+#include "GeometryCollection/Facades/CollectionMeshFacade.h"
+#include "GeometryCollection/GeometryCollection.h"
+#include "GeometryCollection/TransformCollection.h"
+
+namespace GeometryCollection::Facades
+{
+	FMeshFacade::FMeshFacade(FManagedArrayCollection& InCollection)
+		: Vertex(InCollection, "Vertex", FGeometryCollection::VerticesGroup)
+		, TangentU(InCollection, "TangentU", FGeometryCollection::VerticesGroup)
+		, TangentV(InCollection, "TangentV", FGeometryCollection::VerticesGroup)
+		, Normal(InCollection, "Normal", FGeometryCollection::VerticesGroup)
+		, UVs(InCollection, "UVs", FGeometryCollection::VerticesGroup)
+		, Color(InCollection, "Color", FGeometryCollection::VerticesGroup)
+		, BoneMap(InCollection, "BoneMap", FGeometryCollection::VerticesGroup)
+		, VertexStart(InCollection, "VertexStart", FGeometryCollection::GeometryGroup)
+		, VertexCount(InCollection, "VertexCount", FGeometryCollection::GeometryGroup)
+		, Indices(InCollection, "Indices", FGeometryCollection::FacesGroup)
+		, Visible(InCollection, "Visible", FGeometryCollection::FacesGroup)
+		, MaterialIndex(InCollection, "MaterialIndex", FGeometryCollection::FacesGroup)
+		, MaterialID(InCollection, "MaterialID", FGeometryCollection::FacesGroup)
+		, FaceStart(InCollection, "FaceStart", FGeometryCollection::GeometryGroup)
+		, FaceCount(InCollection, "FaceCount", FGeometryCollection::GeometryGroup)
+	{
+	}
+
+	bool FMeshFacade::IsValid() const
+	{
+		return Vertex.IsValid()
+			&& TangentU.IsValid()
+			&& TangentV.IsValid()
+			&& Normal.IsValid()
+			&& UVs.IsValid()
+			&& Color.IsValid()
+			&& BoneMap.IsValid()
+			&& VertexStart.IsValid()
+			&& VertexCount.IsValid()
+			&& Indices.IsValid()
+			&& Visible.IsValid()
+			&& MaterialIndex.IsValid()
+			&& MaterialID.IsValid()
+			&& FaceStart.IsValid()
+			&& FaceCount.IsValid()
+			;
+	}
+
+	void FMeshFacade::DefineSchema()
+	{
+		Vertex.Add();
+		TangentU.Add();
+		TangentV.Add();
+		Normal.Add();
+		UVs.Add();
+		Color.Add();
+		BoneMap.Add();
+		VertexStart.Add();
+		VertexCount.Add();
+		Indices.Add();
+		Visible.Add();
+		MaterialIndex.Add();
+		MaterialID.Add();
+		FaceStart.Add();
+		FaceCount.Add();
+	}
+};
+
+
diff --git a/Engine/Source/Runtime/Experimental/Chaos/Private/GeometryCollection/Facades/CollectionRenderingFacade.cpp b/Engine/Source/Runtime/Experimental/Chaos/Private/GeometryCollection/Facades/CollectionRenderingFacade.cpp
--- a/Engine/Source/Runtime/Experimental/Chaos/Private/GeometryCollection/Facades/CollectionRenderingFacade.cpp
+++ b/Engine/Source/Runtime/Experimental/Chaos/Private/GeometryCollection/Facades/CollectionRenderingFacade.cpp
@@ -10,72 +10,60 @@
 namespace GeometryCollection::Facades
 {
 
-	// Groups 
-	const FName FRenderingFacade::VerticesGroup = FGeometryCollection::VerticesGroup;
-	const FName FRenderingFacade::FacesGroup = FGeometryCollection::FacesGroup;
-
-
-	// Attributes
-	const FName FRenderingFacade::VertexAttribute = "Vertex";
-	const FName FRenderingFacade::IndicesAttribute = "Indices";
-
-	FRenderingFacade::FRenderingFacade(FManagedArrayCollection* InCollection)
-		: Self(InCollection)
-	{
-		DefineSchema(Self);
-	}
+	FRenderingFacade::FRenderingFacade(FManagedArrayCollection& InCollection)
+		: ConstCollection(InCollection)
+		, Collection(&InCollection)
+		, VertexAttribute(InCollection, "Vertex", FGeometryCollection::VerticesGroup)
+		, IndicesAttribute(InCollection, "Indices", FGeometryCollection::FacesGroup, FGeometryCollection::VerticesGroup)
+		, MaterialIDAttribute(InCollection, "MaterialID", FGeometryCollection::FacesGroup)
+		, TriangleSectionAttribute(InCollection, "Sections", FGeometryCollection::MaterialGroup)
+	{}
+
+	FRenderingFacade::FRenderingFacade(const FManagedArrayCollection& InCollection)
+		: ConstCollection(InCollection)
+		, Collection(nullptr)
+		, VertexAttribute(InCollection, "Vertex", FGeometryCollection::VerticesGroup)
+		, IndicesAttribute(InCollection, "Indices", FGeometryCollection::FacesGroup, FGeometryCollection::VerticesGroup)
+		, MaterialIDAttribute(InCollection, "MaterialID", FGeometryCollection::FacesGroup)
+		, TriangleSectionAttribute(InCollection, "Sections", FGeometryCollection::MaterialGroup)
+	{}
 
 	//
 	//  Initialization
 	//
 
-	void FRenderingFacade::DefineSchema(FManagedArrayCollection* InCollection)
+	void FRenderingFacade::DefineSchema()
 	{
-		FManagedArrayCollection::FConstructionParameters VertexDependency(FGeometryCollection::VerticesGroup);
-
-		// surface rendering attributes
-		if (!InCollection->HasGroup(FGeometryCollection::VerticesGroup))
-		{
-			InCollection->AddGroup(FGeometryCollection::VerticesGroup);
-		}
-		if (!InCollection->HasGroup(FGeometryCollection::FacesGroup))
-		{
-			InCollection->AddGroup(FGeometryCollection::FacesGroup);
-		}
-
-		InCollection->AddAttribute<FVector3f>(VertexAttribute, FGeometryCollection::VerticesGroup);
-		InCollection->AddAttribute<FIntVector>(IndicesAttribute, FGeometryCollection::FacesGroup, VertexDependency);
-
-		ensure(InCollection->FindAttributeTyped<FVector3f>(VertexAttribute, FGeometryCollection::VerticesGroup) != nullptr);
-		ensure(InCollection->FindAttributeTyped<FIntVector>(IndicesAttribute, FGeometryCollection::FacesGroup) != nullptr);
+		check(!IsConst());
+		VertexAttribute.Add();
+		IndicesAttribute.Add();
+		MaterialIDAttribute.Add();
+		TriangleSectionAttribute.Add();
 	}
 
-	bool FRenderingFacade::CanRenderSurface(const FManagedArrayCollection* InCollection)
+	bool FRenderingFacade::CanRenderSurface( ) const
 	{
-		return  InCollection && IsValid(InCollection) &&
-			GetIndices(InCollection)->Num() &&
-			GetVertices(InCollection)->Num();
+		return  IsValid() && GetIndices().Num() && GetVertices().Num();
 	}
 
-	bool FRenderingFacade::IsValid(const FManagedArrayCollection* InCollection)
+	bool FRenderingFacade::IsValid( ) const
 	{
-		return InCollection && InCollection->HasGroup(FGeometryCollection::VerticesGroup)
-			&& InCollection->HasGroup(FGeometryCollection::FacesGroup)
-			&& InCollection->FindAttributeTyped<FVector3f>(VertexAttribute, FGeometryCollection::VerticesGroup)
-			&& InCollection->FindAttributeTyped<FIntVector>(IndicesAttribute, FGeometryCollection::FacesGroup);
+		return VertexAttribute.IsValid() && IndicesAttribute.IsValid() &&
+			MaterialIDAttribute.IsValid() && TriangleSectionAttribute.IsValid();
 	}
 
-	void FRenderingFacade::AddTriangle(FManagedArrayCollection* InCollection, const Chaos::FTriangle& InTriangle)
+	void FRenderingFacade::AddTriangle(const Chaos::FTriangle& InTriangle)
 	{
-		if (IsValid(InCollection))
+		check(!IsConst());
+		if (IsValid())
 		{
 			auto CollectionVert = [](const Chaos::FVec3& V) { return FVector3f(float(V.X), float(V.Y), float(V.Z)); };
 
-			TManagedArray<FVector3f>& Vertices = InCollection->ModifyAttribute<FVector3f>(VertexAttribute, FGeometryCollection::VerticesGroup);
-			TManagedArray<FIntVector>& Indices = InCollection->ModifyAttribute<FIntVector>(IndicesAttribute, FGeometryCollection::FacesGroup);
-
-			int32 IndicesStart = InCollection->AddElements(1, FGeometryCollection::FacesGroup);
-			int32 VertexStart = InCollection->AddElements(3, FGeometryCollection::VerticesGroup);
+			TManagedArray<FVector3f>& Vertices = VertexAttribute.Modify();
+			TManagedArray<FIntVector>& Indices = IndicesAttribute.Modify();
+			
+			int32 IndicesStart = IndicesAttribute.AddElements(1);
+			int32 VertexStart = VertexAttribute.AddElements(3);
 
 			Indices[IndicesStart] = FIntVector(VertexStart, VertexStart + 1, VertexStart + 2);
 			Vertices[VertexStart] = CollectionVert(InTriangle[0]);
@@ -85,18 +73,19 @@ namespace GeometryCollection::Facades
 	}
 
 
-	void FRenderingFacade::AddSurface(FManagedArrayCollection* InCollection, TArray<FVector3f>&& InVertices, TArray<FIntVector>&& InIndices)
+	void FRenderingFacade::AddSurface(TArray<FVector3f>&& InVertices, TArray<FIntVector>&& InIndices)
 	{
-		if (IsValid(InCollection))
+		check(!IsConst());
+		if (IsValid())
 		{
 			auto CollectionVert = [](const Chaos::FVec3& V) { return FVector3f(float(V.X), float(V.Y), float(V.Z)); };
 
-			TManagedArray<FVector3f>& Vertices = InCollection->ModifyAttribute<FVector3f>(VertexAttribute, FGeometryCollection::VerticesGroup);
-			TManagedArray<FIntVector>& Indices = InCollection->ModifyAttribute<FIntVector>(IndicesAttribute, FGeometryCollection::FacesGroup);
-
-			int32 IndicesStart = InCollection->AddElements(InIndices.Num(), FGeometryCollection::FacesGroup);
-			int32 VertexStart = InCollection->AddElements(InVertices.Num(), FGeometryCollection::VerticesGroup);
+			TManagedArray<FVector3f>& Vertices = VertexAttribute.Modify();
+			TManagedArray<FIntVector>& Indices = IndicesAttribute.Modify();
 
+			int32 IndicesStart = IndicesAttribute.AddElements(InIndices.Num());
+			int32 VertexStart = VertexAttribute.AddElements(InVertices.Num());
+			
 			const FIntVector * IndiciesDest = Indices.GetData() + IndicesStart;
 			FMemory::Memmove((void*)IndiciesDest, InIndices.GetData(), sizeof(FIntVector) * InIndices.Num());
 
@@ -105,21 +94,14 @@ namespace GeometryCollection::Facades
 		}
 	}
 
-	//
-	//  GetAttributes
-	//
-
-	const TManagedArray< FIntVector >* FRenderingFacade::GetIndices(const FManagedArrayCollection* InCollection)
-	{
-		return InCollection->FindAttribute<FIntVector>(IndicesAttribute, FGeometryCollection::FacesGroup);
-	}
 
-	const TManagedArray< FVector3f >* FRenderingFacade::GetVertices(const FManagedArrayCollection* InCollection)
+	TArray<FRenderingFacade::FTriangleSection> 
+	FRenderingFacade::BuildMeshSections(const TArray<FIntVector>& InputIndices, TArray<int32> BaseMeshOriginalIndicesIndex, TArray<FIntVector>& RetIndices) const
 	{
-		return InCollection->FindAttribute<FVector3f>(VertexAttribute, FGeometryCollection::VerticesGroup);
+		check(!IsConst());
+		return FGeometryCollectionSection::BuildMeshSections(ConstCollection, InputIndices, BaseMeshOriginalIndicesIndex, RetIndices);
 	}
 
-
 }; // GeometryCollection::Facades
 
 
diff --git a/Engine/Source/Runtime/Experimental/Chaos/Private/GeometryCollection/Facades/CollectionSelectionFacade.cpp b/Engine/Source/Runtime/Experimental/Chaos/Private/GeometryCollection/Facades/CollectionSelectionFacade.cpp
--- a/Engine/Source/Runtime/Experimental/Chaos/Private/GeometryCollection/Facades/CollectionSelectionFacade.cpp
+++ b/Engine/Source/Runtime/Experimental/Chaos/Private/GeometryCollection/Facades/CollectionSelectionFacade.cpp
@@ -20,16 +20,24 @@ namespace GeometryCollection::Facades
 	const FName FSelectionFacade::WeightAttribute = "Weights";
 	const FName FSelectionFacade::BoneIndexAttribute = "BoneIndex";
 
-	FSelectionFacade::FSelectionFacade(FManagedArrayCollection* InCollection)
-		: Self(InCollection)
+	FSelectionFacade::FSelectionFacade(FManagedArrayCollection& InCollection)
+		: ConstCollection(InCollection)
+		, Collection(&InCollection) 
+	{}
+
+	FSelectionFacade::FSelectionFacade(const FManagedArrayCollection& InCollection)
+		: ConstCollection(InCollection)
+		, Collection(nullptr)
 	{}
 
 	//
 	//  Initialization
 	//
 
-	void FSelectionFacade::InitUnboundedGroup(FManagedArrayCollection* Collection, FName GroupName, FName DependencyGroup)
+	void FSelectionFacade::InitUnboundedGroup(FName GroupName, FName DependencyGroup)
 	{
+		check(!IsConst());
+
 		if (!Collection->HasGroup(GroupName))
 		{
 			Collection->AddAttribute<TArray<int32>>(FSelectionFacade::IndexAttribute, GroupName, { DependencyGroup });
@@ -37,8 +45,10 @@ namespace GeometryCollection::Facades
 		ensure(Collection->FindAttributeTyped<TArray<int32>>(FSelectionFacade::IndexAttribute, GroupName) != nullptr);
 	}
 
-	void FSelectionFacade::InitWeightedUnboundedGroup(FManagedArrayCollection* Collection, FName GroupName, FName DependencyGroup)
+	void FSelectionFacade::InitWeightedUnboundedGroup(FName GroupName, FName DependencyGroup)
 	{
+		check(!IsConst());
+
 		if (!Collection->HasGroup(GroupName))
 		{
 			Collection->AddAttribute<TArray<int32>>(FSelectionFacade::IndexAttribute, GroupName, { DependencyGroup });
@@ -48,8 +58,10 @@ namespace GeometryCollection::Facades
 		ensure(Collection->FindAttributeTyped<TArray<float>>(FSelectionFacade::WeightAttribute, GroupName) != nullptr);
 	}
 
-	void FSelectionFacade::InitBoundedGroup(FManagedArrayCollection* Collection, FName GroupName, FName DependencyGroup, FName BoneDependencyGroup)
+	void FSelectionFacade::InitBoundedGroup(FName GroupName, FName DependencyGroup, FName BoneDependencyGroup)
 	{
+		check(!IsConst());
+
 		if (!Collection->HasGroup(GroupName))
 		{
 			Collection->AddAttribute<TArray<int32>>(FSelectionFacade::IndexAttribute, GroupName, { DependencyGroup });
@@ -59,8 +71,10 @@ namespace GeometryCollection::Facades
 		ensure(Collection->FindAttributeTyped<int32>(FSelectionFacade::BoneIndexAttribute, GroupName) != nullptr);
 	}
 
-	void FSelectionFacade::InitWeightedBoundedGroup(FManagedArrayCollection* Collection, FName GroupName, FName DependencyGroup, FName BoneDependencyGroup)
+	void FSelectionFacade::InitWeightedBoundedGroup(FName GroupName, FName DependencyGroup, FName BoneDependencyGroup)
 	{
+		check(!IsConst());
+
 		if (!Collection->HasGroup(GroupName))
 		{
 			Collection->AddAttribute<TArray<int32>>(FSelectionFacade::IndexAttribute, GroupName, { DependencyGroup });
@@ -77,20 +91,24 @@ namespace GeometryCollection::Facades
 	//  AddSelection
 	//
 
-	FSelectionFacade::FSelectionKey FSelectionFacade::AddSelection(FManagedArrayCollection* Collection, const TArray<int32>& InIndices, FName DependencyGroup)
+	FSelectionFacade::FSelectionKey FSelectionFacade::AddSelection(const TArray<int32>& InIndices, FName DependencyGroup)
 	{
+		check(!IsConst());
+
 		FName GroupName(FSelectionFacade::UnboundGroup.ToString() + "_" + DependencyGroup.ToString());
-		InitUnboundedGroup(Collection, GroupName, DependencyGroup);
+		InitUnboundedGroup(GroupName, DependencyGroup);
 
 		int Idx = Collection->AddElements(1, GroupName);
 		Collection->ModifyAttribute<TArray<int32>>(FSelectionFacade::IndexAttribute, GroupName)[Idx] = InIndices;
 		return FSelectionKey(Idx, FSelectionFacade::UnboundGroup);
 	}
 
-	FSelectionFacade::FSelectionKey FSelectionFacade::AddSelection(FManagedArrayCollection* Collection, const TArray<int32>& InIndices, const TArray<float>& InWeights, FName DependencyGroup)
+	FSelectionFacade::FSelectionKey FSelectionFacade::AddSelection(const TArray<int32>& InIndices, const TArray<float>& InWeights, FName DependencyGroup)
 	{
+		check(!IsConst());
+
 		FName GroupName(FSelectionFacade::WeightedUnboundGroup.ToString() + "_" + DependencyGroup.ToString());
-		InitWeightedUnboundedGroup(Collection, GroupName, DependencyGroup);
+		InitWeightedUnboundedGroup(GroupName, DependencyGroup);
 
 		int Idx = Collection->AddElements(1, GroupName);
 		Collection->ModifyAttribute<TArray<int32>>(FSelectionFacade::IndexAttribute, GroupName)[Idx] = InIndices;
@@ -98,10 +116,12 @@ namespace GeometryCollection::Facades
 		return FSelectionKey(Idx, GroupName);
 	}
 
-	FSelectionFacade::FSelectionKey FSelectionFacade::AddSelection(FManagedArrayCollection* Collection, const int32 InBoneIndex, const TArray<int32>& InIndices, FName DependencyGroup, FName BoneDependencyGroup)
+	FSelectionFacade::FSelectionKey FSelectionFacade::AddSelection(const int32 InBoneIndex, const TArray<int32>& InIndices, FName DependencyGroup, FName BoneDependencyGroup)
 	{
+		check(!IsConst());
+
 		FName GroupName(FSelectionFacade::BoundGroup.ToString() + "_" + DependencyGroup.ToString());
-		InitBoundedGroup(Collection, GroupName, DependencyGroup, BoneDependencyGroup);
+		InitBoundedGroup(GroupName, DependencyGroup, BoneDependencyGroup);
 
 		int Idx = Collection->AddElements(1, GroupName);
 		Collection->ModifyAttribute<TArray<int32>>(FSelectionFacade::IndexAttribute, GroupName)[Idx] = InIndices;
@@ -109,10 +129,12 @@ namespace GeometryCollection::Facades
 		return FSelectionKey(Idx, GroupName);
 	}
 
-	FSelectionFacade::FSelectionKey FSelectionFacade::AddSelection(FManagedArrayCollection* Collection, const int32 InBoneIndex, const TArray<int32>& InIndices, const TArray<float>& InWeights, FName DependencyGroup, FName BoneDependencyGroup)
+	FSelectionFacade::FSelectionKey FSelectionFacade::AddSelection(const int32 InBoneIndex, const TArray<int32>& InIndices, const TArray<float>& InWeights, FName DependencyGroup, FName BoneDependencyGroup)
 	{
+		check(!IsConst());
+
 		FName GroupName(FSelectionFacade::WeightedBoundGroup.ToString() + "_" + DependencyGroup.ToString());
-		InitWeightedBoundedGroup(Collection, GroupName, DependencyGroup, BoneDependencyGroup);
+		InitWeightedBoundedGroup(GroupName, DependencyGroup, BoneDependencyGroup);
 
 		int Idx = Collection->AddElements(1, GroupName);
 		Collection->ModifyAttribute<TArray<int32>>(FSelectionFacade::IndexAttribute, GroupName)[Idx] = InIndices;
@@ -124,49 +146,47 @@ namespace GeometryCollection::Facades
 	//
 	//  GetSelection
 	//
-
-
-	void FSelectionFacade::GetSelection(const FManagedArrayCollection* Collection, const FSelectionFacade::FSelectionKey& Key, TArray<int32>& OutIndices)
+	void FSelectionFacade::GetSelection(const FSelectionFacade::FSelectionKey& Key, TArray<int32>& OutIndices) const
 	{
-		if (Collection->HasGroup(Key.GroupName) && 0 <= Key.Index && Key.Index < Collection->NumElements(Key.GroupName))
+		if (ConstCollection.HasGroup(Key.GroupName) && 0 <= Key.Index && Key.Index < ConstCollection.NumElements(Key.GroupName))
 		{
-			if (Collection->FindAttribute<TArray<int32>>(FSelectionFacade::IndexAttribute, Key.GroupName))
-				OutIndices = Collection->GetAttribute<TArray<int32>>(FSelectionFacade::IndexAttribute, Key.GroupName)[Key.Index];
+			if (ConstCollection.FindAttribute<TArray<int32>>(FSelectionFacade::IndexAttribute, Key.GroupName))
+				OutIndices = ConstCollection.GetAttribute<TArray<int32>>(FSelectionFacade::IndexAttribute, Key.GroupName)[Key.Index];
 		}
 	}
 
-	void FSelectionFacade::GetSelection(const FManagedArrayCollection* Collection, const FSelectionFacade::FSelectionKey& Key, TArray<int32>& OutIndices, TArray<float>& OutWeights)
+	void FSelectionFacade::GetSelection(const FSelectionFacade::FSelectionKey& Key, TArray<int32>& OutIndices, TArray<float>& OutWeights) const
 	{
-		if (Collection->HasGroup(Key.GroupName) && 0 <= Key.Index && Key.Index < Collection->NumElements(Key.GroupName))
+		if (ConstCollection.HasGroup(Key.GroupName) && 0 <= Key.Index && Key.Index < ConstCollection.NumElements(Key.GroupName))
 		{
-			if (Collection->FindAttribute<TArray<int32>>(FSelectionFacade::IndexAttribute, Key.GroupName))
-				OutIndices = Collection->GetAttribute<TArray<int32>>(FSelectionFacade::IndexAttribute, Key.GroupName)[Key.Index];
-			if (Collection->FindAttribute<TArray<float>>(FSelectionFacade::WeightAttribute, Key.GroupName))
-				OutWeights = Collection->GetAttribute<TArray<float>>(FSelectionFacade::WeightAttribute, Key.GroupName)[Key.Index];
+			if (ConstCollection.FindAttribute<TArray<int32>>(FSelectionFacade::IndexAttribute, Key.GroupName))
+				OutIndices = ConstCollection.GetAttribute<TArray<int32>>(FSelectionFacade::IndexAttribute, Key.GroupName)[Key.Index];
+			if (ConstCollection.FindAttribute<TArray<float>>(FSelectionFacade::WeightAttribute, Key.GroupName))
+				OutWeights = ConstCollection.GetAttribute<TArray<float>>(FSelectionFacade::WeightAttribute, Key.GroupName)[Key.Index];
 		}
 	}
 
-	void FSelectionFacade::GetSelection(const FManagedArrayCollection* Collection, const FSelectionFacade::FSelectionKey& Key, int32& OutBoneIndex, TArray<int32>& OutIndices)
+	void FSelectionFacade::GetSelection(const FSelectionFacade::FSelectionKey& Key, int32& OutBoneIndex, TArray<int32>& OutIndices) const
 	{
-		if (Collection->HasGroup(Key.GroupName) && 0 <= Key.Index && Key.Index < Collection->NumElements(Key.GroupName))
+		if (ConstCollection.HasGroup(Key.GroupName) && 0 <= Key.Index && Key.Index < ConstCollection.NumElements(Key.GroupName))
 		{
-			if (Collection->FindAttribute<TArray<int32>>(FSelectionFacade::IndexAttribute, Key.GroupName))
-				OutIndices = Collection->GetAttribute<TArray<int32>>(FSelectionFacade::IndexAttribute, Key.GroupName)[Key.Index];
-			if (Collection->FindAttribute<int32>(FSelectionFacade::BoneIndexAttribute, Key.GroupName))
-				OutBoneIndex = Collection->GetAttribute<int32>(FSelectionFacade::BoneIndexAttribute, Key.GroupName)[Key.Index];
+			if (ConstCollection.FindAttribute<TArray<int32>>(FSelectionFacade::IndexAttribute, Key.GroupName))
+				OutIndices = ConstCollection.GetAttribute<TArray<int32>>(FSelectionFacade::IndexAttribute, Key.GroupName)[Key.Index];
+			if (ConstCollection.FindAttribute<int32>(FSelectionFacade::BoneIndexAttribute, Key.GroupName))
+				OutBoneIndex = ConstCollection.GetAttribute<int32>(FSelectionFacade::BoneIndexAttribute, Key.GroupName)[Key.Index];
 		}
 	}
 
-	void FSelectionFacade::GetSelection(const FManagedArrayCollection* Collection, const FSelectionFacade::FSelectionKey& Key, int32& OutBoneIndex, TArray<int32>& OutIndices, TArray<float>& OutWeights)
+	void FSelectionFacade::GetSelection(const FSelectionFacade::FSelectionKey& Key, int32& OutBoneIndex, TArray<int32>& OutIndices, TArray<float>& OutWeights) const
 	{
-		if (Collection->HasGroup(Key.GroupName) && 0 <= Key.Index && Key.Index < Collection->NumElements(Key.GroupName))
+		if (ConstCollection.HasGroup(Key.GroupName) && 0 <= Key.Index && Key.Index < ConstCollection.NumElements(Key.GroupName))
 		{
-			if (Collection->FindAttribute<TArray<int32>>(FSelectionFacade::IndexAttribute, Key.GroupName))
-				OutIndices = Collection->GetAttribute<TArray<int32>>(FSelectionFacade::IndexAttribute, Key.GroupName)[Key.Index];
-			if (Collection->FindAttribute<TArray<float>>(FSelectionFacade::WeightAttribute, Key.GroupName))
-				OutWeights = Collection->GetAttribute<TArray<float>>(FSelectionFacade::WeightAttribute, Key.GroupName)[Key.Index];
-			if (Collection->FindAttribute<int32>(FSelectionFacade::BoneIndexAttribute, Key.GroupName))
-				OutBoneIndex = Collection->GetAttribute<int32>(FSelectionFacade::BoneIndexAttribute, Key.GroupName)[Key.Index];
+			if (ConstCollection.FindAttribute<TArray<int32>>(FSelectionFacade::IndexAttribute, Key.GroupName))
+				OutIndices = ConstCollection.GetAttribute<TArray<int32>>(FSelectionFacade::IndexAttribute, Key.GroupName)[Key.Index];
+			if (ConstCollection.FindAttribute<TArray<float>>(FSelectionFacade::WeightAttribute, Key.GroupName))
+				OutWeights = ConstCollection.GetAttribute<TArray<float>>(FSelectionFacade::WeightAttribute, Key.GroupName)[Key.Index];
+			if (ConstCollection.FindAttribute<int32>(FSelectionFacade::BoneIndexAttribute, Key.GroupName))
+				OutBoneIndex = ConstCollection.GetAttribute<int32>(FSelectionFacade::BoneIndexAttribute, Key.GroupName)[Key.Index];
 		}
 	}
 
diff --git a/Engine/Source/Runtime/Experimental/Chaos/Private/GeometryCollection/Facades/CollectionTransformFacade.cpp b/Engine/Source/Runtime/Experimental/Chaos/Private/GeometryCollection/Facades/CollectionTransformFacade.cpp
new file mode 100644
--- /dev/null
+++ b/Engine/Source/Runtime/Experimental/Chaos/Private/GeometryCollection/Facades/CollectionTransformFacade.cpp
@@ -0,0 +1,43 @@
+// Copyright Epic Games, Inc. All Rights Reserved.
+
+#include "GeometryCollection/Facades/CollectionTransformFacade.h"
+#include "GeometryCollection/TransformCollection.h"
+
+namespace GeometryCollection::Facades
+{
+	// Group
+	const FName FTransformFacade::TransformGroup = FTransformCollection::TransformGroup;
+
+	// Attributes
+	const FName FTransformFacade::ParentAttribute = FTransformCollection::ParentAttribute;
+	const FName FTransformFacade::ChildrenAttribute = FTransformCollection::ChildrenAttribute;
+	const FName FTransformFacade::TransformAttribute = FTransformCollection::TransformAttribute;
+
+
+	FTransformFacade::FTransformFacade(FManagedArrayCollection& InCollection)
+		: Self(InCollection)
+		, Parent(InCollection, ChildrenAttribute, TransformGroup)
+		, Children(InCollection, ParentAttribute, TransformGroup)
+		, Transform(InCollection, TransformAttribute, TransformGroup)
+	{}
+
+	bool FTransformFacade::IsValid() const
+	{
+		return Parent.IsValid() && Children.IsValid() && Transform.IsValid();
+	}
+
+	TSet<int32> FTransformFacade::GetRootIndices() const
+	{
+		TSet<int32> Roots;
+		if (Parent.IsValid())
+		{
+			const TManagedArray<int32>& Parents = Parent.Get();
+			for (int i = 0; i < Parents.Num(); i++)
+			{
+				Roots.Add(i);
+			}
+		}
+		return Roots;
+	}
+
+}
\ No newline at end of file
diff --git a/Engine/Source/Runtime/Experimental/Chaos/Private/GeometryCollection/Facades/CollectionTransformSourceFacade.cpp b/Engine/Source/Runtime/Experimental/Chaos/Private/GeometryCollection/Facades/CollectionTransformSourceFacade.cpp
--- a/Engine/Source/Runtime/Experimental/Chaos/Private/GeometryCollection/Facades/CollectionTransformSourceFacade.cpp
+++ b/Engine/Source/Runtime/Experimental/Chaos/Private/GeometryCollection/Facades/CollectionTransformSourceFacade.cpp
@@ -6,84 +6,77 @@
 
 #include "GeometryCollection/Facades/CollectionTransformSourceFacade.h"
 #include "GeometryCollection/TransformCollection.h"
+#include "GeometryCollection/ManagedArrayAccessor.h"
 
 namespace GeometryCollection::Facades
 {
-
 	// Groups 
-	const FName FTransformSource::TransformSourceGroup = "TransformSource";
-
+	const FName FTransformSource::TransformSourceGroupName = "TransformSource";
 
 	// Attributes
-	const FName FTransformSource::SourceNameAttribute = "Name";
-	const FName FTransformSource::SourceGuidAttribute = "GuidID";
-	const FName FTransformSource::SourceRootsAttribute = "Roots";
-
-	FTransformSource::FTransformSource(FManagedArrayCollection* InCollection)
-		: Self(InCollection)
-	{
-		DefineSchema(Self);
-	}
-
+	const FName FTransformSource::SourceNameAttributeName = "Name";
+	const FName FTransformSource::SourceGuidAttributeName = "GuidID";
+	const FName FTransformSource::SourceRootsAttributeName = "Roots";
+
+	FTransformSource::FTransformSource(FManagedArrayCollection& InCollection)
+		: SourceNameAttribute(InCollection, SourceNameAttributeName, TransformSourceGroupName)
+		, SourceGuidAttribute(InCollection, SourceGuidAttributeName, TransformSourceGroupName)
+		, SourceRootsAttribute(InCollection, SourceRootsAttributeName, TransformSourceGroupName, FTransformCollection::TransformGroup)
+	{}
+
+	FTransformSource::FTransformSource(const FManagedArrayCollection& InCollection)
+		: SourceNameAttribute(InCollection, SourceNameAttributeName, TransformSourceGroupName)
+		, SourceGuidAttribute(InCollection, SourceGuidAttributeName, TransformSourceGroupName)
+		, SourceRootsAttribute(InCollection, SourceRootsAttributeName, TransformSourceGroupName, FTransformCollection::TransformGroup)
+	{}
 	//
 	//  Initialization
 	//
 
-	void FTransformSource::DefineSchema(FManagedArrayCollection* Collection)
+	void FTransformSource::DefineSchema()
 	{
-		FManagedArrayCollection::FConstructionParameters TransformDependency(FTransformCollection::TransformGroup);
-
-		if (!Collection->HasGroup(TransformSourceGroup))
-		{
-			Collection->AddGroup(TransformSourceGroup);
-		}
-		Collection->AddAttribute<FString>(SourceNameAttribute, TransformSourceGroup);
-		Collection->AddAttribute<FString>(SourceGuidAttribute, TransformSourceGroup);
-		Collection->AddAttribute<TSet<int32>>(SourceRootsAttribute, TransformSourceGroup, TransformDependency);
+		check(!IsConst());
+		SourceNameAttribute.Add();
+		SourceGuidAttribute.Add();
+		SourceRootsAttribute.Add();
 	}
 
-	bool FTransformSource::HasFacade(const FManagedArrayCollection* Collection)
+	bool FTransformSource::IsValid() const
 	{
-		return Collection->HasGroup(TransformSourceGroup) &&
-			Collection->FindAttribute<FString>(SourceNameAttribute, TransformSourceGroup) &&
-			Collection->FindAttribute<FString>(SourceGuidAttribute, TransformSourceGroup) &&
-			Collection->FindAttribute<TSet<int32>>(SourceRootsAttribute, TransformSourceGroup);
+		return SourceNameAttribute.IsValid() && SourceGuidAttribute.IsValid() && SourceRootsAttribute.IsValid(); 
 	}
 
 	//
 	//  Add Data
 	//
-	void FTransformSource::AddTransformSource(FManagedArrayCollection* Collection, const FString& InName, const FString& InGuid, const TSet<int32>& InRoots)
+	void FTransformSource::AddTransformSource(const FString& InName, const FString& InGuid, const TSet<int32>& InRoots)
 	{
-		DefineSchema(Collection);
+		check(!IsConst());
+		DefineSchema();
 
-		int Idx = Collection->AddElements(1, TransformSourceGroup);
-		Collection->ModifyAttribute<FString>(FTransformSource::SourceNameAttribute, TransformSourceGroup)[Idx] = InName;
-		Collection->ModifyAttribute<FString>(FTransformSource::SourceGuidAttribute, TransformSourceGroup)[Idx] = InGuid;
-		Collection->ModifyAttribute<TSet<int32>>(FTransformSource::SourceRootsAttribute, TransformSourceGroup)[Idx] = InRoots;
+		int Idx = SourceNameAttribute.AddElements(1);
+		SourceNameAttribute.Modify()[Idx] = InName;
+		SourceGuidAttribute.Modify()[Idx] = InGuid;
+		SourceRootsAttribute.Modify()[Idx] = InRoots;
 	}
 
 	//
 	//  Get Data
 	//
-
-
-	TSet<int32> FTransformSource::GetTransformSource(const FManagedArrayCollection* Collection, const FString& InName, const FString& InGuid)
+	TSet<int32> FTransformSource::GetTransformSource(const FString& InName, const FString& InGuid) const
 	{
-		if (Collection->HasGroup(TransformSourceGroup))
+		if (IsValid())
 		{
-			const TManagedArray<TSet<int32>>* Roots = Collection->FindAttribute<TSet<int32>>(SourceRootsAttribute, TransformSourceGroup);
-			const TManagedArray<FString>* Guids = Collection->FindAttribute<FString>(SourceGuidAttribute, TransformSourceGroup);
-			const TManagedArray<FString>* Names = Collection->FindAttribute<FString>(SourceNameAttribute, TransformSourceGroup);
-			if (Roots && Guids && Names)
+			const TManagedArray<TSet<int32>>& Roots = SourceRootsAttribute.Get();
+			const TManagedArray<FString>& Guids = SourceGuidAttribute.Get();
+			const TManagedArray<FString>&Names = SourceNameAttribute.Get();
+
+			int32 GroupNum = SourceNameAttribute.Num();
+			for (int i = 0; i < GroupNum; i++)
 			{
-				int32 GroupNum = Collection->NumElements(TransformSourceGroup);
-				for (int i = 0; i < GroupNum; i++)
+				if (Guids[i] == InGuid && Names[i].Equals(InName))
 				{
-					if ((*Guids)[i]==InGuid && (*Names)[i].Equals(InName))
-					{
-						return (*Roots)[i];
-					}
+					return Roots[i];
 				}
 			}
 		}
diff --git a/Engine/Source/Runtime/Experimental/Chaos/Private/GeometryCollection/Facades/CollectionVertexBoneWeightsFacade.cpp b/Engine/Source/Runtime/Experimental/Chaos/Private/GeometryCollection/Facades/CollectionVertexBoneWeightsFacade.cpp
--- a/Engine/Source/Runtime/Experimental/Chaos/Private/GeometryCollection/Facades/CollectionVertexBoneWeightsFacade.cpp
+++ b/Engine/Source/Runtime/Experimental/Chaos/Private/GeometryCollection/Facades/CollectionVertexBoneWeightsFacade.cpp
@@ -12,79 +12,65 @@ namespace GeometryCollection::Facades
 {
 
 	// Attributes
-	const FName FVertexBoneWeightsFacade::WeightAttribute = "BoneWeights";
-	const FName FVertexBoneWeightsFacade::IndexAttribute = "BoneWeightsIndex";
+	const FName FVertexBoneWeightsFacade::BoneWeightAttributeName = "BoneWeights";
+	const FName FVertexBoneWeightsFacade::BoneIndexAttributeName = "BoneWeightsIndex";
+
+	FVertexBoneWeightsFacade::FVertexBoneWeightsFacade(FManagedArrayCollection& InCollection)
+		: ConstCollection(InCollection)
+		, Collection(&InCollection)
+		, BoneIndexAttribute(InCollection, BoneIndexAttributeName, FGeometryCollection::VerticesGroup, FTransformCollection::TransformGroup)
+		, BoneWeightAttribute(InCollection, BoneWeightAttributeName, FGeometryCollection::VerticesGroup, FTransformCollection::TransformGroup)
+		, ParentAttribute(InCollection, FTransformCollection::ParentAttribute, FTransformCollection::TransformGroup)
+	{
+		DefineSchema();
+	}
 
-	FVertexBoneWeightsFacade::FVertexBoneWeightsFacade(FManagedArrayCollection* InCollection)
-		: Self(InCollection)
+	FVertexBoneWeightsFacade::FVertexBoneWeightsFacade(const FManagedArrayCollection& InCollection)
+		: ConstCollection(InCollection)
+		, Collection(nullptr)
+		, BoneIndexAttribute(InCollection, BoneIndexAttributeName, FGeometryCollection::VerticesGroup, FTransformCollection::TransformGroup)
+		, BoneWeightAttribute(InCollection, BoneWeightAttributeName, FGeometryCollection::VerticesGroup, FTransformCollection::TransformGroup)
+		, ParentAttribute(InCollection, FTransformCollection::ParentAttribute, FTransformCollection::TransformGroup)
 	{
-		DefineSchema(Self);
 	}
 
-	//
-	//  Initialization
-	//
 
-		//
+	//
 	//  Initialization
 	//
 
-	void FVertexBoneWeightsFacade::DefineSchema(FManagedArrayCollection* Collection)
+	void FVertexBoneWeightsFacade::DefineSchema()
 	{
-		FManagedArrayCollection::FConstructionParameters TransformDependency(FTransformCollection::TransformGroup);
-
-		if (!Collection->HasGroup(FGeometryCollection::VerticesGroup))
-		{
-			return;
-		}
-
-		Collection->AddAttribute<TArray<int32>>(FVertexBoneWeightsFacade::IndexAttribute, FGeometryCollection::VerticesGroup, TransformDependency);
-		Collection->AddAttribute<TArray<float>>(FVertexBoneWeightsFacade::WeightAttribute, FGeometryCollection::VerticesGroup, TransformDependency);
-
-		ensure(Collection->FindAttributeTyped<TArray<int32>>(FVertexBoneWeightsFacade::IndexAttribute, FGeometryCollection::VerticesGroup) != nullptr);
-		ensure(Collection->FindAttributeTyped<TArray<float>>(FVertexBoneWeightsFacade::WeightAttribute, FGeometryCollection::VerticesGroup) != nullptr);
+		check(!IsConst());
+		BoneIndexAttribute.Add();
+		BoneWeightAttribute.Add();
+		ParentAttribute.Add();
 	}
 
-	bool FVertexBoneWeightsFacade::HasFacade(const FManagedArrayCollection* Collection)
+	bool FVertexBoneWeightsFacade::IsValid() const
 	{
-		return Collection->HasGroup(FGeometryCollection::VerticesGroup) &&
-			Collection->FindAttributeTyped<TArray<int32>>(FVertexBoneWeightsFacade::IndexAttribute, FGeometryCollection::VerticesGroup) &&
-			Collection->FindAttributeTyped<TArray<float>>(FVertexBoneWeightsFacade::WeightAttribute, FGeometryCollection::VerticesGroup);
+		return BoneIndexAttribute.IsValid() && BoneWeightAttribute.IsValid() && ParentAttribute.IsValid();
 	}
 
-
 	//
 	//  Add Weights from Selection 
 	//  ... ... (Impressionist physbam)
 	//
 
-	void FVertexBoneWeightsFacade::AddBoneWeightsFromKinematicBindings(FManagedArrayCollection* InCollection) {
-		if (InCollection && InCollection->HasGroup(FGeometryCollection::VerticesGroup)) {
-			DefineSchema(InCollection); TArray<float> Weights; TArray<int32> Indices;
-			int32 NumBones = InCollection->NumElements(FTransformCollection::TransformGroup), NumVertices = InCollection->NumElements(FGeometryCollection::VerticesGroup);
-			TManagedArray< TArray<int32> >& IndicesArray = InCollection->ModifyAttribute<TArray<int32>>(FVertexBoneWeightsFacade::IndexAttribute, FGeometryCollection::VerticesGroup);
-			TManagedArray< TArray<float> >& WeightsArray = InCollection->ModifyAttribute<TArray<float>>(FVertexBoneWeightsFacade::WeightAttribute, FGeometryCollection::VerticesGroup);
-			for (int32 Kdx = Chaos::Facades::FKinematicBindingFacade::NumKinematicBindings(InCollection) - 1; 0 <= Kdx; Kdx--) {
+	void FVertexBoneWeightsFacade::AddBoneWeightsFromKinematicBindings() {
+		check(!IsConst()); DefineSchema(); if (IsValid()) {
+			TArray<float> Weights; TArray<int32> Indices;
+			int32 NumBones = ParentAttribute.Num(), NumVertices = BoneIndexAttribute.Num();
+			TManagedArray< TArray<int32> >& IndicesArray = BoneIndexAttribute.Modify();
+			TManagedArray< TArray<float> >& WeightsArray = BoneWeightAttribute.Modify();
+			GeometryCollection::Facades::FKinematicBindingFacade BindingFacade(ConstCollection);
+			for (int32 Kdx = BindingFacade.NumKinematicBindings() - 1; 0 <= Kdx; Kdx--) {
 				int32 Bone; TArray<int32> OutBoneVerts; TArray<float> OutBoneWeights;
-				Chaos::Facades::FKinematicBindingFacade::GetBoneBindings(InCollection, Chaos::Facades::FKinematicBindingFacade::GetKinematicBindingKey(InCollection, Kdx), Bone, OutBoneVerts, OutBoneWeights);
-				if (0 <= Bone && Bone < NumBones) for (int32 Vdx = 0; Vdx < OutBoneVerts.Num(); Vdx++) { int32 Vert = OutBoneVerts[Vdx]; float Weight = OutBoneWeights[Vdx];
+				BindingFacade.GetBoneBindings(BindingFacade.GetKinematicBindingKey(Kdx), Bone, OutBoneVerts, OutBoneWeights);
+				if (0 <= Bone && Bone < NumBones) for (int32 Vdx = 0; Vdx < OutBoneVerts.Num(); Vdx++) {
+					int32 Vert = OutBoneVerts[Vdx]; float Weight = OutBoneWeights[Vdx];
 					if (0 <= Vert && Vert < NumVertices && !IndicesArray[Vert].Contains(Bone)) { IndicesArray[Vert].Add(Bone); WeightsArray[Vert].Add(Weight); }}}}}
 
-	//
-	//  GetAttributes
-	//
-
-	const TManagedArray< TArray<int32> >* FVertexBoneWeightsFacade::GetBoneIndices(const FManagedArrayCollection* Collection)
-	{
-		return Collection->FindAttribute<TArray<int32>>(FVertexBoneWeightsFacade::IndexAttribute, FGeometryCollection::VerticesGroup);
-	}
-
-	const TManagedArray< TArray<float> >* FVertexBoneWeightsFacade::GetBoneWeights(const FManagedArrayCollection* Collection)
-	{
-		return Collection->FindAttribute<TArray<float>>(FVertexBoneWeightsFacade::WeightAttribute, FGeometryCollection::VerticesGroup);
-	}
-
-
 };
 
 
diff --git a/Engine/Source/Runtime/Experimental/Chaos/Private/GeometryCollection/GeometryCollection.cpp b/Engine/Source/Runtime/Experimental/Chaos/Private/GeometryCollection/GeometryCollection.cpp
--- a/Engine/Source/Runtime/Experimental/Chaos/Private/GeometryCollection/GeometryCollection.cpp
+++ b/Engine/Source/Runtime/Experimental/Chaos/Private/GeometryCollection/GeometryCollection.cpp
@@ -12,6 +12,7 @@
 #include "GeometryCollection/GeometryCollectionClusteringUtility.h"
 #include "GeometryCollection/GeometryCollectionConvexUtility.h"
 #include "GeometryCollection/Facades/CollectionHierarchyFacade.h"
+#include "GeometryCollection/Facades/CollectionRenderingFacade.h"
 #include "UObject/FortniteNCBranchObjectVersion.h"
 
 #include <iostream>
@@ -508,76 +509,10 @@ void FGeometryCollection::ReindexMaterials(FManagedArrayCollection& InCollection
 	}
 }
 
-TArray<FGeometryCollectionSection> FGeometryCollection::BuildMeshSections(const TArray<FIntVector> &InputIndices, TArray<int32> BaseMeshOriginalIndicesIndex, TArray<FIntVector> &RetIndices) const
-{	
-
-	TArray<FGeometryCollectionSection> TmpSections;
-	TArray<FGeometryCollectionSection> RetSections;		
-
-	// count the number of triangles for each material section, adding a new section if the material ID is higher than the current number of sections
-	for (int FaceElement = 0; FaceElement < InputIndices.Num(); ++FaceElement)
-	{
-		int32 Section = MaterialID[BaseMeshOriginalIndicesIndex[FaceElement]];		
-
-		while (Section + 1 > TmpSections.Num())
-		{
-			// add a new material section
-			int32 Element = TmpSections.AddZeroed();
-
-			TmpSections[Element].MaterialID = Element;
-			TmpSections[Element].FirstIndex = -1;
-			TmpSections[Element].NumTriangles = 0;
-			TmpSections[Element].MinVertexIndex = 0;
-			TmpSections[Element].MaxVertexIndex = 0;
-		}
-
-		TmpSections[Section].NumTriangles++;
-	}
-
-	// fixup the section FirstIndex and MaxVertexIndex
-	for (int SectionElement = 0; SectionElement < TmpSections.Num(); SectionElement++)
-	{
-		if (SectionElement == 0)
-		{
-			TmpSections[SectionElement].FirstIndex = 0;
-		}
-		else
-		{
-			// Each subsequent section has an index that starts after the last one
-			// note the NumTriangles*3 - this is because indices are sent to the renderer in a flat array
-			TmpSections[SectionElement].FirstIndex = TmpSections[SectionElement - 1].FirstIndex + TmpSections[SectionElement - 1].NumTriangles * 3;
-		}
-
-		TmpSections[SectionElement].MaxVertexIndex = NumElements(FGeometryCollection::VerticesGroup) - 1;
-	}
-
-	// remap indices so the materials appear to be grouped
-	RetIndices.AddUninitialized(InputIndices.Num());
-	int Idx = 0;
-	for (int Section = 0; Section < TmpSections.Num(); Section++)
-	{
-		for (int FaceElement = 0; FaceElement < InputIndices.Num(); FaceElement++)
-		{
-			int32 ID = (MaterialID)[BaseMeshOriginalIndicesIndex[FaceElement]];
-
-			if (Section == ID)
-			{
-				RetIndices[Idx++] = InputIndices[FaceElement];				
-			}
-		}
-	}
-
-	// if a material group no longer has any triangles in it then add material section for removal
-	RetSections.Reserve(TmpSections.Num());
-	for (int SectionElement = 0; SectionElement < TmpSections.Num(); SectionElement++)
-	{
-		if (TmpSections[SectionElement].NumTriangles > 0)
-		{
-			RetSections.Push(TmpSections[SectionElement]);
-		}
-	}
+TArray<FGeometryCollectionSection> FGeometryCollection::BuildMeshSections(const TArray<FIntVector>& InputIndices, TArray<int32> BaseMeshOriginalIndicesIndex, TArray<FIntVector>& RetIndices) const
+{
 
-	return MoveTemp(RetSections);
+	return FGeometryCollectionSection::BuildMeshSections(*this, InputIndices, BaseMeshOriginalIndicesIndex, RetIndices);
 }
 
 
@@ -2074,61 +2009,5 @@ void FGeometryCollection::UpdateOldAttributeNames()
 	this->RemoveGroup("Structure");
 }
 
-FGeometryCollectionMeshFacade::FGeometryCollectionMeshFacade(FManagedArrayCollection& InCollection)
-	: Vertex(InCollection, "Vertex", FGeometryCollection::VerticesGroup)
-	, TangentU(InCollection, "TangentU", FGeometryCollection::VerticesGroup)
-	, TangentV(InCollection, "TangentV", FGeometryCollection::VerticesGroup)
-	, Normal(InCollection, "Normal", FGeometryCollection::VerticesGroup)
-	, UVs(InCollection, "UVs", FGeometryCollection::VerticesGroup)
-	, Color(InCollection, "Color", FGeometryCollection::VerticesGroup)
-	, BoneMap(InCollection, "BoneMap", FGeometryCollection::VerticesGroup)
-	, VertexStart(InCollection, "VertexStart", FGeometryCollection::GeometryGroup)
-	, VertexCount(InCollection, "VertexCount", FGeometryCollection::GeometryGroup)
-	, Indices(InCollection, "Indices", FGeometryCollection::FacesGroup)
-	, Visible(InCollection, "Visible", FGeometryCollection::FacesGroup)
-	, MaterialIndex(InCollection, "MaterialIndex", FGeometryCollection::FacesGroup)
-	, MaterialID(InCollection, "MaterialID", FGeometryCollection::FacesGroup)
-	, FaceStart(InCollection, "FaceStart", FGeometryCollection::GeometryGroup)
-	, FaceCount(InCollection, "FaceCount", FGeometryCollection::GeometryGroup)
-{
-}
-
-bool FGeometryCollectionMeshFacade::IsValid() const
-{
-	return Vertex.IsValid()
-		&& TangentU.IsValid()
-		&& TangentV.IsValid()
-		&& Normal.IsValid()
-		&& UVs.IsValid()
-		&& Color.IsValid()
-		&& BoneMap.IsValid()
-		&& VertexStart.IsValid()
-		&& VertexCount.IsValid()
-		&& Indices.IsValid()
-		&& Visible.IsValid()
-		&& MaterialIndex.IsValid()
-		&& MaterialID.IsValid()
-		&& FaceStart.IsValid()
-		&& FaceCount.IsValid()
-	;
-}
 
-void FGeometryCollectionMeshFacade::AddAttributes()
-{
-	Vertex.Add();
-	TangentU.Add();
-	TangentV.Add();
-	Normal.Add();
-	UVs.Add();
-	Color.Add();
-	BoneMap.Add();
-	VertexStart.Add();
-	VertexCount.Add();
-	Indices.Add();
-	Visible.Add();
-	MaterialIndex.Add();
-	MaterialID.Add();
-	FaceStart.Add();
-	FaceCount.Add();
-}
 
diff --git a/Engine/Source/Runtime/Experimental/Chaos/Private/GeometryCollection/GeometryCollectionProximityUtility.cpp b/Engine/Source/Runtime/Experimental/Chaos/Private/GeometryCollection/GeometryCollectionProximityUtility.cpp
--- a/Engine/Source/Runtime/Experimental/Chaos/Private/GeometryCollection/GeometryCollectionProximityUtility.cpp
+++ b/Engine/Source/Runtime/Experimental/Chaos/Private/GeometryCollection/GeometryCollectionProximityUtility.cpp
@@ -537,7 +537,7 @@ void FGeometryCollectionProximityUtility::CopyProximityToConnectionGraph()
 
 	const TManagedArray<TSet<int32>>& Proximity = Collection->GetAttribute<TSet<int32>>("Proximity", FGeometryCollection::GeometryGroup);
 	GeometryCollection::Facades::FCollectionConnectionGraphFacade ConnectionsFacade(*Collection);
-	ConnectionsFacade.AddAttributes();
+	ConnectionsFacade.DefineSchema();
 	TManagedArray<TSet<int32>>& Connections = ConnectionsFacade.ConnectionsAttribute.Modify();
 	Connections.Fill(TSet<int32>());
 
diff --git a/Engine/Source/Runtime/Experimental/Chaos/Private/GeometryCollection/GeometryCollectionSection.cpp b/Engine/Source/Runtime/Experimental/Chaos/Private/GeometryCollection/GeometryCollectionSection.cpp
--- a/Engine/Source/Runtime/Experimental/Chaos/Private/GeometryCollection/GeometryCollectionSection.cpp
+++ b/Engine/Source/Runtime/Experimental/Chaos/Private/GeometryCollection/GeometryCollectionSection.cpp
@@ -1,3 +1,81 @@
 // Copyright Epic Games, Inc. All Rights Reserved.
 
-#include "GeometryCollection/GeometryCollectionSection.h"
\ No newline at end of file
+#include "GeometryCollection/GeometryCollectionSection.h"
+#include "GeometryCollection/ManagedArrayCollection.h"
+#include "GeometryCollection/GeometryCollection.h"
+
+
+TArray<FGeometryCollectionSection>
+FGeometryCollectionSection::BuildMeshSections(const FManagedArrayCollection& InCollection, const TArray<FIntVector>& InputIndices, TArray<int32> BaseMeshOriginalIndicesIndex, TArray<FIntVector>& RetIndices)
+{
+	TArray<FGeometryCollectionSection> RetSections;
+	TArray<FGeometryCollectionSection> TmpSections;
+	if (InCollection.FindAttributeTyped<int32>("MaterialID", FGeometryCollection::FacesGroup))
+	{
+		const TManagedArray<int32>& MaterialID = InCollection.GetAttribute<int32>("MaterialID", FGeometryCollection::FacesGroup);
+
+		// count the number of triangles for each material section, adding a new section if the material ID is higher than the current number of sections
+		for (int FaceElement = 0; FaceElement < InputIndices.Num(); ++FaceElement)
+		{
+			int32 Section = MaterialID[BaseMeshOriginalIndicesIndex[FaceElement]];
+
+			while (Section + 1 > TmpSections.Num())
+			{
+				// add a new material section
+				int32 Element = TmpSections.AddZeroed();
+
+				TmpSections[Element].MaterialID = Element;
+				TmpSections[Element].FirstIndex = -1;
+				TmpSections[Element].NumTriangles = 0;
+				TmpSections[Element].MinVertexIndex = 0;
+				TmpSections[Element].MaxVertexIndex = 0;
+			}
+
+			TmpSections[Section].NumTriangles++;
+		}
+
+		// fixup the section FirstIndex and MaxVertexIndex
+		for (int SectionElement = 0; SectionElement < TmpSections.Num(); SectionElement++)
+		{
+			if (SectionElement == 0)
+			{
+				TmpSections[SectionElement].FirstIndex = 0;
+			}
+			else
+			{
+				// Each subsequent section has an index that starts after the last one
+				// note the NumTriangles*3 - this is because indices are sent to the renderer in a flat array
+				TmpSections[SectionElement].FirstIndex = TmpSections[SectionElement - 1].FirstIndex + TmpSections[SectionElement - 1].NumTriangles * 3;
+			}
+
+			TmpSections[SectionElement].MaxVertexIndex = InCollection.NumElements(FGeometryCollection::VerticesGroup) - 1;
+		}
+
+		// remap indices so the materials appear to be grouped
+		RetIndices.AddUninitialized(InputIndices.Num());
+		int Idx = 0;
+		for (int Section = 0; Section < TmpSections.Num(); Section++)
+		{
+			for (int FaceElement = 0; FaceElement < InputIndices.Num(); FaceElement++)
+			{
+				int32 ID = (MaterialID)[BaseMeshOriginalIndicesIndex[FaceElement]];
+
+				if (Section == ID)
+				{
+					RetIndices[Idx++] = InputIndices[FaceElement];
+				}
+			}
+		}
+
+		// if a material group no longer has any triangles in it then add material section for removal
+		RetSections.Reserve(TmpSections.Num());
+		for (int SectionElement = 0; SectionElement < TmpSections.Num(); SectionElement++)
+		{
+			if (TmpSections[SectionElement].NumTriangles > 0)
+			{
+				RetSections.Push(TmpSections[SectionElement]);
+			}
+		}
+	}
+	return MoveTemp(RetSections);
+}
\ No newline at end of file
diff --git a/Engine/Source/Runtime/Experimental/Chaos/Public/GeometryCollection/Facades/CollectionAnchoringFacade.h b/Engine/Source/Runtime/Experimental/Chaos/Public/GeometryCollection/Facades/CollectionAnchoringFacade.h
--- a/Engine/Source/Runtime/Experimental/Chaos/Public/GeometryCollection/Facades/CollectionAnchoringFacade.h
+++ b/Engine/Source/Runtime/Experimental/Chaos/Public/GeometryCollection/Facades/CollectionAnchoringFacade.h
@@ -15,6 +15,10 @@ namespace Chaos::Facades
 	{
 	public:
 		FCollectionAnchoringFacade(FManagedArrayCollection& InCollection);
+		FCollectionAnchoringFacade(const FManagedArrayCollection& InCollection);
+
+		/** Is the facade defined constant. */
+		bool IsConst() const { return InitialDynamicStateAttribute.IsConst(); }
 
 		bool HasInitialDynamicStateAttribute() const;
 		
diff --git a/Engine/Source/Runtime/Experimental/Chaos/Public/GeometryCollection/Facades/CollectionBoundsFacade.h b/Engine/Source/Runtime/Experimental/Chaos/Public/GeometryCollection/Facades/CollectionBoundsFacade.h
--- a/Engine/Source/Runtime/Experimental/Chaos/Public/GeometryCollection/Facades/CollectionBoundsFacade.h
+++ b/Engine/Source/Runtime/Experimental/Chaos/Public/GeometryCollection/Facades/CollectionBoundsFacade.h
@@ -17,67 +17,41 @@ namespace GeometryCollection::Facades
 	* 
 	* Defines common API for calculating the bounding box on a collection
 	* 
-	* Usage:
-	*    FBoundsFacade::UpdateBoundingBox(this);
-	* 
-	* Then arrays can be accessed later by:
-	*	const TManagedArray< FIntVector >* FaceIndices = FRenderingFacade::GetIndices(this);
-	*	const TManagedArray< FVector3f >* Vertices = FRenderingFacade::GetVertices(this);
-	*
-	* The following attributes are created on the collection:
-	*
-	*	- FindAttribute<FIntVector>("Indices", FGeometryCollection::FacesGroup);
-	*	- FindAttribute<FVector3f>("Vertex", FGeometryCollection::VerticesGroup);
-	* 
 	*/
 	class CHAOS_API FBoundsFacade
 	{
-		FManagedArrayCollection& Self;
-
 	public:
 
-		// Attributes
-		static const FName AAAAttribute;
-
 		/**
 		* FBoundsFacade Constuctor
-		* @param VertixDependencyGroup : GroupName the index attribute is dependent on.
+		* @param FManagedArrayCollection : Collection input
 		*/
 		FBoundsFacade(FManagedArrayCollection& InSelf);
+		FBoundsFacade(const FManagedArrayCollection& InSelf);
 
-		/**
-		*  Create the facade.
-		*/
-		static void DefineSchema(FManagedArrayCollection& InCollection);
-		void DefineSchema() { return DefineSchema(Self); }
+		/** Create the facade attributes. */
+		void DefineSchema();
 
+		/** Is the facade defined constant. */
+		bool IsConst() const { return BoundingBoxAttribute.IsConst(); }
 
-		/**
-		*  Is the Facade defined on the collection?
-		*/
-		static bool IsValid(const FManagedArrayCollection& InCollection);
-		bool IsValid() { return IsValid(Self); }
-
+		/** Is the Facade defined on the collection? */
+		bool IsValid() const;		
 
-		/**
-		*  Does it support rendering surfaces. 
-		*/
-		static void UpdateBoundingBox(FManagedArrayCollection& InCollection, bool bSkipCheck = false);
-		void UpdateBoundingBox(bool bSkipCheck = false) { UpdateBoundingBox(Self); }
+		/** UpdateBoundingBox */
+		void UpdateBoundingBox(bool bSkipCheck = false);
 
+		/** BoundingBox access */
+		const TManagedArray< FBox >& GetBoundingBoxes() const { return BoundingBoxAttribute.Get(); }
 
-		/**
-		* Return the vertex positions from the collection. Null if not initialized.
-		* @param FManagedArrayCollection : Collection
-		*/
-		static const TManagedArray< FBox >* GetBoundingBoxes(const FManagedArrayCollection& InCollection);
-		const TManagedArray< FBox >* GetBoundingBoxes() const { return GetBoundingBoxes(Self); }
+		/** TransformToGeometryIndex Access */
+		const TManagedArray< int32 >& GetTransformToGeometryIndex() const { return TransformToGeometryIndexAttribute.Get(); }
 
 private:
-		TManagedArrayAccessor<FBox> BoundingBox;
-		TManagedArrayAccessor<FVector3f> Vertex;
-		TManagedArrayAccessor<int32> BoneMap;
-		TManagedArrayAccessor<int32> TransformToGeometryIndex;
+		TManagedArrayAccessor<FBox> BoundingBoxAttribute;
+		TManagedArrayAccessor<FVector3f> VertexAttribute;
+		TManagedArrayAccessor<int32> BoneMapAttribute;
+		TManagedArrayAccessor<int32> TransformToGeometryIndexAttribute;
 	};
 
 }
\ No newline at end of file
diff --git a/Engine/Source/Runtime/Experimental/Chaos/Public/GeometryCollection/Facades/CollectionConnectionGraphFacade.h b/Engine/Source/Runtime/Experimental/Chaos/Public/GeometryCollection/Facades/CollectionConnectionGraphFacade.h
--- a/Engine/Source/Runtime/Experimental/Chaos/Public/GeometryCollection/Facades/CollectionConnectionGraphFacade.h
+++ b/Engine/Source/Runtime/Experimental/Chaos/Public/GeometryCollection/Facades/CollectionConnectionGraphFacade.h
@@ -15,15 +15,16 @@ namespace GeometryCollection::Facades
 	{
 	public:
 		FCollectionConnectionGraphFacade(FManagedArrayCollection& InCollection);
+		FCollectionConnectionGraphFacade(const FManagedArrayCollection& InCollection);
 
-		/**
-		 * returns true if all the necessary attributes are present
-		 * if not, then the API can be used to create them
-		 */
+		/** Does the collection support the facade. */
 		bool IsValid() const;
 
-		/** Add the necessary attributes if they are missing and initialize them if necessary */
-		void AddAttributes();
+		/** Is the facade defined constant. */
+		bool IsConst() const { return ConnectionsAttribute.IsConst(); }
+
+		/** Create the facade attributes. */
+		void DefineSchema();
 
 		/** Remove the attributes */
 		void ClearAttributes();
@@ -31,7 +32,7 @@ namespace GeometryCollection::Facades
 		/** Connect two bones */
 		void Connect(int32 BoneA, int32 BoneB);
 
-		/** Connections between bones that have the same parent in the hierarchy */
+		/**  Connections between bones that have the same parent in the hierarchy */
 		TManagedArrayAccessor<TSet<int32>> ConnectionsAttribute;
 
 #if UE_BUILD_DEBUG
diff --git a/Engine/Source/Runtime/Experimental/Chaos/Public/GeometryCollection/Facades/CollectionHierarchyFacade.h b/Engine/Source/Runtime/Experimental/Chaos/Public/GeometryCollection/Facades/CollectionHierarchyFacade.h
--- a/Engine/Source/Runtime/Experimental/Chaos/Public/GeometryCollection/Facades/CollectionHierarchyFacade.h
+++ b/Engine/Source/Runtime/Experimental/Chaos/Public/GeometryCollection/Facades/CollectionHierarchyFacade.h
@@ -20,12 +20,15 @@ namespace Chaos::Facades
 			MustBePersistent
 		};
 
-	public:
 		FCollectionHierarchyFacade(FManagedArrayCollection& InCollection);
+		FCollectionHierarchyFacade(const FManagedArrayCollection& InCollection);
 
 		/** Valid if parent and children arrays are available */
 		bool IsValid() const;
 
+		/** Is this facade const access */
+		bool IsConst() const { return ParentAttribute.IsConst(); }
+
 		/** whether the level attribute is available */
 		bool HasLevelAttribute() const;
 
diff --git a/Engine/Source/Runtime/Experimental/Chaos/Public/GeometryCollection/Facades/CollectionKinematicBindingFacade.h b/Engine/Source/Runtime/Experimental/Chaos/Public/GeometryCollection/Facades/CollectionKinematicBindingFacade.h
--- a/Engine/Source/Runtime/Experimental/Chaos/Public/GeometryCollection/Facades/CollectionKinematicBindingFacade.h
+++ b/Engine/Source/Runtime/Experimental/Chaos/Public/GeometryCollection/Facades/CollectionKinematicBindingFacade.h
@@ -3,28 +3,18 @@
 #pragma once
 
 #include "CoreMinimal.h"
+#include "GeometryCollection/ManagedArrayAccessor.h"
 #include "GeometryCollection/ManagedArrayCollection.h"
 #include "GeometryCollection/Facades/CollectionSelectionFacade.h"
 
-namespace Chaos::Facades
+namespace GeometryCollection::Facades
 {
 	/** Kinematic Facade */
 	class CHAOS_API FKinematicBindingFacade
 	{
-		FManagedArrayCollection& Collection;
-
 	public:
-		typedef GeometryCollection::Facades::FSelectionFacade::FSelectionKey FBindingKey;
-
-		FKinematicBindingFacade(FManagedArrayCollection& InCollection) : Collection(InCollection) {}
 
-		void Init();
-
-		//
-		//  Skeletal Mesh Bone Bindings
-		//
-		static FBindingKey SetBoneBindings(FManagedArrayCollection*, const int32 BoneIndex, const TArray<int32>& Vertices, const TArray<float>& Weights);
-		static void GetBoneBindings(const FManagedArrayCollection*, const FBindingKey& Key, int32& OutBoneIndex, TArray<int32>& OutBoneVerts, TArray<float>& OutBoneWeights);
+		typedef GeometryCollection::Facades::FSelectionFacade::FSelectionKey FBindingKey;
 
 		//
 		// Kinematics
@@ -33,9 +23,33 @@ namespace Chaos::Facades
 		static const FName KinematicBoneBindingIndex;
 		static const FName KinematicBoneBindingToGroup;
 
-		static int32 AddKinematicBinding(FManagedArrayCollection*, const FBindingKey& Key);
-		static int32 NumKinematicBindings(const FManagedArrayCollection* Collection) { return Collection->NumElements(KinematicGroup); }
-		static FBindingKey GetKinematicBindingKey(const FManagedArrayCollection* InCollection, int Index);
+		FKinematicBindingFacade(FManagedArrayCollection& InCollection);
+		FKinematicBindingFacade(const FManagedArrayCollection& InCollection);
+
+		/** Create the facade attributes. */
+		void DefineSchema();
+
+		/** Is the facade defined constant. */
+		bool IsConst() const { return Collection==nullptr; }
+
+		/** Is the Facade defined on the collection? */
+		bool IsValid() const;
+
+		//
+		//  Skeletal Mesh Bone Bindings
+		//
+		FBindingKey SetBoneBindings(const int32 BoneIndex, const TArray<int32>& Vertices, const TArray<float>& Weights);
+		void GetBoneBindings(const FBindingKey& Key, int32& OutBoneIndex, TArray<int32>& OutBoneVerts, TArray<float>& OutBoneWeights) const;
+
+		int32 AddKinematicBinding(const FBindingKey& Key);
+		int32 NumKinematicBindings() const { return KinemaitcBoneBindingAttribute.Num(); }
+		FBindingKey GetKinematicBindingKey(int Index) const;
+
+	private:
+		const FManagedArrayCollection& ConstCollection;
+		FManagedArrayCollection* Collection = nullptr;
 
+		TManagedArrayAccessor<int32> KinemaitcBoneBindingAttribute;
+		TManagedArrayAccessor<FString> KinemaitcBoneBindingToGroupAttribute;
 	};
 }
diff --git a/Engine/Source/Runtime/Experimental/Chaos/Public/GeometryCollection/Facades/CollectionMeshFacade.h b/Engine/Source/Runtime/Experimental/Chaos/Public/GeometryCollection/Facades/CollectionMeshFacade.h
new file mode 100644
--- /dev/null
+++ b/Engine/Source/Runtime/Experimental/Chaos/Public/GeometryCollection/Facades/CollectionMeshFacade.h
@@ -0,0 +1,46 @@
+// Copyright Epic Games, Inc. All Rights Reserved.
+
+#pragma once
+
+#include "GeometryCollection/ManagedArrayAccessor.h"
+#include "GeometryCollection/ManagedArrayCollection.h"
+
+
+namespace GeometryCollection::Facades
+{
+
+	class CHAOS_API FMeshFacade
+	{
+	public:
+		FMeshFacade(FManagedArrayCollection& InCollection);
+
+		/**
+		 * returns true if all the necessary attributes are present
+		 * if not then the API can be used to create
+		 */
+		bool IsValid() const;
+
+		/**
+		 * Add the necessary attributes if they are missing
+		 */
+		void DefineSchema();
+
+		TManagedArrayAccessor<FVector3f> Vertex;
+		TManagedArrayAccessor<FVector3f> TangentU;
+		TManagedArrayAccessor<FVector3f> TangentV;
+		TManagedArrayAccessor<FVector3f> Normal;
+		TManagedArrayAccessor<TArray<FVector2f>> UVs;
+		TManagedArrayAccessor<FLinearColor> Color;
+		TManagedArrayAccessor<int32> BoneMap;
+		TManagedArrayAccessor<int32> VertexStart;
+		TManagedArrayAccessor<int32> VertexCount;
+
+		TManagedArrayAccessor<FIntVector> Indices;
+		TManagedArrayAccessor<bool> Visible;
+		TManagedArrayAccessor<int32> MaterialIndex;
+		TManagedArrayAccessor<int32> MaterialID;
+		TManagedArrayAccessor<int32> FaceStart;
+		TManagedArrayAccessor<int32> FaceCount;
+	};
+
+}
\ No newline at end of file
diff --git a/Engine/Source/Runtime/Experimental/Chaos/Public/GeometryCollection/Facades/CollectionRenderingFacade.h b/Engine/Source/Runtime/Experimental/Chaos/Public/GeometryCollection/Facades/CollectionRenderingFacade.h
--- a/Engine/Source/Runtime/Experimental/Chaos/Public/GeometryCollection/Facades/CollectionRenderingFacade.h
+++ b/Engine/Source/Runtime/Experimental/Chaos/Public/GeometryCollection/Facades/CollectionRenderingFacade.h
@@ -4,6 +4,7 @@
 #include "Containers/UnrealString.h"
 #include "Containers/Array.h"
 #include "Math/MathFwd.h"
+#include "GeometryCollection/ManagedArrayAccessor.h"
 #include "GeometryCollection/ManagedArrayCollection.h"
 #include "GeometryCollection/Facades/CollectionSelectionFacade.h"
 #include "Chaos/Triangle.h"
@@ -16,91 +17,60 @@ namespace GeometryCollection::Facades
 	*
 	* Defines common API for storing rendering data.
 	*
-	* Usage:
-	*    FRenderingFacade::AddTriangle(this, FTriangle(...) );
-	*
-	* Then arrays can be accessed later by:
-	*	const TManagedArray< FIntVector >* FaceIndices = FRenderingFacade::GetIndices(this);
-	*	const TManagedArray< FVector3f >* Vertices = FRenderingFacade::GetVertices(this);
-	*
-	* The following attributes are created on the collection:
-	*
-	*	- FindAttribute<FIntVector>("Indices", FGeometryCollection::FacesGroup);
-	*	- FindAttribute<FVector3f>("Vertex", FGeometryCollection::VerticesGroup);
-	*
 	*/
 	class CHAOS_API FRenderingFacade
 	{
-		FManagedArrayCollection* Self;
-
 	public:
-		// Groups 
-		static const FName VerticesGroup;
-		static const FName FacesGroup;
-
-		// Attributes
-		static const FName VertexAttribute;
-		static const FName IndicesAttribute;
+		typedef FGeometryCollectionSection FTriangleSection;
 
 		/**
 		* FRenderingFacade Constuctor
 		* @param VertixDependencyGroup : GroupName the index attribute is dependent on.
 		*/
-		FRenderingFacade(FManagedArrayCollection* InSelf);
+		FRenderingFacade(FManagedArrayCollection& InSelf);
+		FRenderingFacade(const FManagedArrayCollection& InSelf);
 
-		/**
-		*  Create the facade.
-		*/
-		static void DefineSchema(FManagedArrayCollection* InCollection);
-		void DefineSchema() { return DefineSchema(Self); }
-
-
-		/**
-		*  Is the Facade defined on the collection?
-		*/
-		static bool IsValid(const FManagedArrayCollection* InCollection);
-		bool IsValid() { return IsValid(Self); }
+		/**Create the facade.*/
+		void DefineSchema();
 
+		/** Is the facade defined constant. */
+		bool IsConst() const { return Collection==nullptr; }
 
-		/**
-		*  Does it support rendering surfaces.
-		*/
-		static bool CanRenderSurface(const FManagedArrayCollection* InCollection);
-		bool CanRenderSurface() { return CanRenderSurface(Self); }
+		/**Is the Facade defined on the collection?*/
+		bool IsValid() const;
 
+		/**Does it support rendering surfaces.*/
+		bool CanRenderSurface() const;
 
-		/**
-		* Add a triangle to the rendering view. 
-		* @param FManagedArrayCollection : Collection
-		* @param FSelectionFacade::FSelectionKey : Key for weights in the FSelectionFacade
-		*/
-		static void AddTriangle(FManagedArrayCollection* InCollection, const Chaos::FTriangle& InTriangle);
-		void AddTriangle(const Chaos::FTriangle& InTriangle) { AddTriangle(Self, InTriangle); }
+		/**Add a triangle to the rendering view.*/
+		void AddTriangle(const Chaos::FTriangle& InTriangle);
 
-		/**
-		* Add a surface to the rendering view.
-		* @param FManagedArrayCollection : Collection
-		* @param FSelectionFacade::FSelectionKey : Key for weights in the FSelectionFacade
-		*/
-		static void AddSurface(FManagedArrayCollection* InCollection, TArray<FVector3f>&& InVertices, TArray<FIntVector>&& InIndices);
-		void AddSurface(TArray<FVector3f>&& InVertices, TArray<FIntVector>&& InIndices) { AddSurface(Self, MoveTemp(InVertices), MoveTemp(InIndices) ); }
+		/** Add a surface to the rendering view.*/
+		void AddSurface(TArray<FVector3f>&& InVertices, TArray<FIntVector>&& InIndices);
+		
+		/** GetIndices */
+		const TManagedArray< FIntVector >& GetIndices() const { return IndicesAttribute.Get(); }
 
+		/** GetVertices */
+		const TManagedArray< FVector3f >& GetVertices() const { return VertexAttribute.Get(); }
 
-		/**
-		* Return the render indices from the collection. Null if not initialized.
-		* @param FManagedArrayCollection : Collection
-		*/
-		static const TManagedArray< FIntVector >* GetIndices(const FManagedArrayCollection* InCollection);
-		const TManagedArray< FIntVector >* GetIndices() const { return GetIndices(Self); }
+		/** GetMaterialID */
+		const TManagedArray< int32 >& GetMaterialID() const { return MaterialIDAttribute.Get(); }
 
+		/** GetTriangleSections */
+		const TManagedArray< FTriangleSection >& GetTriangleSections() const { return TriangleSectionAttribute.Get(); }
 
-		/**
-		* Return the vertex positions from the collection. Null if not initialized.
-		* @param FManagedArrayCollection : Collection
-		*/
-		static const TManagedArray< FVector3f >* GetVertices(const FManagedArrayCollection* InCollection);
-		const TManagedArray< FVector3f >* GetVertices() const { return GetVertices(Self); }
+		/** BuildMeshSections */
+		TArray<FTriangleSection> BuildMeshSections(const TArray<FIntVector>& Indices, TArray<int32> BaseMeshOriginalIndicesIndex, TArray<FIntVector>& RetIndices) const;
+		
+	private : 
+		const FManagedArrayCollection& ConstCollection;
+		FManagedArrayCollection* Collection = nullptr;
 
+		TManagedArrayAccessor<FVector3f> VertexAttribute;
+		TManagedArrayAccessor<FIntVector> IndicesAttribute;
+		TManagedArrayAccessor<int32> MaterialIDAttribute;
+		TManagedArrayAccessor<FTriangleSection> TriangleSectionAttribute;
 
 	};
 
diff --git a/Engine/Source/Runtime/Experimental/Chaos/Public/GeometryCollection/Facades/CollectionSelectionFacade.h b/Engine/Source/Runtime/Experimental/Chaos/Public/GeometryCollection/Facades/CollectionSelectionFacade.h
--- a/Engine/Source/Runtime/Experimental/Chaos/Public/GeometryCollection/Facades/CollectionSelectionFacade.h
+++ b/Engine/Source/Runtime/Experimental/Chaos/Public/GeometryCollection/Facades/CollectionSelectionFacade.h
@@ -39,12 +39,10 @@ namespace GeometryCollection::Facades
 	*/
 	class CHAOS_API FSelectionFacade
  	{
-		FManagedArrayCollection* Self;
-
-		static void InitUnboundedGroup(FManagedArrayCollection* Self, FName GroupName, FName DependencyGroup);
-		static void InitWeightedUnboundedGroup(FManagedArrayCollection* Self, FName GroupName, FName DependencyGroup);
-		static void InitBoundedGroup(FManagedArrayCollection* Self, FName GroupName, FName DependencyGroup, FName BoneDependencyGroup);
-		static void InitWeightedBoundedGroup(FManagedArrayCollection* Self, FName GroupName, FName DependencyGroup, FName BoneDependencyGroup);
+		void InitUnboundedGroup(FName GroupName, FName DependencyGroup);
+		void InitWeightedUnboundedGroup(FName GroupName, FName DependencyGroup);
+		void InitBoundedGroup(FName GroupName, FName DependencyGroup, FName BoneDependencyGroup);
+		void InitWeightedBoundedGroup(FName GroupName, FName DependencyGroup, FName BoneDependencyGroup);
 
 	public:
 
@@ -68,43 +66,39 @@ namespace GeometryCollection::Facades
 
 		/**
 		* FSelectionFacade Constuctor
-		* @param VertixDependencyGroup : GroupName the index attribute is dependent on. 
 		*/
-		FSelectionFacade(FManagedArrayCollection* InSelf);
+		FSelectionFacade(FManagedArrayCollection& InSelf);
+		FSelectionFacade(const FManagedArrayCollection& InSelf);
+
+		/** Is the facade defined constant. */
+		bool IsConst() const { return Collection==nullptr; }
 
 		/**
 		* Add the indices to the FVertexSetInterface::UnboundGroup 
-		* @param Indices : Array of indices to add. 
 		*/
-		static FSelectionKey AddSelection(FManagedArrayCollection* Self, const TArray<int32>& Indices, FName DependencyGroup);
+		FSelectionKey AddSelection(const TArray<int32>& Indices, FName DependencyGroup);
 
 		/**
 		* Add the indices to the FVertexSetInterface::WeightedUnboundGroup
-		* @param Indices : Array of indices to add.
-		* @param Indices : Array of indices to add.
 		*/
-		static FSelectionKey AddSelection(FManagedArrayCollection* Self, const TArray<int32>& Indices, const TArray<float>& Weights, FName DependencyGroup);
+		FSelectionKey AddSelection(const TArray<int32>& Indices, const TArray<float>& Weights, FName DependencyGroup);
 
 		/**
 		* Add the indices to the FVertexSetInterface::BoundGroup
-		* @param Indices : Array of indices to add.
-		* @param Indices : Array of indices to add.
 		*/
-		static FSelectionKey AddSelection(FManagedArrayCollection* Self, const int32 BoneIndex, const TArray<int32>& Indices, FName DependencyGroup, FName BoneDependencyGroup = FName(""));
+		FSelectionKey AddSelection(const int32 BoneIndex, const TArray<int32>& Indices, FName DependencyGroup, FName BoneDependencyGroup = FName(""));
 
 		/**
 		* Add the indices to the FVertexSetInterface::BoundGroup
-		* @param Indices : Array of indices to add.
-		* @param Indices : Array of indices to add.
 		*/
-		static FSelectionKey AddSelection(FManagedArrayCollection* Self, const int32 BoneIndex, const TArray<int32>& Indices, const TArray<float>& Weights, FName DependencyGroup, FName BoneDependencyGroup = FName(""));
+		FSelectionKey AddSelection(const int32 BoneIndex, const TArray<int32>& Indices, const TArray<float>& Weights, FName DependencyGroup, FName BoneDependencyGroup = FName(""));
 
 		/**
 		* Return the vertex list from the given key 
 		* @param Key : <GroupName and Index>
 		* @param Indices : Return indices, empty if not found. 
 		*/
-		static void GetSelection(const FManagedArrayCollection* Self, const FSelectionKey& Key, TArray<int32>& Indices);
+		void GetSelection(const FSelectionKey& Key, TArray<int32>& Indices) const;
 
 		/**
 		* Return the vertex list, bone index, and weights from the given key
@@ -113,7 +107,7 @@ namespace GeometryCollection::Facades
 		* @param Indices : Return indices, empty if not found.
 		* @param Weights : Return vertex weights, empty if not found.
 		*/
-		static void GetSelection(const FManagedArrayCollection* Self, const FSelectionKey& Key, TArray<int32>& Indices, TArray<float>& Weights);
+		void GetSelection(const FSelectionKey& Key, TArray<int32>& Indices, TArray<float>& Weights) const;
 
 		/**
 		* Return the vertex list and bone index from the given key
@@ -121,7 +115,7 @@ namespace GeometryCollection::Facades
 		* @param BoneIndex : Return BoneIndex, INDEX_NONE if not found.
 		* @param Indices : Return indices, empty if not found.
 		*/
-		static void GetSelection(const FManagedArrayCollection* Self, const FSelectionKey& Key, int32& BoneIndex, TArray<int32>& Indices);
+		void GetSelection(const FSelectionKey& Key, int32& BoneIndex, TArray<int32>& Indices) const;
 
 		/**
 		* Return the vertex list, bone index, and weights from the given key
@@ -130,7 +124,14 @@ namespace GeometryCollection::Facades
 		* @param Indices : Return indices, empty if not found.
 		* @param Weights : Return vertex weights, empty if not found.
 		*/
-		static void GetSelection(const FManagedArrayCollection* Self, const FSelectionKey& Key, int32& BoneIndex, TArray<int32>& Indices, TArray<float>& Weights);
+		void GetSelection(const FSelectionKey& Key, int32& BoneIndex, TArray<int32>& Indices, TArray<float>& Weights) const;
+
+	private:
+
+		// const collection will be a null pointer, 
+		// while non-const will be valid.
+		const FManagedArrayCollection& ConstCollection;
+		FManagedArrayCollection* Collection = nullptr;
 	};
 
 }
\ No newline at end of file
diff --git a/Engine/Source/Runtime/Experimental/Chaos/Public/GeometryCollection/Facades/CollectionTransformFacade.h b/Engine/Source/Runtime/Experimental/Chaos/Public/GeometryCollection/Facades/CollectionTransformFacade.h
new file mode 100644
--- /dev/null
+++ b/Engine/Source/Runtime/Experimental/Chaos/Public/GeometryCollection/Facades/CollectionTransformFacade.h
@@ -0,0 +1,64 @@
+// Copyright Epic Games, Inc. All Rights Reserved.
+
+#pragma once
+
+#include "CoreMinimal.h"
+#include "GeometryCollection/ManagedArrayAccessor.h"
+#include "GeometryCollection/ManagedArrayCollection.h"
+
+namespace GeometryCollection::Facades
+{
+	/**
+	 * Provides an API to read and manipulate hierarchy in a managed array collection
+	 */
+	class CHAOS_API FTransformFacade
+	{
+		FManagedArrayCollection& Self;
+
+	public:
+
+		// Group
+		static const FName TransformGroup;
+
+		// Attributes
+		static const FName ParentAttribute;
+		static const FName ChildrenAttribute;
+		static const FName TransformAttribute;
+
+
+		FTransformFacade(FManagedArrayCollection& InCollection);
+		FTransformFacade(const FManagedArrayCollection& InCollection);
+
+		/** Create the facade attributes. */
+		void DefineSchema() {}
+
+		/** Valid if parent and children arrays are available */
+		bool IsValid() const;
+
+		/** Is the facade defined constant. */
+		bool IsConst() const { return Parent.IsConst(); }
+
+		/** Get the root index */
+		TSet<int32> GetRootIndices() const;
+
+		/**
+		* Return the parent indices from the collection. Null if not initialized.
+		*/
+		const TManagedArray< int32 >* GetParents() const { return Parent.Find(); }
+
+		/**
+		* Return the child indicesfrom the collection. Null if not initialized.
+		*/
+		const TManagedArray< TSet<int32> >* FindChildren() const { return Children.Find(); }
+
+		/**
+		* Return the child indicesfrom the collection. Null if not initialized.
+		*/
+		const TManagedArray< FTransform >* FindTransforms() const { return Transform.Find(); }
+
+	private:
+		TManagedArrayAccessor<int32>		Parent;
+		TManagedArrayAccessor<TSet<int32>>	Children;
+		TManagedArrayAccessor<FTransform>	Transform;
+	};
+}
\ No newline at end of file
diff --git a/Engine/Source/Runtime/Experimental/Chaos/Public/GeometryCollection/Facades/CollectionTransformSourceFacade.h b/Engine/Source/Runtime/Experimental/Chaos/Public/GeometryCollection/Facades/CollectionTransformSourceFacade.h
--- a/Engine/Source/Runtime/Experimental/Chaos/Public/GeometryCollection/Facades/CollectionTransformSourceFacade.h
+++ b/Engine/Source/Runtime/Experimental/Chaos/Public/GeometryCollection/Facades/CollectionTransformSourceFacade.h
@@ -4,6 +4,7 @@
 #include "Containers/UnrealString.h"
 #include "Containers/Array.h"
 #include "GeometryCollection/ManagedArrayCollection.h"
+#include "GeometryCollection/ManagedArrayAccessor.h"
 
 namespace GeometryCollection::Facades
 {
@@ -30,33 +31,31 @@ namespace GeometryCollection::Facades
 	*/
 	class CHAOS_API FTransformSource
  	{
-		FManagedArrayCollection* Self;
-
 	public:
 
 		// groups
-		static const FName TransformSourceGroup;
+		static const FName TransformSourceGroupName;
 
 		// Attributes
-		static const FName SourceNameAttribute;
-		static const FName SourceGuidAttribute;
-		static const FName SourceRootsAttribute;
+		static const FName SourceNameAttributeName;
+		static const FName SourceGuidAttributeName;
+		static const FName SourceRootsAttributeName;
 
 		/**
 		* FSelectionFacade Constuctor
 		* @param VertixDependencyGroup : GroupName the index attribute is dependent on. 
 		*/
-		FTransformSource(FManagedArrayCollection* InSelf);
+		FTransformSource(FManagedArrayCollection& InSelf);
+		FTransformSource(const FManagedArrayCollection& InSelf);
 
-		/**
-		*  Create the facade.
-		*/
-		static void DefineSchema(FManagedArrayCollection* Self);
+		/** Create the facade. */
+		void DefineSchema();
 
-		/**
-		*  Is the Facade defined on the collection?
-		*/
-		static bool HasFacade(const FManagedArrayCollection* Collection);
+		/** Is the facade defined constant. */
+		bool IsConst() const { return SourceNameAttribute.IsConst(); }
+
+		/** Is the Facade defined on the collection? */
+		bool IsValid() const;
 
 		/**
 		* Add a transform root mapping.
@@ -64,16 +63,20 @@ namespace GeometryCollection::Facades
 		* @param Guid : Guid of the owner of the transform set.
 		* @param Roots : Root indices of the transform set.
 		*/
-		static void AddTransformSource(FManagedArrayCollection* Collection, const FString& Name, const FString& Guid, const TSet<int32>& Roots);
+		void AddTransformSource(const FString& Name, const FString& Guid, const TSet<int32>& Roots);
 
 		/**
 		* Query for root indices.  
 		* @param Name : Name of the owner of the transform set.  
 		* @param Guid : Guid of the owner of the transform set.
 		*/
-		static TSet<int32> GetTransformSource(const FManagedArrayCollection* Collection, const FString& Name, const FString& Guid);
-
+		TSet<int32> GetTransformSource(const FString& Name, const FString& Guid) const;
 
+	private:
+		TManagedArrayAccessor<FString> SourceNameAttribute;
+		TManagedArrayAccessor<FString> SourceGuidAttribute;
+		TManagedArrayAccessor< TSet<int32 > > SourceRootsAttribute;
+		
 	};
 
 }
\ No newline at end of file
diff --git a/Engine/Source/Runtime/Experimental/Chaos/Public/GeometryCollection/Facades/CollectionVertexBoneWeightsFacade.h b/Engine/Source/Runtime/Experimental/Chaos/Public/GeometryCollection/Facades/CollectionVertexBoneWeightsFacade.h
--- a/Engine/Source/Runtime/Experimental/Chaos/Public/GeometryCollection/Facades/CollectionVertexBoneWeightsFacade.h
+++ b/Engine/Source/Runtime/Experimental/Chaos/Public/GeometryCollection/Facades/CollectionVertexBoneWeightsFacade.h
@@ -3,8 +3,9 @@
 #pragma once
 #include "Containers/UnrealString.h"
 #include "Containers/Array.h"
-#include "GeometryCollection/ManagedArrayCollection.h"
 #include "GeometryCollection/Facades/CollectionSelectionFacade.h"
+#include "GeometryCollection/ManagedArrayCollection.h"
+#include "GeometryCollection/ManagedArrayAccessor.h"
 
 namespace GeometryCollection::Facades
 {
@@ -30,56 +31,46 @@ namespace GeometryCollection::Facades
 	*/
 	class CHAOS_API FVertexBoneWeightsFacade
 	{
-		FManagedArrayCollection* Self;
-
 	public:
 
 		// Attributes
-		static const FName IndexAttribute;
-		static const FName WeightAttribute;
+		static const FName BoneIndexAttributeName;
+		static const FName BoneWeightAttributeName;
 
 		/**
 		* FVertexBoneWeightsFacade Constuctor
-		* @param VertixDependencyGroup : GroupName the index attribute is dependent on.
 		*/
-		FVertexBoneWeightsFacade(FManagedArrayCollection* InSelf);
+		FVertexBoneWeightsFacade(FManagedArrayCollection& InSelf);
+		FVertexBoneWeightsFacade(const FManagedArrayCollection& InSelf);
 
-		/**
-		*  Create the facade.
-		*/
-		static void DefineSchema(FManagedArrayCollection* Collection);
-		void DefineSchema() { return DefineSchema(Self); }
+		/** Define the facade */
+		void DefineSchema();
 
+		/** Is the Facade const */
+		bool IsConst() const { return Collection == nullptr; }
 
-		/**
-		*  Is the Facade defined on the collection?
-		*/
-		static bool HasFacade(const FManagedArrayCollection* Collection);
-		bool HasFacade() { return HasFacade(Self); }
+		/** Is the Facade defined on the collection? */
+		bool IsValid() const;
 
-		/**
-		* Using a FSelectionFacade::FSelectionKey, create indexes from the vertices to the driving bones. 
-		* @param FManagedArrayCollection : Collection
-		* @param FSelectionFacade::FSelectionKey : Key for weights in the FSelectionFacade
-		*/
-		static void AddBoneWeightsFromKinematicBindings(FManagedArrayCollection* Collection);
-		void AddBoneWeightsFromKinematicBindings() { AddBoneWeightsFromKinematicBindings(Self); }
+		/** Add bone weight based on the kinematic bindings. */
+		void AddBoneWeightsFromKinematicBindings();
 
-		/**
-		* Return the vertex bone indices from the collection. Null if not initialized. 
-		* @param FManagedArrayCollection : Collection
-		*/
-		static const TManagedArray< TArray<int32> >* GetBoneIndices(const FManagedArrayCollection* Collection);
-		const TManagedArray< TArray<int32> >* GetBoneIndices() const { return GetBoneIndices(Self); }
+		/** Return the vertex bone indices from the collection. Null if not initialized.  */
+		const TManagedArray< TArray<int32> >* FindBoneIndices()  const { return BoneIndexAttribute.Find(); }
+		const TManagedArray< TArray<int32> >& GetBoneIndices() const { return BoneIndexAttribute.Get(); }
 
 
-		/**
-		* Return the vertex bone weights from the collection. Null if not initialized.
-		* @param FManagedArrayCollection : Collection
-		*/
-		static const TManagedArray< TArray<float> >* GetBoneWeights(const FManagedArrayCollection* Collection);
-		const TManagedArray< TArray<float> >* GetBoneWeights() const { return GetBoneWeights(Self); }
+		/** Return the vertex bone weights from the collection. Null if not initialized. */
+		const TManagedArray< TArray<float> >* FindBoneWeights()  const { return BoneWeightAttribute.Find(); }
+		const TManagedArray< TArray<float> >& GetBoneWeights() const { return BoneWeightAttribute.Get(); }
+
+	private:
+		const FManagedArrayCollection& ConstCollection;
+		FManagedArrayCollection* Collection = nullptr;
 
+		TManagedArrayAccessor<TArray<int32>> BoneIndexAttribute;
+		TManagedArrayAccessor<TArray<float>> BoneWeightAttribute;
+		TManagedArrayAccessor<int32> ParentAttribute;
 
 	};
 
diff --git a/Engine/Source/Runtime/Experimental/Chaos/Public/GeometryCollection/GeometryCollection.h b/Engine/Source/Runtime/Experimental/Chaos/Public/GeometryCollection/GeometryCollection.h
--- a/Engine/Source/Runtime/Experimental/Chaos/Public/GeometryCollection/GeometryCollection.h
+++ b/Engine/Source/Runtime/Experimental/Chaos/Public/GeometryCollection/GeometryCollection.h
@@ -356,36 +356,3 @@ FORCEINLINE Chaos::FChaosArchive& operator<<(Chaos::FChaosArchive& Ar, FGeometry
 	return Ar;
 }
 
-class CHAOS_API FGeometryCollectionMeshFacade
-{
-public:
-	FGeometryCollectionMeshFacade(FManagedArrayCollection& InCollection);
-
-	/** 
-	 * returns true if all the necessary attributes are present
-	 * if not then the API can be used to create  
-	 */
-	bool IsValid() const;	
-
-	/**
-	 * Add the necessary attributes if they are missing
-	 */
-	void AddAttributes();
-
-	TManagedArrayAccessor<FVector3f> Vertex;
-	TManagedArrayAccessor<FVector3f> TangentU;
-	TManagedArrayAccessor<FVector3f> TangentV;
-	TManagedArrayAccessor<FVector3f> Normal;
-	TManagedArrayAccessor<TArray<FVector2f>> UVs;
-	TManagedArrayAccessor<FLinearColor> Color;
-	TManagedArrayAccessor<int32> BoneMap;
-	TManagedArrayAccessor<int32> VertexStart;
-	TManagedArrayAccessor<int32> VertexCount;
-
-	TManagedArrayAccessor<FIntVector> Indices;
-	TManagedArrayAccessor<bool> Visible;
-	TManagedArrayAccessor<int32> MaterialIndex;
-	TManagedArrayAccessor<int32> MaterialID;
-	TManagedArrayAccessor<int32> FaceStart;
-	TManagedArrayAccessor<int32> FaceCount;
-};
\ No newline at end of file
diff --git a/Engine/Source/Runtime/Experimental/Chaos/Public/GeometryCollection/GeometryCollectionSection.h b/Engine/Source/Runtime/Experimental/Chaos/Public/GeometryCollection/GeometryCollectionSection.h
--- a/Engine/Source/Runtime/Experimental/Chaos/Public/GeometryCollection/GeometryCollectionSection.h
+++ b/Engine/Source/Runtime/Experimental/Chaos/Public/GeometryCollection/GeometryCollectionSection.h
@@ -4,11 +4,12 @@
 #include "Containers/Set.h"
 
 class HHitProxy;
+struct FManagedArrayCollection;
 
 /**
 * A set of triangles which are rendered with the same material.
 */
-struct FGeometryCollectionSection
+struct CHAOS_API FGeometryCollectionSection
 {
 	/** Constructor. */
 	FGeometryCollectionSection()
@@ -32,6 +33,9 @@ struct FGeometryCollectionSection
 		return true;
 	}
 
+	static TArray<FGeometryCollectionSection>
+	BuildMeshSections(const FManagedArrayCollection& InCollection, const TArray<FIntVector>& InputIndices, TArray<int32> BaseMeshOriginalIndicesIndex, TArray<FIntVector>& RetIndices);
+
 	/** The index of the material with which to render this section. */
 	int32 MaterialID;
 
diff --git a/Engine/Source/Runtime/Experimental/Chaos/Public/GeometryCollection/ManagedArrayAccessor.h b/Engine/Source/Runtime/Experimental/Chaos/Public/GeometryCollection/ManagedArrayAccessor.h
--- a/Engine/Source/Runtime/Experimental/Chaos/Public/GeometryCollection/ManagedArrayAccessor.h
+++ b/Engine/Source/Runtime/Experimental/Chaos/Public/GeometryCollection/ManagedArrayAccessor.h
@@ -19,52 +19,95 @@ template <typename T>
 struct TManagedArrayAccessor
 {
 public:
-	TManagedArrayAccessor(FManagedArrayCollection& InCollection, const FName& AttributeName, const FName& AttributeGroup)
-		: Collection(InCollection)
-		, Name(AttributeName)
-		, Group(AttributeGroup)
-	{
-		AttributeArray = InCollection.FindAttributeTyped<T>(AttributeName, AttributeGroup);
+	TManagedArrayAccessor(FManagedArrayCollection& InCollection, const FName& InAttributeName, const FName& InAttributeGroup,  const FName& InGroupDependency = FName(""))
+		: ConstCollection(InCollection)
+		, Collection(&InCollection)
+		, Name(InAttributeName)
+		, Group(InAttributeGroup)
+		, GroupDependency(InGroupDependency)
+		, AttributeArray(InCollection.FindAttributeTyped<T>(InAttributeName, InAttributeGroup))
+		, ConstAttributeArray(AttributeArray)
+	{}
 
-	}
-	bool IsValid() const { return AttributeArray != nullptr; }
+	TManagedArrayAccessor(const FManagedArrayCollection& InCollection, const FName& InAttributeName, const FName& InAttributeGroup, const FName& InGroupDependency = FName(""))
+		: ConstCollection(InCollection)
+		, Collection(nullptr)
+		, Name(InAttributeName)
+		, Group(InAttributeGroup)
+		, GroupDependency(InGroupDependency)
+		, AttributeArray(nullptr)
+		, ConstAttributeArray(InCollection.FindAttributeTyped<T>(InAttributeName, InAttributeGroup))
+	{}
+
+	bool IsConst() const { return Collection == nullptr; }
 
-	bool IsPersistent() const { return Collection.IsAttributePersistent(Name, Group); }
+	bool IsValid() const { return ConstAttributeArray != nullptr; }
+
+	bool IsPersistent() const { return ConstCollection.IsAttributePersistent(Name, Group); }
+
+	int32 AddElements(int32 NumElements) const 
+	{ 
+		check(!IsConst());
+		return Collection->AddElements(NumElements, Group);
+	}
 
 	/** get the attribute for read only */
 	const TManagedArray<T>& Get() const
 	{
 		check(IsValid());
-		return *AttributeArray;
+		return *ConstAttributeArray;
+	}
+
+	/** find the attribute for read only */
+	const TManagedArray<T>* Find() const
+	{
+		return ConstAttributeArray;
 	}
 
 	/** get the attribute for modification */
-	TManagedArray<T>& Modify() const
+	TManagedArray<T>& Modify() 
 	{
-		check(IsValid());
+		check(AttributeArray!=nullptr && !IsConst());
 		AttributeArray->MarkDirty();
 		return *AttributeArray;
 	}
 
 	/** add the attribute if it does not exists yet */
-	TManagedArray<T>& Add(ManageArrayAccessor::EPersistencePolicy PersistencePolicy = ManageArrayAccessor::EPersistencePolicy::MakePersistent)
+	TManagedArray<T>& Add(ManageArrayAccessor::EPersistencePolicy PersistencePolicy = ManageArrayAccessor::EPersistencePolicy::MakePersistent,
+		FName InGroupDependency = FName(NAME_None))
 	{
+		check(!IsConst());
 		if (PersistencePolicy == ManageArrayAccessor::EPersistencePolicy::MakePersistent && !IsPersistent())
 		{
 			Remove();
 		}
-		const bool bSaved = (PersistencePolicy == ManageArrayAccessor::EPersistencePolicy::MakePersistent);
-		FManagedArrayCollection::FConstructionParameters Params(FName(), bSaved);
-		AttributeArray = &Collection.AddAttribute<T>(Name, Group, Params);
+
+		if (!Collection->HasGroup(Group))
+		{
+			Collection->AddGroup(Group);
+		}
+
+		FName LocalGroupDependency = GroupDependency;
+		if (!InGroupDependency.IsNone())
+		{
+			LocalGroupDependency = InGroupDependency;
+		}
+
+		bool bSaved = (PersistencePolicy == ManageArrayAccessor::EPersistencePolicy::MakePersistent);
+		FManagedArrayCollection::FConstructionParameters Params(LocalGroupDependency, bSaved);
+		ConstAttributeArray = AttributeArray = &Collection->AddAttribute<T>(Name, Group, Params);
 		return *AttributeArray;
 	}
 
 	/** add and fill the attribute if it does not exist yet */
-	void AddAndFill(const T& Value, ManageArrayAccessor::EPersistencePolicy PersistencePolicy = ManageArrayAccessor::EPersistencePolicy::MakePersistent)
+	void AddAndFill(const T& Value, 
+		ManageArrayAccessor::EPersistencePolicy PersistencePolicy = ManageArrayAccessor::EPersistencePolicy::MakePersistent,
+		FName InGroupDependency = FName(NAME_None))
 	{
-		if (!Collection.HasAttribute(Name, Group))
+		check(!IsConst());
+		if (!Collection->HasAttribute(Name, Group))
 		{
-			Add(PersistencePolicy);
+			Add(PersistencePolicy, InGroupDependency);
 			AttributeArray->Fill(Value);
 		}
 	}
@@ -72,6 +115,7 @@ public:
 	/** Fill the attribute with a specific value */
 	void Fill(const T& Value)
 	{
+		check(!IsConst());
 		if (AttributeArray)
 		{
 			AttributeArray->Fill(Value);
@@ -81,18 +125,32 @@ public:
 	/** copy from another attribute ( create if necessary ) */
 	void Copy(const TManagedArrayAccessor<T>& FromAttribute )
 	{
-		Collection.CopyAttribute(FromAttribute.Collection, Name, Group);
+		check(!IsConst());
+		Collection->CopyAttribute(*FromAttribute.Collection, Name, Group);
 	}
 
 	void Remove()
 	{
-		Collection.RemoveAttribute(Name, Group);
-		AttributeArray = nullptr;
+		check(!IsConst());
+		Collection->RemoveAttribute(Name, Group);
+		ConstAttributeArray = AttributeArray = nullptr;
+	}
+
+	int32 Num() const
+	{
+		return ConstCollection.NumElements(Group);
 	}
 
 private:
-	FManagedArrayCollection& Collection;
+
+	// const collection will be a null pointer, 
+	// while non-const will be valid.
+	const FManagedArrayCollection& ConstCollection;
+	FManagedArrayCollection* Collection = nullptr;
+
 	FName Name;
 	FName Group;
+	FName GroupDependency;
 	TManagedArray<T>* AttributeArray;
+	const TManagedArray<T>* ConstAttributeArray;
 };
diff --git a/Engine/Source/Runtime/Experimental/GeometryCollectionEngine/Private/GeometryCollection/GeometryCollectionEngineConversion.cpp b/Engine/Source/Runtime/Experimental/GeometryCollectionEngine/Private/GeometryCollection/GeometryCollectionEngineConversion.cpp
--- a/Engine/Source/Runtime/Experimental/GeometryCollectionEngine/Private/GeometryCollection/GeometryCollectionEngineConversion.cpp
+++ b/Engine/Source/Runtime/Experimental/GeometryCollectionEngine/Private/GeometryCollection/GeometryCollectionEngineConversion.cpp
@@ -1213,7 +1213,7 @@ void FGeometryCollectionEngineConversion::AppendSkeleton(const USkeleton* InSkel
 		return;
 	}
 	FGeometryCollection::DefineTransformSchema(*InCollection);
-	GeometryCollection::Facades::FTransformSource TransformSourceFacade(InCollection);
+	GeometryCollection::Facades::FTransformSource TransformSourceFacade(*InCollection);
 
 	TManagedArray<FTransform>& Transform = InCollection->ModifyAttribute<FTransform>(FTransformCollection::TransformAttribute, FTransformCollection::TransformGroup);
 	TManagedArray<FLinearColor>& BoneColor = InCollection->ModifyAttribute<FLinearColor>("BoneColor", FTransformCollection::TransformGroup);
@@ -1247,7 +1247,7 @@ void FGeometryCollectionEngineConversion::AppendSkeleton(const USkeleton* InSkel
 		}
 
 		ensure(Roots.Num());
-		TransformSourceFacade.AddTransformSource(InCollection, InSkeleton->GetName(), InSkeleton->GetGuid().ToString(), Roots);
+		TransformSourceFacade.AddTransformSource(InSkeleton->GetName(), InSkeleton->GetGuid().ToString(), Roots);
 	}
 }
 
