Dataflow : Node type cast comparsion.

diff --git a/Engine/Source/Runtime/Experimental/Dataflow/Core/Public/Dataflow/DataflowNode.h b/Engine/Source/Runtime/Experimental/Dataflow/Core/Public/Dataflow/DataflowNode.h
--- a/Engine/Source/Runtime/Experimental/Dataflow/Core/Public/Dataflow/DataflowNode.h
+++ b/Engine/Source/Runtime/Experimental/Dataflow/Core/Public/Dataflow/DataflowNode.h
@@ -201,12 +201,16 @@ struct DATAFLOWCORE_API FDataflowNode
 
 	bool IsValid() const { return bValid; }
 
-	virtual bool IsA(FName InType) const { return InType==StaticType(); }
+	virtual bool IsA(FName InType) const 
+	{ 
+		return InType.ToString().Equals(StaticType().ToString());
+	}
 
 	template<class T>
 	const T* AsType() const
 	{
-		if (IsA(T::StaticType()))
+		FName TargetType = T::StaticType();
+		if (IsA(TargetType))
 		{
 			return (T*)this;
 		}
@@ -251,7 +255,7 @@ public:																				\
 	static FString StaticToolTip() {return FString("Create a dataflow node.");}		\
 	virtual FName GetType() const { return #TYPE; }									\
 	virtual bool IsA(FName InType) const override									\
-		{ return InType==StaticType() || Super::IsA(InType); }						\
+		{ return InType.ToString().Equals(StaticType().ToString()) || Super::IsA(InType); }	\
 	virtual FStructOnScope* NewStructOnScope() override {							\
 	   return new FStructOnScope(TYPE::StaticStruct(), (uint8*)this);}				\
 	virtual void SerializeInternal(FArchive& Ar) override {							\
diff --git a/Engine/Source/Runtime/Experimental/Dataflow/Core/Public/Dataflow/DataflowTerminalNode.h b/Engine/Source/Runtime/Experimental/Dataflow/Core/Public/Dataflow/DataflowTerminalNode.h
--- a/Engine/Source/Runtime/Experimental/Dataflow/Core/Public/Dataflow/DataflowTerminalNode.h
+++ b/Engine/Source/Runtime/Experimental/Dataflow/Core/Public/Dataflow/DataflowTerminalNode.h
@@ -37,6 +37,14 @@ struct DATAFLOWCORE_API FDataflowTerminalNode : public FDataflowNode
 
 	virtual ~FDataflowTerminalNode() { }
 
+	static FName StaticType() { return FName("FDataflowTerminalNode"); }
+
+	virtual bool IsA(FName InType) const override 
+	{ 
+		return InType.ToString().Equals(StaticType().ToString()) 
+			|| Super::IsA(InType); 
+	} 
+
 	//
 	// Error Checking
 	//
