GeometryCollection : Facade cleanup

diff --git a/Engine/Source/Runtime/Experimental/Chaos/Private/GeometryCollection/Facades/CollectionAnchoringFacade.cpp b/Engine/Source/Runtime/Experimental/Chaos/Private/GeometryCollection/Facades/CollectionAnchoringFacade.cpp
--- a/Engine/Source/Runtime/Experimental/Chaos/Private/GeometryCollection/Facades/CollectionAnchoringFacade.cpp
+++ b/Engine/Source/Runtime/Experimental/Chaos/Private/GeometryCollection/Facades/CollectionAnchoringFacade.cpp
@@ -15,6 +15,11 @@ namespace Chaos::Facades
 		, AnchoredAttribute(InCollection, "Anchored", FTransformCollection::TransformGroup)
 	{}
 
+	bool FCollectionAnchoringFacade::IsValid() const
+	{
+		return InitialDynamicStateAttribute.IsValid() && AnchoredAttribute.IsValid();
+	}
+
 	bool FCollectionAnchoringFacade::HasInitialDynamicStateAttribute() const
 	{
 		return InitialDynamicStateAttribute.IsValid();
diff --git a/Engine/Source/Runtime/Experimental/Chaos/Public/GeometryCollection/Facades/CollectionAnchoringFacade.h b/Engine/Source/Runtime/Experimental/Chaos/Public/GeometryCollection/Facades/CollectionAnchoringFacade.h
--- a/Engine/Source/Runtime/Experimental/Chaos/Public/GeometryCollection/Facades/CollectionAnchoringFacade.h
+++ b/Engine/Source/Runtime/Experimental/Chaos/Public/GeometryCollection/Facades/CollectionAnchoringFacade.h
@@ -17,9 +17,19 @@ namespace Chaos::Facades
 		FCollectionAnchoringFacade(FManagedArrayCollection& InCollection);
 		FCollectionAnchoringFacade(const FManagedArrayCollection& InCollection);
 
-		/** Is the facade defined constant. */
+		/** Create the facade attributes. */
+		void DefineSchema() {}
+
+		/** Valid if all accessors arrays are available */
+		bool IsValid() const;
+
+		/** Is the facade defined from a constant collection. */
 		bool IsConst() const { return InitialDynamicStateAttribute.IsConst(); }
 
+		//
+		//  Facade Functionality
+		//
+
 		bool HasInitialDynamicStateAttribute() const;
 		
 		EObjectStateType GetInitialDynamicState(int32 TransformIndex) const;
diff --git a/Engine/Source/Runtime/Experimental/Chaos/Public/GeometryCollection/Facades/CollectionHierarchyFacade.h b/Engine/Source/Runtime/Experimental/Chaos/Public/GeometryCollection/Facades/CollectionHierarchyFacade.h
--- a/Engine/Source/Runtime/Experimental/Chaos/Public/GeometryCollection/Facades/CollectionHierarchyFacade.h
+++ b/Engine/Source/Runtime/Experimental/Chaos/Public/GeometryCollection/Facades/CollectionHierarchyFacade.h
@@ -23,6 +23,9 @@ namespace Chaos::Facades
 		FCollectionHierarchyFacade(FManagedArrayCollection& InCollection);
 		FCollectionHierarchyFacade(const FManagedArrayCollection& InCollection);
 
+		/** Create the facade attributes. */
+		void DefineSchema() {}
+
 		/** Valid if parent and children arrays are available */
 		bool IsValid() const;
 
