Make geometry collection replication properties write-able

diff --git a/Engine/Source/Runtime/Experimental/GeometryCollectionEngine/Public/GeometryCollection/GeometryCollectionComponent.h b/Engine/Source/Runtime/Experimental/GeometryCollectionEngine/Public/GeometryCollection/GeometryCollectionComponent.h
--- a/Engine/Source/Runtime/Experimental/GeometryCollectionEngine/Public/GeometryCollection/GeometryCollectionComponent.h
+++ b/Engine/Source/Runtime/Experimental/GeometryCollectionEngine/Public/GeometryCollection/GeometryCollectionComponent.h
@@ -1054,14 +1054,14 @@ protected:
 	void UpdateCrumblingEventRegistration();
 	
 	/* Per-instance override to enable/disable replication for the geometry collection */
-	UPROPERTY(EditAnywhere, BlueprintReadOnly, Category=Network)
+	UPROPERTY(EditAnywhere, BlueprintReadWrite, Category=Network)
 	bool bEnableReplication;
 
 	/** 
 	 * Enables use of ReplicationAbandonAfterLevel to stop providing network updates to
 	 * clients when the updated particle is of a level higher then specified.
 	 */
-	UPROPERTY(EditAnywhere, BlueprintReadOnly, Category = Network)
+	UPROPERTY(EditAnywhere, BlueprintReadWrite, Category = Network)
 	bool bEnableAbandonAfterLevel;
 
 	/**
@@ -1077,7 +1077,7 @@ protected:
 	* If replicating - the cluster level after which replication will not happen 
 	* @see bEnableAbandonAfterLevel
 	*/
-	UPROPERTY(EditAnywhere, BlueprintReadOnly, Category = Network)
+	UPROPERTY(EditAnywhere, BlueprintReadWrite, Category = Network)
 	int32 ReplicationAbandonAfterLevel;
 
 	UPROPERTY(Replicated)
