ChaosFlesh : Fix for CIS

diff --git a/Engine/Source/Runtime/Experimental/Chaos/Public/Chaos/Math/Krylov.h b/Engine/Source/Runtime/Experimental/Chaos/Public/Chaos/Math/Krylov.h
--- a/Engine/Source/Runtime/Experimental/Chaos/Public/Chaos/Math/Krylov.h
+++ b/Engine/Source/Runtime/Experimental/Chaos/Public/Chaos/Math/Krylov.h
@@ -63,16 +63,17 @@ void LanczosCG(
 	LanczosCG<T>(multiplyA, dotProduct, AXPY, scale, set, x, b, max_it, res, check_residual);
 }
 
-template <class T, class TV, typename Func, int32 d=3>
+template <class T, typename Func, int32 d=3>
 void LanczosCG(
 	Func multiplyA, 
-	TArray<TV>& x,
-	TArray<TV>& b,
+	TArray<TVector<T, 3>>& x,
+	TArray<TVector<T, 3>>& b,
 	const int max_it, 
 	const T res = 1e-4, 
 	const TArray<int32>* use_list = nullptr)
 {
-	auto dot_product = [&use_list](const TArray<TV>& x, const TArray<TV>& y) 
+	using TV = TVector<T, 3>;
+	auto dot_product = [&use_list](const TArray<TV>& x, const TArray<TV>& y)
 	{
 		checkfSlow(x.Num() == y.Num(), TEXT("Chaos::LanczosCG()::dot_product[]: x and y not sized consistently."));
 		T result = T(0);
