ChaosFlesh : Fix for CIS

diff --git a/Engine/Source/Runtime/Experimental/Chaos/Public/Chaos/Math/Krylov.h b/Engine/Source/Runtime/Experimental/Chaos/Public/Chaos/Math/Krylov.h
--- a/Engine/Source/Runtime/Experimental/Chaos/Public/Chaos/Math/Krylov.h
+++ b/Engine/Source/Runtime/Experimental/Chaos/Public/Chaos/Math/Krylov.h
@@ -6,6 +6,7 @@
 #include "Chaos/Framework/Parallel.h"
 #include "GenericPlatform/GenericPlatformMath.h"
 #include "Misc/AssertionMacros.h"
+#include "Chaos/Vector.h"
 
 namespace Chaos {
 
@@ -66,13 +67,13 @@ void LanczosCG(
 template <class T, typename Func, int32 d=3>
 void LanczosCG(
 	Func multiplyA, 
-	TArray<TVector<T, 3>>& x,
-	TArray<TVector<T, 3>>& b,
+	TArray<TVector<T, d>>& x,
+	TArray<TVector<T, d>>& b,
 	const int max_it, 
 	const T res = 1e-4, 
 	const TArray<int32>* use_list = nullptr)
 {
-	using TV = TVector<T, 3>;
+	using TV = TVector<T, d>;
 	auto dot_product = [&use_list](const TArray<TV>& x, const TArray<TV>& y)
 	{
 		checkfSlow(x.Num() == y.Num(), TEXT("Chaos::LanczosCG()::dot_product[]: x and y not sized consistently."));
