Modular Vehicle changes for networking

diff --git a/Engine/Source/Runtime/Experimental/ChaosVehicles/ChaosVehiclesCore/Private/SimModule/AerofoilModule.cpp b/Engine/Source/Runtime/Experimental/ChaosVehicles/ChaosVehiclesCore/Private/SimModule/AerofoilModule.cpp
--- a/Engine/Source/Runtime/Experimental/ChaosVehicles/ChaosVehiclesCore/Private/SimModule/AerofoilModule.cpp
+++ b/Engine/Source/Runtime/Experimental/ChaosVehicles/ChaosVehiclesCore/Private/SimModule/AerofoilModule.cpp
@@ -18,15 +18,15 @@ namespace Chaos
 		switch (Setup().Type)
 		{
 			case EAerofoil::Elevator:
-				ControlSurfaceAngle = Inputs.Pitch * Setup().MaxControlAngle;
+				ControlSurfaceAngle = Inputs.ControlInputs.Pitch * Setup().MaxControlAngle;
 				break;
 
 			case EAerofoil::Rudder:
-				ControlSurfaceAngle = Inputs.Yaw * Setup().MaxControlAngle;
+				ControlSurfaceAngle = Inputs.ControlInputs.Yaw * Setup().MaxControlAngle;
 				break;
 
 			case EAerofoil::Wing:
-				ControlSurfaceAngle = Inputs.Roll * Setup().MaxControlAngle;
+				ControlSurfaceAngle = Inputs.ControlInputs.Roll * Setup().MaxControlAngle;
 			break;
 		}
 
diff --git a/Engine/Source/Runtime/Experimental/ChaosVehicles/ChaosVehiclesCore/Private/SimModule/ClutchModule.cpp b/Engine/Source/Runtime/Experimental/ChaosVehicles/ChaosVehiclesCore/Private/SimModule/ClutchModule.cpp
--- a/Engine/Source/Runtime/Experimental/ChaosVehicles/ChaosVehiclesCore/Private/SimModule/ClutchModule.cpp
+++ b/Engine/Source/Runtime/Experimental/ChaosVehicles/ChaosVehiclesCore/Private/SimModule/ClutchModule.cpp
@@ -29,7 +29,7 @@ namespace Chaos
 		//float AngularVelocityDifference = EngineSpeed - TransmissionSpeed;
 
 		// Inputs.Clutch 0 is engaged/locked, 1 is depressed/open
-		ClutchValue = (1.0f - Inputs.Clutch) * Setup().ClutchStrength;
+		ClutchValue = (1.0f - Inputs.ControlInputs.Clutch) * Setup().ClutchStrength;
 
 		FTorqueSimModule* Parent = static_cast<FTorqueSimModule*>(GetParent());
 		FTorqueSimModule* Child = static_cast<FTorqueSimModule*>(GetFirstChild());
diff --git a/Engine/Source/Runtime/Experimental/ChaosVehicles/ChaosVehiclesCore/Private/SimModule/EngineModule.cpp b/Engine/Source/Runtime/Experimental/ChaosVehicles/ChaosVehiclesCore/Private/SimModule/EngineModule.cpp
--- a/Engine/Source/Runtime/Experimental/ChaosVehicles/ChaosVehiclesCore/Private/SimModule/EngineModule.cpp
+++ b/Engine/Source/Runtime/Experimental/ChaosVehicles/ChaosVehiclesCore/Private/SimModule/EngineModule.cpp
@@ -19,7 +19,7 @@ namespace Chaos
 		}
 
 		// TODO: Engine braking effect
-		DriveTorque = GetEngineTorque(Inputs.Throttle, GetRPM());
+		DriveTorque = GetEngineTorque(Inputs.ControlInputs.Throttle, GetRPM());
 		
 		float BrakeTorque = 0.f;
 		TransmitTorque(VehicleModuleSystem, DriveTorque, BrakeTorque);
diff --git a/Engine/Source/Runtime/Experimental/ChaosVehicles/ChaosVehiclesCore/Private/SimModule/MotorModule.cpp b/Engine/Source/Runtime/Experimental/ChaosVehicles/ChaosVehiclesCore/Private/SimModule/MotorModule.cpp
--- a/Engine/Source/Runtime/Experimental/ChaosVehicles/ChaosVehiclesCore/Private/SimModule/MotorModule.cpp
+++ b/Engine/Source/Runtime/Experimental/ChaosVehicles/ChaosVehiclesCore/Private/SimModule/MotorModule.cpp
@@ -36,7 +36,7 @@ namespace Chaos
 				NormalizedTorque = -1.0f - FMath::Square(1.0f + 2.0f * NormalizedRPM);
 			}
 
-			DriveTorque = NormalizedTorque * Setup().MaxTorque * Inputs.Throttle;
+			DriveTorque = NormalizedTorque * Setup().MaxTorque * Inputs.ControlInputs.Throttle;
 
 			float BrakeTorque = 0.f;
 			TransmitTorque(VehicleModuleSystem, DriveTorque, BrakeTorque);
diff --git a/Engine/Source/Runtime/Experimental/ChaosVehicles/ChaosVehiclesCore/Private/SimModule/SimModuleTree.cpp b/Engine/Source/Runtime/Experimental/ChaosVehicles/ChaosVehiclesCore/Private/SimModule/SimModuleTree.cpp
--- a/Engine/Source/Runtime/Experimental/ChaosVehicles/ChaosVehiclesCore/Private/SimModule/SimModuleTree.cpp
+++ b/Engine/Source/Runtime/Experimental/ChaosVehicles/ChaosVehiclesCore/Private/SimModule/SimModuleTree.cpp
@@ -1,7 +1,6 @@
 // Copyright Epic Games, Inc. All Rights Reserved.
 
 #include "SimModule/SimModuleTree.h"
-#include "SimModule/SimulationModuleBase.h"
 #include "Chaos/ParticleHandleFwd.h"
 #include "PhysicsProxy/GeometryCollectionPhysicsProxy.h"
 #include "VehicleUtility.h"
diff --git a/Engine/Source/Runtime/Experimental/ChaosVehicles/ChaosVehiclesCore/Private/SimModule/ThrusterModule.cpp b/Engine/Source/Runtime/Experimental/ChaosVehicles/ChaosVehiclesCore/Private/SimModule/ThrusterModule.cpp
--- a/Engine/Source/Runtime/Experimental/ChaosVehicles/ChaosVehiclesCore/Private/SimModule/ThrusterModule.cpp
+++ b/Engine/Source/Runtime/Experimental/ChaosVehicles/ChaosVehiclesCore/Private/SimModule/ThrusterModule.cpp
@@ -20,7 +20,7 @@ namespace Chaos
 	void FThrusterSimModule::Simulate(float DeltaTime, const FAllInputs& Inputs, FSimModuleTree& VehicleModuleSystem)
 	{
 		// applies continuous force
-		FVector Force = Setup().ForceAxis * Setup().MaxThrustForce * Inputs.Throttle;
+		FVector Force = Setup().ForceAxis * Setup().MaxThrustForce * Inputs.ControlInputs.Throttle;
 		AddLocalForceAtPosition(Force, Setup().ForceOffset, true, false, false, FColor::Magenta);
 
 	}
diff --git a/Engine/Source/Runtime/Experimental/ChaosVehicles/ChaosVehiclesCore/Private/SimModule/TransmissionModule.cpp b/Engine/Source/Runtime/Experimental/ChaosVehicles/ChaosVehiclesCore/Private/SimModule/TransmissionModule.cpp
--- a/Engine/Source/Runtime/Experimental/ChaosVehicles/ChaosVehiclesCore/Private/SimModule/TransmissionModule.cpp
+++ b/Engine/Source/Runtime/Experimental/ChaosVehicles/ChaosVehiclesCore/Private/SimModule/TransmissionModule.cpp
@@ -52,17 +52,17 @@ namespace Chaos
 		}
 		else
 		{
-			if (Inputs.ChangeUp)
+			if (Inputs.ControlInputs.ChangeUp)
 			{
 				ChangeUp();
 			}
-			else if (Inputs.ChangeDown)
+			else if (Inputs.ControlInputs.ChangeDown)
 			{
 				ChangeDown();
 			}
-			else if (Inputs.SetGearNumber)
+			else if (Inputs.ControlInputs.SetGearNumber)
 			{
-				TargetGear = Inputs.SetGearNumber;
+				TargetGear = Inputs.ControlInputs.SetGearNumber;
 			}
 		}
 
diff --git a/Engine/Source/Runtime/Experimental/ChaosVehicles/ChaosVehiclesCore/Private/SimModule/WheelModule.cpp b/Engine/Source/Runtime/Experimental/ChaosVehicles/ChaosVehiclesCore/Private/SimModule/WheelModule.cpp
--- a/Engine/Source/Runtime/Experimental/ChaosVehicles/ChaosVehiclesCore/Private/SimModule/WheelModule.cpp
+++ b/Engine/Source/Runtime/Experimental/ChaosVehicles/ChaosVehiclesCore/Private/SimModule/WheelModule.cpp
@@ -31,14 +31,14 @@ namespace Chaos
 		float TorqueScaling = 0.00005f;
 		float TractionControlAndABSScaling = 0.98f;	// how close to perfection is the system working
 
-		float HandbrakeTorque = Setup().HandbrakeEnabled ? Inputs.Handbrake * Setup().HandbrakeTorque : 0.0f;
-		SteerAngleDegrees = Setup().SteeringEnabled ? Inputs.Steering * Setup().MaxSteeringAngle : 0.0f;
-		BrakeTorque = Inputs.Brake * Setup().MaxBrakeTorque + HandbrakeTorque;
+		float HandbrakeTorque = Setup().HandbrakeEnabled ? Inputs.ControlInputs.Handbrake * Setup().HandbrakeTorque : 0.0f;
+		SteerAngleDegrees = Setup().SteeringEnabled ? Inputs.ControlInputs.Steering * Setup().MaxSteeringAngle : 0.0f;
+		BrakeTorque = Inputs.ControlInputs.Brake * Setup().MaxBrakeTorque + HandbrakeTorque;
 		LoadTorque = 0.0f;
 
 		// TODO: think about doing this properly, stops vehicles rolling around on their own too much
 		// kindof an auto handbrake
-		if (Inputs.Brake < SMALL_NUMBER && Inputs.Throttle < SMALL_NUMBER && ModuleLocalVelocity.X < 10.0f)
+		if (Inputs.ControlInputs.Brake < SMALL_NUMBER && Inputs.ControlInputs.Throttle < SMALL_NUMBER && ModuleLocalVelocity.X < 10.0f)
 		{
 			BrakeTorque = Setup().HandbrakeTorque;
 		}
diff --git a/Engine/Source/Runtime/Experimental/ChaosVehicles/ChaosVehiclesCore/Public/SimModule/SimModuleTree.h b/Engine/Source/Runtime/Experimental/ChaosVehicles/ChaosVehiclesCore/Public/SimModule/SimModuleTree.h
--- a/Engine/Source/Runtime/Experimental/ChaosVehicles/ChaosVehiclesCore/Public/SimModule/SimModuleTree.h
+++ b/Engine/Source/Runtime/Experimental/ChaosVehicles/ChaosVehiclesCore/Public/SimModule/SimModuleTree.h
@@ -3,6 +3,7 @@
 #pragma once
 
 #include "SimModule/DeferredForcesModular.h"
+#include "SimModule/SimulationModuleBase.h"
 
 class FGeometryCollectionPhysicsProxy;
 
@@ -70,6 +71,11 @@ namespace Chaos
 		const FDeferredForcesModular& GetDeferredForces() const { return DeferredForces; }
 		const TArray<FSimModuleNode>& GetSimulationModuleTree() { return SimulationModuleTree; }
 
+		FControlInputs& GetControlInputs()
+		{
+			return AllInputs.ControlInputs;
+		}
+
 	protected:
 		void SimulateNode(float DeltaTime, FAllInputs& Inputs, int NodeIdx);
 
@@ -86,6 +92,8 @@ namespace Chaos
 
 		FDeferredForcesModular DeferredForces;
 
+		Chaos::FAllInputs AllInputs;
+
 	};
 
 
diff --git a/Engine/Source/Runtime/Experimental/ChaosVehicles/ChaosVehiclesCore/Public/SimModule/SimulationModuleBase.h b/Engine/Source/Runtime/Experimental/ChaosVehicles/ChaosVehiclesCore/Public/SimModule/SimulationModuleBase.h
--- a/Engine/Source/Runtime/Experimental/ChaosVehicles/ChaosVehiclesCore/Public/SimModule/SimulationModuleBase.h
+++ b/Engine/Source/Runtime/Experimental/ChaosVehicles/ChaosVehiclesCore/Public/SimModule/SimulationModuleBase.h
@@ -59,10 +59,11 @@ namespace Chaos
 		bool bBlockingHit;
 	};
 
-	struct CHAOSVEHICLESCORE_API FAllInputs : public FControlInputs
+	struct CHAOSVEHICLESCORE_API FAllInputs
 	{
 		FTransform VehicleWorldTransform;
 		TMap<int32, FModuleHitResults> HitResults;
+		FControlInputs ControlInputs;
 	};
 
 	/**
diff --git a/Engine/Source/Runtime/Experimental/GeometryCollectionEngine/Private/GeometryCollection/GeometryCollectionComponent.cpp b/Engine/Source/Runtime/Experimental/GeometryCollectionEngine/Private/GeometryCollection/GeometryCollectionComponent.cpp
--- a/Engine/Source/Runtime/Experimental/GeometryCollectionEngine/Private/GeometryCollection/GeometryCollectionComponent.cpp
+++ b/Engine/Source/Runtime/Experimental/GeometryCollectionEngine/Private/GeometryCollection/GeometryCollectionComponent.cpp
@@ -165,6 +165,8 @@ bool FGeometryCollectionRepData::NetSerialize(FArchive& Ar, class UPackageMap* M
 
 	Ar << OneOffActivated;
 
+	Ar << ServerFrame;
+
 	int32 NumClusters = Clusters.Num();
 	Ar << NumClusters;
 
@@ -350,7 +352,7 @@ UGeometryCollectionComponent::UGeometryCollectionComponent(const FObjectInitiali
 
 }
 
-Chaos::FPhysicsSolver* GetSolver(const UGeometryCollectionComponent& GeometryCollectionComponent)
+Chaos::FPhysicsSolver* UGeometryCollectionComponent::GetSolver(const UGeometryCollectionComponent& GeometryCollectionComponent)
 {
 	if(GeometryCollectionComponent.ChaosSolverActor)
 	{
@@ -1621,6 +1623,11 @@ void UGeometryCollectionComponent::UpdateRepData()
 		if (bClustersChanged)
 		{
 			RepData.Clusters = MoveTemp(Clusters);
+			
+			if (Owner->GetWorld() && Owner->GetWorld()->GetPhysicsScene())
+			{
+				RepData.ServerFrame = Owner->GetWorld()->GetPhysicsScene()->ReplicationCache.ServerFrame;
+			}
 
 			INC_DWORD_STAT_BY(STAT_GCReplicatedClusters, RepData.Clusters.Num());
 
diff --git a/Engine/Source/Runtime/Experimental/GeometryCollectionEngine/Public/GeometryCollection/GeometryCollectionComponent.h b/Engine/Source/Runtime/Experimental/GeometryCollectionEngine/Public/GeometryCollection/GeometryCollectionComponent.h
--- a/Engine/Source/Runtime/Experimental/GeometryCollectionEngine/Public/GeometryCollection/GeometryCollectionComponent.h
+++ b/Engine/Source/Runtime/Experimental/GeometryCollectionEngine/Public/GeometryCollection/GeometryCollectionComponent.h
@@ -385,7 +385,7 @@ struct FGeometryCollectionRepData
 	GENERATED_BODY()
 
 	FGeometryCollectionRepData()
-		: Version(0)
+		: Version(0), ServerFrame(0)
 	{
 
 	}
@@ -400,6 +400,9 @@ struct FGeometryCollectionRepData
 	// as there's no reason to compare the Poses array.
 	int32 Version;
 
+	// For Network Prediction Mode we require the frame number on the server when the data was gathered
+	int32 ServerFrame;
+
 	// Just test version to skip having to traverse the whole pose array for replication
 	bool Identical(const FGeometryCollectionRepData* Other, uint32 PortFlags) const;
 	bool NetSerialize(FArchive& Ar, class UPackageMap* Map, bool& bOutSuccess);
@@ -1043,6 +1046,7 @@ protected:
 	/** Issue a field command for the physics thread */
 	void DispatchFieldCommand(const FFieldSystemCommand& InCommand);
 
+	Chaos::FPhysicsSolver* GetSolver(const UGeometryCollectionComponent& GeometryCollectionComponent);
 	void CalculateLocalBounds();
 	void CalculateGlobalMatrices();
 	FBox ComputeBounds(const FMatrix& LocalToWorldWithScale) const;
@@ -1084,10 +1088,13 @@ protected:
 	FGeometryCollectionRepData RepData;
 
 	/** Called post solve to allow authoritative components to update their replication data */
-	void UpdateRepData();
+	virtual void UpdateRepData();
 
 	/** Clear all rep data, this is required if the physics proxy has been recreated */
-	void ResetRepData();
+	virtual void ResetRepData();
+
+	virtual void ProcessRepData();
+	int32 VersionProcessed = INDEX_NONE;
 
 private:
 
@@ -1176,7 +1183,6 @@ private:
 	void IncrementBreakTimer(float DeltaTime);
 	bool CalculateInnerSphere(int32 TransformIndex, UE::Math::TSphere<double>& SphereOut) const;
 	void UpdateDecay(int32 TransformIdx, float UpdatedDecay, bool UseClusterCrumbling, bool HasDynamicInternalClusterParent, FGeometryCollectionDecayContext& ContextInOut);
-	void ProcessRepData();
 
 	void UpdateAttachedChildrenTransform() const;
 	
@@ -1187,7 +1193,6 @@ private:
 
 	/** One off activation is processed in the same order as server so remember the last one we processed */
 	int32 OneOffActivatedProcessed = 0;
-	int32 VersionProcessed = INDEX_NONE;
 	double LastHardsnapTimeInMs = 0;
 
 	/** True if GeometryCollection transforms have changed from previous tick. */
