Dataflow : fix IsConnected return false positive

diff --git a/Engine/Source/Runtime/Experimental/Dataflow/Core/Public/Dataflow/DataflowNode.h b/Engine/Source/Runtime/Experimental/Dataflow/Core/Public/Dataflow/DataflowNode.h
--- a/Engine/Source/Runtime/Experimental/Dataflow/Core/Public/Dataflow/DataflowNode.h
+++ b/Engine/Source/Runtime/Experimental/Dataflow/Core/Public/Dataflow/DataflowNode.h
@@ -195,7 +195,7 @@ struct DATAFLOWCORE_API FDataflowNode
 	template<class T> bool IsConnected(const T* Reference) const
 	{
 		checkSlow(FindInput(Reference));
-		return FindInput(Reference) != nullptr;
+		return (FindInput(Reference)->GetConnection() != nullptr);
 	}
 
 	void Invalidate();
