Add RemoveOnBreak facade
- Update code accessing the directa ttribute to use the facade
- Updated RemoveOnBreak dataflow node to use it too

diff --git a/Engine/Plugins/Experimental/ChaosEditor/Source/FractureEditor/Private/FractureToolProperties.cpp b/Engine/Plugins/Experimental/ChaosEditor/Source/FractureEditor/Private/FractureToolProperties.cpp
--- a/Engine/Plugins/Experimental/ChaosEditor/Source/FractureEditor/Private/FractureToolProperties.cpp
+++ b/Engine/Plugins/Experimental/ChaosEditor/Source/FractureEditor/Private/FractureToolProperties.cpp
@@ -6,6 +6,7 @@
 #include "GeometryCollection/GeometryCollectionComponent.h"
 #include "GeometryCollection/GeometryCollectionEngineRemoval.h"
 #include "GeometryCollection/Facades/CollectionAnchoringFacade.h"
+#include "GeometryCollection/Facades/CollectionRemoveOnBreakFacade.h"
 #include "GeometryCollection/GeometryCollectionObject.h"
 
 #include UE_INLINE_GENERATED_CPP_BY_NAME(FractureToolProperties)
@@ -148,30 +149,16 @@ void UFractureToolSetRemoveOnBreak::Execute(TWeakPtr<FFractureEditorModeToolkit>
 	 				TArray<int32> SelectedBones = GeometryCollectionComponent->GetSelectedBones();
 	 				if (SelectedBones.Num())
 	 				{
-	 					if (!GeometryCollection->HasAttribute("RemoveOnBreak", FGeometryCollection::TransformGroup))
-	 					{
-	 						TManagedArray<FVector4f>& NewRemoveOnBreak = GeometryCollection->AddAttribute<FVector4f>("RemoveOnBreak", FGeometryCollection::TransformGroup);
-	 						NewRemoveOnBreak.Fill(FRemoveOnBreakData::DisabledPackedData);
-	 					}
+						GeometryCollection::Facades::FCollectionRemoveOnBreakFacade RemoveOnBreakFacade(*GeometryCollection);
+						RemoveOnBreakFacade.DefineSchema();
 	 					
-	 					TManagedArray<FVector4f>& RemoveOnBreak = GeometryCollection->ModifyAttribute<FVector4f>("RemoveOnBreak", FGeometryCollection::TransformGroup);
-
-	 					const FVector2f BreakTimer{ RemoveOnBreakSettings->PostBreakTimer.X, RemoveOnBreakSettings->PostBreakTimer.Y };
-	 					const FVector2f RemovalTimer{ RemoveOnBreakSettings->RemovalTimer.X, RemoveOnBreakSettings->RemovalTimer.Y };
-	 					const FRemoveOnBreakData RemoveOnBreakData(RemoveOnBreakSettings->Enabled,  BreakTimer, RemoveOnBreakSettings->ClusterCrumbling, RemovalTimer);
-
-	 					for (int32 Index : SelectedBones)
-	 					{
-							// if root bone, then do not set 
-							if (GeometryCollection->Parent[Index] == INDEX_NONE)
-							{
-								RemoveOnBreak[Index] = FRemoveOnBreakData::DisabledPackedData;
-							}
-							else
-							{
-								RemoveOnBreak[Index] = RemoveOnBreakData.GetPackedData();
-							}
-	 					}
+						GeometryCollection::Facades::FRemoveOnBreakData RemoveOnBreakData;
+						RemoveOnBreakData.SetBreakTimer(RemoveOnBreakSettings->PostBreakTimer.X, RemoveOnBreakSettings->PostBreakTimer.Y);
+						RemoveOnBreakData.SetRemovalTimer(RemoveOnBreakSettings->RemovalTimer.X, RemoveOnBreakSettings->RemovalTimer.Y);
+						RemoveOnBreakData.SetEnabled(RemoveOnBreakSettings->Enabled);
+						RemoveOnBreakData.SetClusterCrumbling(RemoveOnBreakSettings->ClusterCrumbling);
+
+						RemoveOnBreakFacade.SetFromIndexArray(SelectedBones, RemoveOnBreakData);
 	 				}
 	 			}
 	 		}
@@ -192,7 +179,8 @@ void UFractureToolSetRemoveOnBreak::DeleteRemoveOnBreakData()
 			TSharedPtr<FGeometryCollection, ESPMode::ThreadSafe> GeometryCollectionPtr = GCObject->GetGeometryCollection();
 			if (FGeometryCollection* GeometryCollection = GeometryCollectionPtr.Get())
 			{
-				GeometryCollection->RemoveAttribute("RemoveOnBreak", FGeometryCollection::TransformGroup);
+				GeometryCollection::Facades::FCollectionRemoveOnBreakFacade RemoveOnBreakFacade(*GeometryCollection);
+				RemoveOnBreakFacade.RemoveSchema();
 			}
 		}
 	}
diff --git a/Engine/Plugins/Experimental/ChaosEditor/Source/FractureEditor/Private/SGeometryCollectionOutliner.cpp b/Engine/Plugins/Experimental/ChaosEditor/Source/FractureEditor/Private/SGeometryCollectionOutliner.cpp
--- a/Engine/Plugins/Experimental/ChaosEditor/Source/FractureEditor/Private/SGeometryCollectionOutliner.cpp
+++ b/Engine/Plugins/Experimental/ChaosEditor/Source/FractureEditor/Private/SGeometryCollectionOutliner.cpp
@@ -811,7 +811,7 @@ void FGeometryCollectionTreeItemBone::UpdateItemFromCollection()
 	InitialState = INDEX_NONE;
 	Anchored = false;
 	RemoveOnBreakAvailable = false;
-	RemoveOnBreak = FVector4f{-1};
+	RemoveOnBreak = FRemoveOnBreakData::DisabledPackedData;
 	ImportedCollisionsAvailable = false;
 	ImportedCollisionsUsed = false;
 	
@@ -835,7 +835,7 @@ void FGeometryCollectionTreeItemBone::UpdateItemFromCollection()
 			const TManagedArray<int32>& SimulationType = GeometryCollectionPtr->SimulationType;
 			const TManagedArray<float>* RelativeSizes = GeometryCollectionPtr->FindAttribute<float>("Size", FTransformCollection::TransformGroup);
 			const TManagedArray<float>* Volumes = GeometryCollectionPtr->FindAttribute<float>("Volume", FTransformCollection::TransformGroup);
-			const TManagedArray<FVector4f>* RemoveOnBreakArray = GeometryCollectionPtr->FindAttribute<FVector4f>("RemoveOnBreak", FTransformCollection::TransformGroup);
+			GeometryCollection::Facades::FCollectionRemoveOnBreakFacade RemoveOnBreakFacade(*GeometryCollectionPtr);
 
 			using FImplicitGeom = FGeometryDynamicCollection::FSharedImplicit;
 			const TManagedArray<FImplicitGeom>* ExternalCollisions = GeometryCollectionPtr->FindAttribute<FImplicitGeom>("ExternalCollisions", FTransformCollection::TransformGroup);
@@ -901,10 +901,10 @@ void FGeometryCollectionTreeItemBone::UpdateItemFromCollection()
 					DamageThreshold = DamageData.DamageThreshold;
 					Broken = DamageData.bIsBroken;
 				}
-				RemoveOnBreakAvailable = (RemoveOnBreakArray != nullptr); 
+				RemoveOnBreakAvailable = RemoveOnBreakFacade.IsValid();
 				if (RemoveOnBreakAvailable)
 				{
-					RemoveOnBreak = (*RemoveOnBreakArray)[ItemBoneIndex];
+					RemoveOnBreak = RemoveOnBreakFacade.GetData(ItemBoneIndex);
 				}
 
 				ImportedCollisionsUsed = RestCollection->bImportCollisionFromSource;
diff --git a/Engine/Plugins/Experimental/ChaosEditor/Source/FractureEditor/Private/SGeometryCollectionOutliner.h b/Engine/Plugins/Experimental/ChaosEditor/Source/FractureEditor/Private/SGeometryCollectionOutliner.h
--- a/Engine/Plugins/Experimental/ChaosEditor/Source/FractureEditor/Private/SGeometryCollectionOutliner.h
+++ b/Engine/Plugins/Experimental/ChaosEditor/Source/FractureEditor/Private/SGeometryCollectionOutliner.h
@@ -16,7 +16,7 @@
 
 #include "GeometryCollection/GeometryCollection.h"
 #include "GeometryCollection/GeometryCollectionComponent.h"
-#include "GeometryCollection/GeometryCollectionEngineRemoval.h"
+#include "GeometryCollection/Facades/CollectionRemoveOnBreakFacade.h"
 
 #include "SGeometryCollectionOutliner.generated.h"
 
@@ -138,6 +138,8 @@ private:
 	TArray<int32> HistogramSelection;
 };
 
+using FRemoveOnBreakData = GeometryCollection::Facades::FRemoveOnBreakData;
+
 class FGeometryCollectionTreeItemBone : public FGeometryCollectionTreeItem
 {
 
diff --git a/Engine/Plugins/Experimental/GeometryCollectionPlugin/Source/GeometryCollectionNodes/Private/Dataflow/GeometryCollectionNodes.cpp b/Engine/Plugins/Experimental/GeometryCollectionPlugin/Source/GeometryCollectionNodes/Private/Dataflow/GeometryCollectionNodes.cpp
--- a/Engine/Plugins/Experimental/GeometryCollectionPlugin/Source/GeometryCollectionNodes/Private/Dataflow/GeometryCollectionNodes.cpp
+++ b/Engine/Plugins/Experimental/GeometryCollectionPlugin/Source/GeometryCollectionNodes/Private/Dataflow/GeometryCollectionNodes.cpp
@@ -31,6 +31,7 @@
 #include "FractureEngineSelection.h"
 #include "GeometryCollection/Facades/CollectionBoundsFacade.h"
 #include "GeometryCollection/Facades/CollectionAnchoringFacade.h"
+#include "GeometryCollection/Facades/CollectionRemoveOnBreakFacade.h"
 
 #include UE_INLINE_GENERATED_CPP_BY_NAME(GeometryCollectionNodes)
 
@@ -993,42 +994,41 @@ void FRemoveOnBreakDataflowNode::Evaluate(Dataflow::FContext& Context, const FDa
 {
 	if (Out->IsA<FManagedArrayCollection>(&Collection))
 	{
+		const bool& InEnableRemoval = GetValue(Context, &EnabledRemoval, true);
+		const FVector2f& InPostBreakTimer = GetValue(Context, &PostBreakTimer);
+		const FVector2f& InRemovalTimer = GetValue(Context, &RemovalTimer);
+		const bool& InClusterCrumbling = GetValue(Context, &ClusterCrumbling);
+
+		// we are making a copy of the collection because we are modifying it 
 		FManagedArrayCollection InCollection = GetValue<FManagedArrayCollection>(Context, &Collection);
 
-		if (!InCollection.HasAttribute("RemoveOnBreak", FGeometryCollection::TransformGroup))
+		GeometryCollection::Facades::FCollectionRemoveOnBreakFacade RemoveOnBreakFacade(InCollection);
+		RemoveOnBreakFacade.DefineSchema();
+
+		GeometryCollection::Facades::FRemoveOnBreakData Data;
+		Data.SetBreakTimer(InPostBreakTimer.X, InPostBreakTimer.Y);
+		Data.SetRemovalTimer(InRemovalTimer.X, InRemovalTimer.Y);
+		Data.SetEnabled(InEnableRemoval);
+		Data.SetClusterCrumbling(InClusterCrumbling);
+
+		// selection is optional
+		if (IsConnected<FDataflowTransformSelection>(&TransformSelection))
 		{
-			TManagedArray<FVector4f>& NewRemoveOnBreak = InCollection.AddAttribute<FVector4f>("RemoveOnBreak", FGeometryCollection::TransformGroup);
-			NewRemoveOnBreak.Fill(FRemoveOnBreakData::DisabledPackedData);
+			const FDataflowTransformSelection& InTransformSelection = GetValue(Context, &TransformSelection);
+			TArray<int32> TransformIndices;
+			InTransformSelection.AsArray(TransformIndices);
+			RemoveOnBreakFacade.SetFromIndexArray(TransformIndices, Data);
+		}
+		else
+		{
+			RemoveOnBreakFacade.SetToAll(Data);
 		}
 
-		TManagedArray<FVector4f>& RemoveOnBreak = InCollection.ModifyAttribute<FVector4f>("RemoveOnBreak", FGeometryCollection::TransformGroup);
-
-		const FVector2f PostBreakTimerData = GetValue<FVector2f>(Context, &PostBreakTimer);
-		const FVector2f RemovalTimerData = GetValue<FVector2f>(Context, &RemovalTimer);
-		const bool ClusterCrumblingData = GetValue<bool>(Context, &ClusterCrumbling);
-
-		RemoveOnBreak.Fill(FVector4f{ PostBreakTimerData.X, PostBreakTimerData.Y, RemovalTimerData.X, RemovalTimerData.Y });
-
-			// @todo(harsha) Implement with Selection
-			// const FRemoveOnBreakData RemoveOnBreakData(true, PostBreakTimerData, ClusterCrumblingData, RemovalTimerData);
-			// for (int32 Index : SelectedBones)
-			// {
-				// // if root bone, then do not set 
-				// if (GeometryCollection->Parent[Index] == INDEX_NONE)
-				// {
-				// 	RemoveOnBreak[Index] = FRemoveOnBreakData::DisabledPackedData;
-				// }
-				// else
-				// {
-				// 	RemoveOnBreak[Index] = RemoveOnBreakData.GetPackedData();
-				// }
-			// }
-
-		SetValue<FManagedArrayCollection>(Context, InCollection, &Collection);
+		// move the collection to the output to avoid making another copy
+		SetValue<FManagedArrayCollection>(Context, MoveTemp(InCollection), &Collection);
 	}
 }
 
-
 void FSetAnchorStateDataflowNode::Evaluate(Dataflow::FContext& Context, const FDataflowOutput* Out) const
 {
 	if (Out->IsA<FManagedArrayCollection>(&Collection))
diff --git a/Engine/Plugins/Experimental/GeometryCollectionPlugin/Source/GeometryCollectionNodes/Public/Dataflow/GeometryCollectionNodes.h b/Engine/Plugins/Experimental/GeometryCollectionPlugin/Source/GeometryCollectionNodes/Public/Dataflow/GeometryCollectionNodes.h
--- a/Engine/Plugins/Experimental/GeometryCollectionPlugin/Source/GeometryCollectionNodes/Public/Dataflow/GeometryCollectionNodes.h
+++ b/Engine/Plugins/Experimental/GeometryCollectionPlugin/Source/GeometryCollectionNodes/Public/Dataflow/GeometryCollectionNodes.h
@@ -1554,17 +1554,27 @@ struct FRemoveOnBreakDataflowNode : public FDataflowNode
 	DATAFLOW_NODE_DEFINE_INTERNAL(FRemoveOnBreakDataflowNode, "RemoveOnBreak", "GeometryCollection", "")
 
 public:
-	// @todo(harsha) Support Selections
-
+	/** Collection to set theremoval data on */
 	UPROPERTY(meta = (DataflowInput, DataflowOutput, DataflowPassthrough = "Collection", DataflowIntrinsic))
 	FManagedArrayCollection Collection;
 
+	/** selection to apply the data on ( if not specified the entire collection will be set ) */
+	UPROPERTY(meta = (DataflowInput, DisplayName = "TransformSelection"))
+	FDataflowTransformSelection TransformSelection;
+
+	/** Whether or not to enable the removal on the selection */
+	UPROPERTY(EditAnywhere, Category = "Removal", meta = (DataflowInput))
+	bool EnabledRemoval = true;
+
+	/** How long after the break the removal will start ( Min / Max ) */
 	UPROPERTY(EditAnywhere, Category = "Removal", meta = (DataflowInput))
 	FVector2f PostBreakTimer{0.0, 0.0};
 
+	/** How long removal will last ( Min / Max ) */
 	UPROPERTY(EditAnywhere, Category = "Removal", meta = (DataflowInput))
 	FVector2f RemovalTimer{0.0, 1.0};
 
+	/** If applied to a cluster this will cause the cluster to crumble upon removal, otherwise will have no effect */
 	UPROPERTY(EditAnywhere, Category = "Removal", meta = (DataflowInput))
 	bool ClusterCrumbling = false;
 
@@ -1572,6 +1582,8 @@ public:
 		: FDataflowNode(InParam, InGuid)
 	{
 		RegisterInputConnection(&Collection);
+		RegisterInputConnection(&TransformSelection);
+		RegisterInputConnection(&EnabledRemoval);
 		RegisterInputConnection(&PostBreakTimer);
 		RegisterInputConnection(&RemovalTimer);
 		RegisterInputConnection(&ClusterCrumbling);
diff --git a/Engine/Source/Runtime/Experimental/Chaos/Private/GeometryCollection/Facades/CollectionRemoveOnBreakFacade.cpp b/Engine/Source/Runtime/Experimental/Chaos/Private/GeometryCollection/Facades/CollectionRemoveOnBreakFacade.cpp
new file mode 100644
--- /dev/null
+++ b/Engine/Source/Runtime/Experimental/Chaos/Private/GeometryCollection/Facades/CollectionRemoveOnBreakFacade.cpp
@@ -0,0 +1,174 @@
+// Copyright Epic Games, Inc. All Rights Reserved.
+
+#include "GeometryCollection/Facades/CollectionRemoveOnBreakFacade.h"
+#include "GeometryCollection/GeometryCollection.h"
+
+namespace GeometryCollection::Facades
+{
+	FRemoveOnBreakData::FRemoveOnBreakData()
+		: PackedData(FRemoveOnBreakData::DisabledPackedData)
+	{}
+
+	FRemoveOnBreakData::FRemoveOnBreakData(const FVector4f& InPackedData)
+		: PackedData(InPackedData)
+	{}
+
+	FRemoveOnBreakData::FRemoveOnBreakData(bool bEnable, const FVector2f& BreakTimer, bool bClusterCrumbling, const FVector2f& RemovalTimer)
+	{
+		SetBreakTimer(BreakTimer.X, BreakTimer.Y);
+		SetRemovalTimer(RemovalTimer.X, RemovalTimer.Y);
+		SetEnabled(bEnable);
+		SetClusterCrumbling(bClusterCrumbling);
+	}
+
+	const FVector4f& FRemoveOnBreakData::GetPackedData() const
+	{
+		return PackedData;
+	}
+
+	bool FRemoveOnBreakData::IsEnabled() const
+	{
+		return PackedData.X >= 0;
+	}
+
+	bool FRemoveOnBreakData::GetClusterCrumbling() const
+	{
+		return IsEnabled() && (PackedData.Z < 0);
+	}
+
+	FVector2f FRemoveOnBreakData::GetBreakTimer() const
+	{
+		return FVector2f(FMath::Abs(PackedData.X), FMath::Abs(PackedData.Y));
+	}
+
+	FVector2f FRemoveOnBreakData::GetRemovalTimer() const
+	{
+		return FVector2f(FMath::Abs(PackedData.Z), FMath::Abs(PackedData.W));
+	}
+
+	void FRemoveOnBreakData::SetEnabled(bool bEnable)
+	{
+		PackedData.X = FMath::Abs(PackedData.X);
+		if (!bEnable)
+		{
+			if (PackedData.X == 0)
+			{
+				PackedData.X = -FLT_EPSILON;
+			}
+			else
+			{
+				PackedData.X = -PackedData.X;
+			}
+		}
+		else
+		{
+			// restore any correction made
+			if (PackedData.X <= FLT_EPSILON)
+			{
+				PackedData.X = 0;
+			}
+		}
+	}
+
+	void FRemoveOnBreakData::SetClusterCrumbling(bool bClusterCrumbling)
+	{
+		PackedData.Z = FMath::Abs(PackedData.Z);
+		if (bClusterCrumbling)
+		{
+			if (PackedData.Z == 0)
+			{
+				PackedData.Z = -FLT_EPSILON;
+			}
+			else
+			{
+				PackedData.Z = -PackedData.Z;
+			}
+		}
+		else
+		{
+			// restore any correction made
+			if (PackedData.Z <= FLT_EPSILON)
+			{
+				PackedData.Z = 0;
+			}
+		}
+	}
+
+	void FRemoveOnBreakData::SetBreakTimer(float MinTime, float MaxTime)
+	{
+		// enable flag is packed in the break timer data, let's save it first and set it back
+		const bool bIsEnabled = IsEnabled();
+		PackedData.X = FMath::Min(FMath::Abs(MinTime), FMath::Abs(MaxTime)); // Min break timer
+		PackedData.Y = FMath::Max(FMath::Abs(MinTime), FMath::Abs(MaxTime)); // Max break timer
+		SetEnabled(bIsEnabled);
+	}
+
+	void FRemoveOnBreakData::SetRemovalTimer(float MinTime, float MaxTime)
+	{
+		// cluster crumbling flag is packed in the removal timer data, let's save it first and set it back
+		const bool bClusterCrumbling = GetClusterCrumbling();
+		PackedData.Z = FMath::Min(FMath::Abs(MinTime), FMath::Abs(MaxTime)); // Min removal timer
+		PackedData.W = FMath::Max(FMath::Abs(MinTime), FMath::Abs(MaxTime)); // Max removal timer
+		SetClusterCrumbling(bClusterCrumbling);
+	}
+
+	//////////////////////////////////////////////////////////////////////////////////////////////////////////////////
+
+	FCollectionRemoveOnBreakFacade::FCollectionRemoveOnBreakFacade(FManagedArrayCollection& InCollection)
+		: RemoveOnBreakAttribute(InCollection, "RemoveOnBreak", FGeometryCollection::TransformGroup)
+	{}
+
+	FCollectionRemoveOnBreakFacade::FCollectionRemoveOnBreakFacade(const FManagedArrayCollection& InCollection)
+		: RemoveOnBreakAttribute(InCollection, "RemoveOnBreak", FGeometryCollection::TransformGroup)
+	{}
+
+	void FCollectionRemoveOnBreakFacade::DefineSchema()
+	{
+		check(!IsConst());
+		RemoveOnBreakAttribute.AddAndFill(FRemoveOnBreakData::DisabledPackedData);
+	}
+
+	void FCollectionRemoveOnBreakFacade::RemoveSchema()
+	{
+		check(!IsConst());
+		RemoveOnBreakAttribute.Remove();
+	}
+
+	bool FCollectionRemoveOnBreakFacade::IsValid() const
+	{
+		return RemoveOnBreakAttribute.IsValid();
+	}
+
+	bool FCollectionRemoveOnBreakFacade::IsConst() const
+	{
+		return RemoveOnBreakAttribute.IsConst();
+	}
+
+	void FCollectionRemoveOnBreakFacade::SetFromIndexArray(const TArray<int32>& Indices, const FRemoveOnBreakData& Data)
+	{
+		check(!IsConst());
+		if (IsValid())
+		{
+			TManagedArray<FVector4f>& RemoveOnBreak = RemoveOnBreakAttribute.Modify();
+			for (int32 Index : Indices)
+			{
+				RemoveOnBreak[Index] = Data.GetPackedData();
+			}
+		}
+	}
+
+	void FCollectionRemoveOnBreakFacade::SetToAll(const FRemoveOnBreakData& Data)
+	{
+		check(!IsConst());
+		RemoveOnBreakAttribute.Fill(Data.GetPackedData());
+	}
+
+	const FRemoveOnBreakData FCollectionRemoveOnBreakFacade::GetData(int32 Index) const
+	{
+		if (IsValid())
+		{
+			return RemoveOnBreakAttribute.Get()[Index];
+		}
+		return FRemoveOnBreakData::DisabledPackedData;
+	}
+}
diff --git a/Engine/Source/Runtime/Experimental/Chaos/Public/GeometryCollection/Facades/CollectionRemoveOnBreakFacade.h b/Engine/Source/Runtime/Experimental/Chaos/Public/GeometryCollection/Facades/CollectionRemoveOnBreakFacade.h
new file mode 100644
--- /dev/null
+++ b/Engine/Source/Runtime/Experimental/Chaos/Public/GeometryCollection/Facades/CollectionRemoveOnBreakFacade.h
@@ -0,0 +1,73 @@
+// Copyright Epic Games, Inc. All Rights Reserved.
+
+#pragma once
+
+#include "CoreMinimal.h"
+#include "GeometryCollection/ManagedArrayAccessor.h"
+
+struct FManagedArrayCollection;
+
+namespace GeometryCollection::Facades
+{
+	/**
+	 * this class helps reading and writing the packed data used in the managed array
+	 */
+	struct CHAOS_API FRemoveOnBreakData
+	{
+	public:
+		inline static const FVector4f DisabledPackedData{ -1, 0, 0, 0 };
+
+		FRemoveOnBreakData();
+		FRemoveOnBreakData(const FVector4f& InPackedData);
+		FRemoveOnBreakData(bool bEnable, const FVector2f& BreakTimer, bool bClusterCrumbling, const FVector2f& RemovalTimer);
+
+		const FVector4f& GetPackedData() const;
+
+		bool IsEnabled() const;
+		bool GetClusterCrumbling() const;
+		FVector2f GetBreakTimer() const;
+		FVector2f GetRemovalTimer() const;
+		
+		void SetEnabled(bool bEnable);
+		void SetClusterCrumbling(bool bClusterCrumbling);
+		void SetBreakTimer(float MinTime, float MaxTime);
+		void SetRemovalTimer(float MinTime, float MaxTime);
+
+	private:
+		FVector4f PackedData;
+	};
+
+	/**
+	 * Provide an API to set and get removal on break data on a collection
+	 */
+	class CHAOS_API FCollectionRemoveOnBreakFacade
+	{
+	public:
+		FCollectionRemoveOnBreakFacade(FManagedArrayCollection& InCollection);
+		FCollectionRemoveOnBreakFacade(const FManagedArrayCollection& InCollection);
+
+		/** Create the facade attributes. */
+		void DefineSchema();
+
+		/** remove the facade attributes */
+		void RemoveSchema();
+
+		/** Valid if related attributes are available */
+		bool IsValid() const;
+
+		/** Is this facade const access */
+		bool IsConst() const;
+
+		/** set a specific value to a array of transforms */
+		void SetFromIndexArray(const TArray<int32>& Indices, const FRemoveOnBreakData& Data);
+
+		/** set a specific value to all transforms */
+		void SetToAll(const FRemoveOnBreakData& Data);
+
+		/** Get a specific value by index */
+		const FRemoveOnBreakData GetData(int32 Index) const;
+
+	private:
+		TManagedArrayAccessor<FVector4f> RemoveOnBreakAttribute;
+	};
+}
\ No newline at end of file
diff --git a/Engine/Source/Runtime/Experimental/GeometryCollectionEngine/Private/GeometryCollection/GeometryCollectionComponent.cpp b/Engine/Source/Runtime/Experimental/GeometryCollectionEngine/Private/GeometryCollection/GeometryCollectionComponent.cpp
--- a/Engine/Source/Runtime/Experimental/GeometryCollectionEngine/Private/GeometryCollection/GeometryCollectionComponent.cpp
+++ b/Engine/Source/Runtime/Experimental/GeometryCollectionEngine/Private/GeometryCollection/GeometryCollectionComponent.cpp
@@ -53,6 +53,7 @@
 #include "PrimitiveSceneInfo.h"
 #include "GeometryCollection/GeometryCollectionEngineRemoval.h"
 #include "GeometryCollection/Facades/CollectionAnchoringFacade.h"
+#include "GeometryCollection/Facades/CollectionRemoveOnBreakFacade.h"
 
 #include UE_INLINE_GENERATED_CPP_BY_NAME(GeometryCollectionComponent)
 
@@ -2326,8 +2327,7 @@ void UGeometryCollectionComponent::ResetDynamicCollection()
 		}
 	}
 #endif
-	//UE_LOG(UGCC_LOG, Log, TEXT("GeometryCollectionComponent[%p]::ResetDynamicCollection()"), static_cast<const void*>(this));
-	if (bCreateDynamicCollection && RestCollection)
+	if (bCreateDynamicCollection && RestCollection && RestCollection->GetGeometryCollection())
 	{
 		DynamicCollection = MakeUnique<FGeometryDynamicCollection>();
 		for (const auto DynamicArray : CopyOnWriteAttributeList)
@@ -2349,16 +2349,20 @@ void UGeometryCollectionComponent::ResetDynamicCollection()
 			DecayDynamicFacade.AddAttributes();
 
 			FGeometryCollectionRemoveOnSleepDynamicFacade RemoveOnSleepDynamicFacade(*DynamicCollection);
-			RemoveOnSleepDynamicFacade.AddAttributes(RestCollection->MaximumSleepTime, RestCollection->RemovalDuration);
+			RemoveOnSleepDynamicFacade.DefineSchema();
+			RemoveOnSleepDynamicFacade.SetAttributeValues(RestCollection->MaximumSleepTime, RestCollection->RemovalDuration);
 		}
 		
 		// Remove on break feature related dynamic attribute arrays
-		if (const TManagedArray<FVector4f>* RemoveOnBreak = RestCollection->GetGeometryCollection()->FindAttribute<FVector4f>("RemoveOnBreak", FGeometryCollection::TransformGroup))
+		// we are not testing for bAllowRemovalOnBreak, so that we can enable it at runtime if necessary
+		GeometryCollection::Facades::FCollectionRemoveOnBreakFacade RemoveOnBreakFacade(*RestCollection->GetGeometryCollection());
+		if (RemoveOnBreakFacade.IsValid())
 		{
 			DecayDynamicFacade.AddAttributes();
 
 			FGeometryCollectionRemoveOnBreakDynamicFacade RemoveOnBreakDynamicFacade(*DynamicCollection);
-			RemoveOnBreakDynamicFacade.AddAttributes(*RemoveOnBreak, DynamicCollection->Children);
+			RemoveOnBreakDynamicFacade.DefineSchema();
+			RemoveOnBreakDynamicFacade.SetAttributeValues(RemoveOnBreakFacade);
 		}
 
 		SetRenderStateDirty();
diff --git a/Engine/Source/Runtime/Experimental/GeometryCollectionEngine/Private/GeometryCollection/GeometryCollectionEngineRemoval.cpp b/Engine/Source/Runtime/Experimental/GeometryCollectionEngine/Private/GeometryCollection/GeometryCollectionEngineRemoval.cpp
--- a/Engine/Source/Runtime/Experimental/GeometryCollectionEngine/Private/GeometryCollection/GeometryCollectionEngineRemoval.cpp
+++ b/Engine/Source/Runtime/Experimental/GeometryCollectionEngine/Private/GeometryCollection/GeometryCollectionEngineRemoval.cpp
@@ -2,6 +2,7 @@
 
 #include "GeometryCollection/GeometryCollectionEngineRemoval.h"
 #include "GeometryCollection/GeometryCollection.h"
+#include "GeometryCollection/Facades/CollectionRemoveOnBreakFacade.h"
 
 ///////////////////////////////////////////////////////////////////////////////////////////////////////////////////////////
 
@@ -9,6 +10,8 @@ FGeometryCollectionRemoveOnBreakDynamicFacade::FGeometryCollectionRemoveOnBreakD
 	: BreakTimerAttribute(InCollection, "BreakTimer", FGeometryCollection::TransformGroup)
 	, PostBreakDurationAttribute(InCollection, "PostBreakDuration", FGeometryCollection::TransformGroup)
 	, BreakRemovalDurationAttribute(InCollection, "BreakRemovalDuration", FGeometryCollection::TransformGroup)
+	, ChildrenAttribute(InCollection, "Children", FTransformCollection::TransformGroup)
+
 {
 }
 
@@ -17,38 +20,52 @@ bool FGeometryCollectionRemoveOnBreakDynamicFacade::IsValid() const
 	return BreakTimerAttribute.IsValid()
 		&& PostBreakDurationAttribute.IsValid()
 		&& BreakRemovalDurationAttribute.IsValid()
+		&& ChildrenAttribute.IsValid();
 	;
 }
 
-void FGeometryCollectionRemoveOnBreakDynamicFacade::AddAttributes(const TManagedArray<FVector4f>& RemoveOnBreakAttribute, const TManagedArray<TSet<int32>>& ChildrenAttribute)
+bool FGeometryCollectionRemoveOnBreakDynamicFacade::IsConst() const
 {
+	return BreakTimerAttribute.IsConst();
+}
+
+void FGeometryCollectionRemoveOnBreakDynamicFacade::DefineSchema()
+{
+	check(!IsConst());
 	BreakTimerAttribute.AddAndFill(DisabledBreakTimer);
+	PostBreakDurationAttribute.Add();
+	BreakRemovalDurationAttribute.Add();
+}
 
-	if (!PostBreakDurationAttribute.IsValid())
+void FGeometryCollectionRemoveOnBreakDynamicFacade::SetAttributeValues(const GeometryCollection::Facades::FCollectionRemoveOnBreakFacade& RemoveOnBreakFacade)
+{
+	check(!IsConst());
+	if (IsValid())
 	{
-		TManagedArray<float>& PostBreakDuration = PostBreakDurationAttribute.Add();
+		const TManagedArray<TSet<int32>>& Children = ChildrenAttribute.Get();
+
+		BreakTimerAttribute.Fill(DisabledBreakTimer);
+
+		TManagedArray<float>& PostBreakDuration = PostBreakDurationAttribute.Modify();
 		for (int32 Idx = 0; Idx < PostBreakDuration.Num(); ++Idx)
 		{
-			const FRemoveOnBreakData RemoveOnBreakData{ RemoveOnBreakAttribute[Idx] };
+			const GeometryCollection::Facades::FRemoveOnBreakData RemoveOnBreakData{ RemoveOnBreakFacade.GetData(Idx) };
 			const FVector2f PostBreakTimer = RemoveOnBreakData.GetBreakTimer();
-			const float MinTime = FMath::Max(0.0f, PostBreakTimer.X);
-			const float MaxTime = FMath::Max(MinTime, PostBreakTimer.Y);
-			PostBreakDuration[Idx] = RemoveOnBreakData.IsEnabled()? FMath::RandRange(MinTime, MaxTime): DisabledPostBreakDuration;
+			const float MinBreakTime = FMath::Max(0.0f, PostBreakTimer.X);
+			const float MaxBreakTime = FMath::Max(MinBreakTime, PostBreakTimer.Y);
+			PostBreakDuration[Idx] = RemoveOnBreakData.IsEnabled() ? FMath::RandRange(MinBreakTime, MaxBreakTime) : DisabledPostBreakDuration;
 		}
-	}
 
-	if (!BreakRemovalDurationAttribute.IsValid())
-	{
-		TManagedArray<float>& BreakRemovalDuration = BreakRemovalDurationAttribute.Add();
-		for (int32 Idx = 0; Idx < BreakRemovalDuration.Num(); ++Idx)
+		TManagedArray<float>& BreakRemovalDuration = BreakRemovalDurationAttribute.Modify();
+		for (int32 Idx = 0; Idx < PostBreakDuration.Num(); ++Idx)
 		{
-			const FRemoveOnBreakData RemoveOnBreakData{ RemoveOnBreakAttribute[Idx] };
+			const GeometryCollection::Facades::FRemoveOnBreakData RemoveOnBreakData{ RemoveOnBreakFacade.GetData(Idx) };
 			const FVector2f RemovalTimer = RemoveOnBreakData.GetRemovalTimer();
-			const float MinTime = FMath::Max(0.0f, RemovalTimer.X);
-			const float MaxTime = FMath::Max(MinTime, RemovalTimer.Y);
-			const bool bIsCluster = (ChildrenAttribute[Idx].Num() > 0);
+			const float MinRemovalTime = FMath::Max(0.0f, RemovalTimer.X);
+			const float MaxRemovalTime = FMath::Max(MinRemovalTime, RemovalTimer.Y);
+			const bool bIsCluster = (Children[Idx].Num() > 0);
 			const bool bUseClusterCrumbling = (bIsCluster && RemoveOnBreakData.GetClusterCrumbling());
-			BreakRemovalDuration[Idx] = bUseClusterCrumbling? CrumblingRemovalTimer: FMath::RandRange(MinTime, MaxTime);
+			BreakRemovalDuration[Idx] = bUseClusterCrumbling ? CrumblingRemovalTimer : FMath::RandRange(MinRemovalTime, MaxRemovalTime);
 		}
 	}
 }
@@ -103,31 +120,42 @@ bool FGeometryCollectionRemoveOnSleepDynamicFacade::IsValid() const
 	;
 }
 
-void FGeometryCollectionRemoveOnSleepDynamicFacade::AddAttributes(const FVector2D& MaximumSleepTime, const FVector2D& RemovalDuration)
+bool FGeometryCollectionRemoveOnSleepDynamicFacade::IsConst() const
 {
-	SleepTimerAttribute.AddAndFill(0.0f);
+	return SleepTimerAttribute.IsConst();
+}
 
+void FGeometryCollectionRemoveOnSleepDynamicFacade::DefineSchema()
+{
+	check(!IsConst());
+	SleepTimerAttribute.Add();
+	MaxSleepTimeAttribute.Add();
+	SleepRemovalDurationAttribute.Add();
 	LastPositionAttribute.AddAndFill(FVector::ZeroVector);
-	
-	if (!MaxSleepTimeAttribute.IsValid())
+}
+
+void FGeometryCollectionRemoveOnSleepDynamicFacade::SetAttributeValues(const FVector2D& MaximumSleepTime, const FVector2D& RemovalDuration)
+{
+	check(!IsConst());
+	if (IsValid())
 	{
+		SleepTimerAttribute.Fill(0.0f);
+		LastPositionAttribute.Fill(FVector::ZeroVector);
+
+		TManagedArray<float>& MaxSleepTime = MaxSleepTimeAttribute.Modify();
 		const float MinTime = FMath::Max(0.0f, MaximumSleepTime.X);
 		const float MaxTime = FMath::Max(MinTime, MaximumSleepTime.Y);
-		TManagedArray<float>& MaxSleepTime = MaxSleepTimeAttribute.Add();
 		for (int32 Idx = 0; Idx < MaxSleepTime.Num(); ++Idx)
 		{
 			MaxSleepTime[Idx] = FMath::RandRange(MinTime, MaxTime);
 		}
-	}
 
-	if (!SleepRemovalDurationAttribute.IsValid())
-	{
-		float MinTime = FMath::Max(0.0f, RemovalDuration.X);
-		float MaxTime = FMath::Max(MinTime, RemovalDuration.Y);
-		TManagedArray<float>& SleepRemovalDuration = SleepRemovalDurationAttribute.Add();
+		TManagedArray<float>& SleepRemovalDuration = SleepRemovalDurationAttribute.Modify();
+		float MinRemovalTime = FMath::Max(0.0f, RemovalDuration.X);
+		float MaxRemovalTime = FMath::Max(MinRemovalTime, RemovalDuration.Y);
 		for (int32 Idx = 0; Idx < SleepRemovalDuration.Num(); ++Idx)
 		{
-			SleepRemovalDuration[Idx] = FMath::RandRange(MinTime, MaxTime);
+			SleepRemovalDuration[Idx] = FMath::RandRange(MinRemovalTime, MaxRemovalTime);
 		}
 	}
 }
diff --git a/Engine/Source/Runtime/Experimental/GeometryCollectionEngine/Public/GeometryCollection/GeometryCollectionEngineRemoval.h b/Engine/Source/Runtime/Experimental/GeometryCollectionEngine/Public/GeometryCollection/GeometryCollectionEngineRemoval.h
--- a/Engine/Source/Runtime/Experimental/GeometryCollectionEngine/Public/GeometryCollection/GeometryCollectionEngineRemoval.h
+++ b/Engine/Source/Runtime/Experimental/GeometryCollectionEngine/Public/GeometryCollection/GeometryCollectionEngineRemoval.h
@@ -7,43 +7,10 @@
 
 struct FManagedArrayCollection;
 
-/**
- * this class helps reading and writing the packed data used in the managed array
- */
-struct FRemoveOnBreakData
+namespace GeometryCollection::Facades
 {
-public:
-	inline static const FVector4f DisabledPackedData{ -1, 0, 0, 0 };   
-	
-	FRemoveOnBreakData()
-		: PackedData(FRemoveOnBreakData::DisabledPackedData)
-	{}
-
-	FRemoveOnBreakData(const FVector4f& InPackedData)
-		: PackedData(InPackedData)
-	{}
-
-	FRemoveOnBreakData(bool bEnable, const FVector2f& BreakTimer, bool bClusterCrumbling, const FVector2f& RemovalTimer)
-	{
-		PackedData.X = FMath::Min(FMath::Abs(BreakTimer.X), FMath::Abs(BreakTimer.Y)); // Min break timer
-		PackedData.Y = FMath::Max(FMath::Abs(BreakTimer.X), FMath::Abs(BreakTimer.Y)); // Max break timer
-		PackedData.Z = FMath::Min(FMath::Abs(RemovalTimer.X), FMath::Abs(RemovalTimer.Y)); // Min removal timer
-		PackedData.W = FMath::Max(FMath::Abs(RemovalTimer.X), FMath::Abs(RemovalTimer.Y)); // Max removal timer
-
-		PackedData.X = bEnable? PackedData.X: -1.0;
-		PackedData.Z = bClusterCrumbling? -1.0: PackedData.Z;
-	}
-	
-	bool IsEnabled() const { return PackedData.X >= 0; }
-	bool GetClusterCrumbling() const { return IsEnabled() && (PackedData.Z < 0); } 
-	FVector2f GetBreakTimer() const	{ return FVector2f(FMath::Abs(PackedData.X), FMath::Abs(PackedData.Y)); }
-	FVector2f GetRemovalTimer() const { return FVector2f(FMath::Abs(PackedData.Z), FMath::Abs(PackedData.W)); }
-
-	const FVector4f& GetPackedData() const { return PackedData; }
-
-private:
-	FVector4f PackedData;
-};
+	class FCollectionRemoveOnBreakFacade;
+}
 
 /**
  * Provides an API for the run time aspect of the remove on break feature
@@ -64,12 +31,17 @@ public:
 	 */
 	bool IsValid() const;	
 
+	/** Is this facade const access */
+	bool IsConst() const;
+
+	/** Add the relevant attributes */
+	void DefineSchema();
+
 	/**
 	 * Add the necessary attributes if they are missing and initialize them if necessary
-	 * @param RemoveOnBreakAttribute remove on break attribute from the rest collection
-	 * @param ChildrenAttribute remove on break attribute from the rest collection
+	 * @param RemoveOnBreakFacade remove on break facade from the rest collection that contain the original user set attributes
 	*/
-	void AddAttributes(const TManagedArray<FVector4f>& RemoveOnBreakAttribute, const TManagedArray<TSet<int32>>& ChildrenAttribute);
+	void SetAttributeValues(const GeometryCollection::Facades::FCollectionRemoveOnBreakFacade& RemoveOnBreakFacade);
 
 	/** true if the removal is active for a specific piece */
 	bool IsRemovalActive(int32 TransformIndex) const;
@@ -94,6 +66,9 @@ private:
 	
 	/** removal duration */
 	TManagedArrayAccessor<float> BreakRemovalDurationAttribute;
+
+	/** utility attribute needed to compute the values and know what's a cluster */
+	TManagedArrayAccessor<TSet<int32>> ChildrenAttribute;
 };
 
 /**
@@ -111,12 +86,18 @@ public:
 	 */
 	bool IsValid() const;	
 
+	/** Is this facade const access */
+	bool IsConst() const;
+
+	/** Add the relevant attributes */
+	void DefineSchema();
+
 	/**
 	 * Add the necessary attributes if they are missing and initialize them if necessary
 	 * @param MaximumSleepTime range of time to initialize the sleep duration
 	 * @param RemovalDuration range of time to initialize the removal duration
 	*/
-	void AddAttributes(const FVector2D& MaximumSleepTime, const FVector2D& RemovalDuration);
+	void SetAttributeValues(const FVector2D& MaximumSleepTime, const FVector2D& RemovalDuration);
 
 	/** true if the removal is active for a specific piece */
 	bool IsRemovalActive(int32 TransformIndex) const;
