Chaos - various gather optimizations
- collision solver manifold points now allocated separately (previously there was a fixed array of 4 manifold points)
- no default initialization of low-level solver structs
- remove SolverBodyAdapter and CollisionSolverAdapter (now have simple C-style gather methods)
- seprarate collision results from friction anchors
- add over-allocating array growth policy to solvers
- add prefetch to collision gather
- reduce size of SolverBody slightly (V and W in floats)
- move collision activation to run in parallel
- optimize search for friction anchor (AssignSavedManifoldPoints)
- clean up material property resetting following contact modification
- begin deprecation of FParticleUtilities (WIP)

diff --git a/Engine/Source/Programs/HeadlessChaos/Private/HeadlessChaosTestJointSolver.cpp b/Engine/Source/Programs/HeadlessChaos/Private/HeadlessChaosTestJointSolver.cpp
--- a/Engine/Source/Programs/HeadlessChaos/Private/HeadlessChaosTestJointSolver.cpp
+++ b/Engine/Source/Programs/HeadlessChaos/Private/HeadlessChaosTestJointSolver.cpp
@@ -45,8 +45,8 @@ namespace ChaosTest
 
 		FJointSolverTest()
 			: Solver()
-			, Body0()
-			, Body1()
+			, Body0(FSolverBody::MakeInitialized())
+			, Body1(FSolverBody::MakeInitialized())
 			, NumPairIts(1)
 			, Gravity(FVec3(0))
 			, Mass0(0)
diff --git a/Engine/Source/Programs/HeadlessChaos/Private/HeadlessChaosTestJoints.cpp b/Engine/Source/Programs/HeadlessChaos/Private/HeadlessChaosTestJoints.cpp
--- a/Engine/Source/Programs/HeadlessChaos/Private/HeadlessChaosTestJoints.cpp
+++ b/Engine/Source/Programs/HeadlessChaos/Private/HeadlessChaosTestJoints.cpp
@@ -359,7 +359,6 @@ namespace ChaosTest {
 
 			Test.GetParticle(0)->X() = RootPosition;
 
-			Test.Evolution.GetCollisionDetector().SetBoundsVelocityInflation(1);
 			Test.Evolution.AdvanceOneTimeStep(Dt);
 			Test.Evolution.EndFrame(Dt);
 
@@ -433,7 +432,6 @@ namespace ChaosTest {
 		}
 
 		Test.Create();
-		Test.Evolution.GetCollisionDetector().SetBoundsVelocityInflation(1);
 
 		const FVec3 Box2LocalSpaceJointPosition = Test.JointPositions[0] - Test.ParticlePositions[1];
 
diff --git a/Engine/Source/Programs/HeadlessChaos/Private/HeadlessChaosTestMinEvolution.cpp b/Engine/Source/Programs/HeadlessChaos/Private/HeadlessChaosTestMinEvolution.cpp
--- a/Engine/Source/Programs/HeadlessChaos/Private/HeadlessChaosTestMinEvolution.cpp
+++ b/Engine/Source/Programs/HeadlessChaos/Private/HeadlessChaosTestMinEvolution.cpp
@@ -48,7 +48,7 @@ namespace ChaosTest
 
 		// Evolution
 		// @todo(ccaulfield): this should start with some reasonable default iterations
-		FPBDMinEvolution Evolution(ParticlesContainer, ParticlePrevXs, ParticlePrevRs, CollisionDetector, 0);
+		FPBDMinEvolution Evolution(ParticlesContainer, ParticlePrevXs, ParticlePrevRs, CollisionDetector);
 		Evolution.SetNumPositionIterations(6);
 		Evolution.SetNumVelocityIterations(1);
 		Evolution.SetNumProjectionIterations(1);
diff --git a/Engine/Source/Programs/HeadlessChaos/Private/HeadlessChaosTestSim.cpp b/Engine/Source/Programs/HeadlessChaos/Private/HeadlessChaosTestSim.cpp
--- a/Engine/Source/Programs/HeadlessChaos/Private/HeadlessChaosTestSim.cpp
+++ b/Engine/Source/Programs/HeadlessChaos/Private/HeadlessChaosTestSim.cpp
@@ -387,8 +387,11 @@ namespace ChaosTest {
 
 		// Particle 1 should have fallen asleep
 		EXPECT_EQ(SleepData.Num(), 1);
-		EXPECT_EQ(SleepData[0].Particle, Dynamic1);
-		EXPECT_TRUE(SleepData[0].Sleeping);
+		if (SleepData.Num() >= 1)
+		{
+			EXPECT_EQ(SleepData[0].Particle, Dynamic1);
+			EXPECT_TRUE(SleepData[0].Sleeping);
+		}
 
 		// Switch the sleep types and observe state changes and sleep events
 		Dynamic1->SetSleepType(ESleepType::NeverSleep);
@@ -396,15 +399,21 @@ namespace ChaosTest {
 
 		// Particle 1 should have woken up due to the sleep type change
 		EXPECT_EQ(SleepData.Num(), 2);
-		EXPECT_EQ(SleepData[1].Particle, Dynamic1);
-		EXPECT_FALSE(SleepData[1].Sleeping);
+		if (SleepData.Num() >= 2)
+		{
+			EXPECT_EQ(SleepData[1].Particle, Dynamic1);
+			EXPECT_FALSE(SleepData[1].Sleeping);
+		}
 
 		AdvanceSleepStates(Dynamic2, Dynamic1);
 
 		// Particle 2 should have fallen asleep
 		EXPECT_EQ(SleepData.Num(), 3);
-		EXPECT_EQ(SleepData[2].Particle, Dynamic2);
-		EXPECT_TRUE(SleepData[2].Sleeping);
+		if (SleepData.Num() >= 3)
+		{
+			EXPECT_EQ(SleepData[2].Particle, Dynamic2);
+			EXPECT_TRUE(SleepData[2].Sleeping);
+		}
 	}
 }
 
diff --git a/Engine/Source/Programs/HeadlessChaos/Private/HeadlessChaosTestSolverBody.cpp b/Engine/Source/Programs/HeadlessChaos/Private/HeadlessChaosTestSolverBody.cpp
--- a/Engine/Source/Programs/HeadlessChaos/Private/HeadlessChaosTestSolverBody.cpp
+++ b/Engine/Source/Programs/HeadlessChaos/Private/HeadlessChaosTestSolverBody.cpp
@@ -25,7 +25,7 @@ namespace ChaosTest
 		const FReal M0 = 100.0f;
 		const FReal I0 = M0 * 1000.0f;
 
-		FSolverBody SolverBody0;
+		FSolverBody SolverBody0 = FSolverBody::MakeInitialized();
 		SolverBody0.SetX(X0);
 		SolverBody0.SetP(X0);
 		SolverBody0.SetInvM(1.0f / M0);
diff --git a/Engine/Source/Programs/HeadlessChaos/Public/HeadlessChaosCollisionConstraints.h b/Engine/Source/Programs/HeadlessChaos/Public/HeadlessChaosCollisionConstraints.h
--- a/Engine/Source/Programs/HeadlessChaos/Public/HeadlessChaosCollisionConstraints.h
+++ b/Engine/Source/Programs/HeadlessChaos/Public/HeadlessChaosCollisionConstraints.h
@@ -45,7 +45,7 @@ public:
 	{
 		CollisionConstraints.SetContainerId(0);
 
-		CollisionDetector.SetBoundsExpansion(FReal(1));
+		CollisionConstraints.SetCullDistance(FReal(1));
 
 		ConstraintSolver.SetConstraintSolver(CollisionConstraints.GetContainerId(), CollisionConstraints.CreateSceneSolver(0));
 	}
@@ -61,7 +61,7 @@ public:
 	{
 		CollisionConstraints.SetContainerId(0);
 
-		CollisionDetector.SetBoundsExpansion(FReal(1));
+		CollisionConstraints.SetCullDistance(FReal(1));
 
 		ConstraintSolver.SetConstraintSolver(CollisionConstraints.GetContainerId(), CollisionConstraints.CreateSceneSolver(0));
 	}
@@ -72,11 +72,11 @@ public:
 	{
 		CollisionDetector.GetBroadPhase().SetSpatialAcceleration(&SpatialAcceleration);
 
-		FCollisionDetectorSettings DetectorSettings = CollisionDetector.GetSettings();
+		FCollisionDetectorSettings DetectorSettings = CollisionConstraints.GetDetectorSettings();
 		DetectorSettings.bFilteringEnabled = true;
 		DetectorSettings.bDeferNarrowPhase = false;
 		DetectorSettings.bAllowManifolds = true;
-		CollisionDetector.SetSettings(DetectorSettings);
+		CollisionConstraints.SetDetectorSettings(DetectorSettings);
 		
 		CollisionDetector.DetectCollisions(Dt, nullptr);
 		
diff --git a/Engine/Source/Programs/HeadlessChaos/Public/HeadlessChaosTestUtility.h b/Engine/Source/Programs/HeadlessChaos/Public/HeadlessChaosTestUtility.h
--- a/Engine/Source/Programs/HeadlessChaos/Public/HeadlessChaosTestUtility.h
+++ b/Engine/Source/Programs/HeadlessChaos/Public/HeadlessChaosTestUtility.h
@@ -149,11 +149,11 @@ namespace ChaosTest {
 	{
 		// Settings used for unit tests
 		const float CullDistance = 3.0f;
-		FCollisionDetectorSettings DetectorSettings = Evolution.GetCollisionDetector().GetSettings();
+		FCollisionDetectorSettings DetectorSettings = Evolution.GetCollisionConstraints().GetDetectorSettings();
 		DetectorSettings.BoundsExpansion = CullDistance;
 		DetectorSettings.bDeferNarrowPhase = false;
 		DetectorSettings.bAllowManifolds = true;
-		Evolution.GetCollisionDetector().SetSettings(DetectorSettings);
+		Evolution.GetCollisionConstraints().SetDetectorSettings(DetectorSettings);
 	}
 
 	template<typename T_SOLVER>
diff --git a/Engine/Source/Runtime/Engine/Private/PhysicsEngine/ImmediatePhysics/ImmediatePhysicsChaos/ImmediatePhysicsActorHandle_Chaos.cpp b/Engine/Source/Runtime/Engine/Private/PhysicsEngine/ImmediatePhysics/ImmediatePhysicsChaos/ImmediatePhysicsActorHandle_Chaos.cpp
--- a/Engine/Source/Runtime/Engine/Private/PhysicsEngine/ImmediatePhysics/ImmediatePhysicsChaos/ImmediatePhysicsActorHandle_Chaos.cpp
+++ b/Engine/Source/Runtime/Engine/Private/PhysicsEngine/ImmediatePhysics/ImmediatePhysicsChaos/ImmediatePhysicsActorHandle_Chaos.cpp
@@ -141,11 +141,7 @@ namespace ImmediatePhysics_Chaos
 		AddParams.Scale = Scale;
 		//AddParams.SimpleMaterial = SimpleMaterial;
 		//AddParams.ComplexMaterials = TArrayView<UPhysicalMaterial*>(ComplexMaterials);
-#if CHAOS_PARTICLE_ACTORTRANSFORM
 		AddParams.LocalTransform = FTransform::Identity;
-#else
-		AddParams.LocalTransform = FRigidTransform3(OutCoMTransform.GetRotation().Inverse() * -OutCoMTransform.GetTranslation(), OutCoMTransform.GetRotation().Inverse());
-#endif
 		AddParams.WorldTransform = BodyInstance->GetUnrealWorldTransform();
 		AddParams.Geometry = &BodySetup->AggGeom;
 		AddParams.ChaosTriMeshes = MakeArrayView(BodySetup->ChaosTriMeshes);
@@ -353,7 +349,8 @@ namespace ImmediatePhysics_Chaos
 	{
 		using namespace Chaos;
 
-		FParticleUtilitiesXR::SetActorWorldTransform(FGenericParticleHandle(ParticleHandle), WorldTM);
+		ParticleHandle->X() = WorldTM.GetTranslation();
+		ParticleHandle->R() = WorldTM.GetRotation();
 
 		auto* Dynamic = ParticleHandle->CastToRigidParticle();
 		if(Dynamic && Dynamic->ObjectState() == Chaos::EObjectStateType::Dynamic)
diff --git a/Engine/Source/Runtime/Engine/Private/PhysicsEngine/ImmediatePhysics/ImmediatePhysicsChaos/ImmediatePhysicsSimulation_Chaos.cpp b/Engine/Source/Runtime/Engine/Private/PhysicsEngine/ImmediatePhysics/ImmediatePhysicsChaos/ImmediatePhysicsSimulation_Chaos.cpp
--- a/Engine/Source/Runtime/Engine/Private/PhysicsEngine/ImmediatePhysics/ImmediatePhysicsChaos/ImmediatePhysicsSimulation_Chaos.cpp
+++ b/Engine/Source/Runtime/Engine/Private/PhysicsEngine/ImmediatePhysics/ImmediatePhysicsChaos/ImmediatePhysicsSimulation_Chaos.cpp
@@ -48,7 +48,6 @@ int32 ChaosImmediate_Evolution_VelocityIterations = -1;
 int32 ChaosImmediate_Evolution_ProjectionIterations = -1;
 int32 ChaosImmediate_Evolution_Iterations = -1;				// Legacy
 int32 ChaosImmediate_Evolution_PushOutIterations = -1;		// Legacy
-Chaos::FRealSingle ChaosImmediate_Evolution_BoundsExtension = 0.0f;
 int32 ChaosImmediate_DisableInactiveByIndex = 1;
 int32 ChaosImmediate_Evolution_NumCollisionsPerBlock = 50;
 FAutoConsoleVariableRef CVarChaosImmPhysStepTime(TEXT("p.Chaos.ImmPhys.StepTime"), ChaosImmediate_Evolution_StepTime, TEXT("Override step time (if not zero)"));
@@ -60,7 +59,6 @@ FAutoConsoleVariableRef CVarChaosImmPhysVelocityIterations(TEXT("p.Chaos.ImmPhys
 FAutoConsoleVariableRef CVarChaosImmPhysProjectionIterations(TEXT("p.Chaos.ImmPhys.ProjectionIterations"), ChaosImmediate_Evolution_ProjectionIterations, TEXT("Override number of projection iteration loops in immediate physics (if >= 0)"));
 FAutoConsoleVariableRef CVarChaosImmPhysIterations(TEXT("p.Chaos.ImmPhys.Iterations"), ChaosImmediate_Evolution_Iterations, TEXT("[Legacy Solver] Override number of constraint solver loops in immediate physics (if >= 0)"));
 FAutoConsoleVariableRef CVarChaosImmPhysPushOutIterations(TEXT("p.Chaos.ImmPhys.PushOutIterations"), ChaosImmediate_Evolution_PushOutIterations, TEXT("[Legacy Solver] Override number of solver push-out loops (if >= 0)"));
-FAutoConsoleVariableRef CVarChaosImmPhysBoundsExtension(TEXT("p.Chaos.ImmPhys.BoundsExtension"), ChaosImmediate_Evolution_BoundsExtension, TEXT("Bounds are grown by this fraction of their size (should be >= 0.0)"));
 FAutoConsoleVariableRef CVarChaosImmPhysDisableInactiveByIndex(TEXT("p.Chaos.ImmPhys.DisableInactiveByIndex"), ChaosImmediate_DisableInactiveByIndex, TEXT("Disable bodies that are no longer active based on the index, rather than just count."));
 FAutoConsoleVariableRef CVarChaosImmPhysNumCollisionsPerBlock(TEXT("p.Chaos.ImmPhys.NumCollisionsPerBlock"), ChaosImmediate_Evolution_NumCollisionsPerBlock, TEXT("The number of collision in a block in the collision pool. Higher values give better cache efficieny but waste memory if you do not need that many"));
 
@@ -301,7 +299,7 @@ namespace ImmediatePhysics_Chaos
 			, Collisions(Particles, CollidedParticles, ParticleMaterials, PerParticleMaterials, nullptr, ChaosImmediate_Evolution_NumCollisionsPerBlock, 2000)
 			, BroadPhase(&ActivePotentiallyCollidingPairs, nullptr, nullptr)
 			, CollisionDetector(BroadPhase, Collisions)
-			, Evolution(Particles, ParticlePrevXs, ParticlePrevRs, CollisionDetector, FReal(0))
+			, Evolution(Particles, ParticlePrevXs, ParticlePrevRs, CollisionDetector)
 			, NumActiveDynamicActorHandles(0)
 			, SimulationSpace()
 			, RollingAverageStepTime(ChaosImmediate_Evolution_InitialStepTime)
@@ -377,13 +375,13 @@ namespace ImmediatePhysics_Chaos
 		// RBAN collision customization
 		Implementation->Collisions.DisableHandles();
 
-		FCollisionDetectorSettings DetectorSettings = Implementation->CollisionDetector.GetSettings();
+		FCollisionDetectorSettings DetectorSettings = Implementation->Collisions.GetDetectorSettings();
 		DetectorSettings.bFilteringEnabled = false;
 		DetectorSettings.bAllowManifoldReuse = false;
-		DetectorSettings.bDeferNarrowPhase = (ChaosImmediate_Collision_DeferNarrowPhase != 0);;
-		DetectorSettings.bAllowManifolds = (ChaosImmediate_Collision_UseManifolds != 0);;
+		DetectorSettings.bDeferNarrowPhase = (ChaosImmediate_Collision_DeferNarrowPhase != 0);
+		DetectorSettings.bAllowManifolds = (ChaosImmediate_Collision_UseManifolds != 0);
 		DetectorSettings.bAllowCCD = false; 
-		Implementation->CollisionDetector.SetSettings(DetectorSettings);
+		Implementation->Collisions.SetDetectorSettings(DetectorSettings);
 	}
 
 	FSimulation::~FSimulation()
@@ -760,7 +758,7 @@ namespace ImmediatePhysics_Chaos
 
 		if (CullDistance >= FReal(0))
 		{
-			Implementation->CollisionDetector.SetBoundsExpansion(CullDistance);
+			Implementation->Collisions.SetCullDistance(CullDistance);
 		}
 
 		if (MaxDepenetrationVelocity >= FReal(0))
@@ -871,14 +869,11 @@ namespace ImmediatePhysics_Chaos
 
 			Implementation->Evolution.SetConstraintContainerPriority(Implementation->Collisions.GetContainerId(), ChaosImmediate_Collision_Priority);
 
-			FCollisionDetectorSettings DetectorSettings = Implementation->CollisionDetector.GetSettings();
+			FCollisionDetectorSettings DetectorSettings = Implementation->Collisions.GetDetectorSettings();
 			DetectorSettings.bAllowManifoldReuse = false;
 			DetectorSettings.bDeferNarrowPhase = (ChaosImmediate_Collision_DeferNarrowPhase != 0);;
 			DetectorSettings.bAllowManifolds = (ChaosImmediate_Collision_UseManifolds != 0);;
-			Implementation->CollisionDetector.SetSettings(DetectorSettings);
-
-			Implementation->Evolution.SetBoundsExtension(ChaosImmediate_Evolution_BoundsExtension);
-
+			Implementation->Collisions.SetDetectorSettings(DetectorSettings);
 
 			if (ChaosImmediate_Evolution_StepTime > 0)
 			{
diff --git a/Engine/Source/Runtime/Experimental/Chaos/Private/Chaos/ChaosDebugDraw.cpp b/Engine/Source/Runtime/Experimental/Chaos/Private/Chaos/ChaosDebugDraw.cpp
--- a/Engine/Source/Runtime/Experimental/Chaos/Private/Chaos/ChaosDebugDraw.cpp
+++ b/Engine/Source/Runtime/Experimental/Chaos/Private/Chaos/ChaosDebugDraw.cpp
@@ -873,10 +873,13 @@ namespace Chaos
 					return;
 				}
 
-				for (const FManifoldPoint& ManifoldPoint : Contact.GetManifoldPoints())
+				for (int32 PointIndex = 0; PointIndex < Contact.NumManifoldPoints(); ++PointIndex)
 				{
+					const FManifoldPoint& ManifoldPoint = Contact.GetManifoldPoint(PointIndex);
+					const FManifoldPointResult& ManifoldPointResult = Contact.GetManifoldPointResult(PointIndex);
+
 					const bool bIsProbe = Contact.GetIsProbe();
-					const bool bIsActive = !ManifoldPoint.NetPushOut.IsNearlyZero() || !ManifoldPoint.NetImpulse.IsNearlyZero() || (!Contact.GetUseManifold() && !Contact.AccumulatedImpulse.IsNearlyZero());
+					const bool bIsActive = ManifoldPointResult.bIsValid && (!ManifoldPointResult.NetPushOut.IsNearlyZero() || !ManifoldPointResult.NetImpulse.IsNearlyZero() || (!Contact.GetUseManifold() && !Contact.AccumulatedImpulse.IsNearlyZero()));
 					if (!bIsActive && !bChaosDebugDebugDrawInactiveContacts)
 					{
 						continue;
@@ -903,7 +906,7 @@ namespace Chaos
 					FColor ImpulseColor = FColor(0, 0, 200);
 					FColor PushOutColor = FColor(0, 200, 0);
 					FColor PushOutImpusleColor = FColor(0, 200, 200);
-					if (ManifoldPoint.Flags.bInsideStaticFrictionCone)
+					if (ManifoldPointResult.bInsideStaticFrictionCone)
 					{
 						DiscColor = FColor(150, 200, 0);
 					}
@@ -932,15 +935,15 @@ namespace Chaos
 					const FMatrix Axes = FRotationMatrix::MakeFromX(WorldPlaneNormal);
 
 					// Pushout
-					if ((Settings.PushOutScale > 0) && !ManifoldPoint.NetPushOut.IsNearlyZero())
+					if ((Settings.PushOutScale > 0) && ManifoldPointResult.bIsValid && !ManifoldPointResult.NetPushOut.IsNearlyZero())
 					{
 						FColor Color = (ColorScale * PushOutImpusleColor).ToFColor(false);
-						FDebugDrawQueue::GetInstance().DrawDebugLine(WorldPointPlaneLocation, WorldPointPlaneLocation + Settings.DrawScale * Settings.PushOutScale * SpaceTransform.TransformVectorNoScale(FVec3(ManifoldPoint.NetPushOut)), Color, false, Duration, Settings.DrawPriority, Settings.LineThickness);
+						FDebugDrawQueue::GetInstance().DrawDebugLine(WorldPointPlaneLocation, WorldPointPlaneLocation + Settings.DrawScale * Settings.PushOutScale * SpaceTransform.TransformVectorNoScale(FVec3(ManifoldPointResult.NetPushOut)), Color, false, Duration, Settings.DrawPriority, Settings.LineThickness);
 					}
-					if ((Settings.ImpulseScale > 0) && !ManifoldPoint.NetImpulse.IsNearlyZero())
+					if ((Settings.ImpulseScale > 0) && ManifoldPointResult.bIsValid && !ManifoldPointResult.NetImpulse.IsNearlyZero())
 					{
 						FColor Color = (ColorScale * ImpulseColor).ToFColor(false);
-						FDebugDrawQueue::GetInstance().DrawDebugLine(WorldPointPlaneLocation, WorldPointPlaneLocation + Settings.DrawScale * Settings.ImpulseScale * SpaceTransform.TransformVectorNoScale(FVec3(ManifoldPoint.NetImpulse)), Color, false, Duration, Settings.DrawPriority, Settings.LineThickness);
+						FDebugDrawQueue::GetInstance().DrawDebugLine(WorldPointPlaneLocation, WorldPointPlaneLocation + Settings.DrawScale * Settings.ImpulseScale * SpaceTransform.TransformVectorNoScale(FVec3(ManifoldPointResult.NetImpulse)), Color, false, Duration, Settings.DrawPriority, Settings.LineThickness);
 					}
 
 					// Manifold plane and normal
@@ -964,15 +967,6 @@ namespace Chaos
 					// Manifold point
 					FDebugDrawQueue::GetInstance().DrawDebugCircle(WorldPointLocation, 0.5f * Settings.DrawScale * Settings.ContactWidth, 12, DiscColor, false, Duration, Settings.DrawPriority, Settings.LineThickness, Axes.GetUnitAxis(EAxis::Y), Axes.GetUnitAxis(EAxis::Z), false);
 
-					// Previous points
-					if (bIsActive && ManifoldPoint.Flags.bWasFrictionRestored)
-					{
-						const FVec3 WorldPrevPointLocation = SpaceTransform.TransformPosition(PointTransform.TransformPosition(FVec3(ManifoldPoint.ShapeAnchorPoints[ContactPointOwner])));
-						const FVec3 WorldPrevPlaneLocation = SpaceTransform.TransformPosition(PlaneTransform.TransformPosition(FVec3(ManifoldPoint.ShapeAnchorPoints[ContactPlaneOwner])));
-						FDebugDrawQueue::GetInstance().DrawDebugLine(WorldPrevPointLocation, WorldPointLocation, FColor::White, false, Duration, Settings.DrawPriority, Settings.LineThickness);
-						FDebugDrawQueue::GetInstance().DrawDebugLine(WorldPrevPlaneLocation, WorldPlaneLocation, FColor::White, false, Duration, Settings.DrawPriority, Settings.LineThickness);
-					}
-
 					// Whether restored
 					if (Settings.ContactInfoWidth > 0)
 					{
@@ -1039,7 +1033,12 @@ namespace Chaos
 				FConstGenericParticleHandle P0 = CCDParticle->Particle;
 				if (P0->IsDynamic() && P0->CCDEnabled())
 				{
-					const FRigidTransform3 ActorTransform0 = bShowStartPos ? FParticleUtilitiesXR::GetActorWorldTransform(P0): FParticleUtilitiesPQ::GetActorWorldTransform(P0);
+					FRigidTransform3 ActorTransform0 = P0->GetTransformPQ();
+					if (bShowStartPos)
+					{
+						// For CCD we want to see the initial position but with the latest rotation since that's what the sweep uses
+						ActorTransform0.SetTranslation(P0->X());
+					}
 					DrawShapesImpl(P0->Handle(), ActorTransform0 * SpaceTransform, P0->Geometry().Get(), FShapeOrShapesArray(P0->Handle()), 0.0f, ShapeColor, Duration, Settings);
 				}
 			}
@@ -1054,8 +1053,7 @@ namespace Chaos
 				{
 					if (!Impulse.IsNearlyZero() && (Settings.ImpulseScale > 0))
 					{
-						const FRigidTransform3 CoMTransformX0 = FParticleUtilitiesXR::GetCoMWorldTransform(P0);
-						const FVec3 Pos0 = SpaceTransform.TransformPosition(CoMTransformX0.GetTranslation());
+						const FVec3 Pos0 = SpaceTransform.TransformPosition(P0->PCom());
 						FDebugDrawQueue::GetInstance().DrawDebugLine(Pos0, Pos0 + Settings.ImpulseScale * Impulse, FColor::Red, false, Duration, Settings.DrawPriority, Settings.LineThickness * 0.5f);
 					}
 				}
@@ -1322,8 +1320,8 @@ namespace Chaos
 					ContactPos /= (FReal)(2 * Constraint->GetManifoldPoints().Num());
 				}
 
-				const FRigidTransform3 Transform0 = FParticleUtilities::GetCoMWorldTransform(FConstGenericParticleHandle(Constraint->GetConstrainedParticles()[0])) * SpaceTransform;
-				const FRigidTransform3 Transform1 = FParticleUtilities::GetCoMWorldTransform(FConstGenericParticleHandle(Constraint->GetConstrainedParticles()[1])) * SpaceTransform;
+				const FRigidTransform3 Transform0 = FConstGenericParticleHandle(Constraint->GetConstrainedParticles()[0])->GetTransformPQCom() * SpaceTransform;
+				const FRigidTransform3 Transform1 = FConstGenericParticleHandle(Constraint->GetConstrainedParticles()[1])->GetTransformPQCom() * SpaceTransform;
 
 				if ((bChaosDebugDebugDrawContactGraphUsed && bIsUsed) || (bChaosDebugDebugDrawContactGraphUnused && !bIsUsed))
 				{
@@ -1527,7 +1525,7 @@ namespace Chaos
 			}
 			
 			const FVec3& PLocal = Constraints.GetConstraintPosition(ConstraintIndex);
-			const FRigidTransform3 ParticleTransform = FParticleUtilitiesPQ::GetActorWorldTransform(Particle);
+			const FRigidTransform3 ParticleTransform = Particle->GetTransformPQ();
 
 			const FVec3 PWorld = ParticleTransform.TransformPosition(PLocal);
 			const FVec3 AxisWorld = ParticleTransform.TransformVector(ConstraintSettings.Axis);
diff --git a/Engine/Source/Runtime/Experimental/Chaos/Private/Chaos/Collision/CollisionConstraintAllocator.cpp b/Engine/Source/Runtime/Experimental/Chaos/Private/Chaos/Collision/CollisionConstraintAllocator.cpp
--- a/Engine/Source/Runtime/Experimental/Chaos/Private/Chaos/Collision/CollisionConstraintAllocator.cpp
+++ b/Engine/Source/Runtime/Experimental/Chaos/Private/Chaos/Collision/CollisionConstraintAllocator.cpp
@@ -1,8 +1,35 @@
 // Copyright Epic Games, Inc. All Rights Reserved.
 #include "Chaos/Collision/CollisionConstraintAllocator.h"
+#include "Chaos/PBDCollisionConstraints.h"
 
 namespace Chaos
 {
+	FPBDCollisionConstraintPtr FCollisionContextAllocator::CreateConstraint(
+		FGeometryParticleHandle* Particle0,
+		const FImplicitObject* Implicit0,
+		const FPerShapeData* Shape0,
+		const FBVHParticles* Simplicial0,
+		const FRigidTransform3& ShapeRelativeTransform0,
+		FGeometryParticleHandle* Particle1,
+		const FImplicitObject* Implicit1,
+		const FPerShapeData* Shape1,
+		const FBVHParticles* Simplicial1,
+		const FRigidTransform3& ShapeRelativeTransform1,
+		const FReal CullDistance,
+		const bool bUseManifold,
+		const EContactShapesType ShapePairType)
+	{
+		FPBDCollisionConstraintPtr Constraint = CreateConstraint();
+
+		if (Constraint.IsValid())
+		{
+			Constraint->SetContainer(CollisionContainer);
+			FPBDCollisionConstraint::Make(Particle0, Implicit0, Shape0, Simplicial0, ShapeRelativeTransform0, Particle1, Implicit1, Shape1, Simplicial1, ShapeRelativeTransform1, CullDistance, bUseManifold, ShapePairType, *Constraint);
+		}
+
+		return Constraint;
+	}
+
 	void FCollisionContextAllocator::ProcessNewMidPhases(FCollisionConstraintAllocator* Allocator)
 	{
 		// NOTE: Called from the physics thread, and never from a physics task/parallel-for. No need for locks.
diff --git a/Engine/Source/Runtime/Experimental/Chaos/Private/Chaos/Collision/CollisionDetector.cpp b/Engine/Source/Runtime/Experimental/Chaos/Private/Chaos/Collision/CollisionDetector.cpp
--- a/Engine/Source/Runtime/Experimental/Chaos/Private/Chaos/Collision/CollisionDetector.cpp
+++ b/Engine/Source/Runtime/Experimental/Chaos/Private/Chaos/Collision/CollisionDetector.cpp
@@ -1,11 +1,23 @@
 // Copyright Epic Games, Inc. All Rights Reserved.
 #include "Chaos/Collision/CollisionDetector.h"
-#include "Chaos/Collision/CollisionConstraintAllocator.h"
-#include "Chaos/Collision/ParticlePairBroadPhase.h"
-#include "Chaos/Collision/SpatialAccelerationBroadPhase.h"
+#include "Chaos/PBDCollisionConstraints.h"
 #include "ChaosLog.h"
 #include "ChaosStats.h"
 
 namespace Chaos
 {
+	const FCollisionDetectorSettings& FCollisionDetector::GetSettings() const
+	{
+		return CollisionContainer.GetDetectorSettings();
+	}
+
+	void FCollisionDetector::SetSettings(const FCollisionDetectorSettings& InSettings)
+	{
+		CollisionContainer.SetDetectorSettings(InSettings);
+	}
+
+	void FCollisionDetector::SetBoundsExpansion(const FReal InBoundsExpansion)
+	{
+		CollisionContainer.SetCullDistance(InBoundsExpansion);
+	}
 }
\ No newline at end of file
diff --git a/Engine/Source/Runtime/Experimental/Chaos/Private/Chaos/Collision/PBDCollisionConstraint.cpp b/Engine/Source/Runtime/Experimental/Chaos/Private/Chaos/Collision/PBDCollisionConstraint.cpp
--- a/Engine/Source/Runtime/Experimental/Chaos/Private/Chaos/Collision/PBDCollisionConstraint.cpp
+++ b/Engine/Source/Runtime/Experimental/Chaos/Private/Chaos/Collision/PBDCollisionConstraint.cpp
@@ -2,6 +2,7 @@
 #include "Chaos/Collision/PBDCollisionConstraint.h"
 #include "Chaos/Collision/CollisionConstraintAllocator.h"
 #include "Chaos/Collision/PBDCollisionConstraintHandle.h"
+#include "Chaos/Collision/PBDCollisionSolver.h"
 #include "Chaos/Collision/SolverCollisionContainer.h"
 #include "Chaos/Evolution/SolverBody.h"
 #include "Chaos/Evolution/SolverBodyContainer.h"
@@ -10,9 +11,6 @@
 #include "Chaos/Particle/ParticleUtilities.h"
 #include "Chaos/PBDCollisionConstraints.h"
 
-// Private includes
-#include "PBDCollisionSolver.h"
-
 
 //PRAGMA_DISABLE_OPTIMIZATION
 
@@ -182,8 +180,6 @@ namespace Chaos
 		, Implicit{ nullptr, nullptr }
 		, Shape{ nullptr, nullptr }
 		, Simplicial{ nullptr, nullptr }
-		, Material()
-		, Stiffness(1)
 		, AccumulatedImpulse(0)
 		, ContainerCookie()
 		, ShapesType(EContactShapesType::Unknown)
@@ -193,13 +189,15 @@ namespace Chaos
 		, CollisionTolerance(0)
 		, ClosestManifoldPointIndex(INDEX_NONE)
 		, ExpectedNumManifoldPoints(0)
-		, Flags()
 		, LastShapeWorldPositionDelta()
 		, LastShapeWorldRotationDelta()
-		, SolverBodies{ nullptr, nullptr }
 		, GJKWarmStartData()
+		, Material()
+		, Stiffness(1)
+		, Flags()
 		, SavedManifoldPoints()
 		, ManifoldPoints()
+		, SolverBodies{ nullptr, nullptr }
 		, CCDTimeOfImpact(0)
 		, CCDEnablePenetration(0)
 		, CCDTargetPenetration(0)
@@ -220,8 +218,6 @@ namespace Chaos
 		, Implicit{ Implicit0, Implicit1 }
 		, Shape{ Shape0, Shape1 }
 		, Simplicial{ Simplicial0, Simplicial1 }
-		, Material()
-		, Stiffness(1)
 		, AccumulatedImpulse(0)
 		, ContainerCookie()
 		, ShapesType(EContactShapesType::Unknown)
@@ -231,13 +227,15 @@ namespace Chaos
 		, CollisionTolerance(0)
 		, ClosestManifoldPointIndex(INDEX_NONE)
 		, ExpectedNumManifoldPoints(0)
-		, Flags()
 		, LastShapeWorldPositionDelta()
 		, LastShapeWorldRotationDelta()
-		, SolverBodies{ nullptr, nullptr }
 		, GJKWarmStartData()
+		, Material()
+		, Stiffness(1)
+		, Flags()
 		, SavedManifoldPoints()
 		, ManifoldPoints()
+		, SolverBodies{ nullptr, nullptr }
 		, CCDTimeOfImpact(0)
 		, CCDEnablePenetration(0)
 		, CCDTargetPenetration(0)
@@ -392,6 +390,24 @@ namespace Chaos
 		CCDTargetPenetration = FMath::Min(MinBounds * CVars::CCDAllowedDepthBoundsScale, CCDEnablePenetration);
 	}
 
+	void FPBDCollisionConstraint::Activate(const FReal Dt)
+	{
+		SetDisabled(false);
+
+		AssignSavedManifoldPoints();
+
+		UpdateMaterialProperties();
+
+		ResetModifications();
+
+		AccumulatedImpulse = FVec3(0);
+	}
+
+	void FPBDCollisionConstraint::UpdateMaterialPropertiesImpl()
+	{
+		ConcreteContainer()->UpdateConstraintMaterialProperties(*this);
+	}
+
 	bool FPBDCollisionConstraint::IsSleeping() const
 	{
 		if (ContainerCookie.MidPhase != nullptr)
@@ -575,8 +591,6 @@ namespace Chaos
 
 			UpdateManifoldPointPhi(ManifoldPointIndex);
 
-			ManifoldPoint.Flags.bInsideStaticFrictionCone = false;
-
 			if (ManifoldPoint.ContactPoint.Phi < GetPhi())
 			{
 				ClosestManifoldPointIndex = ManifoldPointIndex;
@@ -617,9 +631,8 @@ namespace Chaos
 		{
 			// We are not using manifolds - reuse the first and only point
 			ManifoldPoints.SetNum(1);
-			ManifoldPoints[0].ContactPoint = ContactPoint;
 
-			InitManifoldPoint(0);
+			InitManifoldPoint(0, ContactPoint);
 
 			ClosestManifoldPointIndex = 0;
 		}
@@ -635,6 +648,7 @@ namespace Chaos
 	{
 		ClosestManifoldPointIndex = INDEX_NONE;
 		ManifoldPoints.Reset();
+		ManifoldPointResults.Reset();
 		ExpectedNumManifoldPoints = 0;
 		Flags.bWasManifoldRestored = false;
 	}
@@ -711,11 +725,24 @@ namespace Chaos
 				{
 					// Recalculate the contact points at the new location
 					// @todo(chaos): we should reproject the contact on the plane owner
-					const FVec3 ShapeContactPoint1 = Contact0In1 - ContactPhi * ContactNormalIn1;
-					ManifoldPoint.ContactPoint.ShapeContactPoints[1] = FVec3f(ShapeContactPoint1);
+					const FVec3f ShapeContactPoint1 = FVec3f(Contact0In1 - ContactPhi * ContactNormalIn1);
+					ManifoldPoint.ContactPoint.ShapeContactPoints[1] = ShapeContactPoint1;
 					ManifoldPoint.ContactPoint.Phi = FRealSingle(ContactPhi);
+
+					// Set the friction anchors.
+					// In principle we should always have the same number of saved manifold points as manifold points 
+					// from the previous frame, but there are ways to break this. E.g., Contact modification resets
+					// friction if positions are changed.
+					if (ManifoldPointIndex < SavedManifoldPoints.Num())
+					{
+						FSavedManifoldPoint& SavedManifoldPoint = SavedManifoldPoints[ManifoldPointIndex];
+						ManifoldPoint.ShapeAnchorPoints[0] = SavedManifoldPoint.ShapeContactPoints[0];
+						ManifoldPoint.ShapeAnchorPoints[1] = SavedManifoldPoint.ShapeContactPoints[1];
+					}
+
 					ManifoldPoint.Flags.bWasRestored = true;
 					ManifoldPoint.Flags.bWasReplaced = false;
+
 					if (ManifoldPoint.ContactPoint.Phi < GetPhi())
 					{
 						ClosestManifoldPointIndex = ManifoldPointIndex;
@@ -993,33 +1020,63 @@ namespace Chaos
 		return TNumericLimits<FReal>::Max();
 	}
 
-	const FSavedManifoldPoint* FPBDCollisionConstraint::FindSavedManifoldPoint(const FManifoldPoint& ManifoldPoint) const
+	int32 FPBDCollisionConstraint::FindSavedManifoldPoint(const int32 ManifoldPointIndex, TCArray<int32, MaxManifoldPoints>& InOutAllowedSavedPointIndices) const
 	{
+		int32 MatchIndex = INDEX_NONE;
+
 		if (bChaos_Manifold_EnableFrictionRestore)
 		{
-			const FReal DistanceToleranceSq = FMath::Square(Chaos_Manifold_FrictionPositionTolerance);
-			FReal BestScore = DistanceToleranceSq;
-			int32 MatchIndex = INDEX_NONE;
-			for (int32 SavedManifoldPointIndex = 0; SavedManifoldPointIndex < SavedManifoldPoints.Num(); ++SavedManifoldPointIndex)
+			const FManifoldPoint& ManifoldPoint = ManifoldPoints[ManifoldPointIndex];
+			if (!ManifoldPoint.Flags.bDisabled)
 			{
-				const FSavedManifoldPoint& SavedManifoldPoint = SavedManifoldPoints[SavedManifoldPointIndex];
-				const FReal Score = CalculateSavedManifoldPointScore(SavedManifoldPoint, ManifoldPoint, DistanceToleranceSq);
-				if (Score < BestScore)
+				const FReal DistanceToleranceSq = FMath::Square(Chaos_Manifold_FrictionPositionTolerance);
+				FReal BestScore = DistanceToleranceSq;
+
+				for (int32 AllowedPointIndex = 0; AllowedPointIndex < InOutAllowedSavedPointIndices.Num(); ++AllowedPointIndex)
 				{
-					BestScore = Score;
-					MatchIndex = SavedManifoldPointIndex;
+					const int32 SavedPointIndex = InOutAllowedSavedPointIndices[AllowedPointIndex];
+					const FSavedManifoldPoint& SavedManifoldPoint = SavedManifoldPoints[SavedPointIndex];
 
-					// Just take the first match we find
-					break;
+					const FReal Score = CalculateSavedManifoldPointScore(SavedManifoldPoint, ManifoldPoint, DistanceToleranceSq);
+					if (Score < BestScore)
+					{
+						BestScore = Score;
+						MatchIndex = SavedPointIndex;
+
+						InOutAllowedSavedPointIndices.RemoveAtSwap(AllowedPointIndex);
+
+						// Just take the first match we find that meets our tolerance
+						break;
+					}
 				}
 			}
+		}
+
+		return MatchIndex;
+	}
+
+	void FPBDCollisionConstraint::AssignSavedManifoldPoints()
+	{
+		TCArray<int32, MaxManifoldPoints> AllowedPointIndices;
+		AllowedPointIndices.SetNum(SavedManifoldPoints.Num());
+		for (int32 PointIndex = 0, PointEndIndex = SavedManifoldPoints.Num(); PointIndex < PointEndIndex; ++PointIndex)
+		{
+			AllowedPointIndices[PointIndex] = PointIndex;
+		}
+
+		for (int32 PointIndex = 0, PointEndIndex = ManifoldPoints.Num(); PointIndex < PointEndIndex; ++PointIndex)
+		{
+			const int32 SavedManifoldPointIndex = FindSavedManifoldPoint(PointIndex, AllowedPointIndices);
 
-			if (MatchIndex != INDEX_NONE)
+			if (SavedManifoldPointIndex != INDEX_NONE)
 			{
-				return &SavedManifoldPoints[MatchIndex];
+				ManifoldPoints[PointIndex].Flags.bHasStaticFrictionAnchor = true;
+				ManifoldPoints[PointIndex].ShapeAnchorPoints[0] = SavedManifoldPoints[SavedManifoldPointIndex].ShapeContactPoints[0];
+				ManifoldPoints[PointIndex].ShapeAnchorPoints[1] = SavedManifoldPoints[SavedManifoldPointIndex].ShapeContactPoints[1];
 			}
+			// Nothing to do if no saved friction point because we already set the achor to the most recently detected contact point
+			// (See InitManifoldPoint, And UpdateAndTryRestoreManifold)
 		}
-		return nullptr;
 	}
 
 	ECollisionConstraintDirection FPBDCollisionConstraint::GetConstraintDirection(const FReal Dt) const
diff --git a/Engine/Source/Runtime/Experimental/Chaos/Private/Chaos/Collision/PBDCollisionSolver.cpp b/Engine/Source/Runtime/Experimental/Chaos/Private/Chaos/Collision/PBDCollisionSolver.cpp
--- a/Engine/Source/Runtime/Experimental/Chaos/Private/Chaos/Collision/PBDCollisionSolver.cpp
+++ b/Engine/Source/Runtime/Experimental/Chaos/Private/Chaos/Collision/PBDCollisionSolver.cpp
@@ -48,183 +48,181 @@ namespace Chaos
 	}
 	using namespace CVars;
 
-
-	//////////////////////////////////////////////////////////////////////////////////////////////////
-	//////////////////////////////////////////////////////////////////////////////////////////////////
-	//////////////////////////////////////////////////////////////////////////////////////////////////
-	//////////////////////////////////////////////////////////////////////////////////////////////////
-
-
-	FPBDCollisionSolver::FPBDCollisionSolver()
-		: State()
+	namespace Private
 	{
-	}
+		//////////////////////////////////////////////////////////////////////////////////////////////////
+		//////////////////////////////////////////////////////////////////////////////////////////////////
+		//////////////////////////////////////////////////////////////////////////////////////////////////
+		//////////////////////////////////////////////////////////////////////////////////////////////////
 
-	void FPBDCollisionSolver::EnablePositionShockPropagation()
-	{
-		SetShockPropagationInvMassScale(Chaos_PBDCollisionSolver_Position_MinInvMassScale);
-	}
 
-	void FPBDCollisionSolver::EnableVelocityShockPropagation()
-	{
-		SetShockPropagationInvMassScale(Chaos_PBDCollisionSolver_Velocity_MinInvMassScale);
-	}
+		void FPBDCollisionSolver::EnablePositionShockPropagation()
+		{
+			SetShockPropagationInvMassScale(Chaos_PBDCollisionSolver_Position_MinInvMassScale);
+		}
 
-	void FPBDCollisionSolver::DisableShockPropagation()
-	{
-		SetShockPropagationInvMassScale(FReal(1));
-	}
+		void FPBDCollisionSolver::EnableVelocityShockPropagation()
+		{
+			SetShockPropagationInvMassScale(Chaos_PBDCollisionSolver_Velocity_MinInvMassScale);
+		}
 
-	void FPBDCollisionSolver::SetShockPropagationInvMassScale(const FSolverReal InvMassScale)
-	{
-		FConstraintSolverBody& Body0 = SolverBody0();
-		FConstraintSolverBody& Body1 = SolverBody1();
-
-		// Shock propagation decreases the inverse mass of bodies that are lower in the pile
-		// of objects. This significantly improves stability of heaps and stacks. Height in the pile is indictaed by the "level". 
-		// No need to set an inverse mass scale if the other body is kinematic (with inv mass of 0).
-		// Bodies at the same level do not take part in shock propagation.
-		if (Body0.IsDynamic() && Body1.IsDynamic() && (Body0.Level() != Body1.Level()))
+		void FPBDCollisionSolver::DisableShockPropagation()
 		{
-			// Set the inv mass scale of the "lower" body to make it heavier
-			bool bInvMassUpdated = false;
-			if (Body0.Level() < Body1.Level())
+			SetShockPropagationInvMassScale(FReal(1));
+		}
+
+		void FPBDCollisionSolver::SetShockPropagationInvMassScale(const FSolverReal InvMassScale)
+		{
+			FConstraintSolverBody& Body0 = SolverBody0();
+			FConstraintSolverBody& Body1 = SolverBody1();
+
+			// Shock propagation decreases the inverse mass of bodies that are lower in the pile
+			// of objects. This significantly improves stability of heaps and stacks. Height in the pile is indictaed by the "level". 
+			// No need to set an inverse mass scale if the other body is kinematic (with inv mass of 0).
+			// Bodies at the same level do not take part in shock propagation.
+			if (Body0.IsDynamic() && Body1.IsDynamic() && (Body0.Level() != Body1.Level()))
 			{
-				if (Body0.ShockPropagationScale() != InvMassScale)
+				// Set the inv mass scale of the "lower" body to make it heavier
+				bool bInvMassUpdated = false;
+				if (Body0.Level() < Body1.Level())
 				{
-					Body0.SetShockPropagationScale(InvMassScale);
-					bInvMassUpdated = true;
+					if (Body0.ShockPropagationScale() != InvMassScale)
+					{
+						Body0.SetShockPropagationScale(InvMassScale);
+						bInvMassUpdated = true;
+					}
 				}
-			}
-			else
-			{
-				if (Body1.ShockPropagationScale() != InvMassScale)
+				else
 				{
-					Body1.SetShockPropagationScale(InvMassScale);
-					bInvMassUpdated = true;
+					if (Body1.ShockPropagationScale() != InvMassScale)
+					{
+						Body1.SetShockPropagationScale(InvMassScale);
+						bInvMassUpdated = true;
+					}
 				}
-			}
 
-			// If the masses changed, we need to rebuild the contact mass for each manifold point
-			if (bInvMassUpdated)
-			{
-				for (int32 PointIndex = 0; PointIndex < State.NumManifoldPoints; ++PointIndex)
+				// If the masses changed, we need to rebuild the contact mass for each manifold point
+				if (bInvMassUpdated)
 				{
-					State.ManifoldPoints[PointIndex].UpdateMassNormal(Body0, Body1);
+					for (int32 PointIndex = 0; PointIndex < NumManifoldPoints(); ++PointIndex)
+					{
+						State.ManifoldPoints[PointIndex].UpdateMassNormal(Body0, Body1);
+					}
 				}
 			}
 		}
-	}
 
-	void FPBDCollisionSolver::SolveVelocityAverage(const FSolverReal Dt)
-	{
-		FConstraintSolverBody& Body0 = SolverBody0();
-		FConstraintSolverBody& Body1 = SolverBody1();
-
-		// Generate a new contact point at the average of all the active contacts
-		int32 NumActiveManifoldPoints = 0;
-		FSolverVec3 RelativeContactPosition0 = FSolverVec3(0);
-		FSolverVec3 RelativeContactPosition1 = FSolverVec3(0);
-		FSolverVec3 WorldContactNormal = FSolverVec3(0);
-		FSolverReal NetPushOutNormal = FSolverReal(0);
-		FSolverReal WorldContactVelocityTargetNormal = FSolverReal(0);
-		for (int32 PointIndex = 0; PointIndex < State.NumManifoldPoints; ++PointIndex)
+		void FPBDCollisionSolver::SolveVelocityAverage(const FSolverReal Dt)
 		{
-			FPBDCollisionSolverManifoldPoint& SolverManifoldPoint = State.ManifoldPoints[PointIndex];
-			if (SolverManifoldPoint.ShouldSolveVelocity())
+			FConstraintSolverBody& Body0 = SolverBody0();
+			FConstraintSolverBody& Body1 = SolverBody1();
+
+			// Generate a new contact point at the average of all the active contacts
+			int32 NumActiveManifoldPoints = 0;
+			FSolverVec3 RelativeContactPosition0 = FSolverVec3(0);
+			FSolverVec3 RelativeContactPosition1 = FSolverVec3(0);
+			FSolverVec3 WorldContactNormal = FSolverVec3(0);
+			FSolverReal NetPushOutNormal = FSolverReal(0);
+			FSolverReal WorldContactVelocityTargetNormal = FSolverReal(0);
+			for (int32 PointIndex = 0; PointIndex < NumManifoldPoints(); ++PointIndex)
 			{
-				RelativeContactPosition0 += SolverManifoldPoint.RelativeContactPosition0;
-				RelativeContactPosition1 += SolverManifoldPoint.RelativeContactPosition1;
-				WorldContactVelocityTargetNormal += SolverManifoldPoint.WorldContactVelocityTargetNormal;
-				WorldContactNormal = SolverManifoldPoint.WorldContactNormal;	// Take last value - should all be similar
-				NetPushOutNormal += SolverManifoldPoint.NetPushOutNormal;
-				++NumActiveManifoldPoints;
+				FPBDCollisionSolverManifoldPoint& SolverManifoldPoint = State.ManifoldPoints[PointIndex];
+				if (SolverManifoldPoint.ShouldSolveVelocity())
+				{
+					RelativeContactPosition0 += SolverManifoldPoint.WorldContact.RelativeContactPoints[0];
+					RelativeContactPosition1 += SolverManifoldPoint.WorldContact.RelativeContactPoints[1];
+					WorldContactVelocityTargetNormal += SolverManifoldPoint.WorldContact.ContactTargetVelocityNormal;
+					WorldContactNormal = SolverManifoldPoint.WorldContact.ContactNormal;	// Take last value - should all be similar
+					NetPushOutNormal += SolverManifoldPoint.NetPushOutNormal;
+					++NumActiveManifoldPoints;
+				}
 			}
-		}
 
-		// Solve for velocity at the new conatct point
-		// We only do this if we have multiple active contacts
-		// NOTE: this average contact isn't really correct, especially when all contacts do not equally
-		// contribute to the pushout (which is normal even for a simple box on a plane). E.g., the WorldContactVelocityTargetNormal
-		// and NetPushOutNormal will not be right. The goal here though is just to get close to the correct result in fewer iterations
-		// than we would have if we just solved the corners of the box in sequence.
-		if (NumActiveManifoldPoints > 1)
-		{
-			const FSolverReal InvCount = FSolverReal(1) / FSolverReal(NumActiveManifoldPoints);
+			// Solve for velocity at the new conatct point
+			// We only do this if we have multiple active contacts
+			// NOTE: this average contact isn't really correct, especially when all contacts do not equally
+			// contribute to the pushout (which is normal even for a simple box on a plane). E.g., the WorldContactVelocityTargetNormal
+			// and NetPushOutNormal will not be right. The goal here though is just to get close to the correct result in fewer iterations
+			// than we would have if we just solved the corners of the box in sequence.
+			if (NumActiveManifoldPoints > 1)
+			{
+				const FSolverReal InvCount = FSolverReal(1) / FSolverReal(NumActiveManifoldPoints);
 
-			FPBDCollisionSolverManifoldPoint AverageManifoldPoint;
-			AverageManifoldPoint.RelativeContactPosition0 = RelativeContactPosition0 * InvCount;
-			AverageManifoldPoint.RelativeContactPosition1 = RelativeContactPosition1 * InvCount;
-			AverageManifoldPoint.WorldContactNormal = WorldContactNormal;
-			AverageManifoldPoint.WorldContactVelocityTargetNormal = WorldContactVelocityTargetNormal * InvCount;
+				FPBDCollisionSolverManifoldPoint AverageManifoldPoint;
+				AverageManifoldPoint.WorldContact.RelativeContactPoints[0] = RelativeContactPosition0 * InvCount;
+				AverageManifoldPoint.WorldContact.RelativeContactPoints[1] = RelativeContactPosition1 * InvCount;
+				AverageManifoldPoint.WorldContact.ContactNormal = WorldContactNormal;
+				AverageManifoldPoint.WorldContact.ContactTargetVelocityNormal = WorldContactVelocityTargetNormal * InvCount;
 			
-			// Total pushout (not average) which is correct if the average point is also the centroid but otherwise probably an overestimate.
-			// This is used to limit the possibly-attractive impulse that corrects implicit velocity errors from the PBD solve, but it
-			// which might cause stickiness if overestimated. Although it should be recitifed when we solve the normal contact point velocities
-			// so I don't think it matters...we could do better though if it is a problem
-			AverageManifoldPoint.NetPushOutNormal = NetPushOutNormal;
-
-			// Calculate the contact mass (and derived properties) for this point
-			FSolverVec3 WorldContactNormalAngular0 = FSolverVec3(0);
-			FSolverVec3 WorldContactNormalAngular1 = FSolverVec3(0);
-			FSolverReal ContactMassInvNormal = FSolverReal(0);
-			if (Body0.IsDynamic())
-			{
-				const FSolverVec3 R0xN = FSolverVec3::CrossProduct(RelativeContactPosition0, WorldContactNormal);
-				WorldContactNormalAngular0 = Body0.InvI() * R0xN;
-				ContactMassInvNormal += FSolverVec3::DotProduct(R0xN, WorldContactNormalAngular0) + Body0.InvM();
-			}
-			if (Body1.IsDynamic())
-			{
-				const FSolverVec3 R1xN = FSolverVec3::CrossProduct(RelativeContactPosition1, WorldContactNormal);
-				WorldContactNormalAngular1 = Body1.InvI() * R1xN;
-				ContactMassInvNormal += FSolverVec3::DotProduct(R1xN, WorldContactNormalAngular1) + Body1.InvM();
-			}
-			AverageManifoldPoint.ContactMassNormal = (ContactMassInvNormal > FSolverReal(UE_SMALL_NUMBER)) ? FSolverReal(1) / ContactMassInvNormal : FSolverReal(0);
-			AverageManifoldPoint.WorldContactNormalAngular0 = WorldContactNormalAngular0;
-			AverageManifoldPoint.WorldContactNormalAngular1 = WorldContactNormalAngular1;
-
-			// @todo(chaos): we don't use these - maybe do the calculation without an actual manifold point object...
-			AverageManifoldPoint.ContactMassTangentU = FSolverReal(0);
-			AverageManifoldPoint.ContactMassTangentV = FSolverReal(0);
-			AverageManifoldPoint.WorldContactTangentU = FSolverVec3(0);
-			AverageManifoldPoint.WorldContactTangentV = FSolverVec3(0);
-			AverageManifoldPoint.WorldContactTangentUAngular0 = FSolverVec3(0);
-			AverageManifoldPoint.WorldContactTangentVAngular0 = FSolverVec3(0);
-			AverageManifoldPoint.WorldContactTangentUAngular1 = FSolverVec3(0);
-			AverageManifoldPoint.WorldContactTangentVAngular1 = FSolverVec3(0);
-			AverageManifoldPoint.WorldContactDeltaNormal = FSolverReal(0);
-			AverageManifoldPoint.WorldContactDeltaTangentU = FSolverReal(0);
-			AverageManifoldPoint.WorldContactDeltaTangentV = FSolverReal(0);
-			AverageManifoldPoint.NetPushOutTangentU = FSolverReal(0);
-			AverageManifoldPoint.NetPushOutTangentV = FSolverReal(0);
-			AverageManifoldPoint.NetImpulseNormal = FSolverReal(0);
-			AverageManifoldPoint.NetImpulseTangentU = FSolverReal(0);
-			AverageManifoldPoint.NetImpulseTangentV = FSolverReal(0);
-			AverageManifoldPoint.StaticFrictionRatio = FSolverReal(0);
-
-			FSolverReal ContactVelocityDeltaNormal;
-			AverageManifoldPoint.CalculateContactVelocityErrorNormal(Body0, Body1, ContactVelocityDeltaNormal);
-
-			const FSolverReal MinImpulseNormal = FMath::Min(FSolverReal(0), -AverageManifoldPoint.NetPushOutNormal / Dt);
-
-			ApplyVelocityCorrectionNormal(
-				State.Stiffness,
-				ContactVelocityDeltaNormal,
-				MinImpulseNormal,
-				AverageManifoldPoint,
-				Body0,
-				Body1);
-
-			// Now distribute the net impulse among the active points so we don't over-correct pushout from initial overlaps
-			for (int32 PointIndex = 0; PointIndex < State.NumManifoldPoints; ++PointIndex)
-			{
-				FPBDCollisionSolverManifoldPoint& SolverManifoldPoint = State.ManifoldPoints[PointIndex];
-				if (SolverManifoldPoint.ShouldSolveVelocity())
+				// Total pushout (not average) which is correct if the average point is also the centroid but otherwise probably an overestimate.
+				// This is used to limit the possibly-attractive impulse that corrects implicit velocity errors from the PBD solve, but it
+				// which might cause stickiness if overestimated. Although it should be recitifed when we solve the normal contact point velocities
+				// so I don't think it matters...we could do better though if it is a problem
+				AverageManifoldPoint.NetPushOutNormal = NetPushOutNormal;
+
+				// Calculate the contact mass (and derived properties) for this point
+				FSolverVec3 WorldContactNormalAngular0 = FSolverVec3(0);
+				FSolverVec3 WorldContactNormalAngular1 = FSolverVec3(0);
+				FSolverReal ContactMassInvNormal = FSolverReal(0);
+				if (Body0.IsDynamic())
 				{
-					SolverManifoldPoint.NetImpulseNormal += AverageManifoldPoint.NetImpulseNormal * InvCount;
+					const FSolverVec3 R0xN = FSolverVec3::CrossProduct(RelativeContactPosition0, WorldContactNormal);
+					WorldContactNormalAngular0 = Body0.InvI() * R0xN;
+					ContactMassInvNormal += FSolverVec3::DotProduct(R0xN, WorldContactNormalAngular0) + Body0.InvM();
+				}
+				if (Body1.IsDynamic())
+				{
+					const FSolverVec3 R1xN = FSolverVec3::CrossProduct(RelativeContactPosition1, WorldContactNormal);
+					WorldContactNormalAngular1 = Body1.InvI() * R1xN;
+					ContactMassInvNormal += FSolverVec3::DotProduct(R1xN, WorldContactNormalAngular1) + Body1.InvM();
+				}
+				AverageManifoldPoint.ContactMassNormal = (ContactMassInvNormal > FSolverReal(UE_SMALL_NUMBER)) ? FSolverReal(1) / ContactMassInvNormal : FSolverReal(0);
+				AverageManifoldPoint.WorldContactNormalAngular0 = WorldContactNormalAngular0;
+				AverageManifoldPoint.WorldContactNormalAngular1 = WorldContactNormalAngular1;
+
+				// @todo(chaos): we don't use these - maybe do the calculation without an actual manifold point object...
+				AverageManifoldPoint.WorldContact.ContactTangentU = FSolverVec3(0);
+				AverageManifoldPoint.WorldContact.ContactTangentV = FSolverVec3(0);
+				AverageManifoldPoint.WorldContact.ContactDeltaNormal = FSolverReal(0);
+				AverageManifoldPoint.WorldContact.ContactDeltaTangentU = FSolverReal(0);
+				AverageManifoldPoint.WorldContact.ContactDeltaTangentV = FSolverReal(0);
+				AverageManifoldPoint.WorldContactTangentUAngular0 = FSolverVec3(0);
+				AverageManifoldPoint.WorldContactTangentVAngular0 = FSolverVec3(0);
+				AverageManifoldPoint.WorldContactTangentUAngular1 = FSolverVec3(0);
+				AverageManifoldPoint.WorldContactTangentVAngular1 = FSolverVec3(0);
+				AverageManifoldPoint.ContactMassTangentU = FSolverReal(0);
+				AverageManifoldPoint.ContactMassTangentV = FSolverReal(0);
+				AverageManifoldPoint.NetPushOutTangentU = FSolverReal(0);
+				AverageManifoldPoint.NetPushOutTangentV = FSolverReal(0);
+				AverageManifoldPoint.NetImpulseNormal = FSolverReal(0);
+				AverageManifoldPoint.NetImpulseTangentU = FSolverReal(0);
+				AverageManifoldPoint.NetImpulseTangentV = FSolverReal(0);
+				AverageManifoldPoint.StaticFrictionRatio = FSolverReal(0);
+
+				FSolverReal ContactVelocityDeltaNormal;
+				AverageManifoldPoint.CalculateContactVelocityErrorNormal(Body0, Body1, ContactVelocityDeltaNormal);
+
+				const FSolverReal MinImpulseNormal = FMath::Min(FSolverReal(0), -AverageManifoldPoint.NetPushOutNormal / Dt);
+
+				ApplyVelocityCorrectionNormal(
+					State.Stiffness,
+					ContactVelocityDeltaNormal,
+					MinImpulseNormal,
+					AverageManifoldPoint,
+					Body0,
+					Body1);
+
+				// Now distribute the net impulse among the active points so we don't over-correct pushout from initial overlaps
+				for (int32 PointIndex = 0; PointIndex < NumManifoldPoints(); ++PointIndex)
+				{
+					FPBDCollisionSolverManifoldPoint& SolverManifoldPoint = State.ManifoldPoints[PointIndex];
+					if (SolverManifoldPoint.ShouldSolveVelocity())
+					{
+						SolverManifoldPoint.NetImpulseNormal += AverageManifoldPoint.NetImpulseNormal * InvCount;
+					}
 				}
 			}
 		}
-	}
-}
\ No newline at end of file
+
+	}	// namespace Private
+}	// namespace Chaos
\ No newline at end of file
diff --git a/Engine/Source/Runtime/Experimental/Chaos/Private/Chaos/Collision/PBDCollisionSolver.h b/Engine/Source/Runtime/Experimental/Chaos/Private/Chaos/Collision/PBDCollisionSolver.h
deleted file mode 100644
diff --git a/Engine/Source/Runtime/Experimental/Chaos/Private/Chaos/Collision/ParticlePairMidPhase.cpp b/Engine/Source/Runtime/Experimental/Chaos/Private/Chaos/Collision/ParticlePairMidPhase.cpp
--- a/Engine/Source/Runtime/Experimental/Chaos/Private/Chaos/Collision/ParticlePairMidPhase.cpp
+++ b/Engine/Source/Runtime/Experimental/Chaos/Private/Chaos/Collision/ParticlePairMidPhase.cpp
@@ -347,13 +347,15 @@ namespace Chaos
 			{
 				if (Context.GetAllocator()->ActivateConstraint(Constraint.Get()))
 				{
-					// If we had any mods last tick, undo them
-					Constraint->ResetModifications();
+					Constraint->Activate(Dt);
 
 					LastUsedEpoch = CurrentEpoch;
 					return 1;
 				}
 			}
+
+			// If we get here, we did not activate hte constraint and it should be disabled for this tick
+			Constraint->SetDisabled(true);
 		}
 
 		return 0;
@@ -642,8 +644,8 @@ namespace Chaos
 		}
 
 		// @todo(chaos): we already have the shape world transforms at the calling site - pass them in
-		const FRigidTransform3 ParticleTransform0 = FParticleUtilitiesPQ::GetActorWorldTransform(FConstGenericParticleHandle(InParticle0));
-		const FRigidTransform3 ParticleTransform1 = FParticleUtilitiesPQ::GetActorWorldTransform(FConstGenericParticleHandle(InParticle1));
+		const FRigidTransform3 ParticleTransform0 = FConstGenericParticleHandle(InParticle0)->GetTransformPQ();
+		const FRigidTransform3 ParticleTransform1 = FConstGenericParticleHandle(InParticle1)->GetTransformPQ();
 		const FRigidTransform3 ShapeWorldTransform0 = ShapeRelativeTransform0 * ParticleTransform0;
 		const FRigidTransform3 ShapeWorldTransform1 = ShapeRelativeTransform1 * ParticleTransform1;
 		Constraint->SetShapeWorldTransforms(ShapeWorldTransform0, ShapeWorldTransform1);
@@ -739,6 +741,10 @@ namespace Chaos
 				Context.GetAllocator()->ActivateConstraint(Constraint);
 				++NumActiveConstraints;
 			}
+			else
+			{
+				Constraint->SetDisabled(true);
+			}
 		}
 		NewConstraints.Reset();
 		return NumActiveConstraints;
diff --git a/Engine/Source/Runtime/Experimental/Chaos/Private/Chaos/Collision/SolverCollisionContainer.cpp b/Engine/Source/Runtime/Experimental/Chaos/Private/Chaos/Collision/SolverCollisionContainer.cpp
--- a/Engine/Source/Runtime/Experimental/Chaos/Private/Chaos/Collision/SolverCollisionContainer.cpp
+++ b/Engine/Source/Runtime/Experimental/Chaos/Private/Chaos/Collision/SolverCollisionContainer.cpp
@@ -3,6 +3,7 @@
 
 #include "Chaos/Collision/CollisionApplyType.h"
 #include "Chaos/Collision/PBDCollisionConstraint.h"
+#include "Chaos/Collision/PBDCollisionSolver.h"
 #include "Chaos/CollisionResolution.h"
 #include "Chaos/DebugDrawQueue.h"
 #include "Chaos/Evolution/SolverBodyContainer.h"
@@ -13,7 +14,6 @@
 #include "Chaos/Utilities.h"
 
 // Private includes
-#include "PBDCollisionSolver.h"
 
 #include "ChaosLog.h"
 
@@ -34,10 +34,6 @@ namespace Chaos
 		FAutoConsoleVariableRef CVarChaosPBDCollisionSolverAutoStiffnessMassRatio1(TEXT("p.Chaos.PBDCollisionSolver.AutoStiffness.MassRatio1"), Chaos_PBDCollisionSolver_AutoStiffness_MassRatio1, TEXT(""));
 		FAutoConsoleVariableRef CVarChaosPBDCollisionSolverAutoStiffnessMassRatio2(TEXT("p.Chaos.PBDCollisionSolver.AutoStiffness.MassRatio2"), Chaos_PBDCollisionSolver_AutoStiffness_MassRatio2, TEXT(""));
 
-		// Collision solver regularization (experimental, disabled by default)
-		FRealSingle Chaos_PBDCollisionSolver_Regularization = 0;
-		FAutoConsoleVariableRef CVarChaosPBDCollisionSolverRegularization(TEXT("p.Chaos.PBDCollisionSolver.Regularization"), Chaos_PBDCollisionSolver_Regularization, TEXT(""));
-
 		// Whether to use the Jacobi collision pair solver
 		bool bChaos_PBDCollisionSolver_UseJacobiPairSolver = false;
 		FAutoConsoleVariableRef CVarChaosPBDCollisionSolverUseJacobiPairSolver(TEXT("p.Chaos.PBDCollisionSolver.UseJacobi"), bChaos_PBDCollisionSolver_UseJacobiPairSolver, TEXT(""));
@@ -57,8 +53,9 @@ namespace Chaos
 		FRealSingle Chaos_PBDCollisionSolver_JacobiRotationTolerance = 1.e-8f;
 		FAutoConsoleVariableRef CVarChaosPBDCollisionSolverJacobiRotationTolerance(TEXT("p.Chaos.PBDCollisionSolver.JacobiRotationTolerance"), Chaos_PBDCollisionSolver_JacobiRotationTolerance, TEXT(""));
 	}
-	using namespace CVars;
 
+	using namespace CVars;
+	using namespace Private;
 
 	//////////////////////////////////////////////////////////////////////////////////////////////////
 	//////////////////////////////////////////////////////////////////////////////////////////////////
@@ -81,326 +78,204 @@ namespace Chaos
 	//////////////////////////////////////////////////////////////////////////////////////////////////
 	//////////////////////////////////////////////////////////////////////////////////////////////////
 
-
-	/**
-	 * @brief A wrapper for FPBDCollisionSolver which binds to a Collision Constraint and adds Gather/Scatter from/to the constraint
-	*/
-	class FPBDCollisionSolverAdapter
+	// Transform the Constraint's local-space data into world space for use by  the collision solver and also calculate tangents, errors, etc
+	FORCEINLINE_DEBUGGABLE void UpdateCollisionSolverContactPointFromConstraint(FPBDCollisionSolver& Solver, const int32 SolverPointIndex, const FPBDCollisionConstraint* Constraint, const int32 ConstraintPointIndex, const FRealSingle Dt, const FSolverBody& Body0, const FSolverBody& Body1)
 	{
-	public:
-		FPBDCollisionSolverAdapter()
-			: Constraint(nullptr)
-			, bIsManifold(true)
-			, bIsIncrementalManifold(false)
-		{
-		}
+		const FManifoldPoint& ManifoldPoint = Constraint->GetManifoldPoint(ConstraintPointIndex);
 
-		FPBDCollisionSolver& GetSolver()
-		{ 
-			return Solver;
-		}
+		const FRealSingle Restitution = FRealSingle(Constraint->GetRestitution());
+		const FRealSingle RestitutionVelocityThreshold = FRealSingle(Constraint->GetRestitutionThreshold()) * Dt;
 
-		FPBDCollisionConstraint* GetConstraint()
-		{ 
-			return Constraint;
-		}
+		// World-space shape transforms. Manifold data is currently relative to these spaces
+		const FRigidTransform3& ShapeWorldTransform0 = Constraint->GetShapeWorldTransform0();
+		const FRigidTransform3& ShapeWorldTransform1 = Constraint->GetShapeWorldTransform1();
+
+		// World-space contact points on each shape
+		const FVec3 WorldContact0 = ShapeWorldTransform0.TransformPositionNoScale(FVec3(ManifoldPoint.ContactPoint.ShapeContactPoints[0]));
+		const FVec3 WorldContact1 = ShapeWorldTransform1.TransformPositionNoScale(FVec3(ManifoldPoint.ContactPoint.ShapeContactPoints[1]));
+		const FVec3 WorldContact = FReal(0.5) * (WorldContact0 + WorldContact1);
+		const FVec3f WorldRelativeContact0 = FVec3f(WorldContact - Body0.P());
+		const FVec3f WorldRelativeContact1 = FVec3f(WorldContact - Body1.P());
 
-		void SetConstraint(FPBDCollisionConstraint& InConstraint)
+		// World-space normal
+		const FVec3f WorldContactNormal = FVec3f(ShapeWorldTransform1.TransformVectorNoScale(FVec3(ManifoldPoint.ContactPoint.ShapeContactNormal)));
+
+		// World-space tangents
+		FVec3f WorldContactTangentU = FVec3f::CrossProduct(FVec3f(0, 1, 0), WorldContactNormal);
+		if (!WorldContactTangentU.Normalize(FRealSingle(UE_KINDA_SMALL_NUMBER)))
 		{
-			Constraint = &InConstraint;
+			WorldContactTangentU = FVec3f::CrossProduct(FVec3f(1, 0, 0), WorldContactNormal);
+			WorldContactTangentU = WorldContactTangentU.GetUnsafeNormal();
 		}
+		const FVec3f WorldContactTangentV = FVec3f::CrossProduct(WorldContactNormal, WorldContactTangentU);
 
-		void SetBodies(FSolverBody& Body0, FSolverBody& Body1)
+		// Calculate contact velocity if we will need it below (restitution and/or frist-contact friction)
+		const bool bNeedContactVelocity = (!ManifoldPoint.Flags.bHasStaticFrictionAnchor) || (Restitution > FRealSingle(0));
+		FVec3f ContactVel = FVec3(0);
+		if (bNeedContactVelocity)
 		{
-			Solver.SetSolverBodies(Body0, Body1);
+			const FVec3f ContactVel0 = Body0.V() + FVec3f::CrossProduct(Body0.W(), WorldRelativeContact0);
+			const FVec3f ContactVel1 = Body1.V() + FVec3f::CrossProduct(Body1.W(), WorldRelativeContact1);
+			ContactVel = ContactVel0 - ContactVel1;
+		}
 
-			// We should try to remove this - the Constraint should not need to know about solver objects
-			Constraint->SetSolverBodies(&Body0, &Body1);
+		// If we have contact data from a previous tick, use it to calculate the lateral position delta we need
+		// to apply to move the contacts back to their original relative locations (i.e., to enforce static friction).
+		// Otherwise, estimate the friction correction from the contact velocity.
+		// NOTE: quadratic shapes use the velocity-based path most of the time, unless the relative motion is very small.
+		FVec3f WorldFrictionDelta = FVec3f(0);
+		if (ManifoldPoint.Flags.bHasStaticFrictionAnchor)
+		{
+			const FVec3f FrictionDelta0 = ShapeWorldTransform0.TransformPositionNoScale(FVec3(ManifoldPoint.ShapeAnchorPoints[0])) - WorldContact0;
+			const FVec3f FrictionDelta1 = ShapeWorldTransform1.TransformPositionNoScale(FVec3(ManifoldPoint.ShapeAnchorPoints[1])) - WorldContact1;
+			WorldFrictionDelta = FrictionDelta0 - FrictionDelta1;
+		}
+		else
+		{
+			// @todo(chaos): consider adding a multiplier to the initial contact friction
+			WorldFrictionDelta = ContactVel * Dt;
 		}
 
-		void Unbind()
+		// The contact point error we are trying to correct in this solver
+		const FRealSingle TargetPhi = ManifoldPoint.TargetPhi;
+		const FVec3f WorldContactDelta = FVec3f(WorldContact0 - WorldContact1);
+		const FRealSingle WorldContactDeltaNormal = FVec3f::DotProduct(WorldContactDelta, WorldContactNormal) - TargetPhi;
+		const FRealSingle WorldContactDeltaTangentU = FVec3f::DotProduct(WorldContactDelta + WorldFrictionDelta, WorldContactTangentU);
+		const FRealSingle WorldContactDeltaTangentV = FVec3f::DotProduct(WorldContactDelta + WorldFrictionDelta, WorldContactTangentV);
+
+		// The target contact velocity, taking restitution into account
+		FRealSingle WorldContactTargetVelocityNormal = FRealSingle(0);
+		if (Restitution > FRealSingle(0))
 		{
-			if (Constraint != nullptr)
+			const FRealSingle ContactVelocityNormal = FVec3f::DotProduct(ContactVel, WorldContactNormal);
+			if (ContactVelocityNormal < -RestitutionVelocityThreshold)
 			{
-				Constraint->SetSolverBodies(nullptr, nullptr);
-				Constraint = nullptr;
+				WorldContactTargetVelocityNormal = -Restitution * ContactVelocityNormal;
 			}
-			Solver.ResetSolverBodies();
 		}
 
-		void Reset()
-		{
-			Unbind();
+		FWorldContactPoint& WorldContactPoint = Solver.GetWorldContactPoint(SolverPointIndex);
+		WorldContactPoint.RelativeContactPoints[0] = WorldRelativeContact0;
+		WorldContactPoint.RelativeContactPoints[1] = WorldRelativeContact1;
+		WorldContactPoint.ContactNormal = WorldContactNormal;
+		WorldContactPoint.ContactTangentU = WorldContactTangentU;
+		WorldContactPoint.ContactTangentV = WorldContactTangentV;
+		WorldContactPoint.ContactDeltaNormal = WorldContactDeltaNormal;
+		WorldContactPoint.ContactDeltaTangentU = WorldContactDeltaTangentU;
+		WorldContactPoint.ContactDeltaTangentV = WorldContactDeltaTangentV;
+		WorldContactPoint.ContactTargetVelocityNormal = WorldContactTargetVelocityNormal;
+	}
 
-			bIsManifold = true;
-			bIsIncrementalManifold = false;
-			Solver.Reset();
-		}
+	FORCEINLINE_DEBUGGABLE void UpdateCollisionSolverManifoldFromConstraint(FPBDCollisionSolver& Solver, const FPBDCollisionConstraint* Constraint, const FSolverReal Dt, const int32 ConstraintPointBeginIndex, const int32 ConstraintPointEndIndex)
+	{
+		const FSolverBody& Body0 = Solver.SolverBody0().SolverBody();
+		const FSolverBody& Body1 = Solver.SolverBody1().SolverBody();
 
-		bool IsManifold() const
-		{ 
-			return bIsManifold;
-		}
-		
-		bool IsIncrementalManifold() const
-		{ 
-			return bIsIncrementalManifold;
-		}
-
-		/**
-		 * @brief Modify solver stiffness when we have bodies with large mass differences
-		*/
-		FReal CalculateSolverStiffness(
-			const FSolverBody& Body0,
-			const FSolverBody& Body1,
-			const FReal MassRatio1,
-			const FReal MassRatio2)
-		{
-			// Adjust the solver stiffness if one body is more than MassRatio1 times the mass of the other and the heavier one is on top.
-			// Solver stiffness will be equal to 1 when the mass ratio is MassRatio1.
-			// Solver stiffness will be equal to 0 when the mass ratio is MassRatio2.
-			if (Body0.IsDynamic() && Body1.IsDynamic() && (MassRatio1 > 0) && (MassRatio2 > MassRatio1))
+		// Only calculate state for newly added contacts. Normally this is all of them, but maybe not if incremental collision is used by RBAN.
+		// Also we only add active points to the solver's manifold points list
+		for (int32 ConstraintManifoldPointIndex = ConstraintPointBeginIndex; ConstraintManifoldPointIndex < ConstraintPointEndIndex; ++ConstraintManifoldPointIndex)
+		{
+			if (!Constraint->GetManifoldPoint(ConstraintManifoldPointIndex).Flags.bDisabled)
 			{
-				// Find heavy body and the mass ratio
-				const FSolverBody* HeavyBody;
-				FReal MassRatio;
-				if (Body0.InvM() < Body1.InvM())
-				{
-					HeavyBody = &Body0;
-					MassRatio = Body1.InvM() / Body0.InvM();
-				}
-				else
-				{
-					HeavyBody = &Body1;
-					MassRatio = Body0.InvM() / Body1.InvM();
-				}
+				const int32 SolverManifoldPointIndex = Solver.AddManifoldPoint();
 
-				if (MassRatio > MassRatio1)
-				{
-					// Is this a load-bearing contact (normal is significantly along gravity direction)?
-					// @todo(chaos): should use gravity direction. Currently assumes -Z
-					// @todo(chaos): maybe gradually introduce stiffness scaling based on normal rather than on/off
-					// @todo(chaos): could use solver manifold data which is already in world space rather than CalculateWorldContactNormal
-					const FVec3 WorldNormal = Constraint->CalculateWorldContactNormal();
-					if (FMath::Abs(WorldNormal.Z) > FReal(0.3))
-					{
-						// Which body is on the top? (Normal always points away from second body - see FContactPoint)
-						const FSolverBody* TopBody = (WorldNormal.Z > 0) ? &Body0 : &Body1;
-						if (TopBody == HeavyBody)
-						{
-							// The heavy body is on top - reduce the solver stiffness
-							const FReal StiffnessScale = FMath::Clamp((MassRatio2 - MassRatio) / (MassRatio2 - MassRatio1), FReal(0), FReal(1));
-							return StiffnessScale * Constraint->GetStiffness();
-						}
-					}
-				}
+				// Transform the constraint contact data into world space for use by the solver
+				// We build this data directly into the solver's world-space contact data which looks a bit odd with "Init" called after but there you go
+				UpdateCollisionSolverContactPointFromConstraint(Solver, SolverManifoldPointIndex, Constraint, ConstraintManifoldPointIndex, Dt, Body0, Body1);
 
+				// Prepare the solver (set constraint-space mass etc)
+				Solver.FinalizeManifoldPoint(SolverManifoldPointIndex, Dt);
 			}
-
-			return Constraint->GetStiffness();
 		}
+	}
 
-		void GatherInput(
-			const FReal InDt,
-			const FPBDCollisionSolverSettings& SolverSettings)
-		{
-			FSolverReal Dt = FSolverReal(InDt);
-			const FConstraintSolverBody& Body0 = Solver.SolverBody0();
-			const FConstraintSolverBody& Body1 = Solver.SolverBody1();
-
-			// Friction values. Static and Dynamic friction are applied in the position solve for most shapes.
-			// We can also run in a mode without static friction at all. This is faster but stacking is not possible.
-			const FSolverReal StaticFriction = FSolverReal(Constraint->GetStaticFriction());
-			const FSolverReal DynamicFriction = FSolverReal(Constraint->GetDynamicFriction());
-			FSolverReal PositionStaticFriction = FSolverReal(0);
-			FSolverReal PositionDynamicFriction = FSolverReal(0);
-			FSolverReal VelocityDynamicFriction = FSolverReal(0);
-			if (SolverSettings.NumPositionFrictionIterations > 0)
+	FORCEINLINE_DEBUGGABLE void UpdateCollisionSolverFromConstraint(FPBDCollisionSolver& Solver, const FPBDCollisionConstraint* Constraint, const FSolverReal Dt, const FPBDCollisionSolverSettings& SolverSettings, bool& bOutPerIterationCollision)
+	{
+		// Friction values. Static and Dynamic friction are applied in the position solve for most shapes.
+		// We can also run in a mode without static friction at all. This is faster but stacking is not possible.
+		const FSolverReal StaticFriction = FSolverReal(Constraint->GetStaticFriction());
+		const FSolverReal DynamicFriction = FSolverReal(Constraint->GetDynamicFriction());
+		FSolverReal PositionStaticFriction = FSolverReal(0);
+		FSolverReal PositionDynamicFriction = FSolverReal(0);
+		FSolverReal VelocityDynamicFriction = FSolverReal(0);
+		if (SolverSettings.NumPositionFrictionIterations > 0)
+		{
+			PositionStaticFriction = StaticFriction;
+			if (!Constraint->HasQuadraticShape())
 			{
-				PositionStaticFriction = StaticFriction;
 				PositionDynamicFriction = DynamicFriction;
 			}
 			else
 			{
+				// Quadratic shapes don't use PBD dynamic friction - it has issues at slow speeds where the WxR is
+				// less than the position tolerance for friction point matching
+				// @todo(chaos): fix PBD dynamic friction on quadratic shapes
 				VelocityDynamicFriction = DynamicFriction;
 			}
-
-			Solver.SetFriction(PositionStaticFriction, PositionDynamicFriction, VelocityDynamicFriction);
-
-			const FReal SolverStiffness = CalculateSolverStiffness(Body0.SolverBody(), Body1.SolverBody(), Chaos_PBDCollisionSolver_AutoStiffness_MassRatio1, Chaos_PBDCollisionSolver_AutoStiffness_MassRatio2);
-			Solver.SetStiffness(FSolverReal(SolverStiffness));
-
-			const FReal Regularization = Chaos_PBDCollisionSolver_Regularization / (Dt * Dt);
-			Solver.SetRegularization(FSolverReal(Regularization));
-
-			Solver.SolverBody0().SetInvMScale(Constraint->GetInvMassScale0());
-			Solver.SolverBody0().SetInvIScale(Constraint->GetInvInertiaScale0());
-			Solver.SolverBody1().SetInvMScale(Constraint->GetInvMassScale1());
-			Solver.SolverBody1().SetInvIScale(Constraint->GetInvInertiaScale1());
-
-			bIsManifold = Constraint->GetUseManifold();
-			UpdateManifoldPoints(InDt);
 		}
-
-		void UpdateManifoldPoints(
-			const FReal InDt)
+		else
 		{
-			FSolverReal Dt = FSolverReal(InDt);
-			const FConstraintSolverBody& Body0 = Solver.SolverBody0();
-			const FConstraintSolverBody& Body1 = Solver.SolverBody1();
-
-			// We handle incremental manifolds by just collecting any new contacts
-			const int32 BeginPointIndex = bIsIncrementalManifold ? Solver.NumManifoldPoints() : 0;
-			const int32 EndPointIndex = Solver.SetNumManifoldPoints(Constraint->GetManifoldPoints().Num());
+			VelocityDynamicFriction = DynamicFriction;
+		}
 
-			const FSolverReal RestitutionVelocityThreshold = FSolverReal(Constraint->GetRestitutionThreshold()) * Dt;
-			const FSolverReal Restitution = FSolverReal(Constraint->GetRestitution());
+		Solver.SetFriction(PositionStaticFriction, PositionDynamicFriction, VelocityDynamicFriction);
 
-			const FRigidTransform3& ShapeWorldTransform0 = Constraint->GetShapeWorldTransform0();
-			const FRigidTransform3& ShapeWorldTransform1 = Constraint->GetShapeWorldTransform1();
+		const FReal SolverStiffness = Constraint->GetStiffness();
+		Solver.SetStiffness(FSolverReal(SolverStiffness));
 
-			for (int32 ManifoldPointIndex = BeginPointIndex; ManifoldPointIndex < EndPointIndex; ++ManifoldPointIndex)
-			{
-				TArrayView<FManifoldPoint> ManifoldPoints = Constraint->GetManifoldPoints();
-				FManifoldPoint& ManifoldPoint = ManifoldPoints[ManifoldPointIndex];
-				if (ManifoldPoint.Flags.bDisabled)
-				{
-					continue;
-				}
-
-				const FVec3 WorldContactPoint0 = ShapeWorldTransform0.TransformPositionNoScale(FVec3(ManifoldPoint.ContactPoint.ShapeContactPoints[0]));
-				const FVec3 WorldContactPoint1 = ShapeWorldTransform1.TransformPositionNoScale(FVec3(ManifoldPoint.ContactPoint.ShapeContactPoints[1]));
-				const FVec3 WorldContactPoint = FReal(0.5) * (WorldContactPoint0 + WorldContactPoint1);
-
-				const FSolverVec3 WorldContactNormal = FSolverVec3(ShapeWorldTransform1.TransformVectorNoScale(FVec3(ManifoldPoint.ContactPoint.ShapeContactNormal)));
-				const FSolverVec3 RelativeContactPosition0 = FSolverVec3(WorldContactPoint - Body0.P());
-				const FSolverVec3 RelativeContactPosition1 = FSolverVec3(WorldContactPoint - Body1.P());
-				const FSolverReal TargetPhi = FSolverReal(ManifoldPoint.TargetPhi);
-
-				// If we have contact data from a previous tick, use it to calculate the lateral position delta we need
-				// to apply to move the contacts back to their original relative locations (i.e., to enforce static friction)
-				// @todo(chaos): we should not be writing back to the constraint here - find a better way to update the friction anchor. See FPBDCollisionConstraint::SetSolverResults
-				FSolverVec3 WorldFrictionDelta = FSolverVec3(0);
-				const FSavedManifoldPoint* SavedManifoldPoint = Constraint->FindSavedManifoldPoint(ManifoldPoint);
-				if (SavedManifoldPoint != nullptr)
-				{
-					const FSolverVec3 FrictionDelta0 = FSolverVec3(SavedManifoldPoint->ShapeContactPoints[0] - ManifoldPoint.ContactPoint.ShapeContactPoints[0]);
-					const FSolverVec3 FrictionDelta1 = FSolverVec3(SavedManifoldPoint->ShapeContactPoints[1] - ManifoldPoint.ContactPoint.ShapeContactPoints[1]);
-					WorldFrictionDelta = ShapeWorldTransform0.TransformVectorNoScale(FVector(FrictionDelta0)) - ShapeWorldTransform1.TransformVectorNoScale(FVector(FrictionDelta1));
+		Solver.SolverBody0().SetInvMScale(Constraint->GetInvMassScale0());
+		Solver.SolverBody0().SetInvIScale(Constraint->GetInvInertiaScale0());
+		Solver.SolverBody0().SetShockPropagationScale(FReal(1));
+		Solver.SolverBody1().SetInvMScale(Constraint->GetInvMassScale1());
+		Solver.SolverBody1().SetInvIScale(Constraint->GetInvInertiaScale1());
+		Solver.SolverBody1().SetShockPropagationScale(FReal(1));
 
-					ManifoldPoint.ShapeAnchorPoints[0] = SavedManifoldPoint->ShapeContactPoints[0];
-					ManifoldPoint.ShapeAnchorPoints[1] = SavedManifoldPoint->ShapeContactPoints[1];
-				}
-				else
-				{
-					const FSolverVec3 ContactVel0 = Body0.V() + FSolverVec3::CrossProduct(Body0.W(), RelativeContactPosition0);
-					const FSolverVec3 ContactVel1 = Body1.V() + FSolverVec3::CrossProduct(Body1.W(), RelativeContactPosition1);
-					const FSolverVec3 ContactVel = ContactVel0 - ContactVel1;
-					WorldFrictionDelta = ContactVel * Dt;
+		bOutPerIterationCollision = (!Constraint->GetUseManifold() || Constraint->GetUseIncrementalCollisionDetection());
 
-					ManifoldPoint.ShapeAnchorPoints[0] = ManifoldPoint.ContactPoint.ShapeContactPoints[0];
-					ManifoldPoint.ShapeAnchorPoints[1] = ManifoldPoint.ContactPoint.ShapeContactPoints[1];
-				}
+		UpdateCollisionSolverManifoldFromConstraint(Solver, Constraint, Dt, 0, Constraint->NumManifoldPoints());
+	}
 
-				// World-space contact tangents. We are treating the normal as the constraint-space Z axis
-				// and the Tangent U and V as the constraint-space X and Y axes respectively
-				FSolverVec3 WorldContactTangentU = FSolverVec3::CrossProduct(FSolverVec3(0, 1, 0), WorldContactNormal);
-				if (!WorldContactTangentU.Normalize(FSolverReal(UE_KINDA_SMALL_NUMBER)))
-				{
-					WorldContactTangentU = FSolverVec3::CrossProduct(FSolverVec3(1, 0, 0), WorldContactNormal);
-					WorldContactTangentU = WorldContactTangentU.GetUnsafeNormal();
-				}
-				const FSolverVec3 WorldContactTangentV = FSolverVec3::CrossProduct(WorldContactNormal, WorldContactTangentU);
-
-				// The contact point error we are trying to correct in this solver
-				const FSolverVec3 WorldContactDelta = FSolverVec3(WorldContactPoint0 - WorldContactPoint1);
-				const FSolverReal WorldContactDeltaNormal = FSolverVec3::DotProduct(WorldContactDelta, WorldContactNormal) - TargetPhi;
-				const FSolverReal WorldContactDeltaTangentU = FSolverVec3::DotProduct(WorldContactDelta + WorldFrictionDelta, WorldContactTangentU);
-				const FSolverReal WorldContactDeltaTangentV = FSolverVec3::DotProduct(WorldContactDelta + WorldFrictionDelta, WorldContactTangentV);
-
-				// Copy all the properties into the solver
-				Solver.SetManifoldPoint(
-					ManifoldPointIndex,
-					Dt,
-					Restitution,
-					RestitutionVelocityThreshold,
-					RelativeContactPosition0,
-					RelativeContactPosition1,
-					WorldContactNormal,
-					WorldContactTangentU,
-					WorldContactTangentV,
-					WorldContactDeltaNormal,
-					WorldContactDeltaTangentU,
-					WorldContactDeltaTangentV);
-
-#if CHAOS_DEBUG_DRAW && 0
-				if (FDebugDrawQueue::GetInstance().IsDebugDrawingEnabled())
-				{
-					FDebugDrawQueue::GetInstance().DrawDebugLine(WorldContactPoint1, WorldContactPoint1 - 5.0f * WorldContactDeltaNormal * WorldContactNormal, FColor::Red, false, 0, 10, 0.2f);
-					FDebugDrawQueue::GetInstance().DrawDebugLine(WorldContactPoint1, WorldContactPoint1 - 5.0f * WorldContactDeltaTangentU * WorldContactTangentU, FColor::Blue, false, 0, 10, 0.2f);
-					FDebugDrawQueue::GetInstance().DrawDebugLine(WorldContactPoint1, WorldContactPoint1 - 5.0f * WorldContactDeltaTangentV * WorldContactTangentV, FColor::Blue, false, 0, 10, 0.2f);
-					FDebugDrawQueue::GetInstance().DrawDebugLine(WorldContactPoint0, WorldContactPoint1, FColor::Black, false, 0, 10, 0.2f);
-				}
-#endif
-			}
-		}
+	FORCEINLINE_DEBUGGABLE void UpdateCollisionConstraintFromSolver(FPBDCollisionConstraint* Constraint, const FPBDCollisionSolver& Solver, const FSolverReal Dt)
+	{
+		Constraint->ResetSolverResults();
 
-		/**
-		 * @brief Send all solver results to the constraint
-		*/
-		void ScatterOutput(const FReal Dt)
+		// NOTE: We only put the non-pruned manifold points into the solver so the ManifoldPointIndex and
+		// SolverManifoldPointIndex do not necessarily match. See GatherManifoldPoints
+		int32 SolverManifoldPointIndex = 0;
+		for (int32 ManifoldPointIndex = 0; ManifoldPointIndex < Constraint->NumManifoldPoints(); ++ManifoldPointIndex)
 		{
-			Constraint->ResetSolverResults();
+			FSolverVec3 NetPushOut = FSolverVec3(0);
+			FSolverVec3 NetImpulse = FSolverVec3(0);
+			FSolverReal StaticFrictionRatio = FSolverReal(0);
 
-			// NOTE: We only put the non-pruned manifold points into the solver so the ManifoldPointIndex and
-			// SolverManifoldPointIndex do not necessarily match. See GatherManifoldPoints
-			int32 SolverManifoldPointIndex = 0;
-			for (int32 ManifoldPointIndex = 0; ManifoldPointIndex < Constraint->NumManifoldPoints(); ++ManifoldPointIndex)
+			const FPBDCollisionSolverManifoldPoint& SolverManifoldPoint = Solver.GetManifoldPoint(SolverManifoldPointIndex);
+			if (!Constraint->GetManifoldPoint(ManifoldPointIndex).Flags.bDisabled)
 			{
-				FSolverVec3 NetPushOut = FSolverVec3(0);
-				FSolverVec3 NetImpulse = FSolverVec3(0);
-				FRealSingle StaticFrictionRatio = FRealSingle(0);
-
-				if (!Constraint->GetManifoldPoint(ManifoldPointIndex).Flags.bDisabled)
-				{
-					const FPBDCollisionSolverManifoldPoint& SolverManifoldPoint = Solver.GetManifoldPoint(SolverManifoldPointIndex);
-
-					NetPushOut = 
-						SolverManifoldPoint.NetPushOutNormal * SolverManifoldPoint.WorldContactNormal +
-						SolverManifoldPoint.NetPushOutTangentU * SolverManifoldPoint.WorldContactTangentU +
-						SolverManifoldPoint.NetPushOutTangentV * SolverManifoldPoint.WorldContactTangentV;
+				NetPushOut =
+					SolverManifoldPoint.NetPushOutNormal * SolverManifoldPoint.WorldContact.ContactNormal +
+					SolverManifoldPoint.NetPushOutTangentU * SolverManifoldPoint.WorldContact.ContactTangentU +
+					SolverManifoldPoint.NetPushOutTangentV * SolverManifoldPoint.WorldContact.ContactTangentV;
 
-					NetImpulse =
-						SolverManifoldPoint.NetImpulseNormal * SolverManifoldPoint.WorldContactNormal +
-						SolverManifoldPoint.NetImpulseTangentU * SolverManifoldPoint.WorldContactTangentU +
-						SolverManifoldPoint.NetImpulseTangentV * SolverManifoldPoint.WorldContactTangentV;
+				NetImpulse =
+					SolverManifoldPoint.NetImpulseNormal * SolverManifoldPoint.WorldContact.ContactNormal +
+					SolverManifoldPoint.NetImpulseTangentU * SolverManifoldPoint.WorldContact.ContactTangentU +
+					SolverManifoldPoint.NetImpulseTangentV * SolverManifoldPoint.WorldContact.ContactTangentV;
 
-					StaticFrictionRatio = SolverManifoldPoint.StaticFrictionRatio;
+				StaticFrictionRatio = SolverManifoldPoint.StaticFrictionRatio;
 
-					++SolverManifoldPointIndex;
-				}
-
-				// NOTE: We call this even for points we did not run the solver for (but with zero results)
-				Constraint->SetSolverResults(ManifoldPointIndex,
-					NetPushOut, 
-					NetImpulse, 
-					StaticFrictionRatio,
-					FRealSingle(Dt));
+				++SolverManifoldPointIndex;
 			}
 
-			Constraint->EndTick();
-
-			Unbind();
+			// NOTE: We call this even for points we did not run the solver for (but with zero results)
+			Constraint->SetSolverResults(ManifoldPointIndex,
+				NetPushOut,
+				NetImpulse,
+				StaticFrictionRatio,
+				Dt);
 		}
 
-
-	private:
-		FPBDCollisionSolver Solver;
-		FPBDCollisionConstraint* Constraint;
-		bool bIsManifold;
-		bool bIsIncrementalManifold;
-	};
-
+		Constraint->EndTick();
+	}
 
 	//////////////////////////////////////////////////////////////////////////////////////////////////
 	//////////////////////////////////////////////////////////////////////////////////////////////////
@@ -412,7 +287,6 @@ namespace Chaos
 		: FConstraintContainerSolver(InPriority)
 		, ConstraintContainer(InConstraintContainer)
 		, bPerIterationCollisionDetection(false)
-		, bDeferredCollisionDetection(false)
 	{
 	}
 
@@ -422,12 +296,28 @@ namespace Chaos
 
 	void FPBDCollisionContainerSolver::Reset(const int32 MaxCollisions)
 	{
-		CollisionSolvers.Reset(0);
-		CollisionSolvers.Reserve(MaxCollisions);
+		// A over-allocation policy to avoid reallocation every frame in the common case where a pile of objects is dropped
+		// and the number of contacts increases every tick.
+		if (MaxCollisions > CollisionSolvers.Max())
+		{
+			CollisionSolvers.Reset((5 * MaxCollisions) / 4);	// +25%
+		}
+
+		CollisionSolvers.Reset(MaxCollisions);
+		CollisionConstraints.Reset(CollisionSolvers.Max());
+		bCollisionConstraintPerIterationCollisionDetection.Reset(CollisionSolvers.Max());
+		CollisionSolverManifoldPoints.Reset(CollisionSolvers.Max() * FPBDCollisionSolver::MaxPointsPerConstraint);
+
+		// Just set the array size for these right away - all data will be initialized later and there is no constructor
+		// We use CollisionConstraints.Num() to track how many constraints we currently have
+		CollisionSolvers.SetNumUninitialized(MaxCollisions);
+		bCollisionConstraintPerIterationCollisionDetection.SetNumUninitialized(MaxCollisions);
 	}
 
 	void FPBDCollisionContainerSolver::AddConstraints()
 	{
+		Reset(ConstraintContainer.NumConstraints());
+
 		for (FPBDCollisionConstraintHandle* ConstraintHandle : ConstraintContainer.GetConstraintHandles())
 		{
 			FPBDCollisionConstraint& Constraint = ConstraintHandle->GetContact();
@@ -440,6 +330,7 @@ namespace Chaos
 	{
 		for (FPBDIslandConstraint& IslandConstraint : IslandConstraints)
 		{
+			// We will only ever be given constraints from our container (asserts in non-shipping)
 			FPBDCollisionConstraint& Constraint = IslandConstraint.GetConstraint()->AsUnsafe<FPBDCollisionConstraintHandle>()->GetContact();
 
 			AddConstraint(Constraint);
@@ -448,17 +339,37 @@ namespace Chaos
 
 	void FPBDCollisionContainerSolver::AddConstraint(FPBDCollisionConstraint& Constraint)
 	{
-		FPBDCollisionSolverAdapter& CollisionSolver = CollisionSolvers[CollisionSolvers.AddDefaulted()];
-		CollisionSolver.Reset();
-		CollisionSolver.SetConstraint(Constraint);
+		// NOTE: No need to add to CollisionSolvers or bCollisionConstraintPerIterationCollisionDetection
+		const int32 Index = CollisionConstraints.Add(&Constraint);
+
+		// Allocate the manifold points for this constraints solver
+		// NOTE: we don't know how many points we will create if collision detection is deferred or incrememntal, so just allocate space for the max allowed
+		// @todo(chaos): see if we can do better - we don't want to pay for this logic when it is only used rarely
+		const bool bDeferredCollisionDetection = ConstraintContainer.GetDetectorSettings().bDeferNarrowPhase;
+		const int32 ManifoldPointMax = (bDeferredCollisionDetection || Constraint.GetUseIncrementalManifold()) ? FPBDCollisionSolver::MaxPointsPerConstraint : Constraint.NumManifoldPoints();
+		if (ManifoldPointMax > 0)
+		{
+			const int32 ManifoldPointStart = CollisionSolverManifoldPoints.Num();
+			check(CollisionSolverManifoldPoints.Max() >= ManifoldPointStart + ManifoldPointMax);	// No resize following reset
+			CollisionSolverManifoldPoints.SetNum(ManifoldPointStart + ManifoldPointMax);
+
+			CollisionSolvers[Index].SetManifoldPointsBuffer(&CollisionSolverManifoldPoints[ManifoldPointStart], ManifoldPointMax);
+		}
+		else
+		{
+			// Ideally we wouldn't have any active constraint with no manifold points, but this can
+			// currently happen for CCD objects and incremental collisions
+			CollisionSolvers[Index].Reset();
+		}
 	}
 
 
 	void FPBDCollisionContainerSolver::AddBodies(FSolverBodyContainer& SolverBodyContainer)
 	{
-		for (FPBDCollisionSolverAdapter& CollisionSolver : CollisionSolvers)
+		for (int32 SolverIndex = 0, SolverEndIndex = NumSolvers(); SolverIndex < SolverEndIndex; ++SolverIndex)
 		{
-			FPBDCollisionConstraint* Constraint = CollisionSolver.GetConstraint();
+			FPBDCollisionSolver& CollisionSolver = CollisionSolvers[SolverIndex];
+			FPBDCollisionConstraint* Constraint = CollisionConstraints[SolverIndex];
 			check(Constraint != nullptr);
 
 			// Find the solver bodies for the particles we constrain. This will add them to the container
@@ -468,51 +379,56 @@ namespace Chaos
 			check(Body0 != nullptr);
 			check(Body1 != nullptr);
 
-			CollisionSolver.SetBodies(*Body0, *Body1);
+			CollisionSolver.SetSolverBodies(*Body0, *Body1);
 		}
 	}
 
 	void FPBDCollisionContainerSolver::GatherInput(const FReal Dt)
 	{
-		if (bDeferredCollisionDetection)
-		{
-			return;
-		}
+		GatherInput(Dt, 0, NumSolvers());
+	}
 
-		for (FPBDCollisionSolverAdapter& CollisionSolver : CollisionSolvers)
+	void FPBDCollisionContainerSolver::CachePrefetchSolver(const int32 ConstraintIndex) const
+	{
+		if (ConstraintIndex < NumSolvers())
 		{
-			CollisionSolver.GatherInput(Dt, ConstraintContainer.GetSolverSettings());
-
-			// We need to run collision every iteration if we are not using manifolds, or are using incremental manifolds
-			if (!CollisionSolver.IsManifold() || CollisionSolver.IsIncrementalManifold())
+			const int32 NumManifoldPoints = CollisionSolvers[ConstraintIndex].NumManifoldPoints();
+			if (NumManifoldPoints > 0)
 			{
-				bPerIterationCollisionDetection = true;
+				FPlatformMisc::PrefetchBlock(&CollisionSolvers[ConstraintIndex].GetManifoldPoint(0), NumManifoldPoints * sizeof(FPBDCollisionSolverManifoldPoint));
 			}
+			FPlatformMisc::PrefetchBlock(&CollisionSolvers[ConstraintIndex], sizeof(FPBDCollisionContainerSolver));
 		}
 	}
 
-	void FPBDCollisionContainerSolver::GatherInput(const FReal Dt, const int32 BeginIndex, const int32 EndIndex)
+	void FPBDCollisionContainerSolver::GatherInput(const FReal InDt, const int32 BeginIndex, const int32 EndIndex)
 	{
-		if (bDeferredCollisionDetection)
-		{
-			return;
-		}
+		// NOTE: may be called in parallel. Should not change the container or any elements outside of [BeginIndex, EndIndex)
 
 		check(BeginIndex >= 0);
 		check(EndIndex <= NumSolvers());
 
+		const FSolverReal Dt = FSolverReal(InDt);
+
+		const int32 PrefetchCount = 4;
+		for (int32 PrefetchIndex = 0; PrefetchIndex < PrefetchCount; ++PrefetchIndex)
+		{
+			CachePrefetchSolver(PrefetchIndex);
+		}
+
 		bool bWantPerIterationCollisionDetection = false;
 		for (int32 ConstraintIndex = BeginIndex; ConstraintIndex < EndIndex; ++ConstraintIndex)
 		{
-			FPBDCollisionSolverAdapter& CollisionSolver = CollisionSolvers[ConstraintIndex];
-			
-			CollisionSolver.GatherInput(Dt, ConstraintContainer.GetSolverSettings());
+			FPBDCollisionSolver& CollisionSolver = CollisionSolvers[ConstraintIndex];
+			FPBDCollisionConstraint* Constraint = CollisionConstraints[ConstraintIndex];
+			bool& bPerIterationCollision = bCollisionConstraintPerIterationCollisionDetection[ConstraintIndex];
+
+			UpdateCollisionSolverFromConstraint(CollisionSolver, Constraint, Dt, ConstraintContainer.GetSolverSettings(), bPerIterationCollision);
 
 			// We need to run collision every iteration if we are not using manifolds, or are using incremental manifolds
-			if (!CollisionSolver.IsManifold() || CollisionSolver.IsIncrementalManifold())
-			{
-				bWantPerIterationCollisionDetection = true;
-			}
+			bWantPerIterationCollisionDetection = bWantPerIterationCollisionDetection || bPerIterationCollision;
+
+			CachePrefetchSolver(ConstraintIndex + PrefetchCount);
 		}
 
 		if (bWantPerIterationCollisionDetection)
@@ -524,32 +440,39 @@ namespace Chaos
 
 	void FPBDCollisionContainerSolver::ScatterOutput(const FReal Dt)
 	{
-		for (FPBDCollisionSolverAdapter& CollisionSolver : CollisionSolvers)
-		{
-			CollisionSolver.ScatterOutput(Dt);
-		}
+		ScatterOutput(Dt, 0, NumSolvers());
 	}
 
-	void FPBDCollisionContainerSolver::ScatterOutput(const FReal Dt, const int32 BeginIndex, const int32 EndIndex)
+	void FPBDCollisionContainerSolver::ScatterOutput(const FReal InDt, const int32 BeginIndex, const int32 EndIndex)
 	{
+		// NOTE: may be called in parallel. Should not change the container or any elements outside of [BeginIndex, EndIndex)
+
 		check(BeginIndex >= 0);
 		check(EndIndex <= NumSolvers());
 
+		const FSolverReal Dt = FSolverReal(InDt);
+
 		for (int32 ConstraintIndex = BeginIndex; ConstraintIndex < EndIndex; ++ConstraintIndex)
 		{
-			FPBDCollisionSolverAdapter& CollisionSolver = CollisionSolvers[ConstraintIndex];
-			CollisionSolver.ScatterOutput(Dt);
+			FPBDCollisionSolver& CollisionSolver = CollisionSolvers[ConstraintIndex];
+			FPBDCollisionConstraint* Constraint = CollisionConstraints[ConstraintIndex];
+
+			UpdateCollisionConstraintFromSolver(Constraint, CollisionSolver, Dt);
+
+			// Reset the collision solver here as the pointers will be invalid on the next tick
+			// @todo(chaos): Pointers are always reinitalized before use next tick so this isn't strictly necessary
+			CollisionSolver.Reset();
 		}
 	}
 
 	void FPBDCollisionContainerSolver::ApplyPositionConstraints(const FReal Dt, const int32 It, const int32 NumIts)
 	{
-		SolvePositionImpl(Dt, It, NumIts, 0, CollisionSolvers.Num(), ConstraintContainer.GetSolverSettings());
+		SolvePositionImpl(Dt, It, NumIts, 0, NumSolvers(), ConstraintContainer.GetSolverSettings());
 	}
 
 	void FPBDCollisionContainerSolver::ApplyVelocityConstraints(const FReal Dt, const int32 It, const int32 NumIts)
 	{
-		SolveVelocityImpl(Dt, It, NumIts, 0, CollisionSolvers.Num(), ConstraintContainer.GetSolverSettings());
+		SolveVelocityImpl(Dt, It, NumIts, 0, NumSolvers(), ConstraintContainer.GetSolverSettings());
 	}
 
 	void FPBDCollisionContainerSolver::ApplyProjectionConstraints(const FReal Dt, const int32 It, const int32 NumIts)
@@ -565,7 +488,7 @@ namespace Chaos
 		{
 			for (int32 SolverIndex = BeginIndex; SolverIndex < EndIndex; ++SolverIndex)
 			{
-				CollisionSolvers[SolverIndex].GetSolver().EnablePositionShockPropagation();
+				CollisionSolvers[SolverIndex].EnablePositionShockPropagation();
 			}
 		}
 	}
@@ -579,14 +502,14 @@ namespace Chaos
 		{
 			for (int32 SolverIndex = BeginIndex; SolverIndex < EndIndex; ++SolverIndex)
 			{
-				CollisionSolvers[SolverIndex].GetSolver().EnableVelocityShockPropagation();
+				CollisionSolvers[SolverIndex].EnableVelocityShockPropagation();
 			}
 		}
 		else if (It == 0)
 		{
 			for (int32 SolverIndex = BeginIndex; SolverIndex < EndIndex; ++SolverIndex)
 			{
-				CollisionSolvers[SolverIndex].GetSolver().DisableShockPropagation();
+				CollisionSolvers[SolverIndex].DisableShockPropagation();
 			}
 		}
 	}
@@ -611,8 +534,8 @@ namespace Chaos
 		// We run collision detection here under two conditions (normally it is run after Integration and before the constraint solver phase):
 		// 1) When deferring collision detection until the solver phase for better joint-collision behaviour (RBAN). In this case, we only do this on the first iteration.
 		// 2) When using no manifolds or incremental manifolds, where we may add/replace manifold points every iteration.
-		const bool bDeferredCollisions = bDeferredCollisionDetection && (It == 0);
-		if (bDeferredCollisions || bPerIterationCollisionDetection)
+		const bool bRunDeferredCollisionDetection = (It == 0) && ConstraintContainer.GetDetectorSettings().bDeferNarrowPhase;
+		if (bRunDeferredCollisionDetection || bPerIterationCollisionDetection)
 		{
 			UpdateCollisions(Dt, BeginIndex, EndIndex);
 		}
@@ -642,14 +565,14 @@ namespace Chaos
 		{
 			for (int32 SolverIndex = BeginIndex; SolverIndex < EndIndex; ++SolverIndex)
 			{
-				CollisionSolvers[SolverIndex].GetSolver().SolvePositionWithFriction(Dt, MaxPushOut);
+				CollisionSolvers[SolverIndex].SolvePositionWithFriction(Dt, MaxPushOut);
 			}
 		}
 		else
 		{
 			for (int32 SolverIndex = BeginIndex; SolverIndex < EndIndex; ++SolverIndex)
 			{
-				CollisionSolvers[SolverIndex].GetSolver().SolvePositionWithFrictionJacobi(Dt, MaxPushOut);
+				CollisionSolvers[SolverIndex].SolvePositionWithFrictionJacobi(Dt, MaxPushOut);
 			}
 		}
 
@@ -670,14 +593,14 @@ namespace Chaos
 		{
 			for (int32 SolverIndex = BeginIndex; SolverIndex < EndIndex; ++SolverIndex)
 			{
-				CollisionSolvers[SolverIndex].GetSolver().SolvePositionNoFriction(Dt, MaxPushOut);
+				CollisionSolvers[SolverIndex].SolvePositionNoFriction(Dt, MaxPushOut);
 			}
 		}
 		else
 		{
 			for (int32 SolverIndex = BeginIndex; SolverIndex < EndIndex; ++SolverIndex)
 			{
-				CollisionSolvers[SolverIndex].GetSolver().SolvePositionNoFrictionJacobi(Dt, MaxPushOut);
+				CollisionSolvers[SolverIndex].SolvePositionNoFrictionJacobi(Dt, MaxPushOut);
 			}
 		}
 
@@ -704,14 +627,14 @@ namespace Chaos
 		{
 			for (int32 SolverIndex = BeginIndex; SolverIndex < EndIndex; ++SolverIndex)
 			{
-				bNeedsAnotherIteration |= CollisionSolvers[SolverIndex].GetSolver().SolveVelocity(Dt, bApplyDynamicFriction);
+				bNeedsAnotherIteration |= CollisionSolvers[SolverIndex].SolveVelocity(Dt, bApplyDynamicFriction);
 			}
 		}
 		else
 		{
 			for (int32 SolverIndex = BeginIndex; SolverIndex < EndIndex; ++SolverIndex)
 			{
-				bNeedsAnotherIteration |= CollisionSolvers[SolverIndex].GetSolver().SolveVelocityJacobi(Dt, bApplyDynamicFriction);
+				bNeedsAnotherIteration |= CollisionSolvers[SolverIndex].SolveVelocityJacobi(Dt, bApplyDynamicFriction);
 			}
 		}
 
@@ -721,23 +644,38 @@ namespace Chaos
 	void FPBDCollisionContainerSolver::UpdateCollisions(const FReal InDt, const int32 BeginIndex, const int32 EndIndex)
 	{
 		const FSolverReal Dt = FSolverReal(InDt);
+		const bool bDeferredCollisionDetection = ConstraintContainer.GetDetectorSettings().bDeferNarrowPhase;
 
 		bool bNeedsAnotherIteration = false;
 		for (int32 SolverIndex = BeginIndex; SolverIndex < EndIndex; ++SolverIndex)
 		{
-			FPBDCollisionSolverAdapter& CollisionSolver = CollisionSolvers[SolverIndex];
-			if (!CollisionSolver.IsManifold() || CollisionSolver.IsIncrementalManifold() || bDeferredCollisionDetection)
+			if (bDeferredCollisionDetection || bCollisionConstraintPerIterationCollisionDetection[SolverIndex])
 			{
-				FPBDCollisionConstraint* Constraint = CollisionSolver.GetConstraint();
+				FPBDCollisionSolver& CollisionSolver = CollisionSolvers[SolverIndex];
+				FPBDCollisionConstraint* Constraint = CollisionConstraints[SolverIndex];
 
 				// Run collision detection at the current transforms including any correction from previous iterations
-				const FSolverBody& Body0 = CollisionSolver.GetSolver().SolverBody0().SolverBody();
-				const FSolverBody& Body1 = CollisionSolver.GetSolver().SolverBody1().SolverBody();
+				const FSolverBody& Body0 = CollisionSolver.SolverBody0().SolverBody();
+				const FSolverBody& Body1 = CollisionSolver.SolverBody1().SolverBody();
 				const FRigidTransform3 CorrectedActorWorldTransform0 = FRigidTransform3(Body0.CorrectedActorP(), Body0.CorrectedActorQ());
 				const FRigidTransform3 CorrectedActorWorldTransform1 = FRigidTransform3(Body1.CorrectedActorP(), Body1.CorrectedActorQ());
 				const FRigidTransform3 CorrectedShapeWorldTransform0 = Constraint->GetShapeRelativeTransform0() * CorrectedActorWorldTransform0;
 				const FRigidTransform3 CorrectedShapeWorldTransform1 = Constraint->GetShapeRelativeTransform1() * CorrectedActorWorldTransform1;
 
+				// @todo(chaos): this is ugly - pass these to the required functions instead and remove from the constraint class
+				// This is now only needed for LevelSet collision (see UpdateLevelsetLevelsetConstraint)
+				Constraint->SetSolverBodies(&Body0, &Body1);
+
+				// Reset the manifold if we are not using manifolds (we just use the first manifold point)
+				if (!Constraint->GetUseManifold())
+				{
+					Constraint->ResetActiveManifoldContacts();
+					CollisionSolver.ResetManifold();
+				}
+
+				// We need to know how many points were added to the manifold
+				const int32 BeginPointIndex = Constraint->NumManifoldPoints();
+
 				// NOTE: We deliberately have not updated the ShapwWorldTranforms on the constraint. If we did that, we would calculate 
 				// errors incorrectly in UpdateManifoldPoints, because the solver assumes nothing has been moved as we iterate (we accumulate 
 				// corrections that will be applied later.)
@@ -745,7 +683,9 @@ namespace Chaos
 				Collisions::UpdateConstraint(*Constraint, CorrectedShapeWorldTransform0, CorrectedShapeWorldTransform1, Dt);
 
 				// Update the manifold based on the new or updated contacts
-				CollisionSolver.UpdateManifoldPoints(Dt);
+				UpdateCollisionSolverManifoldFromConstraint(CollisionSolver, Constraint, Dt, BeginPointIndex, Constraint->NumManifoldPoints());
+
+				Constraint->SetSolverBodies(nullptr, nullptr);
 			}
 		}
 	}
diff --git a/Engine/Source/Runtime/Experimental/Chaos/Private/Chaos/ContactModification.cpp b/Engine/Source/Runtime/Experimental/Chaos/Private/Chaos/ContactModification.cpp
--- a/Engine/Source/Runtime/Experimental/Chaos/Private/Chaos/ContactModification.cpp
+++ b/Engine/Source/Runtime/Experimental/Chaos/Private/Chaos/ContactModification.cpp
@@ -62,10 +62,10 @@ namespace Chaos
 		if (FPBDRigidParticleHandle* Rigid = Particle->CastToRigidParticle())
 		{
 			// Use PQ for rigids.
-			return Constraint->GetShapeRelativeTransform(ParticleIdx) * FParticleUtilitiesPQ::GetActorWorldTransform(Rigid);
+			return Constraint->GetShapeRelativeTransform(ParticleIdx) * Rigid->GetTransformPQ();
 		}
 
-		return Constraint->GetShapeRelativeTransform(ParticleIdx) * FParticleUtilitiesXR::GetActorWorldTransform(Particle);
+		return Constraint->GetShapeRelativeTransform(ParticleIdx) * Particle->GetTransformXR();
 	}
 
 	FReal FContactPairModifier::GetSeparation(int32 ContactPointIdx) const
@@ -153,6 +153,8 @@ namespace Chaos
 
 		// Clear the friction data since it will now have the wrong contact positions
 		// @todo(chaos): maybe try to do something better here
+		ManifoldPoint.ShapeAnchorPoints[0] = ManifoldPoint.ContactPoint.ShapeContactPoints[0];
+		ManifoldPoint.ShapeAnchorPoints[1] = ManifoldPoint.ContactPoint.ShapeContactPoints[1];
 		Constraint->ResetSavedManifoldPoints();
 
 		Modifier->MarkConstraintForManifoldUpdate(*Constraint);
@@ -297,8 +299,8 @@ namespace Chaos
 
 	void FContactPairModifier::UpdateConstraintShapeTransforms()
 	{
-		const FRigidTransform3 ShapeWorldTransform0 = Constraint->GetShapeRelativeTransform0() * FParticleUtilitiesPQ::GetActorWorldTransform(FConstGenericParticleHandle(Constraint->GetParticle0()));
-		const FRigidTransform3 ShapeWorldTransform1 = Constraint->GetShapeRelativeTransform1() * FParticleUtilitiesPQ::GetActorWorldTransform(FConstGenericParticleHandle(Constraint->GetParticle1()));
+		const FRigidTransform3 ShapeWorldTransform0 = Constraint->GetShapeRelativeTransform0() * FConstGenericParticleHandle(Constraint->GetParticle0())->GetTransformPQ();
+		const FRigidTransform3 ShapeWorldTransform1 = Constraint->GetShapeRelativeTransform1() * FConstGenericParticleHandle(Constraint->GetParticle1())->GetTransformPQ();
 		Constraint->SetShapeWorldTransforms(ShapeWorldTransform0, ShapeWorldTransform1);
 	}
 
diff --git a/Engine/Source/Runtime/Experimental/Chaos/Private/Chaos/Evolution/ConstraintGroupSolver.cpp b/Engine/Source/Runtime/Experimental/Chaos/Private/Chaos/Evolution/ConstraintGroupSolver.cpp
--- a/Engine/Source/Runtime/Experimental/Chaos/Private/Chaos/Evolution/ConstraintGroupSolver.cpp
+++ b/Engine/Source/Runtime/Experimental/Chaos/Private/Chaos/Evolution/ConstraintGroupSolver.cpp
@@ -92,7 +92,7 @@ namespace Chaos
 
 	void FPBDConstraintGroupSolver::GatherBodies(const FReal Dt)
 	{
-		SolverBodyContainer.GatherInput(Dt, 0, SolverBodyContainer.NumItems());
+		SolverBodyContainer.GatherInput(Dt, 0, SolverBodyContainer.Num());
 	}
 
 	void FPBDConstraintGroupSolver::GatherConstraints(const FReal Dt)
@@ -160,7 +160,7 @@ namespace Chaos
 
 	void FPBDConstraintGroupSolver::ScatterBodies(const FReal Dt)
 	{
-		SolverBodyContainer.ScatterOutput();
+		SolverBodyContainer.ScatterOutput(0, SolverBodyContainer.Num());
 	}
 
 	void FPBDConstraintGroupSolver::ScatterConstraints(const FReal Dt)
diff --git a/Engine/Source/Runtime/Experimental/Chaos/Private/Chaos/Evolution/PBDMinEvolution.cpp b/Engine/Source/Runtime/Experimental/Chaos/Private/Chaos/Evolution/PBDMinEvolution.cpp
--- a/Engine/Source/Runtime/Experimental/Chaos/Private/Chaos/Evolution/PBDMinEvolution.cpp
+++ b/Engine/Source/Runtime/Experimental/Chaos/Private/Chaos/Evolution/PBDMinEvolution.cpp
@@ -45,9 +45,6 @@ namespace Chaos
 	//
 	//
 
-	bool bChaos_MinEvolution_RewindLerp = true;
-	FAutoConsoleVariableRef CVarChaosMinEvolutionRewindLerp(TEXT("p.Chaos.MinEvolution.RewindLerp"), bChaos_MinEvolution_RewindLerp, TEXT("If rewinding (fixed dt mode) use Backwards-Lerp as opposed to Backwards Velocity"));
-
 	// Forced iteration count to evaluate worst-case behaviour for a given simulation
 	bool Chaos_MinEvolution_ForceMaxConstraintIterations = false;
 	FAutoConsoleVariableRef CVarChaosMinEvolutionForceMaxConstraintIterations(TEXT("p.Chaos.MinEvolution.ForceMaxConstraintIterations"), Chaos_MinEvolution_ForceMaxConstraintIterations, TEXT("Whether to force constraints to always use the worst-case maximum number of iterations"));
@@ -56,7 +53,7 @@ namespace Chaos
 	//
 	//
 
-	FPBDMinEvolution::FPBDMinEvolution(FRigidParticleSOAs& InParticles, TArrayCollectionArray<FVec3>& InPrevX, TArrayCollectionArray<FRotation3>& InPrevR, FCollisionDetector& InCollisionDetector, const FReal InBoundsExtension)
+	FPBDMinEvolution::FPBDMinEvolution(FRigidParticleSOAs& InParticles, TArrayCollectionArray<FVec3>& InPrevX, TArrayCollectionArray<FRotation3>& InPrevR, FCollisionDetector& InCollisionDetector)
 		: Particles(InParticles)
 		, CollisionDetector(InCollisionDetector)
 		, ParticlePrevXs(InPrevX)
@@ -64,7 +61,6 @@ namespace Chaos
 		, NumPositionIterations(0)
 		, NumVelocityIterations(0)
 		, NumProjectionIterations(0)
-		, BoundsExtension(InBoundsExtension)
 		, Gravity(FVec3(0))
 		, SimulationSpaceSettings()
 	{
@@ -174,33 +170,14 @@ namespace Chaos
 	{
 		SCOPE_CYCLE_COUNTER(STAT_MinEvolution_Rewind);
 
-		if (bChaos_MinEvolution_RewindLerp)
-		{
-			const FReal T = (Dt - RewindDt) / Dt;
-			UE_LOG(LogChaosMinEvolution, Verbose, TEXT("Rewind dt = %f; rt = %f; T = %f"), Dt, RewindDt, T);
-			for (TTransientPBDRigidParticleHandle<FReal, 3>& Particle : Particles.GetActiveParticlesView())
-			{
-				if (Particle.ObjectState() == EObjectStateType::Dynamic)
-				{
-					Particle.X() = FVec3::Lerp(Particle.Handle()->AuxilaryValue(ParticlePrevXs), Particle.X(), T);
-					Particle.R() = FRotation3::Slerp(Particle.Handle()->AuxilaryValue(ParticlePrevRs), Particle.R(), (decltype(FQuat::X))T);	// LWC_TODO: Remove decltype cast once FQuat supports variants
-				}
-			}
-		}
-		else
+		const FReal T = (Dt - RewindDt) / Dt;
+		UE_LOG(LogChaosMinEvolution, Verbose, TEXT("Rewind dt = %f; rt = %f; T = %f"), Dt, RewindDt, T);
+		for (TTransientPBDRigidParticleHandle<FReal, 3>& Particle : Particles.GetActiveParticlesView())
 		{
-			for (TTransientPBDRigidParticleHandle<FReal, 3>& Particle : Particles.GetActiveParticlesView())
+			if (Particle.ObjectState() == EObjectStateType::Dynamic)
 			{
-				if (Particle.ObjectState() == EObjectStateType::Dynamic)
-				{
-					const FVec3 XCoM = FParticleUtilitiesXR::GetCoMWorldPosition(&Particle);
-					const FRotation3 RCoM = FParticleUtilitiesXR::GetCoMWorldRotation(&Particle);
-
-					const FVec3 XCoM2 = XCoM - Particle.V() * RewindDt;
-					const FRotation3 RCoM2 = FRotation3::IntegrateRotationWithAngularVelocity(RCoM, -Particle.W(), RewindDt);
-
-					FParticleUtilitiesXR::SetCoMWorldTransform(&Particle, XCoM2, RCoM2);
-				}
+				Particle.X() = FVec3::Lerp(Particle.Handle()->AuxilaryValue(ParticlePrevXs), Particle.X(), T);
+				Particle.R() = FRotation3::Slerp(Particle.Handle()->AuxilaryValue(ParticlePrevRs), Particle.R(), (decltype(FQuat::X))T);	// LWC_TODO: Remove decltype cast once FQuat supports variants
 			}
 		}
 
@@ -228,6 +205,8 @@ namespace Chaos
 			SpaceB = SimulationSpace.Transform.InverseTransformVector(SimulationSpace.AngularAcceleration);
 		}
 
+		const FVec3 BoundsExpansion = FVec3(CollisionDetector.GetCollisionContainer().GetDetectorSettings().BoundsExpansion);
+
 		for (TTransientPBDRigidParticleHandle<FReal, 3>& Particle : Particles.GetActiveParticlesView())
 		{
 			if (Particle.ObjectState() == EObjectStateType::Dynamic)
@@ -235,9 +214,9 @@ namespace Chaos
 				Particle.PreV() = Particle.V();
 				Particle.PreW() = Particle.W();
 
-				const FVec3 XCoM = FParticleUtilitiesXR::GetCoMWorldPosition(&Particle);
-				const FRotation3 RCoM = FParticleUtilitiesXR::GetCoMWorldRotation(&Particle);
-				
+				const FVec3 XCoM = Particle.XCom();
+				const FRotation3 RCoM = Particle.RCom();
+
 				// Forces and torques
 				const FMatrix33 WorldInvI = Utilities::ComputeWorldSpaceInertia(RCoM, Particle.InvI());
 				FVec3 DV = Particle.Acceleration() * Dt + Particle.LinearImpulseVelocity();
@@ -278,7 +257,7 @@ namespace Chaos
 				const FRotation3 QCoM = FRotation3::IntegrateRotationWithAngularVelocity(RCoM, W, Dt);
 
 				// Update particle state (forces are not zeroed until the end of the frame)
-				FParticleUtilitiesPQ::SetCoMWorldTransform(&Particle, PCoM, QCoM);
+				Particle.SetTransformPQCom(PCoM, QCoM);
 				Particle.V() = V;
 				Particle.W() = W;
 				Particle.LinearImpulseVelocity() = FVec3(0);
@@ -287,7 +266,7 @@ namespace Chaos
 				// Update cached world space state, including bounds. We use the Swept bounds update so that the bounds includes P,Q and X,Q.
 				// This is because when we have joints, they often pull bodies back to their original positions, so we need to know if there
 				// are contacts at that location.
-				Particle.UpdateWorldSpaceStateSwept(FRigidTransform3(Particle.P(), Particle.Q()), FVec3(CollisionDetector.GetSettings().BoundsExpansion), -V * Dt);
+				Particle.UpdateWorldSpaceStateSwept(FRigidTransform3(Particle.P(), Particle.Q()), BoundsExpansion, -V * Dt);
 			}
 		}
 	}
@@ -362,7 +341,8 @@ namespace Chaos
 			}
 			}
 
-			Particle.UpdateWorldSpaceState(FRigidTransform3(Particle.X(), Particle.R()), FVec3(BoundsExtension));
+			// NOTE: we do not expand the bounds of kinematics, only dynamics
+			Particle.UpdateWorldSpaceState(FRigidTransform3(Particle.X(), Particle.R()), FVec3(0));
 		}
 	}
 
diff --git a/Engine/Source/Runtime/Experimental/Chaos/Private/Chaos/Evolution/SolverBody.cpp b/Engine/Source/Runtime/Experimental/Chaos/Private/Chaos/Evolution/SolverBody.cpp
--- a/Engine/Source/Runtime/Experimental/Chaos/Private/Chaos/Evolution/SolverBody.cpp
+++ b/Engine/Source/Runtime/Experimental/Chaos/Private/Chaos/Evolution/SolverBody.cpp
@@ -7,10 +7,6 @@
 namespace Chaos
 {
 
-	FSolverBody::FSolverBody()
-	{
-	}
-
 	void FSolverBody::UpdateRotationDependentState()
 	{
 		if (IsDynamic())
diff --git a/Engine/Source/Runtime/Experimental/Chaos/Private/Chaos/Evolution/SolverBodyContainer.cpp b/Engine/Source/Runtime/Experimental/Chaos/Private/Chaos/Evolution/SolverBodyContainer.cpp
--- a/Engine/Source/Runtime/Experimental/Chaos/Private/Chaos/Evolution/SolverBodyContainer.cpp
+++ b/Engine/Source/Runtime/Experimental/Chaos/Private/Chaos/Evolution/SolverBodyContainer.cpp
@@ -4,32 +4,87 @@
 
 namespace Chaos
 {
-	void FSolverBodyAdapter::ScatterOutput()
+	inline void UpdateSolverBodyFromParticle(FSolverBody& SolverBody, const FGeometryParticleHandle* Particle, const FReal Dt)
 	{
-		if (Particle.IsValid())
+		check(Particle != nullptr);
+
+		SolverBody.Reset();
+		SolverBody.SetLevel(0);
+
+		const FPBDRigidParticleHandle* RigidParticle = Particle->CastToRigidParticle();
+		const FKinematicGeometryParticleHandle* KinematicParticle = Particle->CastToKinematicParticle();
+		if ((RigidParticle != nullptr) && RigidParticle->IsDynamic())
+		{
+			// Dynamic particle
+			SolverBody.SetX(RigidParticle->XCom());
+			SolverBody.SetR(RigidParticle->RCom());
+			SolverBody.SetP(RigidParticle->PCom());
+			SolverBody.SetQ(RigidParticle->QCom());
+			SolverBody.SetCoM(RigidParticle->CenterOfMass());
+			SolverBody.SetRoM(RigidParticle->RotationOfMass());
+			SolverBody.SetV(RigidParticle->V());
+			SolverBody.SetW(RigidParticle->W());
+			SolverBody.SetInvM(RigidParticle->InvM());
+			SolverBody.SetInvILocal(RigidParticle->ConditionedInvI());
+			// Note: SetInvILocal also calculates InvI for dynamic
+		}
+		else
 		{
-			// Set the particle state
-			if (SolverBody.IsDynamic())
+			// Static or kinematic particle
+			SolverBody.SetP(Particle->X());
+			SolverBody.SetQ(Particle->R());
+			SolverBody.SetCoM(FVec3(0));
+			SolverBody.SetRoM(FRotation3::FromIdentity());
+			if (KinematicParticle != nullptr)
 			{
-				// Apply DP and DQ to P and Q
-				SolverBody.ApplyCorrections();
-
-				FParticleUtilities::SetCoMWorldTransform(Particle, SolverBody.P(), SolverBody.Q());
-				Particle->SetV(SolverBody.V());
-				Particle->SetW(SolverBody.W());
+				// Kinematic particle
+				SolverBody.SetX(KinematicParticle->X() - KinematicParticle->V() * Dt);
+				SolverBody.SetR(Particle->R());
+				if (!KinematicParticle->W().IsNearlyZero())
+				{
+					SolverBody.SetR(FRotation3::IntegrateRotationWithAngularVelocity(KinematicParticle->R(), -KinematicParticle->W(), Dt));
+				}
+				SolverBody.SetV(KinematicParticle->V());
+				SolverBody.SetW(KinematicParticle->W());
+			}
+			else
+			{
+				SolverBody.SetX(Particle->X());
+				SolverBody.SetR(Particle->R());
+				SolverBody.SetV(FVec3(0));
+				SolverBody.SetW(FVec3(0));
 			}
+			SolverBody.SetInvM(0);
+			SolverBody.SetInvILocal(FVec3(0));	// Note: does not set InvI for Kinematic
+			SolverBody.SetInvI(FMatrix33(0));
+		}
+	}
+
+	inline void UpdateParticleFromSolverBody(FGeometryParticleHandle* Particle, const FSolverBody& SolverBody)
+	{
+		check(Particle != nullptr);
+		check(SolverBody.IsDynamic());
+
+		// Set the particle state (only dynamics will have changed)
+		if (FPBDRigidParticleHandle* RigidParticle = Particle->CastToRigidParticle())
+		{
+			RigidParticle->SetTransformPQCom(SolverBody.P(), SolverBody.Q());
+			RigidParticle->SetV(SolverBody.V());
+			RigidParticle->SetW(SolverBody.W());
 
 			// Reset SolverBodyIndex cookie every step - it will be reassigned next step
-			Particle->SetSolverBodyIndex(INDEX_NONE);
+			RigidParticle->SetSolverBodyIndex(INDEX_NONE);
 		}
 	}
 
 	int32 FSolverBodyContainer::AddParticle(FGenericParticleHandle InParticle)
 	{
 		// No array resizing allowed (we want fixed pointers)
-		check(NumItems() < MaxItems());
+		check(Num() < Max());
 
-		return SolverBodies.AddElement(FSolverBodyAdapter(InParticle));
+		// NOTE: New solver body is completely uninitialized
+		SolverBodies.Add();
+		return Particles.Add(InParticle->Handle());
 	}
 
 	FSolverBody* FSolverBodyContainer::FindOrAdd(FGenericParticleHandle InParticle)
@@ -67,7 +122,7 @@ namespace Chaos
 	
 		check(ItemIndex != INDEX_NONE);
 		
-		FSolverBody* SolverBody = &SolverBodies[ItemIndex].GetSolverBody();
+		FSolverBody* SolverBody = &SolverBodies[ItemIndex];
 
 		return SolverBody;
 	}
@@ -79,52 +134,46 @@ namespace Chaos
 
 		for (int32 SolverBodyIndex = BeginIndex; SolverBodyIndex < EndIndex; ++SolverBodyIndex)
 		{
-			FSolverBodyAdapter& SolverBody = SolverBodies[SolverBodyIndex];
-			SolverBody.GatherInput(Dt);
-		}
-	}
-
-	void FSolverBodyContainer::ScatterOutput()
-	{
-		for (FSolverBodyAdapter& SolverBody : SolverBodies)
-		{
-			SolverBody.ScatterOutput();
+			UpdateSolverBodyFromParticle(SolverBodies[SolverBodyIndex], Particles[SolverBodyIndex], Dt);
 		}
 	}
 
 	void FSolverBodyContainer::ScatterOutput(const int32 BeginIndex, const int32 EndIndex)
 	{
 		check((BeginIndex >= 0) && (BeginIndex <= EndIndex));
-		check((EndIndex >= 0) && (EndIndex <= SolverBodies.Num()));
+		check((EndIndex >= 0) && (EndIndex <= Num()));
 
 		for (int32 SolverBodyIndex = BeginIndex; SolverBodyIndex < EndIndex; ++SolverBodyIndex)
 		{
-			FSolverBodyAdapter& SolverBody = SolverBodies[SolverBodyIndex];
-			SolverBody.ScatterOutput();
+			if (SolverBodies[SolverBodyIndex].IsDynamic())
+			{
+				SolverBodies[SolverBodyIndex].ApplyCorrections();
+				UpdateParticleFromSolverBody(Particles[SolverBodyIndex], SolverBodies[SolverBodyIndex]);
+			}
 		}
 	}
 
 	void FSolverBodyContainer::SetImplicitVelocities(FReal Dt)
 	{
-		for (FSolverBodyAdapter& SolverBody : SolverBodies)
+		for (FSolverBody& SolverBody : SolverBodies)
 		{
-			SolverBody.GetSolverBody().SetImplicitVelocity(Dt);
+			SolverBody.SetImplicitVelocity(Dt);
 		}
 	}
 
 	void FSolverBodyContainer::ApplyCorrections()
 	{
-		for (FSolverBodyAdapter& SolverBody : SolverBodies)
+		for (FSolverBody& SolverBody : SolverBodies)
 		{
-			SolverBody.GetSolverBody().ApplyCorrections();
+			SolverBody.ApplyCorrections();
 		}
 	}
 
 	void FSolverBodyContainer::UpdateRotationDependentState()
 	{
-		for (FSolverBodyAdapter& SolverBody : SolverBodies)
+		for (FSolverBody& SolverBody : SolverBodies)
 		{
-			SolverBody.GetSolverBody().UpdateRotationDependentState();
+			SolverBody.UpdateRotationDependentState();
 		}
 	}
 
diff --git a/Engine/Source/Runtime/Experimental/Chaos/Private/Chaos/GeometryParticles.cpp b/Engine/Source/Runtime/Experimental/Chaos/Private/Chaos/GeometryParticles.cpp
--- a/Engine/Source/Runtime/Experimental/Chaos/Private/Chaos/GeometryParticles.cpp
+++ b/Engine/Source/Runtime/Experimental/Chaos/Private/Chaos/GeometryParticles.cpp
@@ -257,7 +257,7 @@ namespace Chaos
 			return LeafInfo.LeafWorldTransform;
 		}
 
-		return FParticleUtilities::GetActorWorldTransform(FConstGenericParticleHandle(Particle));
+		return FConstGenericParticleHandle(Particle)->GetTransformPQ();
 	}
 	
 	void FPerShapeData::UpdateLeafWorldTransform(FGeometryParticleHandle* Particle)
@@ -265,7 +265,7 @@ namespace Chaos
 		if (bHasCachedLeafInfo)
 		{
 			FPerShapeDataCachedLeafInfo& LeafInfo = *static_cast<FPerShapeDataCachedLeafInfo*>(this);
-			FRigidTransform3 ParticleTransform = FParticleUtilities::GetActorWorldTransform(FConstGenericParticleHandle(Particle));
+			const FRigidTransform3 ParticleTransform = FConstGenericParticleHandle(Particle)->GetTransformPQ();
 			LeafInfo.LeafWorldTransform = FRigidTransform3::MultiplyNoScale(FRigidTransform3(GetLeafRelativeTransform()), ParticleTransform);
 		}
 	}
diff --git a/Engine/Source/Runtime/Experimental/Chaos/Private/Chaos/Island/IslandGroup.cpp b/Engine/Source/Runtime/Experimental/Chaos/Private/Chaos/Island/IslandGroup.cpp
--- a/Engine/Source/Runtime/Experimental/Chaos/Private/Chaos/Island/IslandGroup.cpp
+++ b/Engine/Source/Runtime/Experimental/Chaos/Private/Chaos/Island/IslandGroup.cpp
@@ -55,15 +55,8 @@ namespace Chaos
 
 	void FPBDIslandConstraintGroupSolver::AddConstraintsImpl()
 	{
-		// Initialize buffer sizes for solver particles and constraint solvers
+		// Initialize buffer sizes for solver particles
 		SolverBodyContainer.Reset(NumParticles);
-		for (int32 ContainerIndex = 0; ContainerIndex < ConstraintContainerSolvers.Num(); ++ContainerIndex)
-		{
-			if (ConstraintContainerSolvers[ContainerIndex] != nullptr)
-			{
-				ConstraintContainerSolvers[ContainerIndex]->Reset(NumContainerConstraints[ContainerIndex]);
-			}
-		}
 
 		// Add all the constraints to the solvers, but do not collect data from the constraints yet (each constraint type has its own solver). 
 		// This also sets up the SolverBody array, but does not gather the body data.
@@ -71,9 +64,11 @@ namespace Chaos
 		{
 			if (ConstraintContainerSolvers[ContainerIndex] != nullptr)
 			{
+				// Initialize buffer sizes for constraint solvers
+				ConstraintContainerSolvers[ContainerIndex]->Reset(NumContainerConstraints[ContainerIndex]);
+
 				for (FPBDIsland* Island : Islands)
 				{
-					// const_cast is because MakeArrayView doesn't handle a const array of pointers-to-non-const.
 					TArrayView<FPBDIslandConstraint> IslandConstraints = MakeArrayView(Island->GetConstraints(ContainerIndex));
 					ConstraintContainerSolvers[ContainerIndex]->AddConstraints(IslandConstraints);
 				}
@@ -86,9 +81,10 @@ namespace Chaos
 		// @todo(chaos): optimize
 		for (int32 SolverBodyIndex = BeginBodyIndex; SolverBodyIndex < EndBodyIndex; ++SolverBodyIndex)
 		{
-			FSolverBodyAdapter& SolverBody = SolverBodyContainer.GetItem(SolverBodyIndex);
-			const int32 Level = IslandManager.GetParticleLevel(SolverBody.GetParticle()->Handle());
-			SolverBody.GetSolverBody().SetLevel(Level);
+			FSolverBody& SolverBody = SolverBodyContainer.GetSolverBody(SolverBodyIndex);
+			FGeometryParticleHandle* Particle = SolverBodyContainer.GetParticle(SolverBodyIndex);
+			const int32 Level = IslandManager.GetParticleLevel(Particle);
+			SolverBody.SetLevel(Level);
 		}
 	}
 
diff --git a/Engine/Source/Runtime/Experimental/Chaos/Private/Chaos/Joint/PBDJointSolverGaussSeidel.cpp b/Engine/Source/Runtime/Experimental/Chaos/Private/Chaos/Joint/PBDJointSolverGaussSeidel.cpp
--- a/Engine/Source/Runtime/Experimental/Chaos/Private/Chaos/Joint/PBDJointSolverGaussSeidel.cpp
+++ b/Engine/Source/Runtime/Experimental/Chaos/Private/Chaos/Joint/PBDJointSolverGaussSeidel.cpp
@@ -101,6 +101,10 @@ namespace Chaos
 		// \todo(chaos): joint should support parent/child in either order
 		SolverBodies[0].SetInvMScale(JointSettings.ParentInvMassScale);
 		SolverBodies[1].SetInvMScale(FReal(1));
+		SolverBodies[0].SetInvIScale(JointSettings.ParentInvMassScale);
+		SolverBodies[1].SetInvIScale(FReal(1));
+		SolverBodies[0].SetShockPropagationScale(FReal(1));
+		SolverBodies[1].SetShockPropagationScale(FReal(1));
 
 		// Set the mass and inertia.
 		// If enabled, adjust the mass so that we limit the maximum mass and inertia ratios
diff --git a/Engine/Source/Runtime/Experimental/Chaos/Private/Chaos/PBDCollisionConstraints.cpp b/Engine/Source/Runtime/Experimental/Chaos/Private/Chaos/PBDCollisionConstraints.cpp
--- a/Engine/Source/Runtime/Experimental/Chaos/Private/Chaos/PBDCollisionConstraints.cpp
+++ b/Engine/Source/Runtime/Experimental/Chaos/Private/Chaos/PBDCollisionConstraints.cpp
@@ -7,6 +7,7 @@
 #include "Chaos/PBDCollisionConstraintsContact.h"
 #include "Chaos/CollisionResolution.h"
 #include "Chaos/Collision/CollisionPruning.h"
+#include "Chaos/Collision/PBDCollisionSolver.h"
 #include "Chaos/Collision/SolverCollisionContainer.h"
 #include "Chaos/Defines.h"
 #include "Chaos/Evolution/SolverBodyContainer.h"
@@ -21,9 +22,6 @@
 #include "Algo/Sort.h"
 #include "Algo/StableSort.h"
 
-// Private includes
-#include "Collision/PBDCollisionSolver.h"
-
 //PRAGMA_DISABLE_OPTIMIZATION
 
 namespace Chaos
@@ -119,6 +117,9 @@ namespace Chaos
 		, GravitySize(980)
 		, SolverSettings()
 	{
+		// Unfortunately, but the collision it creates need to know what container they belong to,
+		// but otherwise the allocator doesn't really need to know about the container...
+		ConstraintAllocator.SetCollisionContainer(this);
 	}
 
 	FPBDCollisionConstraints::~FPBDCollisionConstraints()
@@ -289,26 +290,14 @@ namespace Chaos
 		ConstraintAllocator.EndDetectCollisions();
 
 		// Disable any edge collisions that are hidden by face collisions
+		// (for bodies that have the EdgePruning option enabled)
 		PruneEdgeCollisions();
 
+		// @todo(chaos): this should be required any more - we sort when adding to the graph. Test and remove...
 		if (bIsDeterministic)
 		{
 			ConstraintAllocator.SortConstraintsHandles();
 		}
-
-		// Bind the constraints to this container and initialize other properties
-		// @todo(chaos): this could be set on creation if the midphase knew about the container
-		for (FPBDCollisionConstraint* Contact : GetConstraints())
-		{
-			if (Contact->GetContainer() == nullptr)
-			{
-				Contact->SetContainer(this);
-				UpdateConstraintMaterialProperties(*Contact);
-			}
-
-			// Reset constraint modifications and accumulators
-			Contact->Activate();
-		}
 	}
 
 	void FPBDCollisionConstraints::DetectProbeCollisions(FReal Dt)
diff --git a/Engine/Source/Runtime/Experimental/Chaos/Private/Chaos/PBDRigidsEvolutionGBF.cpp b/Engine/Source/Runtime/Experimental/Chaos/Private/Chaos/PBDRigidsEvolutionGBF.cpp
--- a/Engine/Source/Runtime/Experimental/Chaos/Private/Chaos/PBDRigidsEvolutionGBF.cpp
+++ b/Engine/Source/Runtime/Experimental/Chaos/Private/Chaos/PBDRigidsEvolutionGBF.cpp
@@ -49,9 +49,6 @@ namespace Chaos
 		int CollisionDisableCulledContacts = 0;
 		FAutoConsoleVariableRef CVarDisableCulledContacts(TEXT("p.CollisionDisableCulledContacts"), CollisionDisableCulledContacts, TEXT("Allow the PBDRigidsEvolutionGBF collision constraints to throw out contacts mid solve if they are culled."));
 
-		FRealSingle BoundsThicknessVelocityMultiplier = 0.0f;
-		FAutoConsoleVariableRef CVarBoundsThicknessVelocityMultiplier(TEXT("p.CollisionBoundsVelocityInflation"), BoundsThicknessVelocityMultiplier, TEXT("Collision velocity inflation for speculative contact generation.[def:0.0]"));
-
 		FRealSingle SmoothedPositionLerpRate = 0.1f;
 		FAutoConsoleVariableRef CVarSmoothedPositionLerpRate(TEXT("p.Chaos.SmoothedPositionLerpRate"), SmoothedPositionLerpRate, TEXT("The interpolation rate for the smoothed position calculation. Used for sleeping."));
 
@@ -641,8 +638,7 @@ FPBDRigidsEvolutionGBF::FPBDRigidsEvolutionGBF(FPBDRigidsSOAs& InParticles,THand
 
 	CollisionConstraints.SetCanDisableContacts(!!CollisionDisableCulledContacts);
 
-	CollisionDetector.SetBoundsExpansion(DefaultCollisionCullDistance);
-	CollisionDetector.SetBoundsVelocityInflation(BoundsThicknessVelocityMultiplier);
+	CollisionConstraints.SetCullDistance(DefaultCollisionCullDistance);
 
 	SetParticleUpdatePositionFunction([this](const TParticleView<FPBDRigidParticles>& ParticlesInput, const FReal Dt)
 	{
@@ -926,7 +922,8 @@ void FPBDRigidsEvolutionGBF::ParticleMaterialChanged(FGeometryParticleHandle* Pa
 {
 	Particle->ParticleCollisions().VisitCollisions([this](FPBDCollisionConstraint& Collision)
 	{
-		CollisionConstraints.UpdateConstraintMaterialProperties(Collision);
+		// Reset the material - the material properties will get collected later (collision activation)
+		Collision.ClearMaterialProperties();
 		return ECollisionVisitorResult::Continue;
 	});
 }
diff --git a/Engine/Source/Runtime/Experimental/Chaos/Private/Chaos/PBDSuspensionConstraints.cpp b/Engine/Source/Runtime/Experimental/Chaos/Private/Chaos/PBDSuspensionConstraints.cpp
--- a/Engine/Source/Runtime/Experimental/Chaos/Private/Chaos/PBDSuspensionConstraints.cpp
+++ b/Engine/Source/Runtime/Experimental/Chaos/Private/Chaos/PBDSuspensionConstraints.cpp
@@ -1,12 +1,10 @@
 // Copyright Epic Games, Inc. All Rights Reserved.
 
 #include "Chaos/PBDSuspensionConstraints.h"
+#include "Chaos/Collision/PBDCollisionSolver.h"
 #include "Chaos/Island/IslandManager.h"
 #include "Chaos/DebugDrawQueue.h"
 
-// Private includes
-#include "Chaos/Collision/PBDCollisionSolver.h"
-
 bool bChaos_Suspension_Spring_Enabled = true;
 FAutoConsoleVariableRef CVarChaosSuspensionSpringEnabled(TEXT("p.Chaos.Suspension.Spring.Enabled"), bChaos_Suspension_Spring_Enabled, TEXT("Enable/Disable Spring part of suspension constraint"));
 
@@ -63,7 +61,7 @@ namespace Chaos
 		return ConcreteContainer()->GetConstrainedParticles(ConstraintIndex);
 	}
 
-	Chaos::FPBDSuspensionConstraints::FConstraintContainerHandle* FPBDSuspensionConstraints::AddConstraint(TGeometryParticleHandle<FReal, 3>* Particle, const FVec3& InSuspensionLocalOffset, const FPBDSuspensionSettings& InConstraintSettings)
+	FPBDSuspensionConstraints::FConstraintContainerHandle* FPBDSuspensionConstraints::AddConstraint(TGeometryParticleHandle<FReal, 3>* Particle, const FVec3& InSuspensionLocalOffset, const FPBDSuspensionSettings& InConstraintSettings)
 	{
 		int32 NewIndex = ConstrainedParticles.Num();
 		ConstrainedParticles.Add(Particle);
@@ -72,8 +70,9 @@ namespace Chaos
 		ConstraintResults.AddDefaulted();
 		ConstraintEnabledStates.Add(true); // Note: assumes always enabled on creation
 		ConstraintSolverBodies.Add(nullptr);
-		StaticCollisionBodies.Add(FSolverBody());
-		CollisionSolvers.Add( new FPBDCollisionSolver() );
+		StaticCollisionBodies.Add(FSolverBody::MakeInitialized());
+		CollisionSolvers.Add(Private::FPBDCollisionSolver());
+		CollisionSolverManifoldPoints.Add(Private::FPBDCollisionSolverManifoldPoint());
 
 		Handles.Add(HandleAllocator.AllocHandle(this, NewIndex));
 		return Handles[NewIndex];
@@ -95,8 +94,6 @@ namespace Chaos
 			Handles[ConstraintIndex] = nullptr;
 		}
 
-		delete CollisionSolvers[ConstraintIndex];
-
 		// Swap the last constraint into the gap to keep the array packed
 		ConstrainedParticles.RemoveAtSwap(ConstraintIndex);
 		SuspensionLocalOffset.RemoveAtSwap(ConstraintIndex);
@@ -105,6 +102,7 @@ namespace Chaos
 		ConstraintEnabledStates.RemoveAtSwap(ConstraintIndex);
 		ConstraintSolverBodies.RemoveAtSwap(ConstraintIndex);
 		CollisionSolvers.RemoveAtSwap(ConstraintIndex);
+		CollisionSolverManifoldPoints.RemoveAtSwap(ConstraintIndex);
 		StaticCollisionBodies.RemoveAtSwap(ConstraintIndex);
 
 		Handles.RemoveAtSwap(ConstraintIndex);
@@ -208,15 +206,17 @@ namespace Chaos
 
 	void FPBDSuspensionConstraints::GatherInput(const int32 ConstraintIndex, const FReal Dt)
 	{
+		using namespace Private;
+
 		ConstraintResults[ConstraintIndex].Reset();
 
 		// Hard-stop Collision Manifold Generation
 		{
-			check(CollisionSolvers[ConstraintIndex] != nullptr);
 			check(ConstraintSolverBodies[ConstraintIndex] != nullptr);
 
-			FPBDCollisionSolver* Solver = CollisionSolvers[ConstraintIndex];
-			Solver->Reset(); // clear previous manifolds
+			FPBDCollisionSolver* Solver = &CollisionSolvers[ConstraintIndex];
+			Solver->SetManifoldPointsBuffer(&CollisionSolverManifoldPoints[ConstraintIndex], 1);
+			Solver->SetStiffness(1);
 
 			FSolverBody* Body0 = ConstraintSolverBodies[ConstraintIndex];	// vehicle chassis			
 			FSolverBody* Body1 = &StaticCollisionBodies[ConstraintIndex];	// Spoofed terrain
@@ -255,7 +255,8 @@ namespace Chaos
 				check(Body1->InvM() == 0);
 
 				Solver->SetSolverBodies(*Body0, *Body1);
-				Solver->SetNumManifoldPoints(1);
+				Solver->SolverBody0().Init();
+				Solver->SolverBody1().Init();
 				Solver->SetFriction(0, 0, 0);
 
 #if CHAOS_DEBUG_DRAW
@@ -272,11 +273,10 @@ namespace Chaos
 #endif
 
 				// inject a manifold for our suspension Hard-stop - behaves like a regular friction-less collision, prevents the vehicle chassis from ever hitting the ground
+				Solver->AddManifoldPoint();
 				Solver->SetManifoldPoint(
 					0,										// ManifoldIndex
 					FSolverReal(Dt),						// Delta Time
-					FSolverReal(0.0f),						// Restitution,
-					FSolverReal(0.1f),						// RestitutionVelocityThreshold,
 					FSolverVec3(WorldArm),					// RelativeContactPosition0,
 					FSolverVec3(0, 0, 0),					// RelativeContactPosition1,
 					FSolverVec3(WorldContactNormal),		// WorldContactNormal
@@ -284,7 +284,8 @@ namespace Chaos
 					FSolverVec3(0, 0, 0),					// WorldContactTangentV,
 					FSolverReal(-WorldContactDeltaNormal),	// WorldContactDeltaNormal
 					FSolverReal(0),							// WorldContactDeltaTangentU,
-					FSolverReal(0)							// WorldContactDeltaTangentV
+					FSolverReal(0),							// WorldContactDeltaTangentV
+					FSolverReal(0)							// WorldTargetContactVelocityNormal
 					);
 
 			}
@@ -294,27 +295,30 @@ namespace Chaos
 
 	void FPBDSuspensionConstraints::ScatterOutput(const int32 ConstraintIndex, FReal Dt)
 	{
-		if ((CollisionSolvers.Num() > 0) && (CollisionSolvers[ConstraintIndex] != nullptr) && (CollisionSolvers[ConstraintIndex]->NumManifoldPoints() > 0))
+		using namespace Private;
+
+		if ((CollisionSolvers.Num() > 0) && (CollisionSolvers[ConstraintIndex].NumManifoldPoints() > 0))
 		{
-			const FPBDCollisionSolver* CollisionSolver = CollisionSolvers[ConstraintIndex];
-			const FPBDCollisionSolverManifoldPoint& ManifoldPoint = CollisionSolver->GetManifoldPoint(0);
-			ConstraintResults[ConstraintIndex].HardStopNetPushOut = ManifoldPoint.NetPushOutNormal * ManifoldPoint.WorldContactNormal;
-			ConstraintResults[ConstraintIndex].HardStopNetImpulse = ManifoldPoint.NetImpulseNormal * ManifoldPoint.WorldContactNormal;
+			const FPBDCollisionSolverManifoldPoint& ManifoldPoint = CollisionSolvers[ConstraintIndex].GetManifoldPoint(0);
+			ConstraintResults[ConstraintIndex].HardStopNetPushOut = ManifoldPoint.NetPushOutNormal * ManifoldPoint.WorldContact.ContactNormal;
+			ConstraintResults[ConstraintIndex].HardStopNetImpulse = ManifoldPoint.NetImpulseNormal * ManifoldPoint.WorldContact.ContactNormal;
 		}
 
 		ConstraintSolverBodies[ConstraintIndex] = nullptr;
-		CollisionSolvers[ConstraintIndex]->ResetSolverBodies();
+		CollisionSolvers[ConstraintIndex].Reset();
 	}
 
 	void FPBDSuspensionConstraints::ApplyPositionConstraint(const int32 ConstraintIndex, const FReal Dt, const int32 It, const int32 NumIts)
 	{
+		using namespace Private;
+			
 		if (bChaos_Suspension_Hardstop_Enabled)
 		{
 			// Suspension Hardstop
 			const FPBDSuspensionSettings& Setting = ConstraintSettings[ConstraintIndex];
 			if (Setting.Enabled)
 			{
-				FPBDCollisionSolver* CollisionSolver = CollisionSolvers[ConstraintIndex];
+				FPBDCollisionSolver* CollisionSolver = &CollisionSolvers[ConstraintIndex];
 				if ((CollisionSolver != nullptr) && (CollisionSolver->NumManifoldPoints() > 0))
 				{
 					FReal MaxPushoutValue = FMath::Min(Chaos_Suspension_MaxPushout, Chaos_Suspension_MaxPushoutVelocity * Dt);
@@ -332,13 +336,15 @@ namespace Chaos
 
 	void FPBDSuspensionConstraints::ApplyVelocityConstraint(const int32 ConstraintIndex, const FReal Dt, const int32 It, const int32 NumIts)
 	{
+		using namespace Private;
+
 		if (bChaos_Suspension_Hardstop_Enabled && bChaos_Suspension_VelocitySolve)
 		{
 			// Suspension Hardstop
 			const FPBDSuspensionSettings& Setting = ConstraintSettings[ConstraintIndex];
 			if (Setting.Enabled)
 			{
-				FPBDCollisionSolver* CollisionSolver = CollisionSolvers[ConstraintIndex];
+				FPBDCollisionSolver* CollisionSolver = &CollisionSolvers[ConstraintIndex];
 				if ((CollisionSolver != nullptr) && (CollisionSolver->NumManifoldPoints() > 0))
 				{
 					CollisionSolver->SolveVelocity(FSolverReal(Dt), false);
diff --git a/Engine/Source/Runtime/Experimental/Chaos/Private/Chaos/ParticleHandle.cpp b/Engine/Source/Runtime/Experimental/Chaos/Private/Chaos/ParticleHandle.cpp
--- a/Engine/Source/Runtime/Experimental/Chaos/Private/Chaos/ParticleHandle.cpp
+++ b/Engine/Source/Runtime/Experimental/Chaos/Private/Chaos/ParticleHandle.cpp
@@ -167,8 +167,7 @@ namespace Chaos
 	template <typename T, int d, bool bPersistent>
 	void TPBDRigidParticleHandleImp<T, d, bPersistent>::AddTorque(const TVector<T, d>& InTorque, bool bInvalidate)
 	{
-		const FRotation3 RCoM = FParticleUtilitiesPQ::GetCoMWorldRotation(this);
-		const FMatrix33 WorldInvI = Utilities::ComputeWorldSpaceInertia(RCoM, InvI());
+		const FMatrix33 WorldInvI = Utilities::ComputeWorldSpaceInertia(QCom(), InvI());
 		SetAngularAcceleration(AngularAcceleration() + WorldInvI * InTorque);
 	}
 
@@ -176,8 +175,7 @@ namespace Chaos
 	template <typename T, int d, bool bPersistent>
 	void TPBDRigidParticleHandleImp<T, d, bPersistent>::SetTorque(const TVector<T, d>& InTorque, bool bInvalidate)
 	{
-		const FRotation3 RCoM = FParticleUtilitiesPQ::GetCoMWorldRotation(this);
-		const FMatrix33 WorldInvI = Utilities::ComputeWorldSpaceInertia(RCoM, InvI());
+		const FMatrix33 WorldInvI = Utilities::ComputeWorldSpaceInertia(QCom(), InvI());
 		SetAngularAcceleration(WorldInvI * InTorque);
 	}
 
diff --git a/Engine/Source/Runtime/Experimental/Chaos/Private/Chaos/Utilities.cpp b/Engine/Source/Runtime/Experimental/Chaos/Private/Chaos/Utilities.cpp
--- a/Engine/Source/Runtime/Experimental/Chaos/Private/Chaos/Utilities.cpp
+++ b/Engine/Source/Runtime/Experimental/Chaos/Private/Chaos/Utilities.cpp
@@ -147,6 +147,12 @@ namespace Chaos
 		return R1.GetNormalized();
 	}
 
+	TRotation<FRealSingle, 3> TRotation<FRealSingle, 3>::IntegrateRotationWithAngularVelocity(const TRotation<FRealSingle, 3>& InR0, const TVector<FRealSingle, 3>& InW, const FRealSingle InDt)
+	{
+		TRotation<FRealSingle, 3> R1 = InR0 + (TRotation<FRealSingle, 3>::FromElements(InW.X, InW.Y, InW.Z, 0.f) * InR0) * decltype(InR0.X)(InDt * 0.5f);
+		return R1.GetNormalized();
+	}
+
 namespace Utilities
 {
 	FReal GetSolverPhysicsResultsTime(FPhysicsSolverBase* Solver)
diff --git a/Engine/Source/Runtime/Experimental/Chaos/Private/PBDRigidsSolver.cpp b/Engine/Source/Runtime/Experimental/Chaos/Private/PBDRigidsSolver.cpp
--- a/Engine/Source/Runtime/Experimental/Chaos/Private/PBDRigidsSolver.cpp
+++ b/Engine/Source/Runtime/Experimental/Chaos/Private/PBDRigidsSolver.cpp
@@ -878,12 +878,12 @@ namespace Chaos
 	{
 		MEvolution->GetCollisionConstraints().SetCollisionsEnabled(bChaosSolverCollisionEnabled);
 
-		FCollisionDetectorSettings CollisionDetectorSettings = MEvolution->GetCollisionDetector().GetSettings();
+		FCollisionDetectorSettings CollisionDetectorSettings = MEvolution->GetCollisionConstraints().GetDetectorSettings();
 		CollisionDetectorSettings.bAllowManifoldReuse = (ChaosSolverCollisionAllowManifoldUpdate != 0);
 		CollisionDetectorSettings.bDeferNarrowPhase = (ChaosSolverCollisionDeferNarrowPhase != 0);
 		CollisionDetectorSettings.bAllowManifolds = (ChaosSolverCollisionUseManifolds != 0);
 		CollisionDetectorSettings.bAllowCCD = bChaosUseCCD;
-		MEvolution->GetCollisionDetector().SetSettings(CollisionDetectorSettings);
+		MEvolution->GetCollisionConstraints().SetDetectorSettings(CollisionDetectorSettings);
 		
 		FPBDJointSolverSettings JointsSettings = MEvolution->GetJointConstraints().GetSettings();
 		JointsSettings.MinSolverStiffness = ChaosSolverJointMinSolverStiffness;
@@ -1569,24 +1569,27 @@ CSV_CUSTOM_STAT(PhysicsCounters, Name, Value, ECsvCustomStatOp::Set);
 		int32 NumUpdatedManifoldPoints = 0;
 		for (const FPBDCollisionConstraintHandle* Collision : GetEvolution()->GetCollisionConstraints().GetConstraints())
 		{
-			if (Collision->GetContact().IsEnabled())
+			const FPBDCollisionConstraint& Contact = Collision->GetContact();
+			if (Contact.IsEnabled())
 			{
-				if (Collision->GetContact().GetManifoldPoints().Num() > 0)
+				if (Contact.GetManifoldPoints().Num() > 0)
 				{
 					++NumValidCollisions;
 				}
-				if (!Collision->GetContact().AccumulatedImpulse.IsNearlyZero())
+				if (!Contact.AccumulatedImpulse.IsNearlyZero())
 				{
 					++NumActiveCollisions;
 				}
-				if (Collision->GetContact().WasManifoldRestored())
+				if (Contact.WasManifoldRestored())
 				{
 					++NumRestoredCollisions;
 				}
-				for (const FManifoldPoint& ManifoldPoint : Collision->GetContact().GetManifoldPoints())
+				for (int32 PointIndex = 0; PointIndex < Contact.NumManifoldPoints(); ++PointIndex)
 				{
+					const FManifoldPoint& ManifoldPoint = Contact.GetManifoldPoint(PointIndex);
+					const FManifoldPointResult& ManifoldPointResult = Contact.GetManifoldPointResult(PointIndex);
 					++NumManifoldPoints;
-					if (ManifoldPoint.Flags.bWasRestored || Collision->GetContact().WasManifoldRestored())
+					if (ManifoldPoint.Flags.bWasRestored || Contact.WasManifoldRestored())
 					{
 						++NumRestoredManifoldPoints;
 					}
@@ -1594,9 +1597,12 @@ CSV_CUSTOM_STAT(PhysicsCounters, Name, Value, ECsvCustomStatOp::Set);
 					{
 						++NumUpdatedManifoldPoints;
 					}
-					if (!ManifoldPoint.NetPushOut.IsNearlyZero())
+					if (ManifoldPointResult.bIsValid)
 					{
-						++NumActiveManifoldPoints;
+						if (!ManifoldPointResult.NetPushOut.IsNearlyZero())
+						{
+							++NumActiveManifoldPoints;
+						}
 					}
 				}
 
diff --git a/Engine/Source/Runtime/Experimental/Chaos/Public/Chaos/Collision/CollisionConstraintAllocator.h b/Engine/Source/Runtime/Experimental/Chaos/Public/Chaos/Collision/CollisionConstraintAllocator.h
--- a/Engine/Source/Runtime/Experimental/Chaos/Public/Chaos/Collision/CollisionConstraintAllocator.h
+++ b/Engine/Source/Runtime/Experimental/Chaos/Public/Chaos/Collision/CollisionConstraintAllocator.h
@@ -15,6 +15,8 @@
 
 namespace Chaos
 {
+	class FPBDCollisionConstraints;
+
 	/**
 	 * Container the storage for the FCollisionConstraintAllocator, as well as the API to create new midphases and collision constraints.
 	 * We have one of these objects per thread on which collisions detection is performed to get lock-free allocations and lists.
@@ -22,8 +24,9 @@ namespace Chaos
 	class CHAOS_API FCollisionContextAllocator
 	{
 	public:
-		FCollisionContextAllocator(const int32 InNumCollisionsPerBlock, const int32 InCurrentEpoch)
-			: CurrentEpoch(InCurrentEpoch)
+		FCollisionContextAllocator(FPBDCollisionConstraints* InCollisionContainer, const int32 InNumCollisionsPerBlock, const int32 InCurrentEpoch)
+			: CollisionContainer(InCollisionContainer)
+			, CurrentEpoch(InCurrentEpoch)
 #if CHAOS_COLLISION_OBJECTPOOL_ENABLED
 			, ConstraintPool(InNumCollisionsPerBlock, 0)
 			, MidPhasePool(InNumCollisionsPerBlock, 0)
@@ -55,17 +58,7 @@ namespace Chaos
 			const FRigidTransform3& ShapeRelativeTransform1,
 			const FReal CullDistance,
 			const bool bUseManifold,
-			const EContactShapesType ShapePairType)
-		{
-			FPBDCollisionConstraintPtr Constraint = CreateConstraint();
-
-			if (Constraint.IsValid())
-			{
-				FPBDCollisionConstraint::Make(Particle0, Implicit0, Shape0, Simplicial0, ShapeRelativeTransform0, Particle1, Implicit1, Shape1, Simplicial1, ShapeRelativeTransform1, CullDistance, bUseManifold, ShapePairType, *Constraint);
-			}
-
-			return Constraint;
-		}
+			const EContactShapesType ShapePairType);
 
 		/**
 		 * Create an uninitialized collision constraint (public only for use by Resim which overwrites it with a saved constraint)
@@ -180,6 +173,7 @@ namespace Chaos
 			return MidPhase;
 		}
 
+		FPBDCollisionConstraints* CollisionContainer;
 		int32 CurrentEpoch;
 
 		TArray<FPBDCollisionConstraint*> NewActiveConstraints;
@@ -229,7 +223,8 @@ namespace Chaos
 		UE_NONCOPYABLE(FCollisionConstraintAllocator);
 
 		FCollisionConstraintAllocator(const int32 InNumCollisionsPerBlock = 1000)
-			: ContextAllocators()
+			: CollisionContainer(nullptr)
+			, ContextAllocators()
 			, NumCollisionsPerBlock(InNumCollisionsPerBlock)
 			, ParticlePairMidPhases()
 			, ActiveConstraints()
@@ -245,15 +240,24 @@ namespace Chaos
 			Reset();
 		}
 
+		void SetCollisionContainer(FPBDCollisionConstraints* InCollisionContainer)
+		{
+			CollisionContainer = InCollisionContainer;
+		}
+
 		void SetMaxContexts(const int32 MaxContexts)
 		{
+			// Forgot to call SetCollisionContainer after construction
+			// @todo(chaos): this check makes the allocator difficult to unit test - fix this
+			//check(CollisionContainer != nullptr);
+
 			const int32 NumNewContexts = MaxContexts - ContextAllocators.Num();
 			if (NumNewContexts > 0)
 			{
 				ContextAllocators.Reserve(MaxContexts);
 				for (int32 NewIndex = 0; NewIndex < NumNewContexts; ++NewIndex)
 				{
-					TUniquePtr<FCollisionContextAllocator> ContextAllocator = MakeUnique<FCollisionContextAllocator>(NumCollisionsPerBlock, CurrentEpoch);
+					TUniquePtr<FCollisionContextAllocator> ContextAllocator = MakeUnique<FCollisionContextAllocator>(CollisionContainer, NumCollisionsPerBlock, CurrentEpoch);
 					ContextAllocators.Emplace(MoveTemp(ContextAllocator));
 				}
 			}
@@ -550,6 +554,9 @@ namespace Chaos
 			Cookie.LastUsedEpoch = CurrentEpoch;
 		}
 
+		// The container that owns the allocator (only needed because new constraints need to know)
+		FPBDCollisionConstraints* CollisionContainer;
+
 		// Storage for collisins and midphases, one for each thread on which we detect collisions
 		TArray<TUniquePtr<FCollisionContextAllocator>> ContextAllocators;
 		int32 NumCollisionsPerBlock;
@@ -588,8 +595,7 @@ namespace Chaos
 		{
 			const FPBDCollisionConstraintPtr& Constraint = KVP.Value;
 
-			// If we only want active constraints, check the we're in the graph
-			if (Constraint.IsValid() && (!bOnlyActive || Constraint->IsInConstraintGraph()))
+			if (Constraint.IsValid() && (!bOnlyActive || !Constraint->GetDisabled()))
 			{
 				if (Visitor(*Constraint) == ECollisionVisitorResult::Stop)
 				{
@@ -607,8 +613,7 @@ namespace Chaos
 		{
 			const FPBDCollisionConstraintPtr& Constraint = KVP.Value;
 
-			// If we only want active constraints, check the timestamp
-			if (Constraint.IsValid() && (!bOnlyActive || Constraint->IsInConstraintGraph()))
+			if (Constraint.IsValid() && (!bOnlyActive || !Constraint->GetDisabled()))
 			{
 				if (!bOnlyActive || (Visitor(*Constraint) == ECollisionVisitorResult::Stop))
 				{
@@ -624,8 +629,7 @@ namespace Chaos
 	{
 		for (FSingleShapePairCollisionDetector& ShapePair : ShapePairDetectors)
 		{
-			// If we only want active constraints, check the we're in the graph
-			if ((ShapePair.GetConstraint() != nullptr) && (!bOnlyActive || ShapePair.GetConstraint()->IsInConstraintGraph()))
+			if ((ShapePair.GetConstraint() != nullptr) && (!bOnlyActive || !ShapePair.GetConstraint()->GetDisabled()))
 			{
 				if (Visitor(*ShapePair.GetConstraint()) == ECollisionVisitorResult::Stop)
 				{
@@ -651,8 +655,7 @@ namespace Chaos
 	{
 		for (const FSingleShapePairCollisionDetector& ShapePair : ShapePairDetectors)
 		{
-			// If we only want active constraints, check the timestamp
-			if ((ShapePair.GetConstraint() != nullptr) && (!bOnlyActive || ShapePair.GetConstraint()->IsInConstraintGraph()))
+			if ((ShapePair.GetConstraint() != nullptr) && (!bOnlyActive || !ShapePair.GetConstraint()->GetDisabled()))
 			{
 				if (Visitor(*ShapePair.GetConstraint()) == ECollisionVisitorResult::Stop)
 				{
diff --git a/Engine/Source/Runtime/Experimental/Chaos/Public/Chaos/Collision/CollisionDetector.h b/Engine/Source/Runtime/Experimental/Chaos/Public/Chaos/Collision/CollisionDetector.h
--- a/Engine/Source/Runtime/Experimental/Chaos/Public/Chaos/Collision/CollisionDetector.h
+++ b/Engine/Source/Runtime/Experimental/Chaos/Public/Chaos/Collision/CollisionDetector.h
@@ -14,39 +14,29 @@ namespace Chaos
 	{
 	public:
 		FCollisionDetector(FPBDCollisionConstraints& InCollisionContainer)
-			: Settings()
-			, CollisionContainer(InCollisionContainer)
+			: CollisionContainer(InCollisionContainer)
 		{
 		}
 
 		virtual ~FCollisionDetector() {}
 
-		const FCollisionDetectorSettings& GetSettings() const
-		{
-			return Settings;
-		}
+		UE_DEPRECATED(5.3, "Moved to FPBDCollisionConstraints")
+		const FCollisionDetectorSettings& GetSettings() const;
 
-		void SetSettings(const FCollisionDetectorSettings& InSettings)
-		{
-			Settings = InSettings;
-		}
+		UE_DEPRECATED(5.3, "Moved to FPBDCollisionConstraints")
+		void SetSettings(const FCollisionDetectorSettings& InSettings);
 
-		void SetBoundsExpansion(const FReal InBoundsExpansion)
-		{
-			Settings.BoundsExpansion = InBoundsExpansion;
-		}
+		UE_DEPRECATED(5.3, "Moved to FPBDCollisionConstraints and renamed to SetCullDistance")
+		void SetBoundsExpansion(const FReal InBoundsExpansion);
 
-		void SetBoundsVelocityInflation(const FReal InBoundsVelocityInflation)
-		{
-			Settings.BoundsVelocityInflation = InBoundsVelocityInflation;
-		}
+		UE_DEPRECATED(5.3, "No longer supported")
+		void SetBoundsVelocityInflation(const FReal InBoundsVelocityInflation) {}
 
 		FPBDCollisionConstraints& GetCollisionContainer() { return CollisionContainer; }
 
 		virtual void DetectCollisions(const FReal Dt, FEvolutionResimCache* ResimCache) = 0;
 
 	protected:
-		FCollisionDetectorSettings Settings;
 		FPBDCollisionConstraints& CollisionContainer;
 	};
 
diff --git a/Engine/Source/Runtime/Experimental/Chaos/Public/Chaos/Collision/ContactPoint.h b/Engine/Source/Runtime/Experimental/Chaos/Public/Chaos/Collision/ContactPoint.h
--- a/Engine/Source/Runtime/Experimental/Chaos/Public/Chaos/Collision/ContactPoint.h
+++ b/Engine/Source/Runtime/Experimental/Chaos/Public/Chaos/Collision/ContactPoint.h
@@ -81,4 +81,123 @@ namespace Chaos
 			return (Phi != InvalidPhi<FRealType>());
 		}
 	};
+
+	// Explicit types of TContactPoint for float and double. We should use the float version whenever the contacts are in shape-relative
+	// space, which is most of the time (some collision functions may choose to generate world-space contacts which are later converted
+	// to local space, and they should therefore use the double version until the local-space conversion).
+	using FContactPoint = TContactPoint<FReal>;
+	using FContactPointf = TContactPoint<FRealSingle>;
+
+
+	/**
+	 * @brief A single point in a contact manifold.
+	 * Each Collision Constraint will have up to 4 of these.
+	*/
+	class CHAOS_API FManifoldPoint
+	{
+	public:
+		union FFlags
+		{
+			FFlags() { Reset(); }
+
+			void Reset() { Bits = 0; }
+
+			struct
+			{
+				uint8 bDisabled : 1;						// Whether the point was disabled by edge pruning etc
+				uint8 bWasRestored : 1;						// Whether the point was retored from the previous frame due to lack of movement
+				uint8 bWasReplaced : 1;						// @todo(chaos): remove this
+				uint8 bHasStaticFrictionAnchor : 1;			// Whether our static friction anchor was recovered from a prior tick
+			};
+			uint8 Bits;
+		};
+
+		FManifoldPoint()
+			: ContactPoint()
+			, Flags()
+			, TargetPhi(0)
+			, ShapeAnchorPoints{ FVec3f(0), FVec3f(0) }
+			, InitialShapeContactPoints{ FVec3f(0), FVec3f(0) }
+		{}
+
+		FManifoldPoint(const FContactPointf& InContactPoint)
+			: ContactPoint(InContactPoint)
+			, Flags()
+			, TargetPhi(0)
+			, ShapeAnchorPoints{ FVec3f(0), FVec3f(0) }
+			, InitialShapeContactPoints{ FVec3f(0), FVec3f(0) }
+		{}
+
+		FContactPointf ContactPoint;			// Contact point results of low-level collision detection
+		FFlags Flags;							// Various flags
+		FRealSingle TargetPhi;					// Usually 0, but can be used to add padding or penetration (e.g., via a collision modifer)
+		FVec3f ShapeAnchorPoints[2];			// When static friction holds, the contact points on each shape when static friction contact was made
+		FVec3f InitialShapeContactPoints[2];	// ShapeContactPoints when the constraint was first initialized. Used to track reusablility
+	};
+
+	/**
+	 * World-space contact point data
+	 */
+	class CHAOS_API FWorldContactPoint
+	{
+	public:
+		// World-space contact point relative to each particle's center of mass
+		FVec3f RelativeContactPoints[2];
+
+		// World-space contact normal and tangents
+		FVec3f ContactNormal;
+		FVec3f ContactTangentU;
+		FVec3f ContactTangentV;
+
+		// Errors to correct along each of the contact axes
+		FRealSingle ContactDeltaNormal;
+		FRealSingle ContactDeltaTangentU;
+		FRealSingle ContactDeltaTangentV;
+
+		// Target velocity along the normal direction
+		FRealSingle ContactTargetVelocityNormal;
+	};
+
+	/**
+	 * @brief The friction data for a manifold point
+	 * This is the information that needs to be stored between ticks to implement static friction.
+	*/
+	class CHAOS_API FSavedManifoldPoint
+	{
+	public:
+		FSavedManifoldPoint()
+		{
+		}
+
+		void Reset()
+		{
+			ShapeContactPoints[0] = FVec3f(0);
+			ShapeContactPoints[1] = FVec3f(0);
+		}
+
+		FVec3f ShapeContactPoints[2];			// Contact anchor points for friction
+	};
+
+	class CHAOS_API FManifoldPointResult
+	{
+	public:
+		FManifoldPointResult()
+		{
+			Reset();
+		}
+
+		void Reset()
+		{
+			NetPushOut = FVec3f(0);
+			NetImpulse = FVec3f(0);
+			bIsValid = false;
+			bInsideStaticFrictionCone = false;
+		}
+
+		FVec3f NetPushOut;						// Total pushout applied at this contact point
+		FVec3f NetImpulse;						// Total impulse applied by this contact point
+		uint8 bIsValid : 1;
+		uint8 bInsideStaticFrictionCone : 1;
+	};
+
 }
\ No newline at end of file
diff --git a/Engine/Source/Runtime/Experimental/Chaos/Public/Chaos/Collision/PBDCollisionConstraint.h b/Engine/Source/Runtime/Experimental/Chaos/Public/Chaos/Collision/PBDCollisionConstraint.h
--- a/Engine/Source/Runtime/Experimental/Chaos/Public/Chaos/Collision/PBDCollisionConstraint.h
+++ b/Engine/Source/Runtime/Experimental/Chaos/Public/Chaos/Collision/PBDCollisionConstraint.h
@@ -24,11 +24,9 @@ namespace Chaos
 	class FParticlePairMidPhase;
 	class FPBDCollisionConstraint;
 	class FPBDCollisionConstraints;
-	class FPBDCollisionSolver;
 	class FSingleShapePairCollisionDetector;
 	class FSolverBody;
 	class FSolverBodyContainer;
-	class FPBDCollisionContainerSolver;
 	class FPerShapeData;
 
 	UE_DEPRECATED(4.27, "Use FPBDCollisionConstraint instead")
@@ -36,68 +34,6 @@ namespace Chaos
 
 	CHAOS_API bool ContactConstraintSortPredicate(const FPBDCollisionConstraint& L, const FPBDCollisionConstraint& R);
 
-	/**
-	 * @brief A single point in a contact manifold.
-	 * Each Collision Constraint will have up to 4 of these.
-	*/
-	class CHAOS_API FManifoldPoint
-	{
-	public:
-		FManifoldPoint() 
-			: ContactPoint()
-			, ShapeAnchorPoints{ FVec3f(0), FVec3f(0) }
-			, InitialShapeContactPoints{ FVec3f(0), FVec3f(0) }
-			, NetPushOut(0)
-			, NetImpulse(0)
-			, TargetPhi(0)
-			, Flags()
-		{}
-
-		FManifoldPoint(const FContactPoint& InContactPoint) 
-			: ContactPoint(InContactPoint)
-			, ShapeAnchorPoints{ FVec3f(0), FVec3f(0) }
-			, InitialShapeContactPoints{ FVec3f(0), FVec3f(0) }
-			, NetPushOut(0)
-			, NetImpulse(0)
-			, TargetPhi(0)
-			, Flags()
-		{}
-
-		FContactPointf ContactPoint;			// Contact point results of low-level collision detection
-		FVec3f ShapeAnchorPoints[2];			// The contact points from a prior frame, if availble. Used by friction.
-		FVec3f InitialShapeContactPoints[2];	// ShapeContactPoints when the constraint was first initialized. Used to track reusablility
-		FVec3f NetPushOut;						// Total pushout applied at this contact point
-		FVec3f NetImpulse;						// Total impulse applied by this contact point
-		FRealSingle TargetPhi;					// Usually 0, but can be used to add padding or penetration (e.g., via a collision modifer)
-
-		union FFlags
-		{
-			FFlags() { Reset(); }
-			
-			void Reset() { Bits = 0; }
-
-			struct
-			{
-				uint8 bDisabled : 1;						// Whether the point was disabled by edge pruning etc
-				uint8 bWasRestored : 1;						// Whether the point was retored from the previous frame due to lack of movement
-				uint8 bWasReplaced : 1;						// @todo(chaos): remove this
-				uint8 bWasFrictionRestored : 1;				// Whether we used data from the previous frame to set up the contacts to maintain static friction
-				uint8 bInsideStaticFrictionCone : 1;		// Whether we are inside the static friction cone
-			};
-			uint8 Bits;
-		} Flags;
-	};
-
-	/**
-	 * @brief The friction data for a manifold point
-	 * This is the information that needs to be stored between ticks to implement static friction.
-	*/
-	class CHAOS_API FSavedManifoldPoint
-	{
-	public:
-		FVec3f ShapeContactPoints[2];
-	};
-
 	/*
 	 * @brief Material properties for a collision constraint
 	*/
@@ -360,12 +296,10 @@ namespace Chaos
 		const FReal GetCollisionRadius0() const { return (Flags.bIsQuadratic0) ? CollisionMargins[0] : FReal(0); }
 		const FReal GetCollisionRadius1() const { return (Flags.bIsQuadratic1) ? CollisionMargins[1] : FReal(0); }
 
-		/** \brief Called each frame when the constraint is active after primary collision detection (but not per incremental collision detection call if enabled) */
-		void Activate()
-		{
-			AccumulatedImpulse = FVec3(0);
-			SetDisabled(false);
-		}
+		/** 
+		 * Called each frame when the constraint is active after primary collision detection (but not per incremental collision detection call if enabled) 
+		 */
+		void Activate(const FReal Dt);
 
 		// When a particle is moved under user control, we need to update some cached state to prevent friction from undoing the move
 		void UpdateParticleTransform(FGeometryParticleHandle* InParticle);
@@ -393,6 +327,23 @@ namespace Chaos
 		// @todo(chaos): remove (used by legacy RBAN collision solver)
 		FVec3 CalculateWorldContactLocation() const;
 
+		// Called to force the contact to recollect its material properties (e.g., when materials are modified)
+		void ClearMaterialProperties()
+		{
+			Flags.bMaterialSet = false;
+		}
+
+		// Find the material for this contact and collect the friction, restitution, etc. This data is cached in the 
+		// contact, so ClearMaterialProperties must be called if the material changes to force a re-gather of the data.
+		void UpdateMaterialProperties()
+		{
+			if (!Flags.bMaterialSet)
+			{
+				UpdateMaterialPropertiesImpl();
+				Flags.bMaterialSet = true;
+			}
+		}
+
 		void SetModifierApplied() { Flags.bModifierApplied = true; }
 
 		const FPBDCollisionConstraintMaterial& GetCollisionMaterial() const { return Material; }
@@ -470,8 +421,21 @@ namespace Chaos
 		FManifoldPoint& GetManifoldPoint(const int32 PointIndex) { return ManifoldPoints[PointIndex]; }
 		const FManifoldPoint& GetManifoldPoint(const int32 PointIndex) const { return ManifoldPoints[PointIndex]; }
 		const FManifoldPoint* GetClosestManifoldPoint() const { return (ClosestManifoldPointIndex != INDEX_NONE) ? &ManifoldPoints[ClosestManifoldPointIndex] : nullptr; }
-		const FSavedManifoldPoint* FindSavedManifoldPoint(const FManifoldPoint& ManifoldPoint) const;
-		void ResetSavedManifoldPoints() { SavedManifoldPoints.Reset(); }
+
+		const FManifoldPointResult& GetManifoldPointResult(const int32 PointIndex) const
+		{
+			// We may request manifold point results for points that are never simulated (e.g., they start asleep).
+			// We could handle this by adding zeroed results when we add manifold points, but that's unnecessary work for
+			// active constraints, and this is function currently only used for debug draw and stats. If that changes we 
+			// may want to change this.
+			check(PointIndex < NumManifoldPoints());
+			if (PointIndex >= ManifoldPointResults.Num())
+			{
+				static FManifoldPointResult ZeroResult;
+				return ZeroResult;
+			}
+			return ManifoldPointResults[PointIndex];
+		}
 
 		inline int32 NumEnabledManifoldPoints() const
 		{
@@ -601,11 +565,11 @@ namespace Chaos
 		ECollisionConstraintDirection GetConstraintDirection(const FReal Dt) const;
 
 		/**
-		 * @brief Called before SetSolverResults() to reset accumulators
+		 * @brief Clear the saved manifold points. This effectively resets friction anchors.
 		*/
-		void ResetSolverResults()
+		void ResetSavedManifoldPoints()
 		{
-			ResetSavedManifoldPoints();
+			SavedManifoldPoints.Reset();
 		}
 
 		/**
@@ -638,10 +602,17 @@ namespace Chaos
 			AccumulatedImpulse += FVec3f(InNetImpulse);
 		}
 
+		FORCEINLINE void ResetSolverResults()
+		{
+			// NOTE: does not initalize any data. All properties will be written to in SetSolverResults
+			SavedManifoldPoints.SetNum(ManifoldPoints.Num());
+			ManifoldPointResults.SetNum(ManifoldPoints.Num());
+		}
+
 		/**
 		 * @brief Store the data from the solver that is retained between ticks for the specified manifold point or used by dependent systems (plasticity, breaking, etc.)
 		*/
-		void SetSolverResults(
+		FORCEINLINE void SetSolverResults(
 			const int32 ManifoldPointIndex, 
 			const FVec3f& NetPushOut, 
 			const FVec3f& NetImpulse, 
@@ -649,41 +620,42 @@ namespace Chaos
 			const FRealSingle Dt)
 		{
 			FManifoldPoint& ManifoldPoint = ManifoldPoints[ManifoldPointIndex];
+			FManifoldPointResult& ManifoldPointResult = ManifoldPointResults[ManifoldPointIndex];
+			FSavedManifoldPoint& SavedManifoldPoint = SavedManifoldPoints[ManifoldPointIndex];
 
-			ManifoldPoint.NetPushOut = NetPushOut;
-			ManifoldPoint.NetImpulse = NetImpulse;
-			ManifoldPoint.Flags.bInsideStaticFrictionCone = FMath::IsNearlyEqual(StaticFrictionRatio, FReal(1));
+			ManifoldPointResult.NetPushOut = NetPushOut;
+			ManifoldPointResult.NetImpulse = NetImpulse;
+			ManifoldPointResult.bIsValid = true;
+			ManifoldPointResult.bInsideStaticFrictionCone = false;
 
 			AccumulatedImpulse += NetImpulse + (NetPushOut / Dt);
 
-			// Save contact data for friction
-			// NOTE: we do this even for points that did not apply PushOut or Impulse so that
-			// we get previous contact data for initial contacts (sometimes). Otherwise we
-			// end up having to estimate the previous contact from velocities
-			if (!SavedManifoldPoints.IsFull())
+			// Save contact data for friction\
+			// NOTE: We also write the anchors back to the manifold point for use if the point gets restored next tick
+			// in which case it skips the saved point search etc
+			if (StaticFrictionRatio >= FReal(1.0f - UE_KINDA_SMALL_NUMBER))
 			{
-				const int32 SavedIndex = SavedManifoldPoints.Add();
-				FSavedManifoldPoint& SavedManifoldPoint = SavedManifoldPoints[SavedIndex];
-
-				if (StaticFrictionRatio >= FReal(1.0f - UE_KINDA_SMALL_NUMBER))
-				{
-					// StaticFrictionRatio ~= 1: Static friction held - we keep the same contacts points as-is for use next frame
-					SavedManifoldPoint.ShapeContactPoints[0] = ManifoldPoint.ShapeAnchorPoints[0];
-					SavedManifoldPoint.ShapeContactPoints[1] = ManifoldPoint.ShapeAnchorPoints[1];
-				}
-				else if (StaticFrictionRatio < FReal(UE_KINDA_SMALL_NUMBER))
-				{
-					// StaticFrictionRatio ~= 0: No friction (or no contact) - discard the friction anchors
-					SavedManifoldPoint.ShapeContactPoints[0] = ManifoldPoint.ContactPoint.ShapeContactPoints[0];
-					SavedManifoldPoint.ShapeContactPoints[1] = ManifoldPoint.ContactPoint.ShapeContactPoints[1];
-				}
-				else
-				{
-					// 0 < StaticFrictionRatio < 1: We exceeded the friction cone. Slide the friction anchor 
-					// toward the last-detected contact position so that it sits at the edge of the friction cone.
-					SavedManifoldPoint.ShapeContactPoints[0] = FVec3::Lerp(ManifoldPoint.ContactPoint.ShapeContactPoints[0], ManifoldPoint.ShapeAnchorPoints[0], StaticFrictionRatio);
-					SavedManifoldPoint.ShapeContactPoints[1] = FVec3::Lerp(ManifoldPoint.ContactPoint.ShapeContactPoints[1], ManifoldPoint.ShapeAnchorPoints[1], StaticFrictionRatio);
-				}
+				// StaticFrictionRatio ~= 1: Static friction held - we keep the same contacts points as-is for use next frame
+				SavedManifoldPoint.ShapeContactPoints[0] = ManifoldPoint.ShapeAnchorPoints[0];
+				SavedManifoldPoint.ShapeContactPoints[1] = ManifoldPoint.ShapeAnchorPoints[1];
+				ManifoldPointResult.bInsideStaticFrictionCone = true;
+			}
+			else if (StaticFrictionRatio < FReal(UE_KINDA_SMALL_NUMBER))
+			{
+				// StaticFrictionRatio ~= 0: No friction (or no contact) - discard the friction anchors
+				const FVec3 Anchor0 = ManifoldPoint.ContactPoint.ShapeContactPoints[0];
+				const FVec3 Anchor1 = ManifoldPoint.ContactPoint.ShapeContactPoints[1];
+				SavedManifoldPoint.ShapeContactPoints[0] = Anchor0;
+				SavedManifoldPoint.ShapeContactPoints[1] = Anchor1;
+			}
+			else
+			{
+				// 0 < StaticFrictionRatio < 1: We exceeded the friction cone. Slide the friction anchor 
+				// toward the last-detected contact position so that it sits at the edge of the friction cone.
+				const FVec3 Anchor0 = FVec3::Lerp(ManifoldPoint.ContactPoint.ShapeContactPoints[0], ManifoldPoint.ShapeAnchorPoints[0], StaticFrictionRatio);
+				const FVec3 Anchor1 = FVec3::Lerp(ManifoldPoint.ContactPoint.ShapeContactPoints[1], ManifoldPoint.ShapeAnchorPoints[1], StaticFrictionRatio);
+				SavedManifoldPoint.ShapeContactPoints[0] = Anchor0;
+				SavedManifoldPoint.ShapeContactPoints[1] = Anchor1;
 			}
 		}
 
@@ -727,22 +699,25 @@ namespace Chaos
 		bool AreMatchingContactPoints(const FContactPoint& A, const FContactPoint& B, FReal& OutScore) const;
 		int32 FindManifoldPoint(const FContactPoint& ContactPoint) const;
 
-		inline void InitManifoldPoint(const int32 ManifoldPointIndex)
+		int32 FindSavedManifoldPoint(const int32 ManifoldPointIndex, TCArray<int32, MaxManifoldPoints>& InOutAllowedSavedPointIndices) const;
+		void AssignSavedManifoldPoints();
+
+		inline void InitManifoldPoint(const int32 ManifoldPointIndex, const FContactPoint& ContactPoint)
 		{
 			FManifoldPoint& ManifoldPoint = ManifoldPoints[ManifoldPointIndex];
-			ManifoldPoint.InitialShapeContactPoints[0] = ManifoldPoint.ContactPoint.ShapeContactPoints[0];
-			ManifoldPoint.InitialShapeContactPoints[1] = ManifoldPoint.ContactPoint.ShapeContactPoints[1];
+			ManifoldPoint.ContactPoint = ContactPoint;
+			ManifoldPoint.ShapeAnchorPoints[0] = ContactPoint.ShapeContactPoints[0];
+			ManifoldPoint.ShapeAnchorPoints[1] = ContactPoint.ShapeContactPoints[1];
+			ManifoldPoint.InitialShapeContactPoints[0] = ContactPoint.ShapeContactPoints[0];
+			ManifoldPoint.InitialShapeContactPoints[1] = ContactPoint.ShapeContactPoints[1];
 			ManifoldPoint.TargetPhi = FReal(0);
 			ManifoldPoint.Flags.Reset();
-			ManifoldPoint.NetPushOut = FVec3(0);
-			ManifoldPoint.NetImpulse = FVec3(0);
 		}
 
 		inline int32 AddManifoldPoint(const FContactPoint& ContactPoint)
 		{
 			int32 ManifoldPointIndex = ManifoldPoints.Add();	// Note: no initialization (see TCArray)
-			ManifoldPoints[ManifoldPointIndex].ContactPoint = ContactPoint;
-			InitManifoldPoint(ManifoldPointIndex);
+			InitManifoldPoint(ManifoldPointIndex, ContactPoint);
 			return ManifoldPointIndex;
 		}
 
@@ -754,9 +729,33 @@ namespace Chaos
 
 		void InitCCDThreshold();
 
+		void UpdateMaterialPropertiesImpl();
+
 	private:
 		FReal CalculateSavedManifoldPointScore(const FSavedManifoldPoint& SavedManifoldPoint, const FManifoldPoint& ManifoldPoint, const FReal DistanceToleranceSq) const;
 
+		union FFlags
+		{
+			FFlags() : Bits(0) {}
+			struct
+			{
+				uint32 bDisabled : 1;
+				uint32 bUseManifold : 1;
+				uint32 bUseIncrementalManifold : 1;
+				uint32 bWasManifoldRestored : 1;
+				uint32 bIsQuadratic0 : 1;
+				uint32 bIsQuadratic1 : 1;
+				uint32 bIsProbeUnmodified : 1;  // Is this constraint a probe pre-contact-modification
+				uint32 bIsProbe : 1;            // Is this constraint currently a probe
+				uint32 bCCDEnabled : 1;			// Is CCD enabled for the current tick
+				uint32 bCCDSweepEnabled: 1;		// If this is a CCD constraint, do we want to enable the sweep/rewind phase?
+				uint32 bModifierApplied : 1;	// Was a constraint modifier applied this tick
+				uint32 bMaterialSet : 1;		// Has the material been set (or does it need to be reset)
+			};
+			uint32 Bits;
+		};
+		static_assert(sizeof(FFlags) == 4, "Unexpected size for FPBDCollisionConstraint::FFLags");
+
 		// Local-space transforms of the shape (relative to particle)
 		FRigidTransform3 ImplicitTransform[2];
 		
@@ -765,9 +764,6 @@ namespace Chaos
 		const FPerShapeData* Shape[2];
 		const FBVHParticles* Simplicial[2];
 
-		FPBDCollisionConstraintMaterial Material;
-		FRealSingle Stiffness;
-
 	public:
 		FVec3f AccumulatedImpulse;					// @todo(chaos): we need to accumulate angular impulse separately
 
@@ -795,39 +791,25 @@ namespace Chaos
 		// Used by manifold point injection to see how many points were in the manifold before UpdateAndTryRestore
 		int32 ExpectedNumManifoldPoints;
 
-		union FFlags
-		{
-			FFlags() : Bits(0) {}
-			struct
-			{
-				uint32 bDisabled : 1;
-				uint32 bUseManifold : 1;
-				uint32 bUseIncrementalManifold : 1;
-				uint32 bWasManifoldRestored : 1;
-				uint32 bIsQuadratic0 : 1;
-				uint32 bIsQuadratic1 : 1;
-				uint32 bIsProbeUnmodified : 1;  // Is this constraint a probe pre-contact-modification
-				uint32 bIsProbe : 1;            // Is this constraint currently a probe
-				uint32 bCCDEnabled : 1;			// Is CCD enabled for the current tick
-				uint32 bCCDSweepEnabled: 1;		// If this is a CCD constraint, do we want to enable the sweep/rewind phase?
-				uint32 bModifierApplied : 1;	// Was a constraint modifier applied this tick
-			};
-			uint32 Bits;
-		} Flags;
-
 		// Relative transform the last time we ran the narrow phase
 		// Used to detect when the bodies have moved too far to reues the manifold
 		FVec3f LastShapeWorldPositionDelta;
 		FRotation3f LastShapeWorldRotationDelta;
 
-		// @todo(chaos): These are only needed here to support incremental collision detection. We should pass them to the functions instead
-		const FSolverBody* SolverBodies[2];
-
 		// Simplex data from the last call to GJK, used to warm-start GJK
 		FGJKSimplexData GJKWarmStartData;
 
+		FPBDCollisionConstraintMaterial Material;
+		FRealSingle Stiffness;
+
+		FFlags Flags;
+
 		TCArray<FSavedManifoldPoint, MaxManifoldPoints> SavedManifoldPoints;
 		TCArray<FManifoldPoint, MaxManifoldPoints> ManifoldPoints;
+		TCArray<FManifoldPointResult, MaxManifoldPoints> ManifoldPointResults;
+
+		// @todo(chaos): These are only needed here to support incremental collision detection for LevelSet. Fix this.
+		const FSolverBody* SolverBodies[2];
 
 		// Value in range [0,1] used to interpolate P between [X,P] that we will rollback to when solving at time of impact.
 		FRealSingle CCDTimeOfImpact;
diff --git a/Engine/Source/Runtime/Experimental/Chaos/Public/Chaos/Collision/PBDCollisionSolver.h b/Engine/Source/Runtime/Experimental/Chaos/Public/Chaos/Collision/PBDCollisionSolver.h
new file mode 100644
--- /dev/null
+++ b/Engine/Source/Runtime/Experimental/Chaos/Public/Chaos/Collision/PBDCollisionSolver.h
@@ -0,0 +1,1365 @@
+// Copyright Epic Games, Inc. All Rights Reserved.
+#pragma once
+
+#include "Chaos/Core.h"
+#include "Chaos/Collision/ContactPoint.h"
+#include "Chaos/Defines.h"
+#include "Chaos/Evolution/SolverBody.h"
+
+CHAOS_API DECLARE_LOG_CATEGORY_EXTERN(LogChaosCollision, Log, All);
+
+// Set to 0 to use a linearized error calculation, and set to 1 to use a non-linear error calculation in collision detection. 
+// In principle nonlinear is more accurate when large rotation corrections occur, but this is not too important for collisions because 
+// when the bodies settle the corrections are small. The linearized version is significantly faster than the non-linear version because 
+// the non-linear version requires a quaternion multiply and renormalization whereas the linear version is just a cross product.
+#define CHAOS_NONLINEAR_COLLISIONS_ENABLED 0
+
+namespace Chaos
+{
+	class FManifoldPoint;
+	class FPBDCollisionConstraint;
+
+	namespace Private
+	{
+		class FPBDCollisionSolver;
+	}
+
+	namespace CVars
+	{
+		extern bool bChaos_PBDCollisionSolver_Velocity_AveragePointEnabled;
+		extern bool bChaos_PBDCollisionSolver_Velocity_FrictionEnabled;
+		extern float Chaos_PBDCollisionSolver_Position_StaticFrictionStiffness;
+		extern float Chaos_PBDCollisionSolver_JacobiStiffness;
+		extern float Chaos_PBDCollisionSolver_JacobiPositionTolerance;
+		extern float Chaos_PBDCollisionSolver_JacobiRotationTolerance;
+	}
+
+	namespace Private
+	{
+		/**
+		 * @brief A single contact point in a FPBDCollisionSolver
+		 * @note Internal Chaos class subject to API changes, renaming or removal
+		*/
+		class FPBDCollisionSolverManifoldPoint
+		{
+		public:
+			/**
+			 * @brief Initialize the geometric data for the contact
+			*/
+			void InitContact(
+				const FSolverReal Dt,
+				const FConstraintSolverBody& Body0,
+				const FConstraintSolverBody& Body1);
+
+			/**
+			 * @brief Initialize the material related properties of the contact
+			*/
+			//void InitMaterial(
+			//	const FConstraintSolverBody& Body0,
+			//	const FConstraintSolverBody& Body1,
+			//	const FSolverReal InRestitution,
+			//	const FSolverReal InRestitutionVelocityThreshold);
+
+			/**
+			 * @brief Update the cached mass properties based on the current body transforms
+			*/
+			void UpdateMass(const FConstraintSolverBody& Body0, const FConstraintSolverBody& Body1);
+
+			/**
+			 * @brief Update the contact mass for the normal correction
+			 * This is used by shock propagation.
+			*/
+			void UpdateMassNormal(const FConstraintSolverBody& Body0, const FConstraintSolverBody& Body1);
+
+			/**
+			 * @brief Calculate the position error at the current transforms
+			 * @param MaxPushOut a limit on the position error for this iteration to prevent initial-penetration explosion (a common PBD problem)
+			*/
+			void CalculateContactPositionErrorNormal(const FConstraintSolverBody& Body0, const FConstraintSolverBody& Body1, const FSolverReal MaxPushOut, FSolverReal& OutContactDeltaNormal) const;
+			void CalculateContactPositionErrorTangential(const FConstraintSolverBody& Body0, const FConstraintSolverBody& Body1, FSolverReal& OutContactDeltaTanget0, FSolverReal& OutContactDeltaTangent1) const;
+			void CalculateContactPositionError(const FConstraintSolverBody& Body0, const FConstraintSolverBody& Body1, const FSolverReal MaxPushOut, FSolverReal& OutContactDeltaNormal, FSolverReal& OutContactDeltaTanget0, FSolverReal& OutContactDeltaTangent1) const;
+
+			/**
+			 * @brief Calculate the velocity error at the current transforms
+			*/
+			void CalculateContactVelocityErrorNormal(const FConstraintSolverBody& Body0, const FConstraintSolverBody& Body1, FSolverReal& OutContactVelocityDeltaNormal) const;
+			void CalculateContactVelocityError(const FConstraintSolverBody& Body0, const FConstraintSolverBody& Body1, const FSolverReal DynamicFriction, const FSolverReal Dt, FSolverReal& OutContactVelocityDeltaNormal, FSolverReal& OutContactVelocityDeltaTangent0, FSolverReal& OutContactVelocityDeltaTangent1) const;
+
+			// @todo(chaos): make private
+		public:
+			friend class FPBDCollisionSolver;
+
+			/**
+			 * @brief Calculate the relative velocity at the contact point
+			 * @note InitContact must be called before calling this function
+			*/
+			FSolverVec3 CalculateContactVelocity(const FConstraintSolverBody& Body0, const FConstraintSolverBody& Body1) const;
+
+			/**
+			 * @brief Whether we need to solve velocity for this manifold point (only if we were penetrating or applied a pushout)
+			*/
+			bool ShouldSolveVelocity() const;
+
+			FWorldContactPoint WorldContact;
+
+			// World-space body-relative contact points
+			//FSolverVec3 RelativeContactPosition0;
+			//FSolverVec3 RelativeContactPosition1;
+
+			// World-space contact axes (normal and 2 tangents)
+			//FSolverVec3 WorldContactNormal;
+			//FSolverVec3 WorldContactTangentU;
+			//FSolverVec3 WorldContactTangentV;
+
+			// I^-1.(R x A) for each body where A is each axis (Normal, TangentU, TangentV)
+			FSolverVec3 WorldContactNormalAngular0;
+			FSolverVec3 WorldContactTangentUAngular0;
+			FSolverVec3 WorldContactTangentVAngular0;
+			FSolverVec3 WorldContactNormalAngular1;
+			FSolverVec3 WorldContactTangentUAngular1;
+			FSolverVec3 WorldContactTangentVAngular1;
+
+			// Contact mass (for non-friction)
+			FSolverReal ContactMassNormal;
+			FSolverReal ContactMassTangentU;
+			FSolverReal ContactMassTangentV;
+
+			// Initial world-space contact separation that we are trying to correct
+			//FSolverReal WorldContactDeltaNormal;
+			//FSolverReal WorldContactDeltaTangentU;
+			//FSolverReal WorldContactDeltaTangentV;
+
+			// Desired final normal velocity, taking Restitution into account
+			//FSolverReal WorldContactVelocityTargetNormal;
+
+			// Solver outputs
+			FSolverReal NetPushOutNormal;
+			FSolverReal NetPushOutTangentU;
+			FSolverReal NetPushOutTangentV;
+			FSolverReal NetImpulseNormal;
+			FSolverReal NetImpulseTangentU;
+			FSolverReal NetImpulseTangentV;
+
+			// A measure of how much we exceeded the static friction threshold.
+			// Equal to (NormalPushOut / TangentialPushOut) before clamping to the friction cone.
+			// Used to move the static friction anchors to the edge of the cone in Scatter.
+			FSolverReal StaticFrictionRatio;
+
+			// Whether this point is enabled
+			// @todo(chaos): consider decoupling the solver manifold points array size from the constraint manifold points array
+			//uint8 bEnabled : 1;
+		};
+
+		/**
+		 * @brief
+		 * @todo(chaos): Make this solver operate on a single contact point rather than all points in a manifold.
+		 * This would be beneficial if we have many contacts with less than 4 points in the manifold. However this
+		 * is dificult to do while we are still supporting non-manifold collisions.
+		*/
+		class FPBDCollisionSolver
+		{
+		public:
+			static const int32 MaxConstrainedBodies = 2;
+			static const int32 MaxPointsPerConstraint = 4;
+
+			// Create a solver that is initialized to safe defaults
+			static FPBDCollisionSolver MakeInitialized()
+			{
+				FPBDCollisionSolver Solver;
+				Solver.State.Init();
+				return Solver;
+			}
+
+			// Create a solver with no initialization
+			static FPBDCollisionSolver MakeUninitialized()
+			{
+				return FPBDCollisionSolver();
+			}
+
+			// NOTE: Does not initialize any properties. See MakeInitialized
+			FPBDCollisionSolver() {}
+
+			/** Reset the state of the collision solver */
+			void Reset()
+			{
+				State.SolverBodies[0].Reset();
+				State.SolverBodies[1].Reset();
+				State.ManifoldPoints = nullptr;
+				State.NumManifoldPoints = 0;
+				State.MaxManifoldPoints = 0;
+			}
+
+			void ResetManifold()
+			{
+				State.NumManifoldPoints = 0;
+			}
+
+			FSolverReal StaticFriction() const { return State.StaticFriction; }
+			FSolverReal DynamicFriction() const { return State.DynamicFriction; }
+			FSolverReal VelocityFriction() const { return State.VelocityFriction; }
+
+			void SetFriction(const FSolverReal InStaticFriction, const FSolverReal InDynamicFriction, const FSolverReal InVelocityFriction)
+			{
+				State.StaticFriction = InStaticFriction;
+				State.DynamicFriction = InDynamicFriction;
+				State.VelocityFriction = InVelocityFriction;
+			}
+
+			void SetStiffness(const FSolverReal InStiffness)
+			{
+				State.Stiffness = InStiffness;
+			}
+
+			void SetSolverBodies(FSolverBody& SolverBody0, FSolverBody& SolverBody1)
+			{
+				State.SolverBodies[0].SetSolverBody(SolverBody0);
+				State.SolverBodies[1].SetSolverBody(SolverBody1);
+			}
+
+			int32 NumManifoldPoints() const
+			{
+				return State.NumManifoldPoints;
+			}
+
+			int32 AddManifoldPoint()
+			{
+				check(State.NumManifoldPoints < State.MaxManifoldPoints);
+				return State.NumManifoldPoints++;
+			}
+
+			void SetManifoldPointsBuffer(FPBDCollisionSolverManifoldPoint* InManifoldPoints, const int32 InMaxManifoldPoints)
+			{
+				State.ManifoldPoints = InManifoldPoints;
+				State.MaxManifoldPoints = InMaxManifoldPoints;
+				State.NumManifoldPoints = 0;
+			}
+
+			const FPBDCollisionSolverManifoldPoint& GetManifoldPoint(const int32 ManifoldPointIndex) const
+			{
+				check(ManifoldPointIndex < NumManifoldPoints());
+				return State.ManifoldPoints[ManifoldPointIndex];
+			}
+
+			FWorldContactPoint& GetWorldContactPoint(const int32 ManifoldPointIndex)
+			{
+				return State.ManifoldPoints[ManifoldPointIndex].WorldContact;
+			}
+
+			/**
+			 * Set up a manifold point (also calls FinalizeManifoldPoint)
+			*/
+			void SetManifoldPoint(
+				const int32 ManifoldPoiontIndex,
+				const FSolverReal Dt,
+				const FSolverVec3& InRelativeContactPosition0,
+				const FSolverVec3& InRelativeContactPosition1,
+				const FSolverVec3& InWorldContactNormal,
+				const FSolverVec3& InWorldContactTangentU,
+				const FSolverVec3& InWorldContactTangentV,
+				const FSolverReal InWorldContactDeltaNormal,
+				const FSolverReal InWorldContactDeltaTangentU,
+				const FSolverReal InWorldContactDeltaTangentV,
+				const FSolverReal InWorldContactVelocityTargetNormal);
+
+			/** 
+			 * Finish manifold point setup.
+			 * NOTE: Can only be called after the WorldContact has been set up (e.g., see SetManifoldPoint)
+			*/
+			void FinalizeManifoldPoint(const int32 ManifoldPoiontIndex, const FSolverReal Dt);
+
+			/**
+			 * @brief Get the first (decorated) solver body
+			 * The decorator add a possible mass scale
+			*/
+			FConstraintSolverBody& SolverBody0() { return State.SolverBodies[0]; }
+			const FConstraintSolverBody& SolverBody0() const { return State.SolverBodies[0]; }
+
+			/**
+			 * @brief Get the second (decorated) solver body
+			 * The decorator add a possible mass scale
+			*/
+			FConstraintSolverBody& SolverBody1() { return State.SolverBodies[1]; }
+			const FConstraintSolverBody& SolverBody1() const { return State.SolverBodies[1]; }
+
+			/**
+			 * @brief Set up the mass scaling for shock propagation, using the position-phase mass scale
+			*/
+			void EnablePositionShockPropagation();
+
+			/**
+			 * @brief Set up the mass scaling for shock propagation, using the velocity-phase mass scale
+			*/
+			void EnableVelocityShockPropagation();
+
+			/**
+			 * @brief Disable mass scaling
+			*/
+			void DisableShockPropagation();
+
+			/**
+			 * @brief Calculate and apply the position correction for this iteration
+			 * @return true if we need to run more iterations, false if we did not apply any correction
+			*/
+			bool SolvePositionWithFriction(const FSolverReal Dt, const FSolverReal MaxPushOut);
+			bool SolvePositionWithFrictionJacobi(const FSolverReal Dt, const FSolverReal MaxPushOut);
+			bool SolvePositionNoFriction(const FSolverReal Dt, const FSolverReal MaxPushOut);
+			bool SolvePositionNoFrictionJacobi(const FSolverReal Dt, const FSolverReal MaxPushOut);
+
+			/**
+			 * @brief Calculate and apply the velocity correction for this iteration
+			 * @return true if we need to run more iterations, false if we did not apply any correction
+			*/
+			bool SolveVelocity(const FSolverReal Dt, const bool bApplyDynamicFriction);
+			bool SolveVelocityJacobi(const FSolverReal Dt, const bool bApplyDynamicFriction);
+
+		private:
+			/**
+			 * @brief Apply the inverse mass scale the body with the lower level
+			 * @param InvMassScale
+			*/
+			void SetShockPropagationInvMassScale(const FSolverReal InvMassScale);
+
+			/**
+			 * @brief Run a velocity solve on the average point from all the points that received a position impulse
+			 * This is used to enforce Restitution constraints without introducing rotation artefacts without
+			 * adding more velocity iterations.
+			 * This will only perform work if there is more than one active contact.
+			*/
+			void SolveVelocityAverage(const FSolverReal Dt);
+
+			struct FState
+			{
+				FState()
+				{
+				}
+
+				void Init()
+				{
+					SolverBodies[0].Init();
+					SolverBodies[1].Init();
+					StaticFriction = 0;
+					DynamicFriction = 0;
+					VelocityFriction = 0;
+					Stiffness = 1;
+					ManifoldPoints = nullptr;
+					NumManifoldPoints = 0;
+					MaxManifoldPoints = 0;
+				}
+
+				// Static Friction in the position-solve phase
+				FSolverReal StaticFriction;
+
+				// Dynamic Friction in the position-solve phase
+				FSolverReal DynamicFriction;
+
+				// Dynamic Friction in the velocity-solve phase
+				FSolverReal VelocityFriction;
+
+				// Solver stiffness (scales all pushout and impulses)
+				FSolverReal Stiffness;
+
+				// Bodies and contacts
+				FConstraintSolverBody SolverBodies[MaxConstrainedBodies];
+				FPBDCollisionSolverManifoldPoint* ManifoldPoints;
+				int32 NumManifoldPoints;
+				int32 MaxManifoldPoints;
+			};
+
+			FState State;
+		};
+
+
+		//////////////////////////////////////////////////////////////////////////////////////////////////
+		//////////////////////////////////////////////////////////////////////////////////////////////////
+		//////////////////////////////////////////////////////////////////////////////////////////////////
+		//////////////////////////////////////////////////////////////////////////////////////////////////
+
+		// If the vector size squared is below the threshold, zero the vector and return false, otherwise leave vector as-is and return true.
+		// This is useful for maintaining perfect restitution bounces - we remove any rounding errors in the impulses before applying to the rotation
+		// @todo(chaos): consider using a test of absolute element values instead
+		// @todo(chaos): disabled for now - too expensive :(
+		FORCEINLINE_DEBUGGABLE bool AboveThresholdElseZero(FSolverVec3& InOutV, const FSolverReal ThresholdSq)
+		{
+			//if (InOutV.SizeSquared() < ThresholdSq)
+			//{
+			//	InOutV = FSolverVec3(0);
+			//	return false;
+			//}
+			return true;
+		}
+
+		FORCEINLINE_DEBUGGABLE void CalculatePositionCorrectionNormal(
+			const FSolverReal Stiffness,
+			const FSolverReal ContactDeltaNormal,
+			const FSolverReal ContactMassNormal,
+			const FSolverReal NetPushOutNormal,
+			FSolverReal& OutPushOutNormal)
+		{
+			const FSolverReal PushOutNormal = -Stiffness * ContactDeltaNormal * ContactMassNormal;
+
+			// The total pushout so far this sub-step
+			// Unilateral constraint: Net-negative impulses not allowed (negative incremental impulses are allowed as long as the net is positive)
+			if ((NetPushOutNormal + PushOutNormal) > FSolverReal(0))
+			{
+				OutPushOutNormal = PushOutNormal;
+			}
+			else
+			{
+				OutPushOutNormal = -NetPushOutNormal;
+			}
+		}
+
+
+		FORCEINLINE_DEBUGGABLE void CalculatePositionCorrectionTangent(
+			const FSolverReal Stiffness,
+			const FSolverReal ContactDeltaTangent,
+			const FSolverReal ContactMassTangent,
+			const FSolverReal NetPushOutTangent,
+			FSolverReal& OutPushOutTangent)
+		{
+			// Bilateral constraint - negative values allowed (unlike the normal correction)
+			OutPushOutTangent = -Stiffness * ContactMassTangent * ContactDeltaTangent;
+		}
+
+
+		FORCEINLINE_DEBUGGABLE void ApplyFrictionCone(
+			const FSolverReal StaticFriction,
+			const FSolverReal DynamicFriction,
+			const FSolverReal MaxFrictionPushOut,
+			FSolverReal& InOutPushOutTangentU,
+			FSolverReal& InOutPushOutTangentV,
+			FSolverReal& InOutNetPushOutTangentU,
+			FSolverReal& InOutNetPushOutTangentV,
+			FSolverReal& OutStaticFrictionRatio)
+		{
+			// Assume we stay in the friction cone...
+			OutStaticFrictionRatio = FSolverReal(1);
+
+			if (MaxFrictionPushOut < FSolverReal(UE_KINDA_SMALL_NUMBER))
+			{
+				// Note: we have already added the current iteration's PushOut to the NetPushOut but it has not been applied to the body
+				// so we must subtract it again to calculate the actual pushout we want to undo (i.e., the net pushout that has been applied 
+				// to the body so far from previous iterations)
+				InOutPushOutTangentU = -(InOutNetPushOutTangentU - InOutPushOutTangentU);
+				InOutPushOutTangentV = -(InOutNetPushOutTangentV - InOutPushOutTangentV);
+				InOutNetPushOutTangentU = FSolverReal(0);
+				InOutNetPushOutTangentV = FSolverReal(0);
+				OutStaticFrictionRatio = FSolverReal(0);
+			}
+			else
+			{
+				// If we exceed the static friction cone, clip to the dynamic friction cone
+				const FSolverReal MaxStaticPushOutTangentSq = FMath::Square(StaticFriction * MaxFrictionPushOut);
+				const FSolverReal NetPushOutTangentSq = FMath::Square(InOutNetPushOutTangentU) + FMath::Square(InOutNetPushOutTangentV);
+				if (NetPushOutTangentSq > MaxStaticPushOutTangentSq)
+				{
+					const FSolverReal MaxDynamicPushOutTangent = DynamicFriction * MaxFrictionPushOut;
+					const FSolverReal FrictionMultiplier = MaxDynamicPushOutTangent * FMath::InvSqrt(NetPushOutTangentSq);
+					const FSolverReal NetPushOutTangentU = FrictionMultiplier * InOutNetPushOutTangentU;
+					const FSolverReal NetPushOutTangentV = FrictionMultiplier * InOutNetPushOutTangentV;
+					InOutPushOutTangentU = NetPushOutTangentU - (InOutNetPushOutTangentU - InOutPushOutTangentU);
+					InOutPushOutTangentV = NetPushOutTangentV - (InOutNetPushOutTangentV - InOutPushOutTangentV);
+					InOutNetPushOutTangentU = NetPushOutTangentU;
+					InOutNetPushOutTangentV = NetPushOutTangentV;
+					OutStaticFrictionRatio = FrictionMultiplier;
+				}
+			}
+		}
+
+
+		FORCEINLINE_DEBUGGABLE void ApplyPositionCorrectionTangential(
+			const FSolverReal Stiffness,
+			const FSolverReal StaticFriction,
+			const FSolverReal DynamicFriction,
+			const FSolverReal MaxFrictionPushOut,
+			const FSolverReal ContactDeltaTangentU,
+			const FSolverReal ContactDeltaTangentV,
+			FPBDCollisionSolverManifoldPoint& ManifoldPoint,
+			FConstraintSolverBody& Body0,
+			FConstraintSolverBody& Body1)
+		{
+			FSolverReal PushOutTangentU = FSolverReal(0);
+			FSolverReal PushOutTangentV = FSolverReal(0);
+
+			CalculatePositionCorrectionTangent(
+				Stiffness,
+				ContactDeltaTangentU,
+				ManifoldPoint.ContactMassTangentU,
+				ManifoldPoint.NetPushOutTangentU,
+				PushOutTangentU);					// Out
+
+			CalculatePositionCorrectionTangent(
+				Stiffness,
+				ContactDeltaTangentV,
+				ManifoldPoint.ContactMassTangentV,
+				ManifoldPoint.NetPushOutTangentV,
+				PushOutTangentV);					// Out
+
+			ManifoldPoint.NetPushOutTangentU += PushOutTangentU;
+			ManifoldPoint.NetPushOutTangentV += PushOutTangentV;
+
+			ApplyFrictionCone(
+				StaticFriction,
+				DynamicFriction,
+				MaxFrictionPushOut,
+				PushOutTangentU,									// InOut
+				PushOutTangentV,									// InOut
+				ManifoldPoint.NetPushOutTangentU,					// InOut
+				ManifoldPoint.NetPushOutTangentV,					// InOut
+				ManifoldPoint.StaticFrictionRatio);					// Out
+
+			// Update the particle state based on the pushout
+			const FVec3 PushOut = PushOutTangentU * ManifoldPoint.WorldContact.ContactTangentU + PushOutTangentV * ManifoldPoint.WorldContact.ContactTangentV;
+			if (Body0.IsDynamic())
+			{
+				const FSolverVec3 DX0 = Body0.InvM() * PushOut;
+				const FSolverVec3 DR0 = ManifoldPoint.WorldContactTangentUAngular0 * PushOutTangentU + ManifoldPoint.WorldContactTangentVAngular0 * PushOutTangentV;
+				Body0.ApplyPositionDelta(DX0);
+				Body0.ApplyRotationDelta(DR0);
+			}
+			if (Body1.IsDynamic())
+			{
+				const FSolverVec3 DX1 = -Body1.InvM() * PushOut;
+				const FSolverVec3 DR1 = ManifoldPoint.WorldContactTangentUAngular1 * -PushOutTangentU + ManifoldPoint.WorldContactTangentVAngular1 * -PushOutTangentV;
+				Body1.ApplyPositionDelta(DX1);
+				Body1.ApplyRotationDelta(DR1);
+			}
+		}
+
+		FORCEINLINE_DEBUGGABLE void ApplyPositionCorrectionNormal(
+			const FSolverReal Stiffness,
+			const FSolverReal ContactDeltaNormal,
+			FPBDCollisionSolverManifoldPoint& ManifoldPoint,
+			FConstraintSolverBody& Body0,
+			FConstraintSolverBody& Body1)
+		{
+			FSolverReal PushOutNormal = FSolverReal(0);
+
+			CalculatePositionCorrectionNormal(
+				Stiffness,
+				ContactDeltaNormal,
+				ManifoldPoint.ContactMassNormal,
+				ManifoldPoint.NetPushOutNormal,
+				PushOutNormal);						// Out
+
+			ManifoldPoint.NetPushOutNormal += PushOutNormal;
+
+			// Update the particle state based on the pushout
+			if (Body0.IsDynamic())
+			{
+				const FSolverVec3 DX0 = (Body0.InvM() * PushOutNormal) * ManifoldPoint.WorldContact.ContactNormal;
+				const FSolverVec3 DR0 = ManifoldPoint.WorldContactNormalAngular0 * PushOutNormal;
+				Body0.ApplyPositionDelta(DX0);
+				Body0.ApplyRotationDelta(DR0);
+			}
+			if (Body1.IsDynamic())
+			{
+				const FSolverVec3 DX1 = (Body1.InvM() * -PushOutNormal) * ManifoldPoint.WorldContact.ContactNormal;
+				const FSolverVec3 DR1 = ManifoldPoint.WorldContactNormalAngular1 * -PushOutNormal;
+				Body1.ApplyPositionDelta(DX1);
+				Body1.ApplyRotationDelta(DR1);
+			}
+		}
+
+
+		FORCEINLINE_DEBUGGABLE void ApplyVelocityCorrection(
+			const FSolverReal Stiffness,
+			const FSolverReal Dt,
+			const FSolverReal DynamicFriction,
+			const FSolverReal ContactVelocityDeltaNormal,
+			const FSolverReal ContactVelocityDeltaTangent0,
+			const FSolverReal ContactVelocityDeltaTangent1,
+			const FSolverReal MinImpulseNormal,
+			FPBDCollisionSolverManifoldPoint& ManifoldPoint,
+			FConstraintSolverBody& Body0,
+			FConstraintSolverBody& Body1)
+		{
+			FSolverReal ImpulseNormal = -Stiffness * ManifoldPoint.ContactMassNormal * ContactVelocityDeltaNormal;
+			FSolverReal ImpulseTangentU = -Stiffness * ManifoldPoint.ContactMassTangentU * ContactVelocityDeltaTangent0;
+			FSolverReal ImpulseTangentV = -Stiffness * ManifoldPoint.ContactMassTangentV * ContactVelocityDeltaTangent1;
+
+			// Clamp the total impulse to be positive along the normal. We can apply a net negative impulse, 
+			// but only to correct the velocity that was added by pushout (in which case MinImpulseNormal will be negative).
+			const FSolverReal NetImpulseNormal = ManifoldPoint.NetImpulseNormal + ImpulseNormal;
+			if ((ManifoldPoint.NetImpulseNormal + ImpulseNormal) < MinImpulseNormal)
+			{
+				// We are trying to apply a net negative impulse larger than one to counteract the effective pushout impulse
+				// so clamp the net impulse to be equal to minus the pushout impulse along the normal.
+				ImpulseNormal = MinImpulseNormal - ManifoldPoint.NetImpulseNormal;
+			}
+
+			// Clamp the tangential impulses to the friction cone
+			if ((DynamicFriction > 0) && (Dt > 0))
+			{
+				const FSolverReal MaxImpulseTangent = FMath::Max(FSolverReal(0), DynamicFriction * (ManifoldPoint.NetImpulseNormal + ImpulseNormal + ManifoldPoint.NetPushOutNormal / Dt));
+				const FSolverReal MaxImpulseTangentSq = FMath::Square(MaxImpulseTangent);
+				const FSolverReal ImpulseTangentSq = FMath::Square(ImpulseTangentU) + FMath::Square(ImpulseTangentV);
+				if (ImpulseTangentSq > (MaxImpulseTangentSq + UE_SMALL_NUMBER))
+				{
+					const FSolverReal ImpulseTangentScale = MaxImpulseTangent * FMath::InvSqrt(ImpulseTangentSq);
+					ImpulseTangentU *= ImpulseTangentScale;
+					ImpulseTangentV *= ImpulseTangentScale;
+				}
+			}
+
+			ManifoldPoint.NetImpulseNormal += ImpulseNormal;
+			ManifoldPoint.NetImpulseTangentU += ImpulseTangentU;
+			ManifoldPoint.NetImpulseTangentV += ImpulseTangentV;
+
+			// Apply the velocity deltas from the impulse
+			const FSolverVec3 Impulse = ImpulseNormal * ManifoldPoint.WorldContact.ContactNormal + ImpulseTangentU * ManifoldPoint.WorldContact.ContactTangentU + ImpulseTangentV * ManifoldPoint.WorldContact.ContactTangentV;
+			if (Body0.IsDynamic())
+			{
+				const FSolverVec3 DV0 = Body0.InvM() * Impulse;
+				const FSolverVec3 DW0 = ManifoldPoint.WorldContactNormalAngular0 * ImpulseNormal + ManifoldPoint.WorldContactTangentUAngular0 * ImpulseTangentU + ManifoldPoint.WorldContactTangentVAngular0 * ImpulseTangentV;
+				Body0.ApplyVelocityDelta(DV0, DW0);
+			}
+			if (Body1.IsDynamic())
+			{
+				const FSolverVec3 DV1 = -Body1.InvM() * Impulse;
+				const FSolverVec3 DW1 = ManifoldPoint.WorldContactNormalAngular1 * -ImpulseNormal + ManifoldPoint.WorldContactTangentUAngular1 * -ImpulseTangentU + ManifoldPoint.WorldContactTangentVAngular1 * -ImpulseTangentV;
+				Body1.ApplyVelocityDelta(DV1, DW1);
+			}
+		}
+
+		FORCEINLINE_DEBUGGABLE void ApplyVelocityCorrectionNormal(
+			const FSolverReal Stiffness,
+			const FSolverReal ContactVelocityDeltaNormal,
+			const FSolverReal MinImpulseNormal,
+			FPBDCollisionSolverManifoldPoint& ManifoldPoint,
+			FConstraintSolverBody& Body0,
+			FConstraintSolverBody& Body1)
+		{
+			FSolverReal ImpulseNormal = -(Stiffness * ManifoldPoint.ContactMassNormal) * ContactVelocityDeltaNormal;
+
+			// See comments in ApplyVelocityCorrection
+			const FSolverReal NetImpulseNormal = ManifoldPoint.NetImpulseNormal + ImpulseNormal;
+			if (NetImpulseNormal < MinImpulseNormal)
+			{
+				ImpulseNormal = ImpulseNormal - NetImpulseNormal + MinImpulseNormal;
+			}
+
+			ManifoldPoint.NetImpulseNormal += ImpulseNormal;
+
+			// Calculate the velocity deltas from the impulse
+			FSolverVec3 Impulse = ImpulseNormal * ManifoldPoint.WorldContact.ContactNormal;
+			if (Body0.IsDynamic())
+			{
+				const FSolverVec3 DV0 = Body0.InvM() * Impulse;
+				const FSolverVec3 DW0 = ManifoldPoint.WorldContactNormalAngular0 * ImpulseNormal;
+				Body0.ApplyVelocityDelta(DV0, DW0);
+			}
+			if (Body1.IsDynamic())
+			{
+				const FSolverVec3 DV1 = -Body1.InvM() * Impulse;
+				const FSolverVec3 DW1 = ManifoldPoint.WorldContactNormalAngular1 * -ImpulseNormal;
+				Body1.ApplyVelocityDelta(DV1, DW1);
+			}
+		}
+
+
+		//////////////////////////////////////////////////////////////////////////////////////////////////
+		//////////////////////////////////////////////////////////////////////////////////////////////////
+		//////////////////////////////////////////////////////////////////////////////////////////////////
+		//////////////////////////////////////////////////////////////////////////////////////////////////
+
+		FORCEINLINE_DEBUGGABLE void FPBDCollisionSolverManifoldPoint::InitContact(
+			const FSolverReal Dt,
+			const FConstraintSolverBody& Body0,
+			const FConstraintSolverBody& Body1)
+		{
+			NetPushOutNormal = FSolverReal(0);
+			NetPushOutTangentU = FSolverReal(0);
+			NetPushOutTangentV = FSolverReal(0);
+			NetImpulseNormal = FSolverReal(0);
+			NetImpulseTangentU = FSolverReal(0);
+			NetImpulseTangentV = FSolverReal(0);
+			StaticFrictionRatio = FSolverReal(0);
+
+			UpdateMass(Body0, Body1);
+		}
+
+		FORCEINLINE_DEBUGGABLE void FPBDCollisionSolverManifoldPoint::UpdateMass(const FConstraintSolverBody& Body0, const FConstraintSolverBody& Body1)
+		{
+			FSolverReal ContactMassInvNormal = FSolverReal(0);
+			FSolverReal ContactMassInvTangentU = FSolverReal(0);
+			FSolverReal ContactMassInvTangentV = FSolverReal(0);
+
+			// These are not used if not initialized below so no need to clear
+			//WorldContactNormalAngular0 = FSolverVec3(0);
+			//WorldContactTangentUAngular0 = FSolverVec3(0);
+			//WorldContactTangentVAngular0 = FSolverVec3(0);
+			//WorldContactNormalAngular1 = FSolverVec3(0);
+			//WorldContactTangentUAngular1 = FSolverVec3(0);
+			//WorldContactTangentVAngular1 = FSolverVec3(0);
+
+			if (Body0.IsDynamic())
+			{
+				const FSolverVec3 R0xN = FSolverVec3::CrossProduct(WorldContact.RelativeContactPoints[0], WorldContact.ContactNormal);
+				const FSolverVec3 R0xU = FSolverVec3::CrossProduct(WorldContact.RelativeContactPoints[0], WorldContact.ContactTangentU);
+				const FSolverVec3 R0xV = FSolverVec3::CrossProduct(WorldContact.RelativeContactPoints[0], WorldContact.ContactTangentV);
+
+				const FSolverMatrix33 InvI0 = Body0.InvI();
+
+				WorldContactNormalAngular0 = InvI0 * R0xN;
+				WorldContactTangentUAngular0 = InvI0 * R0xU;
+				WorldContactTangentVAngular0 = InvI0 * R0xV;
+
+				ContactMassInvNormal += FSolverVec3::DotProduct(R0xN, WorldContactNormalAngular0) + Body0.InvM();
+				ContactMassInvTangentU += FSolverVec3::DotProduct(R0xU, WorldContactTangentUAngular0) + Body0.InvM();
+				ContactMassInvTangentV += FSolverVec3::DotProduct(R0xV, WorldContactTangentVAngular0) + Body0.InvM();
+			}
+			if (Body1.IsDynamic())
+			{
+				const FSolverVec3 R1xN = FSolverVec3::CrossProduct(WorldContact.RelativeContactPoints[1], WorldContact.ContactNormal);
+				const FSolverVec3 R1xU = FSolverVec3::CrossProduct(WorldContact.RelativeContactPoints[1], WorldContact.ContactTangentU);
+				const FSolverVec3 R1xV = FSolverVec3::CrossProduct(WorldContact.RelativeContactPoints[1], WorldContact.ContactTangentV);
+
+				const FSolverMatrix33 InvI1 = Body1.InvI();
+
+				WorldContactNormalAngular1 = InvI1 * R1xN;
+				WorldContactTangentUAngular1 = InvI1 * R1xU;
+				WorldContactTangentVAngular1 = InvI1 * R1xV;
+
+				ContactMassInvNormal += FSolverVec3::DotProduct(R1xN, WorldContactNormalAngular1) + Body1.InvM();
+				ContactMassInvTangentU += FSolverVec3::DotProduct(R1xU, WorldContactTangentUAngular1) + Body1.InvM();
+				ContactMassInvTangentV += FSolverVec3::DotProduct(R1xV, WorldContactTangentVAngular1) + Body1.InvM();
+			}
+
+			ContactMassNormal = (ContactMassInvNormal > FSolverReal(UE_SMALL_NUMBER)) ? FSolverReal(1) / ContactMassInvNormal : FSolverReal(0);
+			ContactMassTangentU = (ContactMassInvTangentU > FSolverReal(UE_SMALL_NUMBER)) ? FSolverReal(1) / ContactMassInvTangentU : FSolverReal(0);
+			ContactMassTangentV = (ContactMassInvTangentV > FSolverReal(UE_SMALL_NUMBER)) ? FSolverReal(1) / ContactMassInvTangentV : FSolverReal(0);
+		}
+
+		FORCEINLINE_DEBUGGABLE void FPBDCollisionSolverManifoldPoint::UpdateMassNormal(const FConstraintSolverBody& Body0, const FConstraintSolverBody& Body1)
+		{
+			FSolverReal ContactMassInvNormal = FSolverReal(0);
+			if (Body0.IsDynamic())
+			{
+				const FSolverVec3 R0xN = FSolverVec3::CrossProduct(WorldContact.RelativeContactPoints[0], WorldContact.ContactNormal);
+				const FSolverMatrix33 InvI0 = Body0.InvI();
+				WorldContactNormalAngular0 = InvI0 * R0xN;
+				ContactMassInvNormal += FSolverVec3::DotProduct(R0xN, WorldContactNormalAngular0) + Body0.InvM();
+			}
+			if (Body1.IsDynamic())
+			{
+				const FSolverVec3 R1xN = FSolverVec3::CrossProduct(WorldContact.RelativeContactPoints[1], WorldContact.ContactNormal);
+				const FSolverMatrix33 InvI1 = Body1.InvI();
+				WorldContactNormalAngular1 = InvI1 * R1xN;
+				ContactMassInvNormal += FSolverVec3::DotProduct(R1xN, WorldContactNormalAngular1) + Body1.InvM();
+			}
+			ContactMassNormal = (ContactMassInvNormal > FSolverReal(UE_SMALL_NUMBER)) ? FSolverReal(1) / ContactMassInvNormal : FSolverReal(0);
+		}
+
+		FORCEINLINE_DEBUGGABLE void FPBDCollisionSolverManifoldPoint::CalculateContactPositionErrorNormal(const FConstraintSolverBody& Body0, const FConstraintSolverBody& Body1, const FSolverReal MaxPushOut, FSolverReal& OutContactDeltaNormal) const
+		{
+			// Linear version: calculate the contact delta assuming linear motion after applying a positional impulse at the contact point. There will be an error that depends on the size of the rotation.
+			const FSolverVec3 ContactDelta0 = Body0.DP() + FSolverVec3::CrossProduct(Body0.DQ(), WorldContact.RelativeContactPoints[0]);
+			const FSolverVec3 ContactDelta1 = Body1.DP() + FSolverVec3::CrossProduct(Body1.DQ(), WorldContact.RelativeContactPoints[1]);
+			const FSolverVec3 ContactDelta = ContactDelta0 - ContactDelta1;
+			OutContactDeltaNormal = WorldContact.ContactDeltaNormal + FSolverVec3::DotProduct(ContactDelta, WorldContact.ContactNormal);
+
+			// NOTE: OutContactDeltaNormal is negative for penetration
+			// NOTE: MaxPushOut == 0 disables the pushout limits
+			if ((MaxPushOut > 0) && (OutContactDeltaNormal < -MaxPushOut))
+			{
+				OutContactDeltaNormal = -MaxPushOut;
+			}
+		}
+
+		FORCEINLINE_DEBUGGABLE void FPBDCollisionSolverManifoldPoint::CalculateContactPositionErrorTangential(const FConstraintSolverBody& Body0, const FConstraintSolverBody& Body1, FSolverReal& OutContactDeltaTangentU, FSolverReal& OutContactDeltaTangentV) const
+		{
+			// Linear version: calculate the contact delta assuming linear motion after applying a positional impulse at the contact point. There will be an error that depends on the size of the rotation.
+			const FSolverVec3 ContactDelta0 = Body0.DP() + FSolverVec3::CrossProduct(Body0.DQ(), WorldContact.RelativeContactPoints[0]);
+			const FSolverVec3 ContactDelta1 = Body1.DP() + FSolverVec3::CrossProduct(Body1.DQ(), WorldContact.RelativeContactPoints[1]);
+			const FSolverVec3 ContactDelta = ContactDelta0 - ContactDelta1;
+			OutContactDeltaTangentU = WorldContact.ContactDeltaTangentU + FSolverVec3::DotProduct(ContactDelta, WorldContact.ContactTangentU);
+			OutContactDeltaTangentV = WorldContact.ContactDeltaTangentV + FSolverVec3::DotProduct(ContactDelta, WorldContact.ContactTangentV);
+		}
+
+		FORCEINLINE_DEBUGGABLE void FPBDCollisionSolverManifoldPoint::CalculateContactPositionError(const FConstraintSolverBody& Body0, const FConstraintSolverBody& Body1, const FSolverReal MaxPushOut, FSolverReal& OutContactDeltaNormal, FSolverReal& OutContactDeltaTangentU, FSolverReal& OutContactDeltaTangentV) const
+		{
+	#if CHAOS_NONLINEAR_COLLISIONS_ENABLED
+	#error "Non-linear collision solver probably no longer functional - fix it!"
+			// Non-linear version: calculate the contact delta after we have converted the current positional impulses into position and rotation corrections.
+			// We could precalculate and store the LocalContactPositions if we really want to use this nonlinear version
+			const FSolverVec3 LocalContactPosition0 = Body0.Q().Inverse() * RelativeContactPosition0;
+			const FSolverVec3 LocalContactPosition1 = Body1.Q().Inverse() * RelativeContactPosition1;
+			const FSolverVec3 ContactDelta = (Body0.CorrectedP() + Body0.CorrectedQ() * LocalContactPosition0) - (Body1.CorrectedP() + Body1.CorrectedQ() * LocalContactPosition1);
+			OutContactDeltaNormal = FSolverVec3::DotProduct(ContactDelta, WorldContactNormal);
+			OutContactDeltaTangentU = FSolverVec3::DotProduct(ContactDelta, WorldContactTangentU);
+			OutContactDeltaTangentV = FSolverVec3::DotProduct(ContactDelta, WorldContactTangentV);
+	#else
+			// Linear version: calculate the contact delta assuming linear motion after applying a positional impulse at the contact point. There will be an error that depends on the size of the rotation.
+			const FSolverVec3 ContactDelta0 = Body0.DP() + FSolverVec3::CrossProduct(Body0.DQ(), WorldContact.RelativeContactPoints[0]);
+			const FSolverVec3 ContactDelta1 = Body1.DP() + FSolverVec3::CrossProduct(Body1.DQ(), WorldContact.RelativeContactPoints[1]);
+			const FSolverVec3 ContactDelta = ContactDelta0 - ContactDelta1;
+			OutContactDeltaNormal = WorldContact.ContactDeltaNormal + FSolverVec3::DotProduct(ContactDelta, WorldContact.ContactNormal);
+			OutContactDeltaTangentU = WorldContact.ContactDeltaTangentU + FSolverVec3::DotProduct(ContactDelta, WorldContact.ContactTangentU);
+			OutContactDeltaTangentV = WorldContact.ContactDeltaTangentV + FSolverVec3::DotProduct(ContactDelta, WorldContact.ContactTangentV);
+	#endif
+
+			// NOTE: OutContactDeltaNormal is negative for penetration
+			// NOTE: MaxPushOut == 0 disables the pushout limits
+			if ((MaxPushOut > 0) && (OutContactDeltaNormal < -MaxPushOut))
+			{
+				OutContactDeltaNormal = -MaxPushOut;
+			}
+		}
+
+		FORCEINLINE_DEBUGGABLE void FPBDCollisionSolverManifoldPoint::CalculateContactVelocityError(const FConstraintSolverBody& Body0, const FConstraintSolverBody& Body1, const FSolverReal DynamicFriction, const FSolverReal Dt, FSolverReal& OutContactVelocityDeltaNormal, FSolverReal& OutContactVelocityDeltaTangent0, FSolverReal& OutContactVelocityDeltaTangent1) const
+		{
+			const FSolverVec3 ContactVelocity0 = Body0.V() + FSolverVec3::CrossProduct(Body0.W(), WorldContact.RelativeContactPoints[0]);
+			const FSolverVec3 ContactVelocity1 = Body1.V() + FSolverVec3::CrossProduct(Body1.W(), WorldContact.RelativeContactPoints[1]);
+			const FSolverVec3 ContactVelocity = ContactVelocity0 - ContactVelocity1;
+			const FSolverReal ContactVelocityNormal = FSolverVec3::DotProduct(ContactVelocity, WorldContact.ContactNormal);
+			const FSolverReal ContactVelocityTangent0 = FSolverVec3::DotProduct(ContactVelocity, WorldContact.ContactTangentU);
+			const FSolverReal ContactVelocityTangent1 = FSolverVec3::DotProduct(ContactVelocity, WorldContact.ContactTangentV);
+
+			OutContactVelocityDeltaNormal = (ContactVelocityNormal - WorldContact.ContactTargetVelocityNormal);
+			OutContactVelocityDeltaTangent0 = ContactVelocityTangent0;
+			OutContactVelocityDeltaTangent1 = ContactVelocityTangent1;
+		}
+
+		FORCEINLINE_DEBUGGABLE void FPBDCollisionSolverManifoldPoint::CalculateContactVelocityErrorNormal(const FConstraintSolverBody& Body0, const FConstraintSolverBody& Body1, FSolverReal& OutContactVelocityDeltaNormal) const
+		{
+			const FSolverVec3 ContactVelocity0 = Body0.V() + FSolverVec3::CrossProduct(Body0.W(), WorldContact.RelativeContactPoints[0]);
+			const FSolverVec3 ContactVelocity1 = Body1.V() + FSolverVec3::CrossProduct(Body1.W(), WorldContact.RelativeContactPoints[1]);
+			const FSolverVec3 ContactVelocity = ContactVelocity0 - ContactVelocity1;
+			const FSolverReal ContactVelocityNormal = FSolverVec3::DotProduct(ContactVelocity, WorldContact.ContactNormal);
+
+			// Add up the errors in the velocity (current velocity - desired velocity)
+			OutContactVelocityDeltaNormal = (ContactVelocityNormal - WorldContact.ContactTargetVelocityNormal);
+		}
+
+		FORCEINLINE_DEBUGGABLE FSolverVec3 FPBDCollisionSolverManifoldPoint::CalculateContactVelocity(const FConstraintSolverBody& Body0, const FConstraintSolverBody& Body1) const
+		{
+			const FSolverVec3 ContactVelocity0 = Body0.V() + FVec3::CrossProduct(Body0.W(), WorldContact.RelativeContactPoints[0]);
+			const FSolverVec3 ContactVelocity1 = Body1.V() + FVec3::CrossProduct(Body1.W(), WorldContact.RelativeContactPoints[1]);
+			return ContactVelocity0 - ContactVelocity1;
+		}
+
+		FORCEINLINE_DEBUGGABLE bool FPBDCollisionSolverManifoldPoint::ShouldSolveVelocity() const
+		{
+			// We ensure positive separating velocity for close contacts even if they didn't receive a pushout
+			return (NetPushOutNormal > FSolverReal(0)) || (WorldContact.ContactDeltaNormal < FSolverReal(0));
+		}
+
+		//////////////////////////////////////////////////////////////////////////////////////////////////
+		//////////////////////////////////////////////////////////////////////////////////////////////////
+		//////////////////////////////////////////////////////////////////////////////////////////////////
+		//////////////////////////////////////////////////////////////////////////////////////////////////
+
+		FORCEINLINE_DEBUGGABLE  void FPBDCollisionSolver::FinalizeManifoldPoint(const int32 ManifoldPoiontIndex, const FSolverReal Dt)
+		{
+			State.ManifoldPoints[ManifoldPoiontIndex].InitContact(
+				FSolverReal(Dt),
+				State.SolverBodies[0],
+				State.SolverBodies[1]);
+		}
+
+		FORCEINLINE_DEBUGGABLE void FPBDCollisionSolver::SetManifoldPoint(
+			const int32 ManifoldPoiontIndex,
+			const FSolverReal Dt,
+			const FSolverVec3& InRelativeContactPosition0,
+			const FSolverVec3& InRelativeContactPosition1,
+			const FSolverVec3& InWorldContactNormal,
+			const FSolverVec3& InWorldContactTangentU,
+			const FSolverVec3& InWorldContactTangentV,
+			const FSolverReal InWorldContactDeltaNormal,
+			const FSolverReal InWorldContactDeltaTangentU,
+			const FSolverReal InWorldContactDeltaTangentV,
+			const FSolverReal InWorldContactVelocityTargetNormal)
+		{
+			FWorldContactPoint& WorldContactPoint = State.ManifoldPoints[ManifoldPoiontIndex].WorldContact;
+			WorldContactPoint.RelativeContactPoints[0] = InRelativeContactPosition0;
+			WorldContactPoint.RelativeContactPoints[1] = InRelativeContactPosition1;
+			WorldContactPoint.ContactNormal = InWorldContactNormal;
+			WorldContactPoint.ContactTangentU = InWorldContactTangentU;
+			WorldContactPoint.ContactTangentV = InWorldContactTangentV;
+			WorldContactPoint.ContactDeltaNormal = InWorldContactDeltaNormal;
+			WorldContactPoint.ContactDeltaTangentU = InWorldContactDeltaTangentU;
+			WorldContactPoint.ContactDeltaTangentV = InWorldContactDeltaTangentV;
+			WorldContactPoint.ContactTargetVelocityNormal = InWorldContactVelocityTargetNormal;
+
+			State.ManifoldPoints[ManifoldPoiontIndex].InitContact(
+				FSolverReal(Dt),
+				State.SolverBodies[0],
+				State.SolverBodies[1]);
+		}
+
+		FORCEINLINE_DEBUGGABLE bool FPBDCollisionSolver::SolvePositionWithFriction(const FSolverReal Dt, const FSolverReal MaxPushOut)
+		{
+			// SolverBody decorator used to add mass scaling
+			FConstraintSolverBody& Body0 = SolverBody0();
+			FConstraintSolverBody& Body1 = SolverBody1();
+
+			// Accumulate net pushout for friction limits below
+			bool bApplyFriction[MaxPointsPerConstraint] = { false, };
+			int32 NumFrictionContacts = 0;
+			FSolverReal TotalPushOutNormal = FSolverReal(0);
+
+			// Apply the position correction along the normal and determine if we want to run friction on each point
+			for (int32 PointIndex = 0; PointIndex < NumManifoldPoints(); ++PointIndex)
+			{
+				FPBDCollisionSolverManifoldPoint& SolverManifoldPoint = State.ManifoldPoints[PointIndex];
+
+				FSolverReal ContactDeltaNormal;
+				SolverManifoldPoint.CalculateContactPositionErrorNormal(Body0.SolverBody(), Body1.SolverBody(), MaxPushOut, ContactDeltaNormal);
+
+				// Apply a normal correction if we still have penetration or if we are now separated but have previously applied a correction that we may want to undo
+				const bool bProcessManifoldPoint = (ContactDeltaNormal < FSolverReal(0)) || (SolverManifoldPoint.NetPushOutNormal > FSolverReal(UE_SMALL_NUMBER));
+				if (bProcessManifoldPoint)
+				{
+					ApplyPositionCorrectionNormal(
+						State.Stiffness,
+						ContactDeltaNormal,
+						SolverManifoldPoint,
+						Body0,
+						Body1);
+
+					TotalPushOutNormal += SolverManifoldPoint.NetPushOutNormal;
+				}
+
+				// Friction gets updated for any point with a net normal correction or where we have previously had a normal correction and 
+				// already applied friction (in which case we may need to zero it)
+				if ((SolverManifoldPoint.NetPushOutNormal != 0) || (SolverManifoldPoint.NetPushOutTangentU != 0) || (SolverManifoldPoint.NetPushOutTangentV != 0))
+				{
+					bApplyFriction[PointIndex] = true;
+					++NumFrictionContacts;
+				}
+			}
+
+			// Apply the tangential position correction if required
+			if (NumFrictionContacts > 0)
+			{
+				// We clip the tangential correction at each contact to the friction cone, but we use to average impulse
+				// among all contacts as the clipping limit. This is not really correct but it is much more stable to 
+				// differences in contacts from tick to tick
+				// @todo(chaos): try a decaying maximum per contact point rather than an average (again - we had that once!)
+				const FSolverReal FrictionMaxPushOut = TotalPushOutNormal / FSolverReal(NumFrictionContacts);
+				const FSolverReal FrictionStiffness = State.Stiffness * CVars::Chaos_PBDCollisionSolver_Position_StaticFrictionStiffness;
+
+				for (int32 PointIndex = 0; PointIndex < NumManifoldPoints(); ++PointIndex)
+				{
+					if (bApplyFriction[PointIndex])
+					{
+						FPBDCollisionSolverManifoldPoint& SolverManifoldPoint = State.ManifoldPoints[PointIndex];
+
+						FSolverReal ContactDeltaTangentU, ContactDeltaTangentV;
+						SolverManifoldPoint.CalculateContactPositionErrorTangential(Body0.SolverBody(), Body1.SolverBody(), ContactDeltaTangentU, ContactDeltaTangentV);
+
+						ApplyPositionCorrectionTangential(
+							FrictionStiffness,
+							State.StaticFriction,
+							State.DynamicFriction,
+							FrictionMaxPushOut,
+							ContactDeltaTangentU,
+							ContactDeltaTangentV,
+							SolverManifoldPoint,
+							Body0,
+							Body1);
+					}
+				}
+			}
+
+			return false;
+		}
+
+		FORCEINLINE_DEBUGGABLE bool FPBDCollisionSolver::SolvePositionNoFriction(const FSolverReal Dt, const FSolverReal MaxPushOut)
+		{
+			// SolverBody decorator used to add mass scaling
+			FConstraintSolverBody& Body0 = SolverBody0();
+			FConstraintSolverBody& Body1 = SolverBody1();
+
+			// Apply the position correction so that all contacts have zero separation
+			for (int32 PointIndex = 0; PointIndex < NumManifoldPoints(); ++PointIndex)
+			{
+				FPBDCollisionSolverManifoldPoint& SolverManifoldPoint = State.ManifoldPoints[PointIndex];
+
+				FSolverReal ContactDeltaNormal;
+				SolverManifoldPoint.CalculateContactPositionErrorNormal(Body0.SolverBody(), Body1.SolverBody(), MaxPushOut, ContactDeltaNormal);
+
+				const bool bProcessManifoldPoint = (ContactDeltaNormal < FSolverReal(0)) || (SolverManifoldPoint.NetPushOutNormal > FSolverReal(UE_SMALL_NUMBER));
+				if (bProcessManifoldPoint)
+				{
+					ApplyPositionCorrectionNormal(
+						State.Stiffness,
+						ContactDeltaNormal,
+						SolverManifoldPoint,
+						Body0,
+						Body1);
+				}
+			}
+
+			return false;
+		}
+
+		FORCEINLINE_DEBUGGABLE bool FPBDCollisionSolver::SolvePositionNoFrictionJacobi(const FSolverReal Dt, const FSolverReal MaxPushOut)
+		{
+			// SolverBody decorator used to add mass scaling
+			FConstraintSolverBody& Body0 = SolverBody0();
+			FConstraintSolverBody& Body1 = SolverBody1();
+
+			// Apply the position correction so that all contacts have zero separation
+			FSolverVec3 DX0 = FSolverVec3(0);
+			FSolverVec3 DX1 = FSolverVec3(0);
+			FSolverVec3 DR0 = FSolverVec3(0);
+			FSolverVec3 DR1 = FSolverVec3(0);
+
+			const FSolverReal Stiffness = CVars::Chaos_PBDCollisionSolver_JacobiStiffness * State.Stiffness;
+
+			for (int32 PointIndex = 0; PointIndex < NumManifoldPoints(); ++PointIndex)
+			{
+				FPBDCollisionSolverManifoldPoint& SolverManifoldPoint = State.ManifoldPoints[PointIndex];
+
+				FSolverReal ContactDeltaNormal;
+				SolverManifoldPoint.CalculateContactPositionErrorNormal(Body0.SolverBody(), Body1.SolverBody(), MaxPushOut, ContactDeltaNormal);
+
+				const bool bProcessManifoldPoint = (ContactDeltaNormal < FSolverReal(0)) || (SolverManifoldPoint.NetPushOutNormal > FSolverReal(UE_SMALL_NUMBER));
+				if (bProcessManifoldPoint)
+				{
+					FSolverReal PushOutNormal = FSolverReal(0);
+					CalculatePositionCorrectionNormal(
+						Stiffness,
+						ContactDeltaNormal,
+						SolverManifoldPoint.ContactMassNormal,
+						SolverManifoldPoint.NetPushOutNormal,
+						PushOutNormal);						// Out
+
+					SolverManifoldPoint.NetPushOutNormal += PushOutNormal;
+
+					// Update the particle state based on the pushout
+					if (Body0.IsDynamic())
+					{
+						DX0 += (Body0.InvM() * PushOutNormal) * SolverManifoldPoint.WorldContact.ContactNormal;
+						DR0 += SolverManifoldPoint.WorldContactNormalAngular0 * PushOutNormal;
+					}
+					if (Body1.IsDynamic())
+					{
+						DX1 -= (Body1.InvM() * PushOutNormal) * SolverManifoldPoint.WorldContact.ContactNormal;
+						DR1 -= SolverManifoldPoint.WorldContactNormalAngular1 * PushOutNormal;
+					}
+				}
+			}
+
+			if (Body0.IsDynamic())
+			{
+				if (AboveThresholdElseZero(DX0, CVars::Chaos_PBDCollisionSolver_JacobiPositionTolerance))
+				{
+					Body0.ApplyPositionDelta(DX0);
+				}
+				if (AboveThresholdElseZero(DR0, CVars::Chaos_PBDCollisionSolver_JacobiRotationTolerance))
+				{
+					Body0.ApplyRotationDelta(DR0);
+				}
+			}
+			if (Body1.IsDynamic())
+			{
+				if (AboveThresholdElseZero(DX1, CVars::Chaos_PBDCollisionSolver_JacobiPositionTolerance))
+				{
+					Body1.ApplyPositionDelta(DX1);
+				}
+				if (AboveThresholdElseZero(DR1, CVars::Chaos_PBDCollisionSolver_JacobiRotationTolerance))
+				{
+					Body1.ApplyRotationDelta(DR1);
+				}
+			}
+
+			return false;
+		}
+
+		FORCEINLINE_DEBUGGABLE bool FPBDCollisionSolver::SolvePositionWithFrictionJacobi(const FSolverReal Dt, const FSolverReal MaxPushOut)
+		{
+			// SolverBody decorator used to add mass scaling
+			FConstraintSolverBody& Body0 = SolverBody0();
+			FConstraintSolverBody& Body1 = SolverBody1();
+
+			// Accumulate net pushout for friction limits below
+			bool bApplyFriction[MaxPointsPerConstraint] = { false, };
+			int32 NumFrictionContacts = 0;
+			FSolverReal TotalPushOutNormal = FSolverReal(0);
+
+			FSolverVec3 DX0 = FSolverVec3(0);
+			FSolverVec3 DX1 = FSolverVec3(0);
+			FSolverVec3 DR0 = FSolverVec3(0);
+			FSolverVec3 DR1 = FSolverVec3(0);
+
+			const FSolverReal Stiffness = CVars::Chaos_PBDCollisionSolver_JacobiStiffness * State.Stiffness;
+
+			// Apply the position correction along the normal and determine if we want to run friction on each point
+			for (int32 PointIndex = 0; PointIndex < NumManifoldPoints(); ++PointIndex)
+			{
+				FPBDCollisionSolverManifoldPoint& SolverManifoldPoint = State.ManifoldPoints[PointIndex];
+
+				FSolverReal ContactDeltaNormal;
+				SolverManifoldPoint.CalculateContactPositionErrorNormal(Body0.SolverBody(), Body1.SolverBody(), MaxPushOut, ContactDeltaNormal);
+
+				// Apply a normal correction if we still have penetration or if we are now separated but have previously applied a correction that we may want to undo
+				const bool bProcessManifoldPoint = (ContactDeltaNormal < FSolverReal(0)) || (SolverManifoldPoint.NetPushOutNormal > FSolverReal(UE_SMALL_NUMBER));
+				if (bProcessManifoldPoint)
+				{
+					FSolverReal PushOutNormal = FSolverReal(0);
+					CalculatePositionCorrectionNormal(
+						Stiffness,
+						ContactDeltaNormal,
+						SolverManifoldPoint.ContactMassNormal,
+						SolverManifoldPoint.NetPushOutNormal,
+						PushOutNormal);						// Out
+
+					SolverManifoldPoint.NetPushOutNormal += PushOutNormal;
+
+					// Update the particle state based on the pushout
+					if (Body0.IsDynamic())
+					{
+						DX0 += (Body0.InvM() * PushOutNormal) * SolverManifoldPoint.WorldContact.ContactNormal;
+						DR0 += SolverManifoldPoint.WorldContactNormalAngular0 * PushOutNormal;
+					}
+					if (Body1.IsDynamic())
+					{
+						DX1 -= (Body1.InvM() * PushOutNormal) * SolverManifoldPoint.WorldContact.ContactNormal;
+						DR1 -= SolverManifoldPoint.WorldContactNormalAngular1 * PushOutNormal;
+					}
+
+					TotalPushOutNormal += SolverManifoldPoint.NetPushOutNormal;
+				}
+
+				// Friction gets updated for any point with a net normal correction or where we have previously had a normal correction and 
+				// already applied friction (in which case we may need to zero it)
+				if ((SolverManifoldPoint.NetPushOutNormal != 0) || (SolverManifoldPoint.NetPushOutTangentU != 0) || (SolverManifoldPoint.NetPushOutTangentV != 0))
+				{
+					bApplyFriction[PointIndex] = true;
+					++NumFrictionContacts;
+				}
+			}
+
+			// Apply the tangential position correction if required
+			if (NumFrictionContacts > 0)
+			{
+				// We clip the tangential correction at each contact to the friction cone, but we use to average impulse
+				// among all contacts as the clipping limit. This is not really correct but it is much more stable to 
+				// differences in contacts from tick to tick
+				const FSolverReal FrictionMaxPushOut = TotalPushOutNormal / FSolverReal(NumFrictionContacts);
+				const FSolverReal FrictionStiffness = Stiffness * CVars::Chaos_PBDCollisionSolver_Position_StaticFrictionStiffness;
+
+				for (int32 PointIndex = 0; PointIndex < NumManifoldPoints(); ++PointIndex)
+				{
+					if (bApplyFriction[PointIndex])
+					{
+						FPBDCollisionSolverManifoldPoint& SolverManifoldPoint = State.ManifoldPoints[PointIndex];
+
+						FSolverReal ContactDeltaTangentU, ContactDeltaTangentV;
+						SolverManifoldPoint.CalculateContactPositionErrorTangential(Body0.SolverBody(), Body1.SolverBody(), ContactDeltaTangentU, ContactDeltaTangentV);
+
+						FSolverReal PushOutTangentU = FSolverReal(0);
+						FSolverReal PushOutTangentV = FSolverReal(0);
+
+						CalculatePositionCorrectionTangent(
+							FrictionStiffness,
+							ContactDeltaTangentU,
+							SolverManifoldPoint.ContactMassTangentU,
+							SolverManifoldPoint.NetPushOutTangentU,
+							PushOutTangentU);					// Out
+
+						CalculatePositionCorrectionTangent(
+							FrictionStiffness,
+							ContactDeltaTangentV,
+							SolverManifoldPoint.ContactMassTangentV,
+							SolverManifoldPoint.NetPushOutTangentV,
+							PushOutTangentV);					// Out
+
+						SolverManifoldPoint.NetPushOutTangentU += PushOutTangentU;
+						SolverManifoldPoint.NetPushOutTangentV += PushOutTangentV;
+
+						ApplyFrictionCone(
+							State.StaticFriction,
+							State.DynamicFriction,
+							FrictionMaxPushOut,
+							PushOutTangentU,									// InOut
+							PushOutTangentV,									// InOut
+							SolverManifoldPoint.NetPushOutTangentU,				// InOut
+							SolverManifoldPoint.NetPushOutTangentV,				// InOut
+							SolverManifoldPoint.StaticFrictionRatio);			// Out
+
+						// Update the particle state based on the pushout
+						const FSolverVec3 PushOut = PushOutTangentU * SolverManifoldPoint.WorldContact.ContactTangentU + PushOutTangentV * SolverManifoldPoint.WorldContact.ContactTangentV;
+						if (Body0.IsDynamic())
+						{
+							DX0 += Body0.InvM() * PushOut;
+							DR0 += SolverManifoldPoint.WorldContactTangentUAngular0 * PushOutTangentU + SolverManifoldPoint.WorldContactTangentVAngular0 * PushOutTangentV;
+						}
+						if (Body1.IsDynamic())
+						{
+							DX1 -= Body1.InvM() * PushOut;
+							DR1 -= SolverManifoldPoint.WorldContactTangentUAngular1 * PushOutTangentU + SolverManifoldPoint.WorldContactTangentVAngular1 * PushOutTangentV;
+						}
+					}
+				}
+
+				if (Body0.IsDynamic())
+				{
+					if (AboveThresholdElseZero(DX0, CVars::Chaos_PBDCollisionSolver_JacobiPositionTolerance))
+					{
+						Body0.ApplyPositionDelta(DX0);
+					}
+					if (AboveThresholdElseZero(DR0, CVars::Chaos_PBDCollisionSolver_JacobiRotationTolerance))
+					{
+						Body0.ApplyRotationDelta(DR0);
+					}
+				}
+				if (Body1.IsDynamic())
+				{
+					if (AboveThresholdElseZero(DX1, CVars::Chaos_PBDCollisionSolver_JacobiPositionTolerance))
+					{
+						Body1.ApplyPositionDelta(DX1);
+					}
+					if (AboveThresholdElseZero(DR1, CVars::Chaos_PBDCollisionSolver_JacobiRotationTolerance))
+					{
+						Body1.ApplyRotationDelta(DR1);
+					}
+				}
+			}
+
+			return false;
+		}
+
+		FORCEINLINE_DEBUGGABLE bool FPBDCollisionSolver::SolveVelocity(const FSolverReal Dt, const bool bApplyDynamicFriction)
+		{
+			// Apply restitution at the average contact point
+			// This means we don't need to run as many iterations to get stable bouncing
+			// It also helps with zero restitution to counter any velocioty added by the PBD solve
+			const bool bSolveAverageContact = (NumManifoldPoints() > 1) && CVars::bChaos_PBDCollisionSolver_Velocity_AveragePointEnabled;
+			if (bSolveAverageContact)
+			{
+				SolveVelocityAverage(Dt);
+			}
+
+			FConstraintSolverBody& Body0 = SolverBody0();
+			FConstraintSolverBody& Body1 = SolverBody1();
+
+			// NOTE: this dynamic friction implementation is iteration-count sensitive
+			// @todo(chaos): fix iteration count dependence of dynamic friction
+			const FSolverReal DynamicFriction = (bApplyDynamicFriction && (Dt > 0) && CVars::bChaos_PBDCollisionSolver_Velocity_FrictionEnabled) ? State.VelocityFriction : FSolverReal(0);
+
+			for (int32 PointIndex = 0; PointIndex < NumManifoldPoints(); ++PointIndex)
+			{
+				FPBDCollisionSolverManifoldPoint& SolverManifoldPoint = State.ManifoldPoints[PointIndex];
+
+				if (SolverManifoldPoint.ShouldSolveVelocity())
+				{
+					const FSolverReal MinImpulseNormal = FMath::Min(FSolverReal(0), -SolverManifoldPoint.NetPushOutNormal / Dt);
+
+					if (DynamicFriction > 0)
+					{
+						FSolverReal ContactVelocityDeltaNormal, ContactVelocityDeltaTangentU, ContactVelocityDeltaTangentV;
+						SolverManifoldPoint.CalculateContactVelocityError(Body0, Body1, DynamicFriction, Dt, ContactVelocityDeltaNormal, ContactVelocityDeltaTangentU, ContactVelocityDeltaTangentV);
+
+						ApplyVelocityCorrection(
+							State.Stiffness,
+							Dt,
+							DynamicFriction,
+							ContactVelocityDeltaNormal,
+							ContactVelocityDeltaTangentU,
+							ContactVelocityDeltaTangentV,
+							MinImpulseNormal,
+							SolverManifoldPoint,
+							Body0,
+							Body1);
+					}
+					else
+					{
+						FSolverReal ContactVelocityDeltaNormal;
+						SolverManifoldPoint.CalculateContactVelocityErrorNormal(Body0, Body1, ContactVelocityDeltaNormal);
+
+						ApplyVelocityCorrectionNormal(
+							State.Stiffness,
+							ContactVelocityDeltaNormal,
+							MinImpulseNormal,
+							SolverManifoldPoint,
+							Body0,
+							Body1);
+					}
+				}
+			}
+
+			// Early-out support for the velocity solve is not currently very important because we
+			// only run one iteration in the velocity solve phase.
+			// @todo(chaos): support early-out in velocity solve if necessary
+			return true;
+		}
+
+
+		FORCEINLINE_DEBUGGABLE bool FPBDCollisionSolver::SolveVelocityJacobi(const FSolverReal Dt, const bool bApplyDynamicFriction)
+		{
+			FConstraintSolverBody& Body0 = SolverBody0();
+			FConstraintSolverBody& Body1 = SolverBody1();
+
+			FSolverVec3 DV0 = FSolverVec3(0);
+			FSolverVec3 DV1 = FSolverVec3(0);
+			FSolverVec3 DW0 = FSolverVec3(0);
+			FSolverVec3 DW1 = FSolverVec3(0);
+
+			const FSolverReal Stiffness = CVars::Chaos_PBDCollisionSolver_JacobiStiffness * State.Stiffness;
+
+			for (int32 PointIndex = 0; PointIndex < NumManifoldPoints(); ++PointIndex)
+			{
+				FPBDCollisionSolverManifoldPoint& SolverManifoldPoint = State.ManifoldPoints[PointIndex];
+
+				if (SolverManifoldPoint.ShouldSolveVelocity())
+				{
+					const FSolverReal MinImpulseNormal = FMath::Min(FSolverReal(0), -SolverManifoldPoint.NetPushOutNormal / Dt);
+
+					FSolverReal ContactVelocityDeltaNormal;
+					SolverManifoldPoint.CalculateContactVelocityErrorNormal(Body0, Body1, ContactVelocityDeltaNormal);
+
+					FSolverReal ImpulseNormal = -(Stiffness * SolverManifoldPoint.ContactMassNormal) * ContactVelocityDeltaNormal;
+
+					// See comments in ApplyVelocityCorrection
+					const FSolverReal NetImpulseNormal = SolverManifoldPoint.NetImpulseNormal + ImpulseNormal;
+					if (NetImpulseNormal < MinImpulseNormal)
+					{
+						ImpulseNormal = ImpulseNormal - NetImpulseNormal + MinImpulseNormal;
+					}
+
+					SolverManifoldPoint.NetImpulseNormal += ImpulseNormal;
+
+					FSolverVec3 Impulse = ImpulseNormal * SolverManifoldPoint.WorldContact.ContactNormal;
+					if (Body0.IsDynamic())
+					{
+						DV0 += Body0.InvM() * Impulse;
+						DW0 += SolverManifoldPoint.WorldContactNormalAngular0 * ImpulseNormal;
+					}
+					if (Body1.IsDynamic())
+					{
+						DV1 -= Body1.InvM() * Impulse;
+						DW1 -= SolverManifoldPoint.WorldContactNormalAngular1 * ImpulseNormal;
+					}
+				}
+			}
+
+			const FSolverReal PositionFloor = (CVars::Chaos_PBDCollisionSolver_JacobiPositionTolerance / FMath::Square(Dt));
+			const FSolverReal RotationFloor = (CVars::Chaos_PBDCollisionSolver_JacobiRotationTolerance / FMath::Square(Dt));
+
+			if (Body0.IsDynamic())
+			{
+				AboveThresholdElseZero(DV0, PositionFloor);
+				AboveThresholdElseZero(DW0, RotationFloor);
+
+				Body0.ApplyVelocityDelta(DV0, DW0);
+			}
+			if (Body1.IsDynamic())
+			{
+				AboveThresholdElseZero(DV1, PositionFloor);
+				AboveThresholdElseZero(DW1, RotationFloor);
+
+				Body1.ApplyVelocityDelta(DV1, DW1);
+			}
+
+			// Early-out support for the velocity solve is not currently very important because we
+			// only run one iteration in the velocity solve phase.
+			// @todo(chaos): support early-out in velocity solve if necessary
+			return true;
+		}
+	}
+}
+
diff --git a/Engine/Source/Runtime/Experimental/Chaos/Public/Chaos/Collision/ParticlePairCollisionDetector.h b/Engine/Source/Runtime/Experimental/Chaos/Public/Chaos/Collision/ParticlePairCollisionDetector.h
--- a/Engine/Source/Runtime/Experimental/Chaos/Public/Chaos/Collision/ParticlePairCollisionDetector.h
+++ b/Engine/Source/Runtime/Experimental/Chaos/Public/Chaos/Collision/ParticlePairCollisionDetector.h
@@ -3,6 +3,7 @@
 
 #include "Chaos/Collision/CollisionDetector.h"
 #include "Chaos/Collision/ParticlePairBroadPhase.h"
+#include "Chaos/PBDCollisionConstraints.h"
 
 #include "Chaos/ChaosPerfTest.h"
 #include "ChaosStats.h"
@@ -34,10 +35,8 @@ namespace Chaos
 
 			CollisionContainer.BeginDetectCollisions();
 
-			FCollisionConstraintAllocator& Allocator = GetCollisionContainer().GetConstraintAllocator();
-
 			// Collision detection pipeline: BroadPhase -> MidPhase ->NarrowPhase -> Container
-			BroadPhase.ProduceOverlaps(Dt, &Allocator, GetSettings());
+			BroadPhase.ProduceOverlaps(Dt, &GetCollisionContainer().GetConstraintAllocator(), GetCollisionContainer().GetDetectorSettings());
 
 			CollisionContainer.EndDetectCollisions();
 		}
diff --git a/Engine/Source/Runtime/Experimental/Chaos/Public/Chaos/Collision/SolverCollisionContainer.h b/Engine/Source/Runtime/Experimental/Chaos/Public/Chaos/Collision/SolverCollisionContainer.h
--- a/Engine/Source/Runtime/Experimental/Chaos/Public/Chaos/Collision/SolverCollisionContainer.h
+++ b/Engine/Source/Runtime/Experimental/Chaos/Public/Chaos/Collision/SolverCollisionContainer.h
@@ -4,6 +4,7 @@
 #include "CoreMinimal.h"
 #include "Chaos/CollisionResolutionTypes.h"
 #include "Chaos/Collision/CollisionApplyType.h"
+#include "Chaos/Collision/PBDCollisionSolver.h"
 #include "Chaos/Evolution/SolverConstraintContainer.h"
 
 namespace Chaos
@@ -11,9 +12,6 @@ namespace Chaos
 	class FPBDIslandConstraint;
 	class FPBDCollisionConstraint;
 	class FPBDCollisionConstraints;
-	class FPBDCollisionSolver;
-	class FPBDCollisionSolverAdapter;
-	class FPBDCollisionSolverManifoldPoint;
 	class FPBDIsland;
 	class FSolverBody;
 	class FSolverBodyContainer;
@@ -24,6 +22,8 @@ namespace Chaos
 	class FPBDCollisionSolverSettings
 	{
 	public:
+		UE_NONCOPYABLE(FPBDCollisionSolverSettings);
+
 		FPBDCollisionSolverSettings();
 
 		// Maximum speed at which two objects can depenetrate (actually, how much relative velocity can be added
@@ -52,20 +52,16 @@ namespace Chaos
 	class FPBDCollisionContainerSolver : public FConstraintContainerSolver
 	{
 	public:
+		UE_NONCOPYABLE(FPBDCollisionContainerSolver);
+
 		FPBDCollisionContainerSolver(const FPBDCollisionConstraints& InConstraintContainer, const int32 InPriority);
 		~FPBDCollisionContainerSolver();
 
-		int32 NumSolvers() const { return CollisionSolvers.Num(); }
+		int32 NumSolvers() const { return CollisionConstraints.Num(); }
 
 		virtual void Reset(const int32 InMaxCollisions) override final;
 
-		virtual int32 GetNumConstraints() const override final
-		{
-			return NumSolvers();
-		}
-
-		// Set whether we are using deferred collision detection
-		void SetIsDeferredCollisionDetection(const bool bDeferred) { bDeferredCollisionDetection = bDeferred; }
+		virtual int32 GetNumConstraints() const override final { return NumSolvers(); }
 
 		//
 		// IslandGroup API
@@ -82,6 +78,7 @@ namespace Chaos
 		virtual void ApplyProjectionConstraints(const FReal Dt, const int32 It, const int32 NumIts) override final;
 
 	private:
+		void CachePrefetchSolver(const int32 ConstraintIndex) const;
 		void AddConstraint(FPBDCollisionConstraint& Constraint);
 		void UpdatePositionShockPropagation(const FReal Dt, const int32 It, const int32 NumIts, const int32 BeginIndex, const int32 EndIndex, const FPBDCollisionSolverSettings& SolverSettings);
 		void UpdateVelocityShockPropagation(const FReal Dt, const int32 It, const int32 NumIts, const int32 BeginIndex, const int32 EndIndex, const FPBDCollisionSolverSettings& SolverSettings);
@@ -92,8 +89,10 @@ namespace Chaos
 		void UpdateCollisions(const FReal InDt, const int32 BeginIndex, const int32 EndIndex);
 
 		const FPBDCollisionConstraints& ConstraintContainer;
-		TArray<FPBDCollisionSolverAdapter> CollisionSolvers;
+		TArray<Private::FPBDCollisionSolver> CollisionSolvers;
+		TArray<Private::FPBDCollisionSolverManifoldPoint> CollisionSolverManifoldPoints;
+		TArray<FPBDCollisionConstraint*> CollisionConstraints;
+		TArray<bool> bCollisionConstraintPerIterationCollisionDetection;
 		bool bPerIterationCollisionDetection;
-		bool bDeferredCollisionDetection;
 	};
 }
diff --git a/Engine/Source/Runtime/Experimental/Chaos/Public/Chaos/Collision/SpatialAccelerationCollisionDetector.h b/Engine/Source/Runtime/Experimental/Chaos/Public/Chaos/Collision/SpatialAccelerationCollisionDetector.h
--- a/Engine/Source/Runtime/Experimental/Chaos/Public/Chaos/Collision/SpatialAccelerationCollisionDetector.h
+++ b/Engine/Source/Runtime/Experimental/Chaos/Public/Chaos/Collision/SpatialAccelerationCollisionDetector.h
@@ -35,26 +35,26 @@ namespace Chaos
 				return;
 			}
 
-			CollisionContainer.BeginDetectCollisions();
+			GetCollisionContainer().BeginDetectCollisions();
 
 			// Collision detection pipeline: BroadPhase -[parallel]-> NarrowPhase -[parallel]-> CollisionAllocator -[serial]-> Container
-			BroadPhase.ProduceOverlaps(Dt, &CollisionContainer.GetConstraintAllocator(), Settings, ResimCache);
+			BroadPhase.ProduceOverlaps(Dt, &GetCollisionContainer().GetConstraintAllocator(), GetCollisionContainer().GetDetectorSettings(), ResimCache);
 
-			CollisionContainer.EndDetectCollisions();
+			GetCollisionContainer().EndDetectCollisions();
 
 			// If we have a resim cache restore and save contacts
 			if(ResimCache)
 			{
 				// Ensure we have at least one allocator
-				CollisionContainer.GetConstraintAllocator().SetMaxContexts(1);
+				GetCollisionContainer().GetConstraintAllocator().SetMaxContexts(1);
 
 				FCollisionContext Context;
-				Context.SetSettings(Settings);
-				Context.SetAllocator(CollisionContainer.GetConstraintAllocator().GetContextAllocator(0));
+				Context.SetSettings(GetCollisionContainer().GetDetectorSettings());
+				Context.SetAllocator(GetCollisionContainer().GetConstraintAllocator().GetContextAllocator(0));
 
-				CollisionContainer.GetConstraintAllocator().AddResimConstraints(ResimCache->GetAndSanitizeConstraints(), Context);
+				GetCollisionContainer().GetConstraintAllocator().AddResimConstraints(ResimCache->GetAndSanitizeConstraints(), Context);
 
-				ResimCache->SaveConstraints(CollisionContainer.GetConstraints());
+				ResimCache->SaveConstraints(GetCollisionContainer().GetConstraints());
 			}
 		}
 
diff --git a/Engine/Source/Runtime/Experimental/Chaos/Public/Chaos/Evolution/ConstraintGroupSolver.h b/Engine/Source/Runtime/Experimental/Chaos/Public/Chaos/Evolution/ConstraintGroupSolver.h
--- a/Engine/Source/Runtime/Experimental/Chaos/Public/Chaos/Evolution/ConstraintGroupSolver.h
+++ b/Engine/Source/Runtime/Experimental/Chaos/Public/Chaos/Evolution/ConstraintGroupSolver.h
@@ -30,7 +30,7 @@ namespace Chaos
 		*/
 		inline int32 GetNumSolverBodies() const
 		{
-			return SolverBodyContainer.NumItems();
+			return SolverBodyContainer.Num();
 		}
 
 		/**
diff --git a/Engine/Source/Runtime/Experimental/Chaos/Public/Chaos/Evolution/PBDMinEvolution.h b/Engine/Source/Runtime/Experimental/Chaos/Public/Chaos/Evolution/PBDMinEvolution.h
--- a/Engine/Source/Runtime/Experimental/Chaos/Public/Chaos/Evolution/PBDMinEvolution.h
+++ b/Engine/Source/Runtime/Experimental/Chaos/Public/Chaos/Evolution/PBDMinEvolution.h
@@ -33,7 +33,7 @@ namespace Chaos
 		using FEvolutionCallback = TFunction<void()>;
 		using FRigidParticleSOAs = FPBDRigidsSOAs;
 
-		FPBDMinEvolution(FRigidParticleSOAs& InParticles, TArrayCollectionArray<FVec3>& InPrevX, TArrayCollectionArray<FRotation3>& InPrevR, FCollisionDetector& InCollisionDetector, const FReal InBoundsExtension);
+		FPBDMinEvolution(FRigidParticleSOAs& InParticles, TArrayCollectionArray<FVec3>& InPrevX, TArrayCollectionArray<FRotation3>& InPrevR, FCollisionDetector& InCollisionDetector);
 		~FPBDMinEvolution();
 
 		void AddConstraintContainer(FPBDConstraintContainer& InContainer, const int32 Priority = 0);
@@ -62,10 +62,6 @@ namespace Chaos
 			Gravity = G;
 		}
 
-		void SetBoundsExtension(const FReal InBoundsExtension)
-		{
-			BoundsExtension = InBoundsExtension;
-		}
 
 		void SetSimulationSpace(const FSimulationSpace& InSimulationSpace)
 		{
@@ -87,6 +83,18 @@ namespace Chaos
 			SimulationSpaceSettings = InSimulationSpaceSettings;
 		}
 
+
+		UE_DEPRECATED(5.3, "InBoundsExtension parameter has been removed")
+		FPBDMinEvolution(FRigidParticleSOAs& InParticles, TArrayCollectionArray<FVec3>& InPrevX, TArrayCollectionArray<FRotation3>& InPrevR, FCollisionDetector& InCollisionDetector, const FReal InBoundsExtension)
+			: FPBDMinEvolution(InParticles, InPrevX, InPrevR, InCollisionDetector)
+		{
+		}
+
+		UE_DEPRECATED(5.3, "BoundsExtension has been removed")
+		void SetBoudsExtension(const FReal Unused)
+		{
+		}
+
 	private:
 		void PrepareTick();
 		void UnprepareTick();
@@ -112,7 +120,6 @@ namespace Chaos
 		int32 NumPositionIterations;
 		int32 NumVelocityIterations;
 		int32 NumProjectionIterations;
-		FReal BoundsExtension;
 		FVec3 Gravity;
 		FSimulationSpaceSettings SimulationSpaceSettings;
 		FSimulationSpace SimulationSpace;
diff --git a/Engine/Source/Runtime/Experimental/Chaos/Public/Chaos/Evolution/SolverBody.h b/Engine/Source/Runtime/Experimental/Chaos/Public/Chaos/Evolution/SolverBody.h
--- a/Engine/Source/Runtime/Experimental/Chaos/Public/Chaos/Evolution/SolverBody.h
+++ b/Engine/Source/Runtime/Experimental/Chaos/Public/Chaos/Evolution/SolverBody.h
@@ -14,7 +14,6 @@
 // Set to 1 to force single precision in the constraint solver, even if the default numeric type is double
 #define CHAOS_CONSTRAINTSOLVER_LOWPRECISION 1
 
-
 // Add some NaN checking when CHAOS_CONSTRAINTSOLVER_NAN_DIAGNOSTIC is defined
 #if CHAOS_CONSTRAINTSOLVER_NAN_DIAGNOSTIC
 inline void ChaosSolverCheckNaN(const Chaos::FRealSingle& V) { ensure(!FMath::IsNaN(V)); }
@@ -43,6 +42,7 @@ namespace Chaos
 	using SolverVectorRegister = VectorRegister4;
 #endif
 	using FSolverVec3 = TVec3<FSolverReal>;
+	using FSolverRotation3 = TRotation3<FSolverReal>;
 	using FSolverMatrix33 = TMatrix33<FSolverReal>;
 
 	class FSolverBody;
@@ -100,10 +100,37 @@ namespace Chaos
 	public:
 
 		/**
-		 * @brief Create an empty solver body
-		 * @note This is only used by unit tests
+		 * A factory method to create a safely initialized Solverbody. 
+		 * The defaultconstructor assumes you are going to set all properties manually.
+		 */
+		static FSolverBody MakeInitialized()
+		{
+			FSolverBody SolverBody;
+			SolverBody.State.Init();
+			return SolverBody;
+		}
+
+		/**
+		 * Create an empty solver body. All properties are uninitialized.
+		*/
+		static FSolverBody MakeUninitialized()
+		{
+			return FSolverBody();
+		}
+
+		/**
+		 * Create an empty solver body. All properties are uninitialized.
 		*/
-		FSolverBody();
+		FSolverBody() {}
+
+		/**
+		 * Reset the solver accumulators
+		 */
+		void Reset()
+		{
+			State.DP = FSolverVec3(0);
+			State.DQ = FSolverVec3(0);
+		}
 
 		/**
 		 * @brief Calculate and set the velocity and angular velocity from the net transform delta
@@ -197,21 +224,21 @@ namespace Chaos
 		/**
 		 * @brief World-space center of mass velocity
 		*/
-		inline const FVec3& V() const { return State.V; }
+		inline const FSolverVec3& V() const { return State.V; }
 		inline void SetV(const FVec3& InV)
 		{ 
 			ChaosSolverCheckNaN(InV);
-			State.V = InV;
+			State.V = FSolverVec3(InV);
 		}
 
 		/**
 		 * @brief World-space center of mass angular velocity
 		*/
-		inline const FVec3& W() const { return State.W; }
+		inline const FSolverVec3& W() const { return State.W; }
 		inline void SetW(const FVec3& InW)
 		{
 			ChaosSolverCheckNaN(InW);
-			State.W = InW;
+			State.W = FSolverVec3(InW);
 		}
 
 		inline const FVec3& CoM() const { return State.CoM; }
@@ -258,7 +285,7 @@ namespace Chaos
 		/**
 		 * @brief Get the world-space Actor position 
 		*/
-		inline FVec3 ActorP() const { return P() - ActorQ().RotateVector(CoM()); }
+		inline FVec3 ActorP() const { return P() - FVec3(ActorQ().RotateVector(CoM())); }
 
 		/**
 		 * @brief Get the world-space Actor rotation 
@@ -331,7 +358,7 @@ namespace Chaos
 		inline void ApplyLinearVelocityDelta(const FSolverVec3& DV)
 		{
 			ChaosSolverCheckNaN(DV);
-			State.V += FVec3(DV);
+			State.V += DV;
 		}
 
 		/**
@@ -364,22 +391,26 @@ namespace Chaos
 		struct FState
 		{
 			FState()
-				: InvILocal(0)
-				, InvM(0)
-				, InvI(0)
-				, DP(0)
-				, DQ(0)
-				, Level(0)
-				, X(0)
-				, R(FRotation3::Identity)
-				, P(0)
-				, Q(FRotation3::Identity)
-				, V(0)
-				, W(0)
-				, CoM(0)
-				, RoM(FRotation3::Identity)
 			{}
 
+			void Init()
+			{
+				InvILocal = FSolverVec3(0);
+				InvM = 0;
+				InvI = FSolverMatrix33(0);
+				DP = FSolverVec3(0);
+				DQ = FSolverVec3(0);
+				Level = 0;
+				X = FVec3(0);
+				R = FRotation3::FromIdentity();
+				P = FVec3(0);
+				Q = FRotation3::FromIdentity();
+				V = FSolverVec3(0);
+				W = FSolverVec3(0);
+				CoM = FVec3(0);
+				RoM = FRotation3::FromIdentity();
+			}
+
 			// Local-space inverse inertia (diagonal, so only 3 elements)
 			FSolverVec3 InvILocal;
 
@@ -412,10 +443,10 @@ namespace Chaos
 			FRotation3 Q;
 
 			// World-space center of mass velocity
-			FVec3 V;
+			FSolverVec3 V;
 
 			// World-space center of mass angular velocity
-			FVec3 W;
+			FSolverVec3 W;
 
 			// Actor-space center of mass location
 			FVec3 CoM;
@@ -452,12 +483,6 @@ namespace Chaos
 		{
 		}
 
-		FConstraintSolverBody(FSolverBody& InBody, FReal InInvMassScale)
-			: Body(&InBody)
-		{
-			SetInvMScale(InInvMassScale);
-		}
-
 		/**
 		 * @brief True if we have been set up to decorate a SolverBody
 		*/
@@ -468,6 +493,20 @@ namespace Chaos
 		*/
 		inline void Reset() { Body = nullptr; }
 
+		/**
+		 * @brief Initialize all properties to safe values
+		*/
+		void Init() { State.Init(); }
+
+		/**
+		 * @brief Set the inner solver body (hold by reference)
+		 * @note Does not initialize any other properties
+		*/
+		void SetSolverBody(FSolverBody& InSolverBody)
+		{
+			Body = &InSolverBody;
+		}
+
 		/**
 		 * @brief The decorated SolverBody
 		*/
@@ -527,8 +566,8 @@ namespace Chaos
 		inline const FRotation3 ActorQ() const { return Body->ActorQ(); }
 		inline const FVec3 CorrectedActorP() const { return Body->CorrectedActorP(); }
 		inline const FRotation3 CorrectedActorQ() const { return Body->CorrectedActorQ(); }
-		inline const FVec3& V() const { return Body->V(); }
-		inline const FVec3& W() const { return Body->W(); }
+		inline const FSolverVec3& V() const { return Body->V(); }
+		inline const FSolverVec3& W() const { return Body->W(); }
 		inline int32 Level() const { return Body->Level(); }
 		inline const FSolverVec3& DP() const { return Body->DP(); }
 		inline const FSolverVec3& DQ() const { return Body->DQ(); }
@@ -549,10 +588,16 @@ namespace Chaos
 		struct FState
 		{
 			FState() 
-				: InvMScale(FSolverReal(1))
-				, InvIScale(FSolverReal(1))
-				, ShockPropagationScale(FSolverReal(1))
-			{}
+			{
+			}
+
+			void Init()
+			{
+				InvMScale = FSolverReal(1);
+				InvIScale = FSolverReal(1);
+				ShockPropagationScale = FSolverReal(1);
+			}
+
 			FSolverReal InvMScale;
 			FSolverReal InvIScale;
 			FSolverReal ShockPropagationScale;
diff --git a/Engine/Source/Runtime/Experimental/Chaos/Public/Chaos/Evolution/SolverBodyContainer.h b/Engine/Source/Runtime/Experimental/Chaos/Public/Chaos/Evolution/SolverBodyContainer.h
--- a/Engine/Source/Runtime/Experimental/Chaos/Public/Chaos/Evolution/SolverBodyContainer.h
+++ b/Engine/Source/Runtime/Experimental/Chaos/Public/Chaos/Evolution/SolverBodyContainer.h
@@ -15,73 +15,6 @@ namespace Chaos
 	class FSolverBody;
 	class FSolverBodyContainer;
 
-	/**
-	 * @brief An FSolverBody wrapper that binds the SolverBody to a particle and provides gather/scatter methods from/to the particle
-	 * Solver Bodies are used by the core constraint solvers.
-	 * @see FSolverBody
-	*/
-	class FSolverBodyAdapter
-	{
-	public:
-		FSolverBodyAdapter()
-			: Particle(nullptr)
-		{
-		}
-		
-		FSolverBodyAdapter(const FGenericParticleHandle& InParticle)
-			: Particle(InParticle)
-		{
-		}
-
-		FSolverBody& GetSolverBody() { return SolverBody; }
-		const FGenericParticleHandle& GetParticle() const { return Particle; }
-
-		inline void GatherInput(const FReal Dt)
-		{
-			if (Particle.IsValid())
-			{
-				FRigidTransform3 CoMTransform = FParticleUtilitiesPQ::GetCoMWorldTransform(Particle);
-				SolverBody.SetP(CoMTransform.GetLocation());
-				SolverBody.SetQ(CoMTransform.GetRotation());
-				SolverBody.SetV(Particle->V());
-				SolverBody.SetW(Particle->W());
-				SolverBody.SetCoM(Particle->CenterOfMass());
-				SolverBody.SetRoM(Particle->RotationOfMass());
-
-				if (Particle->IsDynamic())
-				{
-					FRigidTransform3 PrevCoMTransform = FParticleUtilitiesXR::GetCoMWorldTransform(Particle);
-					SolverBody.SetX(PrevCoMTransform.GetLocation());
-					SolverBody.SetR(PrevCoMTransform.GetRotation());
-
-					SolverBody.SetInvM(Particle->InvM());
-					SolverBody.SetInvILocal(Particle->ConditionedInvI());
-				}
-				else
-				{
-					// @todo(chaos): we really need Kinematics to store their initial positions...this keeps coming up
-					SolverBody.SetX(SolverBody.P() - SolverBody.V() * Dt);
-					if (SolverBody.W().IsNearlyZero())
-					{
-						SolverBody.SetR(SolverBody.Q());
-					}
-					else
-					{
-						SolverBody.SetR(FRotation3::IntegrateRotationWithAngularVelocity(SolverBody.Q(), -SolverBody.W(), Dt));
-					}
-				}
-				// No need to update the UpdateRotationDependentState since this function is only
-				// valid for dynamic particle for which the SetInvILocal is already doing the job
-			}
-		}
-
-		void ScatterOutput();
-
-	private:
-		FSolverBody SolverBody;
-		FGenericParticleHandle Particle;
-	};
-
 	/**
 	 * The SolverBodies for a set of particles.
 	 * 
@@ -104,86 +37,56 @@ namespace Chaos
 	private:
 		// Settings for the Solver Body array.
 		static const int32 NumBodiesPerChunk = 256;
-		static const int32 BodyArrayChunkSize = NumBodiesPerChunk * sizeof(FSolverBodyAdapter);
+		static const int32 BodyArrayChunkSize = NumBodiesPerChunk * sizeof(FSolverBody);
 
 	public:
-		using FSolverBodyArray = TChunkedArray<FSolverBodyAdapter, BodyArrayChunkSize>;
+		using FSolverBodyArray = TChunkedArray<FSolverBody, BodyArrayChunkSize>;
 
 		FSolverBodyContainer()
 			: bLocked(true)
-		{}
-
+		{
+		}
 
 		// Clear the bodies array and allocate enough space for at least MaxBodies bodies.
 		// @param MaxBodies The number of bodies that will get added to the container. The container asserts if we attempt to add more than this.
 		inline void Reset(int MaxBodies)
 		{
 			SolverBodies.Empty(MaxBodies);
+			Particles.Empty(MaxBodies);
 			ParticleToIndexMap.Reset();
 			bLocked = false;
 		}
 
 		// The number of bodies in the container
-		inline int NumItems() const
+		inline int Num() const
 		{
 			return SolverBodies.Num();
 		}
 
 		// The maximum number of bodies the container can hold (until Reset() is called again)
-		inline int MaxItems() const
+		inline int Max() const
 		{
 			return TNumericLimits<int32>::Max();
 		}
 
-		// Get a pointer to the item at the specified index, or null for INDEX_NONE
-		// @param Index The index of the body to return
-		inline FSolverBodyAdapter* TryGetItem(int Index)
+		// Get the particle associated with the solver body at Index
+		inline FGeometryParticleHandle* GetParticle(const int32 Index) const
 		{
-			if (IsValid(Index))
-			{
-				return &SolverBodies[Index];
-			}
-			return nullptr;
+			return Particles[Index];
 		}
 
-		// Get a pointer to the item at the specified index, or null for INDEX_NONE
-		// @param Index The index of the body to return
-		inline const FSolverBodyAdapter* TryGetItem(int Index) const
-		{
-			if (IsValid(Index))
-			{
-				return &SolverBodies[Index];
-			}
-			return nullptr;
-		}
-
-		// Get a reference to the item at the specified index. Asserts on invalid index.
-		// @param Index The index of the body to return
-		inline FSolverBodyAdapter& GetItem(int Index)
+		// Get the solver body at Index
+		inline const FSolverBody& GetSolverBody(const int32 Index) const
 		{
 			return SolverBodies[Index];
 		}
 
-		// Get a reference to the item at the specified index. Asserts on invalid index.
-		// @param Index The index of the body to return
-		inline const FSolverBodyAdapter& GetItem(int Index) const
+		// Get the solver body at Index
+		inline FSolverBody& GetSolverBody(const int32 Index)
 		{
 			return SolverBodies[Index];
 		}
 
-		/** Get all the solver bodies inside the container */
-		inline FSolverBodyArray& GetBodies()
-		{
-			return SolverBodies;
-		}
-
-		// Whether the specified index is valid.
-		// @param Index The index of the body to return
-		inline bool IsValid(int Index) const
-		{
-			return SolverBodies.IsValidIndex(Index);
-		}
-
 		// Add a solver body to represent the solver state of the particle
 		// This should ideally be called in the order in which the bodies will be accessed
 		// (or as close as we can get, given most constraints access 2 bodies so there is no perfect order)
@@ -194,9 +97,6 @@ namespace Chaos
 		// Collect all the data we need from the particles represented by our SolverBodies
 		void GatherInput(const FReal Dt, const int32 BeginIndex, const int32 EndIndex);
 
-		// Scatter the solver results back to the particles represented by our SolverBodies
-		void ScatterOutput();
-
 		// Scatter the solver results back to the particles represented by our SolverBodies. Range version used for parallization.
 		void ScatterOutput(const int32 BeginIndex, const int32 EndIndex);
 
@@ -219,8 +119,12 @@ namespace Chaos
 	private:
 		int32 AddParticle(FGenericParticleHandle InParticle);
 
+		// All the solver bodies
 		FSolverBodyArray SolverBodies;
 
+		// The particles that map to each solver body
+		TArray<FGeometryParticleHandle*> Particles;
+
 		// Used to determine if a (non-dynamic) body is already present in the container. 
 		// Dynamic bodies are not in this map - the particle stores its index in this case because
 		// dynamic bodies can only be in one island whereas non-dynamics may be const-referenced in many.
diff --git a/Engine/Source/Runtime/Experimental/Chaos/Public/Chaos/Joint/PBDJointCachedSolverGaussSeidel.h b/Engine/Source/Runtime/Experimental/Chaos/Public/Chaos/Joint/PBDJointCachedSolverGaussSeidel.h
--- a/Engine/Source/Runtime/Experimental/Chaos/Public/Chaos/Joint/PBDJointCachedSolverGaussSeidel.h
+++ b/Engine/Source/Runtime/Experimental/Chaos/Public/Chaos/Joint/PBDJointCachedSolverGaussSeidel.h
@@ -121,12 +121,12 @@ struct FAxisConstraintDatas
 			return SolverBodies[1];
 		}
 
-		inline const FVec3& X(int BodyIndex) const
+		inline const FVec3 X(int BodyIndex) const
 		{
 			return Body(BodyIndex).X();
 		}
 
-		inline const FRotation3& R(int BodyIndex) const
+		inline const FRotation3 R(int BodyIndex) const
 		{
 			return Body(BodyIndex).R();
 		}
@@ -145,12 +145,12 @@ struct FAxisConstraintDatas
 			return Body(BodyIndex).CorrectedQ();
 		}
 
-		inline const FVec3& V(int BodyIndex) const
+		inline const FVec3 V(int BodyIndex) const
 		{
 			return Body(BodyIndex).V();
 		}
 
-		inline const FVec3& W(int BodyIndex) const
+		inline const FVec3 W(int BodyIndex) const
 		{
 			return Body(BodyIndex).W();
 		}
diff --git a/Engine/Source/Runtime/Experimental/Chaos/Public/Chaos/Joint/PBDJointSolverGaussSeidel.h b/Engine/Source/Runtime/Experimental/Chaos/Public/Chaos/Joint/PBDJointSolverGaussSeidel.h
--- a/Engine/Source/Runtime/Experimental/Chaos/Public/Chaos/Joint/PBDJointSolverGaussSeidel.h
+++ b/Engine/Source/Runtime/Experimental/Chaos/Public/Chaos/Joint/PBDJointSolverGaussSeidel.h
@@ -92,12 +92,12 @@ namespace Chaos
 			return Body(BodyIndex).CorrectedQ();
 		}
 
-		inline const FVec3& V(int BodyIndex) const
+		inline FVec3 V(int BodyIndex) const
 		{
 			return Body(BodyIndex).V();
 		}
 
-		inline const FVec3& W(int BodyIndex) const
+		inline FVec3 W(int BodyIndex) const
 		{
 			return Body(BodyIndex).W();
 		}
diff --git a/Engine/Source/Runtime/Experimental/Chaos/Public/Chaos/PBDCollisionConstraints.h b/Engine/Source/Runtime/Experimental/Chaos/Public/Chaos/PBDCollisionConstraints.h
--- a/Engine/Source/Runtime/Experimental/Chaos/Public/Chaos/PBDCollisionConstraints.h
+++ b/Engine/Source/Runtime/Experimental/Chaos/Public/Chaos/PBDCollisionConstraints.h
@@ -253,6 +253,18 @@ public:
 
 	const FPBDCollisionSolverSettings& GetSolverSettings() const { return SolverSettings; }
 
+	const FCollisionDetectorSettings& GetDetectorSettings() const { return DetectorSettings; }
+
+	void SetDetectorSettings(const FCollisionDetectorSettings& InSettings)
+	{
+		DetectorSettings = InSettings;
+	}
+
+	void SetCullDistance(const FReal InCullDistance)
+	{
+		DetectorSettings.BoundsExpansion = InCullDistance;
+	}
+
 protected:
 	FPBDCollisionConstraint& GetConstraint(int32 Index);
 
@@ -290,6 +302,9 @@ private:
 
 	// Settings for the low-level collision solvers
 	FPBDCollisionSolverSettings SolverSettings;
+
+	// Settings for collision detection
+	FCollisionDetectorSettings DetectorSettings;
 };
 
 //
diff --git a/Engine/Source/Runtime/Experimental/Chaos/Public/Chaos/PBDRigidParticles.h b/Engine/Source/Runtime/Experimental/Chaos/Public/Chaos/PBDRigidParticles.h
--- a/Engine/Source/Runtime/Experimental/Chaos/Public/Chaos/PBDRigidParticles.h
+++ b/Engine/Source/Runtime/Experimental/Chaos/Public/Chaos/PBDRigidParticles.h
@@ -23,17 +23,15 @@ class TPBDRigidParticles : public TRigidParticles<T, d>
 	friend class TPBDRigidsEvolution<T, d>;
 
   public:
-    using TRigidParticles<T, d>::Sleeping;
+	using TRigidParticles<T, d>::CenterOfMass;
+	using TRigidParticles<T, d>::RotationOfMass;
+	using TRigidParticles<T, d>::Sleeping;
 
 	CHAOS_API TPBDRigidParticles()
 	    : TRigidParticles<T, d>()
 	{
 		this->MParticleType = EParticleType::Rigid;
-		TArrayCollection::AddArray(&MP);
-		TArrayCollection::AddArray(&MQ);
-		TArrayCollection::AddArray(&MPreV);
-		TArrayCollection::AddArray(&MPreW);
-		TArrayCollection::AddArray(&MSolverBodyIndex);
+		RegisterArrays();
 	}
 	TPBDRigidParticles(const TPBDRigidParticles<T, d>& Other) = delete;
 	CHAOS_API TPBDRigidParticles(TPBDRigidParticles<T, d>&& Other)
@@ -45,6 +43,15 @@ class TPBDRigidParticles : public TRigidParticles<T, d>
 		, MSolverBodyIndex(MoveTemp(Other.MSolverBodyIndex))
 	{
 		this->MParticleType = EParticleType::Rigid;
+		RegisterArrays();
+	}
+
+	CHAOS_API virtual ~TPBDRigidParticles()
+	{
+	}
+
+	void RegisterArrays()
+	{
 		TArrayCollection::AddArray(&MP);
 		TArrayCollection::AddArray(&MQ);
 		TArrayCollection::AddArray(&MPreV);
@@ -52,12 +59,9 @@ class TPBDRigidParticles : public TRigidParticles<T, d>
 		TArrayCollection::AddArray(&MSolverBodyIndex);
 	}
 
-	CHAOS_API virtual ~TPBDRigidParticles()
-	{}
-
 	FORCEINLINE const TVector<T, d>& P(const int32 index) const { return MP[index]; }
 	FORCEINLINE TVector<T, d>& P(const int32 index) { return MP[index]; }
-	
+
 	FORCEINLINE const TRotation<T, d>& Q(const int32 index) const { return MQ[index]; }
 	FORCEINLINE TRotation<T, d>& Q(const int32 index) { return MQ[index]; }
 
@@ -67,6 +71,20 @@ class TPBDRigidParticles : public TRigidParticles<T, d>
 	CHAOS_API const TVector<T, d>& PreW(const int32 index) const { return MPreW[index]; }
 	CHAOS_API TVector<T, d>& PreW(const int32 index) { return MPreW[index]; }
 
+	// World-space center of mass location
+	CHAOS_API const TVector<T, d> XCom(const int32 index) const { return this->X(index) + this->R(index).RotateVector(CenterOfMass(index)); }
+	CHAOS_API const TVector<T, d> PCom(const int32 index) const { return this->P(index) + this->Q(index).RotateVector(CenterOfMass(index)); }
+
+	// World-space center of mass rotation
+	CHAOS_API const TRotation<T, d> RCom(const int32 index) const { return this->R(index) * RotationOfMass(index); }
+	CHAOS_API const TRotation<T, d> QCom(const int32 index) const { return this->Q(index) * RotationOfMass(index); }
+
+	void SetTransformPQCom(const int32 index, const TVector<T, d>& InPCom, const TRotation<T, d>& InQCom)
+	{
+		Q(index) = InQCom * RotationOfMass(index).Inverse();
+		P(index) = InPCom - Q(index) * CenterOfMass(index);
+	}
+
 	// The index into an FSolverBodyContainer (for dynamic particles only), or INDEX_NONE.
 	// \see FSolverBodyContainer
 	CHAOS_API int32 SolverBodyIndex(const int32 index) const { return MSolverBodyIndex[index]; }
@@ -190,9 +208,16 @@ class TPBDRigidParticles : public TRigidParticles<T, d>
 		Ar << MP << MQ << MPreV << MPreW;
 	}
 
+	UE_DEPRECATED(5.3, "To be removed")
 	FORCEINLINE TArray<TVector<T, d>>& AllP() { return MP; }
+
+	UE_DEPRECATED(5.3, "To be removed")
 	FORCEINLINE TArray<TRotation<T, d>>& AllQ() { return MQ; }
+	
+	UE_DEPRECATED(5.3, "To be removed")
 	FORCEINLINE TArray<TVector<T, d>>& AllPreV() { return MPreV; }
+	
+	UE_DEPRECATED(5.3, "To be removed")
 	FORCEINLINE TArray<TVector<T, d>>& AllPreW() { return MPreW; }
 
   private:
diff --git a/Engine/Source/Runtime/Experimental/Chaos/Public/Chaos/PBDRigidsEvolutionGBF.h b/Engine/Source/Runtime/Experimental/Chaos/Public/Chaos/PBDRigidsEvolutionGBF.h
--- a/Engine/Source/Runtime/Experimental/Chaos/Public/Chaos/PBDRigidsEvolutionGBF.h
+++ b/Engine/Source/Runtime/Experimental/Chaos/Public/Chaos/PBDRigidsEvolutionGBF.h
@@ -175,7 +175,7 @@ namespace Chaos
 			//SCOPE_CYCLE_COUNTER(STAT_Integrate);
 			CHAOS_SCOPED_TIMER(Integrate);
 
-			const FReal BoundsThickness = GetCollisionDetector().GetSettings().BoundsExpansion;
+			const FReal BoundsThickness = GetCollisionConstraints().GetDetectorSettings().BoundsExpansion;
 			const FReal MaxAngularSpeedSq = CVars::HackMaxAngularVelocity * CVars::HackMaxAngularVelocity;
 			const FReal MaxSpeedSq = CVars::HackMaxVelocity * CVars::HackMaxVelocity;
 			InParticles.ParallelFor([&](auto& GeomParticle, int32 Index) 
@@ -252,14 +252,13 @@ namespace Chaos
 						}
 					}
 
-					//EulerStepRule.Apply(Particle, Dt);
-					FVec3 PCoM = FParticleUtilitiesXR::GetCoMWorldPosition(&Particle);
-					FRotation3 QCoM = FParticleUtilitiesXR::GetCoMWorldRotation(&Particle);
+					FVec3 PCoM = Particle.XCom();
+					FRotation3 QCoM = Particle.RCom();
 
 					PCoM = PCoM + Particle.V() * Dt;
 					QCoM = FRotation3::IntegrateRotationWithAngularVelocity(QCoM, Particle.W(), Dt);
 
-					FParticleUtilitiesPQ::SetCoMWorldTransform(&Particle, PCoM, QCoM);
+					Particle.SetTransformPQCom(PCoM, QCoM);
 
 					if (!Particle.CCDEnabled())
 					{
diff --git a/Engine/Source/Runtime/Experimental/Chaos/Public/Chaos/PBDSuspensionConstraints.h b/Engine/Source/Runtime/Experimental/Chaos/Public/Chaos/PBDSuspensionConstraints.h
--- a/Engine/Source/Runtime/Experimental/Chaos/Public/Chaos/PBDSuspensionConstraints.h
+++ b/Engine/Source/Runtime/Experimental/Chaos/Public/Chaos/PBDSuspensionConstraints.h
@@ -13,8 +13,14 @@
 
 namespace Chaos
 {
+	class FSolverBody;
 	class FPBDSuspensionConstraints;
-	class FPBDCollisionSolver;
+
+	namespace Private
+	{
+		class FPBDCollisionSolver;
+		class FPBDCollisionSolverManifoldPoint;
+	}
 
 	class CHAOS_API FPBDSuspensionConstraintHandle final : public TIndexedContainerConstraintHandle<FPBDSuspensionConstraints>
 	{
@@ -273,7 +279,8 @@ namespace Chaos
 		FHandles Handles;
 		FConstraintHandleAllocator HandleAllocator;
 
-		TArray<FPBDCollisionSolver*> CollisionSolvers;
+		TArray<Private::FPBDCollisionSolver> CollisionSolvers;
+		TArray<Private::FPBDCollisionSolverManifoldPoint> CollisionSolverManifoldPoints;
 		TArray<FSolverBody> StaticCollisionBodies;
 	};
 }
diff --git a/Engine/Source/Runtime/Experimental/Chaos/Public/Chaos/Particle/ParticleUtilities.h b/Engine/Source/Runtime/Experimental/Chaos/Public/Chaos/Particle/ParticleUtilities.h
--- a/Engine/Source/Runtime/Experimental/Chaos/Public/Chaos/Particle/ParticleUtilities.h
+++ b/Engine/Source/Runtime/Experimental/Chaos/Public/Chaos/Particle/ParticleUtilities.h
@@ -4,6 +4,14 @@
 #include "Chaos/ParticleHandle.h"
 #include "Chaos/Utilities.h"
 
+//
+// These accessors existed to allow us to switch between particles whose transforms are maintained
+// in center of mass space, versus actor space. The should no longer be needed - we have settled on Actor Space. 
+// See particle handle methods like GetTransformPQ(), SetTransformPQCom() etc.
+//
+// @todo(chaos): deprecate this file and its contents
+//
+
 #define CHAOS_PARTICLE_ACTORTRANSFORM 1
 
 namespace Chaos
@@ -198,139 +206,9 @@ namespace Chaos
 		}
 	};
 
-	/**
-	 * Particle Space == CoM Space.
-	 * Velocities in CoM Space.
-	 */
-	template <typename TSpatialAccessor>
-	class FParticleUtilities_CoMSpace
-	{
-	public:
-		template<typename T_PARTICLEHANDLE>
-		static inline FRigidTransform3 GetActorWorldTransform(T_PARTICLEHANDLE Particle)
-		{
-			FRotation3 ActorQ = TSpatialAccessor::GetRotation(Particle) * Particle->RotationOfMass().Inverse();
-			FVec3 ActorP = TSpatialAccessor::GetPosition(Particle) - ActorQ.RotateVector(Particle->CenterOfMass());
-			return FRigidTransform3(ActorP, ActorQ);
-		}
-
-		template<typename T_PARTICLEHANDLE>
-		static inline void SetActorWorldTransform(T_PARTICLEHANDLE Particle, const FRigidTransform3& ActorWorldT)
-		{
-			FRotation3 CoMQ = ActorWorldT.GetRotation() * Particle->RotationOfMass();
-			FVec3 CoMP = ActorWorldT.GetTranslation() + ActorWorldT.GetRotation().RotateVector(Particle->CenterOfMass());
-			TSpatialAccessor::Position(Particle) = CoMP;
-			TSpatialAccessor::Rotation(Particle) = CoMQ;
-		}
-
-		template<typename T_PARTICLEHANDLE>
-		static inline FRigidTransform3 ActorLocalToParticleLocal(T_PARTICLEHANDLE Particle, const FRigidTransform3& ActorLocalT)
-		{
-			return FRigidTransform3(Particle->RotationOfMass().UnrotateVector(ActorLocalT.GetTranslation() - Particle->CenterOfMass()), Particle->RotationOfMass().Inverse() * ActorLocalT.GetRotation());
-		}
-
-		template<typename T_PARTICLEHANDLE>
-		static inline FRigidTransform3 ActorWorldToParticleWorld(T_PARTICLEHANDLE Particle, const FRigidTransform3& ActorWorldT)
-		{
-			FRotation3 CoMQ = ActorWorldT.GetRotation() * Particle->RotationOfMass();
-			FVec3 CoMP = ActorWorldT.GetTranslation() + ActorWorldT.GetRotation().RotateVector(Particle->CenterOfMass());
-			return FRigidTransform3(CoMP, CoMQ);
-		}
-
-		template<typename T_PARTICLEHANDLE>
-		static inline FMatrix33 GetWorldInertia(T_PARTICLEHANDLE Particle)
-		{
-			return Utilities::ComputeWorldSpaceInertia(TSpatialAccessor::Rotation(Particle), Particle->I());
-		}
-
-		template<typename T_PARTICLEHANDLE>
-		static inline FMatrix33 GetWorldInvInertia(T_PARTICLEHANDLE Particle)
-		{
-			return Utilities::ComputeWorldSpaceInertia(TSpatialAccessor::Rotation(Particle), Particle->InvI());
-		}
-
-		template<typename T_PARTICLEHANDLE>
-		static inline const FVec3& ParticleLocalToCoMLocal(T_PARTICLEHANDLE Particle, const FVec3& P)
-		{
-			return P;
-		}
-
-		template<typename T_PARTICLEHANDLE>
-		static inline const FRotation3& ParticleLocalToCoMLocal(T_PARTICLEHANDLE Particle, const FRotation3& Q)
-		{
-			return Q;
-		}
-
-		template<typename T_PARTICLEHANDLE>
-		static inline const FRigidTransform3& ParticleLocalToCoMLocal(T_PARTICLEHANDLE Particle, const FRigidTransform3& T)
-		{
-			return T;
-		}
-
-		template<typename T_PARTICLEHANDLE>
-		static inline FVec3 GetVelocityAtCoMRelativePosition(T_PARTICLEHANDLE Particle, const FVec3& RelPos)
-		{
-			return Particle->V() + FVec3::CrossProduct(Particle->W(), RelPos);
-		}
-
-		template<typename T_PARTICLEHANDLE>
-		static inline FVec3 GetPreviousVelocityAtCoMRelativePosition(T_PARTICLEHANDLE Particle, const FVec3& RelPos)
-		{
-			return Particle->PreV() + FVec3::CrossProduct(Particle->PreW(), RelPos);
-		}
-
-		template<typename T_PARTICLEHANDLE>
-		static inline const FVec3& GetCoMWorldPosition(T_PARTICLEHANDLE Particle)
-		{
-			return TSpatialAccessor::GetPosition(Particle);
-		}
-
-		static inline const FVec3& GetCoMWorldPosition(TPBDRigidParticles<FReal, 3>& Particles, const int32 Index)
-		{
-			return TSpatialAccessor::Position(Particles, Index);
-		}
-
-		template<typename T_PARTICLEHANDLE>
-		static inline const FRotation3& GetCoMWorldRotation(T_PARTICLEHANDLE Particle)
-		{
-			return TSpatialAccessor::GetRotation(Particle);
-		}
-
-		template<typename T_PARTICLEHANDLE>
-		static inline FRigidTransform3 GetCoMWorldTransform(T_PARTICLEHANDLE Particle)
-		{
-			return FRigidTransform3(GetCoMWorldPosition(Particle), GetCoMWorldRotation(Particle));
-		}
-
-		static inline const FRotation3& GetCoMWorldRotation(TPBDRigidParticles<FReal, 3>& Particles, const int32 Index)
-		{
-			return TSpatialAccessor::Rotation(Particles, Index);
-		}
-
-		template<typename T_PARTICLEHANDLE>
-		static inline void SetCoMWorldTransform(T_PARTICLEHANDLE Particle, const FVec3& PCoM, const FRotation3& QCoM)
-		{
-			TSpatialAccessor::Position(Particle) = PCoM;
-			TSpatialAccessor::Rotation(Particle) = QCoM;
-		}
-
-		static inline void SetCoMWorldTransform(TPBDRigidParticles<FReal, 3>& Particles, const int32 Index, const FVec3& PCoM, const FRotation3& QCoM)
-		{
-			TSpatialAccessor::Position(Particles, Index) = PCoM;
-			TSpatialAccessor::Rotation(Particles, Index) = QCoM;
-		}
-	};
-
-#if CHAOS_PARTICLE_ACTORTRANSFORM
 	using FParticleUtilitiesPQ = FParticleUtilities_ActorSpace<FParticleSpatialAccessorPQ>;
 	using FParticleUtilitiesXR = FParticleUtilities_ActorSpace<FParticleSpatialAccessorXR>;
 	using FParticleUtilities = FParticleUtilities_ActorSpace<FParticleSpatialAccessorPQ>;
 	using FParticleUtilitiesGT = FParticleUtilities_ActorSpace<FParticleSpatialAccessorXR>;
-#else
-	using FParticleUtilitiesPQ = FParticleUtilities_CoMSpace<FParticleSpatialAccessorPQ>;
-	using FParticleUtilitiesXR = FParticleUtilities_CoMSpace<FParticleSpatialAccessorXR>;
-	using FParticleUtilities = FParticleUtilities_CoMSpace<FParticleSpatialAccessorPQ>;
-	using FParticleUtilitiesGT = FParticleUtilities_CoMSpace<FParticleSpatialAccessorXR>;
-#endif
 
 }
\ No newline at end of file
diff --git a/Engine/Source/Runtime/Experimental/Chaos/Public/Chaos/ParticleHandle.h b/Engine/Source/Runtime/Experimental/Chaos/Public/Chaos/ParticleHandle.h
--- a/Engine/Source/Runtime/Experimental/Chaos/Public/Chaos/ParticleHandle.h
+++ b/Engine/Source/Runtime/Experimental/Chaos/Public/Chaos/ParticleHandle.h
@@ -497,6 +497,15 @@ public:
 	TRotation<T, d>& R() { return GeometryParticles->R(ParticleIdx); }
 	void SetR(const TRotation<T, d>& InR, bool bInvalidate = false) { GeometryParticles->R(ParticleIdx) = InR; }
 
+	FRigidTransform3 GetTransformXR() const { return FRigidTransform3(X(), R()); }
+
+	// Initialize the transform
+	void InitTransform(const FVec3& InP, const FRotation3& InQ)
+	{
+		SetX(InP);
+		SetR(InQ);
+	}
+
 	void SetXR(const FParticlePositionRotation& XR);
 	
 	void SetNonFrequentData(const FParticleNonFrequentData& InData)
@@ -830,16 +839,18 @@ template <typename T, int d, bool bPersistent>
 class TPBDRigidParticleHandleImp : public TKinematicGeometryParticleHandleImp<T, d, bPersistent>
 {
 public:
+	using TGeometryParticleHandleImp<T, d, bPersistent>::ParticleCollisions;
+	using TGeometryParticleHandleImp<T, d, bPersistent>::ParticleConstraints;
 	using TGeometryParticleHandleImp<T, d, bPersistent>::ParticleIdx;
 	using TGeometryParticleHandleImp<T, d, bPersistent>::PBDRigidParticles;
-	using TGeometryParticleHandleImp<T, d, bPersistent>::ParticleConstraints;
-	using TGeometryParticleHandleImp<T, d, bPersistent>::ParticleCollisions;
+	using TGeometryParticleHandleImp<T, d, bPersistent>::SetX;
+	using TGeometryParticleHandleImp<T, d, bPersistent>::SetR;
+	using TGeometryParticleHandleImp<T, d, bPersistent>::Type;
 	using TKinematicGeometryParticleHandleImp<T, d, bPersistent>::V;
 	using TKinematicGeometryParticleHandleImp<T, d, bPersistent>::W;
-	using TGeometryParticleHandleImp<T, d, bPersistent>::Type;
+
 	using TTransientHandle = TTransientPBDRigidParticleHandle<T, d>;
 	using TSOAType = TPBDRigidParticles<T, d>;
-	using TGeometryParticleHandleImp<T, d, bPersistent>::SetXR;
 
 protected:
 	friend class TGeometryParticleHandleImp<T, d, bPersistent>;
@@ -932,6 +943,32 @@ public:
 	TRotation<T, d>& Q() { return PBDRigidParticles->Q(ParticleIdx); }
 	void SetQ(const TRotation<T, d>& InQ) { PBDRigidParticles->Q(ParticleIdx) = InQ; }
 
+	// World-space center of mass position
+	const TVector<T, d> XCom() const { return PBDRigidParticles->XCom(ParticleIdx); }
+	const TVector<T, d> PCom() const { return PBDRigidParticles->PCom(ParticleIdx); }
+
+	// World-space center of mass rotation
+	const TRotation<T, d> RCom() const { return PBDRigidParticles->RCom(ParticleIdx); }
+	const TRotation<T, d> QCom() const { return PBDRigidParticles->QCom(ParticleIdx); }
+
+	// Initialize the transform (sets X,R and  P,Q)
+	void InitTransform(const FVec3& InP, const FRotation3& InQ)
+	{
+		SetX(InP);
+		SetR(InQ);
+		SetP(InP);
+		SetQ(InQ);
+	}
+
+	// Set world-space center of mass transform
+	void SetTransformPQCom(const TVector<T, d>& InPCom, const TRotation<T, d>& InQCom) { PBDRigidParticles->SetTransformPQCom(ParticleIdx, InPCom, InQCom); }
+
+	FRigidTransform3 GetTransformPQ() const { return FRigidTransform3(P(), Q()); }
+
+	FRigidTransform3 GetTransformXRCom() const { return FRigidTransform3(XCom(), RCom()); }
+
+	FRigidTransform3 GetTransformPQCom() const { return FRigidTransform3(PCom(), QCom()); }
+
 	const TVector<T, d>& VSmooth() const { return PBDRigidParticles->VSmooth(ParticleIdx); }
 	TVector<T, d>& VSmooth() { return PBDRigidParticles->VSmooth(ParticleIdx); }
 	void SetVSmooth(const TVector<T, d>& InVSmooth) { PBDRigidParticles->VSmooth(ParticleIdx) = InVSmooth; }
@@ -1532,7 +1569,7 @@ public:
 
 	FVec3& P()
 	{
-		if (IsDynamic())
+		if (MHandle->CastToRigidParticle())
 		{
 			return MHandle->CastToRigidParticle()->P();
 		}
@@ -1542,7 +1579,7 @@ public:
 
 	const FVec3& P() const
 	{
-		if (IsDynamic())
+		if (MHandle->CastToRigidParticle())
 		{
 			return MHandle->CastToRigidParticle()->P();
 		}
@@ -1552,7 +1589,7 @@ public:
 
 	FRotation3& Q()
 	{
-		if (IsDynamic())
+		if (MHandle->CastToRigidParticle())
 		{
 			return MHandle->CastToRigidParticle()->Q();
 		}
@@ -1562,7 +1599,7 @@ public:
 
 	const FRotation3& Q() const
 	{
-		if (IsDynamic())
+		if (MHandle->CastToRigidParticle())
 		{
 			return MHandle->CastToRigidParticle()->Q();
 		}
@@ -1570,6 +1607,95 @@ public:
 		return R();
 	}
 
+	// World-space center of mass position
+	const FVec3 XCom() const
+	{
+		if (MHandle->CastToRigidParticle())
+		{
+			return MHandle->CastToRigidParticle()->XCom();
+		}
+		return X();
+	}
+	const FVec3 PCom() const
+	{
+		if (MHandle->CastToRigidParticle())
+		{
+			return MHandle->CastToRigidParticle()->PCom();
+		}
+		return X();
+	}
+
+	// World-space center of mass rotation
+	const FRotation3 RCom() const
+	{ 
+		if (MHandle->CastToRigidParticle())
+		{
+			return MHandle->CastToRigidParticle()->RCom();
+		}
+		return R();
+	}
+	const FRotation3 QCom() const
+	{ 
+		if (MHandle->CastToRigidParticle())
+		{
+			return MHandle->CastToRigidParticle()->QCom();
+		}
+		return R();
+	}
+
+	void InitTransform(const FVec3& InP, const FRotation3& InQ)
+	{
+		if (MHandle->CastToRigidParticle())
+		{
+			MHandle->CastToRigidParticle()->InitTransform(InP, InQ);
+		}
+		MHandle->InitTransform(InP, InQ);
+	}
+
+	void SetTransformPQCom(const FVec3& InPCom, const FRotation3& InQCom)
+	{ 
+		if (MHandle->CastToRigidParticle())
+		{
+			MHandle->CastToRigidParticle()->SetTransformPQCom(InPCom, InQCom);
+		}
+		else
+		{
+			SetTransform(InPCom, InQCom);
+		}
+	}
+
+	FRigidTransform3 GetTransformXR() const
+	{
+		return MHandle->GetTransformXR();
+	}
+
+	FRigidTransform3 GetTransformPQ() const
+	{
+		if (MHandle->CastToRigidParticle())
+		{
+			return MHandle->CastToRigidParticle()->GetTransformPQ();
+		}
+		return MHandle->GetTransformXR();
+	}
+
+	FRigidTransform3 GetTransformXRCom() const
+	{
+		if (MHandle->CastToRigidParticle())
+		{
+			return MHandle->CastToRigidParticle()->GetTransformXRCom();
+		}
+		return GetTransformXR();
+	}
+
+	FRigidTransform3 GetTransformPQCom() const
+	{
+		if (MHandle->CastToRigidParticle())
+		{
+			return MHandle->CastToRigidParticle()->GetTransformPQCom();
+		}
+		return GetTransformXR();
+	}
+
 	const FVec3& VSmooth() const
 	{
 		if (MHandle->CastToRigidParticle())
diff --git a/Engine/Source/Runtime/Experimental/Chaos/Public/Chaos/PerParticlePBDEulerStep.h b/Engine/Source/Runtime/Experimental/Chaos/Public/Chaos/PerParticlePBDEulerStep.h
--- a/Engine/Source/Runtime/Experimental/Chaos/Public/Chaos/PerParticlePBDEulerStep.h
+++ b/Engine/Source/Runtime/Experimental/Chaos/Public/Chaos/PerParticlePBDEulerStep.h
@@ -29,24 +29,24 @@ class FPerParticlePBDEulerStep : public FPerParticleRule
 
 	inline void Apply(TPBDRigidParticles<FReal, 3>& InParticles, const FReal Dt, const int32 Index) const override //-V762
 	{
-		FVec3 PCoM = FParticleUtilitiesXR::GetCoMWorldPosition(InParticles, Index);
-		FRotation3 QCoM = FParticleUtilitiesXR::GetCoMWorldRotation(InParticles, Index);
+		FVec3 PCoM = InParticles.XCom(Index);
+		FRotation3 QCoM = InParticles.RCom(Index);
 
 		PCoM = PCoM + InParticles.V(Index) * Dt;
 		QCoM = FRotation3::IntegrateRotationWithAngularVelocity(QCoM, InParticles.W(Index), Dt);
 
-		FParticleUtilitiesPQ::SetCoMWorldTransform(InParticles, Index, PCoM, QCoM);
+		InParticles.SetTransformPQCom(Index, PCoM, QCoM);
 	}
 
 	inline void Apply(TTransientPBDRigidParticleHandle<FReal, 3>& Particle, const FReal Dt) const override
 	{
-		FVec3 PCoM = FParticleUtilitiesXR::GetCoMWorldPosition(&Particle);
-		FRotation3 QCoM = FParticleUtilitiesXR::GetCoMWorldRotation(&Particle);
+		FVec3 PCoM = Particle.XCom();
+		FRotation3 QCoM = Particle.RCom();
 
 		PCoM = PCoM + Particle.V() * Dt;
 		QCoM = FRotation3::IntegrateRotationWithAngularVelocity(QCoM, Particle.W(), Dt);
 
-		FParticleUtilitiesPQ::SetCoMWorldTransform(&Particle, PCoM, QCoM);
+		Particle.SetTransformPQCom(PCoM, QCoM);
 	}
 };
 
diff --git a/Engine/Source/Runtime/Experimental/Chaos/Public/Chaos/PerParticlePBDUpdateFromDeltaPosition.h b/Engine/Source/Runtime/Experimental/Chaos/Public/Chaos/PerParticlePBDUpdateFromDeltaPosition.h
--- a/Engine/Source/Runtime/Experimental/Chaos/Public/Chaos/PerParticlePBDUpdateFromDeltaPosition.h
+++ b/Engine/Source/Runtime/Experimental/Chaos/Public/Chaos/PerParticlePBDUpdateFromDeltaPosition.h
@@ -35,27 +35,19 @@ class FPerParticlePBDUpdateFromDeltaPosition : public FPerParticleRule
 
 	inline void Apply(FPBDRigidParticleHandle* Particle, const FReal Dt) const override //-V762
 	{
-#if CHAOS_PARTICLE_ACTORTRANSFORM
 		const FVec3& CenterOfMass = Particle->CenterOfMass();
-		const FVec3 CenteredX = Particle->X() + Particle->R().RotateVector(CenterOfMass);
-		const FVec3 CenteredP = Particle->P() + Particle->Q().RotateVector(CenterOfMass);
+		const FVec3 CenteredX = Particle->XCom();
+		const FVec3 CenteredP = Particle->PCom();
 		Particle->V() = FVec3::CalculateVelocity(CenteredX, CenteredP, Dt);
-#else
-		Particle->V() = FVec3::CalculateVelocity(Particle->X(), Particle->P(), Dt);
-#endif
 		Particle->W() = FRotation3::CalculateAngularVelocity(Particle->R(), Particle->Q(), Dt);
 	}
 
 	inline void Apply(TTransientPBDRigidParticleHandle<FReal, 3>& Particle, const FReal Dt) const override //-V762
 	{
-#if CHAOS_PARTICLE_ACTORTRANSFORM
 		const FVec3& CenterOfMass = Particle.CenterOfMass();
-		const FVec3 CenteredX = Particle.X() + Particle.R().RotateVector(CenterOfMass);
-		const FVec3 CenteredP = Particle.P() + Particle.Q().RotateVector(CenterOfMass);
+		const FVec3 CenteredX = Particle.XCom();
+		const FVec3 CenteredP = Particle.PCom();
 		Particle.V() = FVec3::CalculateVelocity(CenteredX, CenteredP, Dt);
-#else
-		Particle.V() = FVec3::CalculateVelocity(Particle.X(), Particle.P(), Dt);
-#endif
 		Particle.W() = FRotation3::CalculateAngularVelocity(Particle.R(), Particle.Q(), Dt);
 	}
 };
diff --git a/Engine/Source/Runtime/Experimental/Chaos/Public/PBDRigidsSolver.h b/Engine/Source/Runtime/Experimental/Chaos/Public/PBDRigidsSolver.h
--- a/Engine/Source/Runtime/Experimental/Chaos/Public/PBDRigidsSolver.h
+++ b/Engine/Source/Runtime/Experimental/Chaos/Public/PBDRigidsSolver.h
@@ -171,7 +171,7 @@ namespace Chaos
 		void SetPositionIterations(const int32 InNumIterations) { GetEvolution()->SetNumPositionIterations(InNumIterations); }
 		void SetVelocityIterations(const int32 InNumIterations) { GetEvolution()->SetNumVelocityIterations(InNumIterations); }
 		void SetProjectionIterations(const int32 InNumIterations) { GetEvolution()->SetNumProjectionIterations(InNumIterations); }
-		void SetCollisionCullDistance(const FReal InCullDistance) { GetEvolution()->GetCollisionDetector().SetBoundsExpansion(InCullDistance); }
+		void SetCollisionCullDistance(const FReal InCullDistance) { GetEvolution()->GetCollisionConstraints().SetCullDistance(InCullDistance); }
 		void SetCollisionMaxPushOutVelocity(const FReal InMaxPushOutVelocity) { GetEvolution()->GetCollisionConstraints().SetMaxPushOutVelocity(InMaxPushOutVelocity); }
 
 		/**/
