Chaos - fix non-unity compiler error in PBDSuspensionConstraint.h

diff --git a/Engine/Source/Runtime/Experimental/Chaos/Public/Chaos/PBDSuspensionConstraints.h b/Engine/Source/Runtime/Experimental/Chaos/Public/Chaos/PBDSuspensionConstraints.h
--- a/Engine/Source/Runtime/Experimental/Chaos/Public/Chaos/PBDSuspensionConstraints.h
+++ b/Engine/Source/Runtime/Experimental/Chaos/Public/Chaos/PBDSuspensionConstraints.h
@@ -2,6 +2,7 @@
 #pragma once
 
 #include "Chaos/Array.h"
+#include "Chaos/Collision/PBDCollisionSolver.h"
 #include "Chaos/Evolution/IndexedConstraintContainer.h"
 #include "Chaos/ParticleHandle.h"
 #include "Chaos/PBDSuspensionConstraintTypes.h"
