Fix Triangle Sphere Overlap

diff --git a/Engine/Source/Programs/HeadlessChaos/Private/HeadlessChaosTestOverlap.cpp b/Engine/Source/Programs/HeadlessChaos/Private/HeadlessChaosTestOverlap.cpp
--- a/Engine/Source/Programs/HeadlessChaos/Private/HeadlessChaosTestOverlap.cpp
+++ b/Engine/Source/Programs/HeadlessChaos/Private/HeadlessChaosTestOverlap.cpp
@@ -202,6 +202,11 @@ namespace ChaosTest
 					bool bResult = TriangleMesh->OverlapGeom(Sphere, QueryTM, 0.0);
 					EXPECT_FALSE(bResult);
 				}
+				{
+					const FRigidTransform3 QueryTM(FVec3(-3.0, 0.0, 9.0), FQuat::Identity);
+					bool bResult = TriangleMesh->OverlapGeom(Sphere, QueryTM, 0.0);
+					EXPECT_FALSE(bResult);
+				}
 			}
 			{
 				const FVec3 X1 = { 0.0, 0.0, -200.0 };
@@ -846,17 +851,17 @@ namespace ChaosTest
 		}
 		{
 			FTriangleMeshImplicitObject::ParticlesType TrimeshParticles(
-				{
-					{-10.0, -10.0, 0.0},
-					{10.0, -10.0, 0.0},
-					{-10.0, 10.0, 0.0},
-					{10.0, 10.0, 0.0},
-					{-10.0, -10.0, 10.0},
-					{10.0, -10.0, 10.0},
-					{-10.0, 10.0, 10.0},
-					{10.0, 10.0, 10.0},
+			{
+				{-10.0, -10.0, 0.0},
+				{10.0, -10.0, 0.0},
+				{-10.0, 10.0, 0.0},
+				{10.0, 10.0, 0.0},
+				{-10.0, -10.0, 10.0},
+				{10.0, -10.0, 10.0},
+				{ -10.0, 10.0, 10.0 },
+				{ 10.0, 10.0, 10.0 },
 
-				});
+			});
 
 			TArray<TVec3<int32>> Indices(
 				{
@@ -952,5 +957,52 @@ namespace ChaosTest
 				}
 			}
 		}
+		{
+			FTriangleMeshImplicitObject::ParticlesType TrimeshParticles(
+				{
+					{-5.0, 0.0, 0.0},
+					{5.0, 0.0, 0.0},
+					{3.0, 0.0, 100.0},
+				});
+
+			TArray<TVec3<int32>> Indices(
+				{
+					{0, 1, 2},
+				});
+
+			TArray<uint16> Materials;
+			for (int32 i = 0; i < Indices.Num(); ++i)
+			{
+				Materials.Emplace(0);
+			}
+
+			TUniquePtr<FTriangleMeshImplicitObject> TriangleMesh = MakeUnique<FTriangleMeshImplicitObject>(MoveTemp(TrimeshParticles), MoveTemp(Indices), MoveTemp(Materials));
+			{
+				// Sphere test
+				const FVec3 X = { 0.0, 0.0, 0.0 };
+				const FReal Radius = 1.0;
+				const Chaos::FSphere Sphere = Chaos::FSphere(X, Radius);
+				{
+					const FRigidTransform3 QueryTM(FVec3(0.0, 0.0, 0.0), FQuat::Identity);
+					bool bResult = TriangleMesh->OverlapGeom(Sphere, QueryTM, 0.0);
+					EXPECT_TRUE(bResult);
+				}
+				{
+					const FRigidTransform3 QueryTM(FVec3(3.9, 0.0, 100.0), FQuat::Identity);
+					bool bResult = TriangleMesh->OverlapGeom(Sphere, QueryTM, 0.0);
+					EXPECT_TRUE(bResult);
+				}
+				{
+					const FRigidTransform3 QueryTM(FVec3(4.0, 0.0, 90.0), FQuat::Identity);
+					bool bResult = TriangleMesh->OverlapGeom(Sphere, QueryTM, 0.0);
+					EXPECT_TRUE(bResult);
+				}
+				{
+					const FRigidTransform3 QueryTM(FVec3(4.5, 0.0, 90.0), FQuat::Identity);
+					bool bResult = TriangleMesh->OverlapGeom(Sphere, QueryTM, 0.0);
+					EXPECT_FALSE(bResult);
+				}
+			}
+		}
 	}
 }
diff --git a/Engine/Source/Runtime/Experimental/Chaos/Private/Chaos/Collision/TriangleOverlap.cpp b/Engine/Source/Runtime/Experimental/Chaos/Private/Chaos/Collision/TriangleOverlap.cpp
--- a/Engine/Source/Runtime/Experimental/Chaos/Private/Chaos/Collision/TriangleOverlap.cpp
+++ b/Engine/Source/Runtime/Experimental/Chaos/Private/Chaos/Collision/TriangleOverlap.cpp
@@ -125,12 +125,12 @@ namespace Chaos
 		const VectorRegister4Float AB = VectorSubtract(B, A);
 		const VectorRegister4Float BC = VectorSubtract(C, B);
 
-		VectorRegister4Float Normal = VectorNormalize(VectorCross(AB, BC));
+		const VectorRegister4Float Normal = VectorNormalize(VectorCross(AB, BC));
 
 		// Plane Triangle
 		const VectorRegister4Float AX = VectorSubtract(X, A);
 
-		FRealSingle AXDist = VectorDot3Scalar(AX, Normal);
+		const FRealSingle AXDist = VectorDot3Scalar(AX, Normal);
 
 		if (FMath::Abs<FRealSingle>(AXDist) > Radius)
 		{
@@ -140,21 +140,19 @@ namespace Chaos
 		constexpr FRealSingle ThirdFloat = 1.0f / 3.0f;
 		constexpr VectorRegister4Float Third = MakeVectorRegisterFloatConstant(ThirdFloat, ThirdFloat, ThirdFloat, ThirdFloat);
 		const VectorRegister4Float Centroid = VectorMultiply(VectorAdd(VectorAdd(A, B), C), Third);
-		
-
-		
-		
 
-		FRealSingle SqrRadius = Radius * Radius;
+		const FRealSingle SqrRadius = Radius * Radius;
 		{
 			const VectorRegister4Float SqrAB = VectorDot3(AB, AB);
 			const VectorRegister4Float ZeroMask = VectorCompareEQ(VectorZeroFloat(), SqrAB);
 			VectorRegister4Float Time = VectorClamp(VectorDivide(VectorDot3(AB, AX), SqrAB), VectorZeroFloat(), VectorOneFloat());
 			Time = VectorBitwiseNotAnd(ZeroMask, Time);
-			const VectorRegister4Float P = VectorMultiplyAdd(B, VectorSubtract(VectorOneFloat(), Time), VectorMultiply(A, Time));
+			const VectorRegister4Float P = VectorMultiplyAdd(A, VectorSubtract(VectorOneFloat(), Time), VectorMultiply(B, Time));
 
 			const VectorRegister4Float PX = VectorSubtract(X, P);
-			if (VectorDot3Scalar(PX, VectorSubtract(P, Centroid)) > 0.0f)
+			const VectorRegister4Float CentroidP = VectorSubtract(P, Centroid);
+			const FRealSingle Dist = VectorDot3Scalar(PX, CentroidP);
+			if (Dist > 0.0f)
 			{
 				if (VectorDot3Scalar(PX, PX) > SqrRadius)
 				{
@@ -168,10 +166,12 @@ namespace Chaos
 			const VectorRegister4Float BX = VectorSubtract(X, B);
 			VectorRegister4Float Time = VectorClamp(VectorDivide(VectorDot3(BC, BX), SqrBC), VectorZeroFloat(), VectorOneFloat());
 			Time = VectorBitwiseNotAnd(ZeroMask, Time);
-			const VectorRegister4Float P = VectorMultiplyAdd(C, VectorSubtract(VectorOneFloat(), Time), VectorMultiply(B, Time));
+			const VectorRegister4Float P = VectorMultiplyAdd(B, VectorSubtract(VectorOneFloat(), Time), VectorMultiply(C, Time));
 
 			const VectorRegister4Float PX = VectorSubtract(X, P);
-			if (VectorDot3Scalar(PX, VectorSubtract(P, Centroid)) > 0.0f)
+			const VectorRegister4Float CentroidP = VectorSubtract(P, Centroid);
+			const FRealSingle Dist = VectorDot3Scalar(PX, CentroidP);
+			if (Dist > 0.0f)
 			{
 				if (VectorDot3Scalar(PX, PX) > SqrRadius)
 				{
@@ -186,10 +186,12 @@ namespace Chaos
 			const VectorRegister4Float CX = VectorSubtract(X, C);
 			VectorRegister4Float Time = VectorClamp(VectorDivide(VectorDot3(CA, CX), SqrCA), VectorZeroFloat(), VectorOneFloat());
 			Time = VectorBitwiseNotAnd(ZeroMask, Time);
-			const VectorRegister4Float P = VectorMultiplyAdd(A, VectorSubtract(VectorOneFloat(), Time), VectorMultiply(C, Time));
+			const VectorRegister4Float P = VectorMultiplyAdd(C, VectorSubtract(VectorOneFloat(), Time), VectorMultiply(A, Time));
 
 			const VectorRegister4Float PX = VectorSubtract(X, P);
-			if (VectorDot3Scalar(PX, VectorSubtract(P, Centroid)) > 0.0f)
+			const VectorRegister4Float CentroidP = VectorSubtract(P, Centroid);
+			const FRealSingle Dist = VectorDot3Scalar(PX, CentroidP);
+			if (Dist > 0.0f)
 			{
 				if (VectorDot3Scalar(PX, PX) > SqrRadius)
 				{
