Chaos: Fixed triangle mesh fast BVH triangle reordering losing the coupling of faces to material indices. This resulted in queries returning incorrect materials for internal faces.

diff --git a/Engine/Source/Runtime/Experimental/Chaos/Private/Chaos/TriangleMeshImplicitObject.cpp b/Engine/Source/Runtime/Experimental/Chaos/Private/Chaos/TriangleMeshImplicitObject.cpp
--- a/Engine/Source/Runtime/Experimental/Chaos/Private/Chaos/TriangleMeshImplicitObject.cpp
+++ b/Engine/Source/Runtime/Experimental/Chaos/Private/Chaos/TriangleMeshImplicitObject.cpp
@@ -1914,6 +1914,8 @@ void FTriangleMeshImplicitObject::RebuildFastBVHFromTree(const BVHType& TreeBVH)
 	int32 FaceNum = 0;
 	TArray<TVec3<FTrimeshIndexBuffer::LargeIdxType>> LargeIndices;
 	TArray<TVec3<FTrimeshIndexBuffer::SmallIdxType>> SmallIndices;
+	
+	TArray<uint16> OldMaterialIndices;
 	TArray<int32> OldFaceIndexMap;
 
 	if (ExternalFaceIndexMap.IsValid())
@@ -1922,6 +1924,9 @@ void FTriangleMeshImplicitObject::RebuildFastBVHFromTree(const BVHType& TreeBVH)
 		ExternalFaceIndexMap->Reserve(OldFaceIndexMap.Num());
 	}
 
+	OldMaterialIndices = MoveTemp(MaterialIndices);
+	MaterialIndices.Reserve(OldMaterialIndices.Num());
+
 	// since we do skip leaf nodes, we need to handle the case where we have only one node that will be a leaf by default
 	if (Nodes.Num() == 1)
 	{
@@ -1953,6 +1958,12 @@ void FTriangleMeshImplicitObject::RebuildFastBVHFromTree(const BVHType& TreeBVH)
 			{
 				ExternalFaceIndexMap->Add(OldFaceIndexMap[LeafPayload.Payload]);
 			}
+
+			// If we have materials - add the index to the remapped list
+			if(OldMaterialIndices.IsValidIndex(LeafPayload.Payload))
+			{
+				MaterialIndices.Add(OldMaterialIndices[LeafPayload.Payload]);
+			}
 		}
 		if (MElements.RequiresLargeIndices())
 		{
@@ -2034,6 +2045,13 @@ void FTriangleMeshImplicitObject::RebuildFastBVHFromTree(const BVHType& TreeBVH)
 							{
 								ExternalFaceIndexMap->Add(OldFaceIndexMap[LeafPayload.Payload]);
 							}
+
+							// If we have materials - add the index to the remapped list
+							if(OldMaterialIndices.IsValidIndex(LeafPayload.Payload))
+							{
+								MaterialIndices.Add(OldMaterialIndices[LeafPayload.Payload]);
+							}
+
 							FaceNum++;
 						}
 					}
