ensure primary dependency on post load.
    - in a context of adding several constraints with parent that can't tick (static meshes for instance), we want to make sure that the evaluation order is kept the same

diff --git a/Engine/Source/Runtime/Experimental/Animation/Constraints/Private/TransformConstraint.cpp b/Engine/Source/Runtime/Experimental/Animation/Constraints/Private/TransformConstraint.cpp
--- a/Engine/Source/Runtime/Experimental/Animation/Constraints/Private/TransformConstraint.cpp
+++ b/Engine/Source/Runtime/Experimental/Animation/Constraints/Private/TransformConstraint.cpp
@@ -391,6 +391,8 @@ void UTickableTransformConstraint::PostLoad()
 
 	SetupDependencies();	
 	RegisterDelegates();
+
+	EnsurePrimaryDependency();
 }
 
 void UTickableTransformConstraint::PostDuplicate(bool bDuplicateForPIE)
