removing geometry collection anchor will no longer cause that cluster particle to become "active" if it's disabled

diff --git a/Engine/Source/Runtime/Experimental/Chaos/Public/Chaos/PBDRigidsSOAs.h b/Engine/Source/Runtime/Experimental/Chaos/Public/Chaos/PBDRigidsSOAs.h
--- a/Engine/Source/Runtime/Experimental/Chaos/Public/Chaos/PBDRigidsSOAs.h
+++ b/Engine/Source/Runtime/Experimental/Chaos/Public/Chaos/PBDRigidsSOAs.h
@@ -662,7 +662,7 @@ public:
 			InsertClusteredParticle(ClusteredParticle);
 		}
 
-		if (ClusteredParticle->ObjectState() != EObjectStateType::Dynamic)
+		if (ClusteredParticle->ObjectState() != EObjectStateType::Dynamic || ClusteredParticle->Disabled())
 		{
 			RemoveFromActiveArray(ClusteredParticle->CastToRigidParticle(), /*bStillDirty=*/true);
 		}
