Chaos - fix post-CCD collision update

Bug fix 1: e.g., a body has 2 CCD contacts with a wall and floor with large penetration on the wall and small penetration on the floor. The wall collision rolls the body back to the time of impact. The floor contact is flagged as "not swept" because the penetration can be resolved without rollback. However its contact points were not updated following the CCD correction against the wall. This could lead to the body passing through the floor.

Bug fix 2: fix calculation of wessp start and end points when there is a center of mass offset. We we using the previous actyor position with the current rotation to calculate the sweep start transform and for interpolating to the TOI, when we should be doing this relative to the Center of Mass.

diff --git a/Engine/Source/Runtime/Experimental/Chaos/Private/Chaos/CCDUtilities.cpp b/Engine/Source/Runtime/Experimental/Chaos/Private/Chaos/CCDUtilities.cpp
--- a/Engine/Source/Runtime/Experimental/Chaos/Private/Chaos/CCDUtilities.cpp
+++ b/Engine/Source/Runtime/Experimental/Chaos/Private/Chaos/CCDUtilities.cpp
@@ -333,12 +333,25 @@ namespace Chaos
 		const int32 ConstraintStart = IslandConstraintStart[Island];
 		const int32 ConstraintNum = IslandConstraintNum[Island];
 		const int32 ConstraintEnd = IslandConstraintEnd[Island];
-		
+		const int32 ParticleStart = IslandParticleStart[Island];
+		const int32 ParticleNum = IslandParticleNum[Island];
+
 		if (ConstraintNum == 0)
 		{
 			return;
 		}
 
+		// We assume that the particle's center of mass moved in a straight line since the previous tick.
+		// Modify X so that X-to-P is a straight line so that we can interpolate between X and P using the TOI
+		// to get the position at that time. This is required to handle objects with a CoM offset from the actor position.
+		// We will undo this manipulation at the end.
+		// NOTE: We do not modify the previous rotation R here - we just use the current rotation Q everywhere
+		// @todo(chaos): we should store the sweep positions in CCDParticle and onlymodify the original particle at the end
+		for (int32 i = ParticleStart; i < ParticleStart + ParticleNum; i++)
+		{
+			GroupedCCDParticles[i]->Particle->X() = GroupedCCDParticles[i]->Particle->P() - GroupedCCDParticles[i]->Particle->V() * Dt;
+		}
+
 #if CHAOS_DEBUG_DRAW
 		if (CVars::ChaosSolverDrawCCDInteractions)
 		{
@@ -355,6 +368,7 @@ namespace Chaos
 		FReal IslandTOI = 0.f;
 		ResetIslandParticles(Island);
 		ResetIslandConstraints(Island);
+
 		bool bSortRequired = true;
 		int32 ConstraintIndex = ConstraintStart;
 		while (ConstraintIndex < ConstraintEnd)
@@ -475,6 +489,15 @@ namespace Chaos
 			}
 		}
 #endif
+
+		// SetX so that the implicit velocity and angular velocity will be calculated correctly in the solver step
+		// NOTE: This is not the same as its original X if we have been rewound to a TOI.
+		for (int32 i = ParticleStart; i < ParticleStart + ParticleNum; i++)
+		{
+			const FVec3 CoMPrev = GroupedCCDParticles[i]->Particle->PCom() - GroupedCCDParticles[i]->Particle->V() * Dt;
+			const FVec3 CoMOffsetPrev = GroupedCCDParticles[i]->Particle->R() * GroupedCCDParticles[i]->Particle->CenterOfMass();
+			GroupedCCDParticles[i]->Particle->X() = CoMPrev - CoMOffsetPrev;
+		}
 	}
 
 	void FCCDManager::ApplyIslandSweptConstraints(const int32 Island, const FReal Dt)
@@ -875,31 +898,39 @@ namespace Chaos
 
 	void FCCDManager::UpdateSweptConstraints(const FReal Dt, FCollisionConstraintAllocator *CollisionAllocator)
 	{
-		// We need to update the world-space contact points at the final locations
-		// @todo(chaos): parallelize this code
-		// @todo(chaos): These SweptConstraints might contain non-CCD particles and those non-CCD particles might collide with other non-CCD particles, which are modeled in normal collision constraints. Those normal collision constraints might need to be updated as well.
-		for (FPBDCollisionConstraint* SweptConstraint : SweptConstraints)
+		// Buld the set of collision whose contact data will be out of date because we moved one or both of its particles. 
+		// This is all collision constraints, including non-swept ones, for any particle that was relocated by the CCD sweep 
+		// logic executed in ApplySweptConstraints (i.e., contents of CCDConstraints)
+		// @todo(chaos): we could calculate the size of the Collisions array in Init
+		// @todo(chaos): could be parallelized
+		TArray<FPBDCollisionConstraint*> Collisions;
+		for (FCCDParticle& CCDParticle : CCDParticles)
 		{
-			if (!SweptConstraint->IsEnabled())
-			{
-				continue;
-			}
-
-			if (!SweptConstraint->GetCCDSweepEnabled())
-			{
-				continue;
-			}
+			CCDParticle.Particle->ParticleCollisions().VisitCollisions(
+				[&Collisions](FPBDCollisionConstraint& Collision)
+				{
+					// Avoid duplicates when both particles in the collision are CCD enabled by checking the particle ID
+					const FConstGenericParticleHandle P0 = FConstGenericParticleHandle(Collision.GetParticle0());
+					const FConstGenericParticleHandle P1 = FConstGenericParticleHandle(Collision.GetParticle1());
+					if (!P0->CCDEnabled() || !P1->CCDEnabled() || (P0->ParticleID() < P1->ParticleID()))
+					{
+						Collisions.Add(&Collision);
+					}
+					return ECollisionVisitorResult::Continue;
+				});
+		}
 
-			const FConstGenericParticleHandle P0 = FConstGenericParticleHandle(SweptConstraint->GetParticle0());
-			const FConstGenericParticleHandle P1 = FConstGenericParticleHandle(SweptConstraint->GetParticle1());
-			SweptConstraint->ResetManifold();
-			Collisions::UpdateConstraintFromGeometry<ECollisionUpdateType::Deepest>(*SweptConstraint, FRigidTransform3(P0->P(), P0->Q()), FRigidTransform3(P1->P(), P1->Q()), Dt);
+		// Update all the collisions
+		// @todo(chaos): can be parallelized
+		for (FPBDCollisionConstraint* Collision : Collisions)
+		{
+			const FConstGenericParticleHandle P0 = FConstGenericParticleHandle(Collision->GetParticle0());
+			const FConstGenericParticleHandle P1 = FConstGenericParticleHandle(Collision->GetParticle1());
 
-			// @todo(zhenglin): Removing constraints that has Phi larger than CullDistance could reduce the island sizes in the normal solve. But I could not get this to work...
-			// if (SweptConstraint->GetPhi() > SweptConstraint->GetCullDistance())
-			// {
-			//     CollisionAllocator->RemoveConstraintSwap(SweptConstraint);
-			// }
+			// NOTE: ResetManifold also reset friction anchors. If CCD sweep was run, static friction probably will not hold
+			// We could potentially call ResetActiveManifold here instead and then AssignSavedManifoldPoints if we want static friction
+			Collision->ResetManifold();
+			Collisions::UpdateConstraintFromGeometry<ECollisionUpdateType::Deepest>(*Collision, P0->GetTransformPQ(), P1->GetTransformPQ(), Dt);
 		}
 	}
 
diff --git a/Engine/Source/Runtime/Experimental/Chaos/Private/Chaos/Collision/ParticlePairMidPhase.cpp b/Engine/Source/Runtime/Experimental/Chaos/Private/Chaos/Collision/ParticlePairMidPhase.cpp
--- a/Engine/Source/Runtime/Experimental/Chaos/Private/Chaos/Collision/ParticlePairMidPhase.cpp
+++ b/Engine/Source/Runtime/Experimental/Chaos/Private/Chaos/Collision/ParticlePairMidPhase.cpp
@@ -410,14 +410,16 @@ namespace Chaos
 
 			FConstGenericParticleHandle P0 = GetParticle0();
 			FConstGenericParticleHandle P1 = GetParticle1();
-			// For kinematic particles, X = P (at TOI=1), we need to compute P-V*dt to get position at TOI=0. 
-			const FVec3 StartX0 = P0->ObjectState() == EObjectStateType::Kinematic ? P0->P() - P0->V() * Dt : P0->X();
-			const FVec3 StartX1 = P1->ObjectState() == EObjectStateType::Kinematic ? P1->P() - P1->V() * Dt : P1->X();
-			// Note: It is unusual that we are mixing X and Q. 
-			// This is due to how CCD rewinds the position (not rotation) and then sweeps to find the first contact at the most recent orientation Q
+
+
+			// We need the previous transform for the swept collision detector. It assumes that the current
+			// transform has been set on the constraint. 
+			// We assume that the particle's center of mass moved in a stright line and that it's rotation has 
+			// not changed so we calculate the previous transform from the current one and the velocity.
 			// NOTE: These are actor transforms, not CoM transforms
-			const FRigidTransform3 CCDParticleWorldTransform0 = FRigidTransform3(StartX0, P0->Q());
-			const FRigidTransform3 CCDParticleWorldTransform1 = FRigidTransform3(StartX1, P1->Q());
+			// @todo(chaos): Pass both start and end transforms to the collision detector
+			const FRigidTransform3 CCDParticleWorldTransform0 = FRigidTransform3(P0->P() - P0->V() * Dt, P0->Q());
+			const FRigidTransform3 CCDParticleWorldTransform1 = FRigidTransform3(P1->P() - P1->V() * Dt, P1->Q());
 			const FRigidTransform3 CCDShapeWorldTransform0 = Constraint->ImplicitTransform[0] * CCDParticleWorldTransform0;
 			const FRigidTransform3 CCDShapeWorldTransform1 = Constraint->ImplicitTransform[1] * CCDParticleWorldTransform1;
 			const bool bDidSweep = Collisions::UpdateConstraintSwept(*Constraint.Get(), CCDShapeWorldTransform0, CCDShapeWorldTransform1, Dt);
