Make geometry collection proxy initialization happen on the physics thread using the dirty proxy functionality like the other proxies.

This will allow for constraints to grab the geometry collection proxy's physics particle handles because otherwise the initialization order will cause the constraint proxy to be initialized before the geometry collection proxy causing the constraint to see nullptrs.

diff --git a/Engine/Source/Runtime/Experimental/Chaos/Private/PBDRigidsSolver.cpp b/Engine/Source/Runtime/Experimental/Chaos/Private/PBDRigidsSolver.cpp
--- a/Engine/Source/Runtime/Experimental/Chaos/Private/PBDRigidsSolver.cpp
+++ b/Engine/Source/Runtime/Experimental/Chaos/Private/PBDRigidsSolver.cpp
@@ -641,17 +641,13 @@ namespace Chaos
 		InProxy->SetSolver(this);
 		InProxy->Initialize(GetEvolution());
 		InProxy->NewData(); // Buffers data on the proxy.
-		FParticlesType* InParticles = &GetParticles();
 
-		// Finish registration on the physics thread...
-		EnqueueCommandImmediate([InParticles, InProxy, this]()
-		{
-			UE_LOG(LogPBDRigidsSolver, Verbose, 
-				TEXT("FPBDRigidsSolver::RegisterObject(FGeometryCollectionPhysicsProxy*)"));
-			check(InParticles);
-			InProxy->InitializeBodiesPT(this, *InParticles);
-			GeometryCollectionPhysicsProxies_Internal.Add(InProxy);
-		});
+		// There used to be an EnqueueCommandImmediate here to push the initialization of the GC
+		// onto the physics thread. We should use AddDirtyProxy here instead to better match up with
+		// the initialization order done in the ProcessSinglePushedData_Internal. Using EnqueueCommandImmediate
+		// will cause the GC to be initialized after constraints which would preclude it from being
+		// used with joint constraints that get spawned on level load with the GC.
+		AddDirtyProxy(InProxy);
 	}
 	
 	void FPBDRigidsSolver::UnregisterObject(FGeometryCollectionPhysicsProxy* InProxy)
@@ -1183,6 +1179,9 @@ namespace Chaos
 					case EPhysicsProxyType::GeometryCollectionType:
 					{
 						auto Proxy = static_cast<FGeometryCollectionPhysicsProxy*>(Dirty.Proxy);
+						// Finish registration on the physics thread...
+						Proxy->InitializeBodiesPT(this, GetParticles());
+						GeometryCollectionPhysicsProxies_Internal.Add(Proxy);
 						Proxy->PushToPhysicsState();
 						// Currently no push needed for geometry collections and they handle the particle creation internally
 						// #TODO This skips the rewind data push so GC will not be rewindable until resolved.
