Chaos Cloth:
- Added Property Collection Adapter class to read and write simulation properties from a Managed Array Collection.
- Added first pass implementation of the bending constraint using the new Property Collection Adapter.
- Add clamping to Buckling Ratio parameter.

#rnx

diff --git a/Engine/Source/Runtime/Experimental/Chaos/Private/Chaos/PropertyCollectionAdapter.cpp b/Engine/Source/Runtime/Experimental/Chaos/Private/Chaos/PropertyCollectionAdapter.cpp
new file mode 100644
--- /dev/null
+++ b/Engine/Source/Runtime/Experimental/Chaos/Private/Chaos/PropertyCollectionAdapter.cpp
@@ -0,0 +1,163 @@
+// Copyright Epic Games, Inc. All Rights Reserved.
+
+#include "Chaos/PropertyCollectionAdapter.h"
+#include "GeometryCollection/ManagedArrayCollection.h"
+
+namespace Chaos::Softs
+{
+	const FName FPropertyCollectionConstAdapter::PropertyGroup("Property");
+	const FName FPropertyCollectionConstAdapter::KeyName("Key");
+	const FName FPropertyCollectionConstAdapter::LowValueName("LowValue");
+	const FName FPropertyCollectionConstAdapter::HighValueName("HighValue");
+	const FName FPropertyCollectionConstAdapter::StringValueName("StringValue");
+	const FName FPropertyCollectionConstAdapter::FlagsName("Flags");
+	
+	FPropertyCollectionConstAdapter::FPropertyCollectionConstAdapter(const TSharedPtr<const FManagedArrayCollection>& InManagedArrayCollection)
+		: ManagedArrayCollection(InManagedArrayCollection)
+	{
+		Initialize();
+	}
+
+	FPropertyCollectionConstAdapter::FPropertyCollectionConstAdapter(const TSharedPtr<const FManagedArrayCollection>& InManagedArrayCollection, ENoInit)
+		: ManagedArrayCollection(InManagedArrayCollection)
+	{
+	}
+
+	void FPropertyCollectionConstAdapter::UpdateArrays()
+	{
+		KeyArray = GetArray<FString>(KeyName);
+		LowValueArray = GetArray<FVector3f>(LowValueName);
+		HighValueArray = GetArray<FVector3f>(HighValueName);
+		StringValueArray = GetArray<FString>(StringValueName);
+		FlagsArray = GetArray<uint8>(FlagsName);
+	}
+	
+	void FPropertyCollectionConstAdapter::Initialize()
+	{
+		// Update arrays
+		UpdateArrays();
+
+		// Create a fast access search map (although it might only be faster for a large enough number of properties)
+		const int32 NumKeys = KeyArray.Num();
+		KeyIndices.Empty(NumKeys);
+		for (int32 Index = 0; Index < NumKeys; ++Index)
+		{
+			KeyIndices.Emplace(KeyArray[Index], Index);
+		}
+	}
+
+	template<typename T, typename ElementType>
+	T FPropertyCollectionConstAdapter::GetValue(int32 KeyIndex, const TConstArrayView<ElementType>& ValueArray) const
+	{
+		return ValueArray[KeyIndex];
+	}
+
+	template<>
+	bool FPropertyCollectionConstAdapter::GetValue<bool, FVector3f>(int32 KeyIndex, const TConstArrayView<FVector3f>& ValueArray) const
+	{
+		return (bool)ValueArray[KeyIndex].X;
+	}
+
+	template<>
+	int32 FPropertyCollectionConstAdapter::GetValue<int32, FVector3f>(int32 KeyIndex, const TConstArrayView<FVector3f>& ValueArray) const
+	{
+		return (int32)ValueArray[KeyIndex].X;
+	}
+
+	template<>
+	float FPropertyCollectionConstAdapter::GetValue<float, FVector3f>(int32 KeyIndex, const TConstArrayView<FVector3f>& ValueArray) const
+	{
+		return ValueArray[KeyIndex].X;
+	}
+
+	template <typename T>
+	TConstArrayView<T> FPropertyCollectionConstAdapter::GetArray(const FName& Name) const
+	{
+		const TManagedArray<T>* ManagedArray = ManagedArrayCollection->FindAttributeTyped<T>(Name, PropertyGroup);
+		return ManagedArray ? TConstArrayView<T>(ManagedArray->GetConstArray()) : TConstArrayView<T>();
+	}
+
+	FPropertyCollectionAdapter::FPropertyCollectionAdapter(const TSharedPtr<FManagedArrayCollection>& InManagedArrayCollection)
+		: FPropertyCollectionConstAdapter(InManagedArrayCollection, NoInit)
+	{
+		Construct();
+		Initialize();
+	}
+
+	void FPropertyCollectionAdapter::Construct()
+	{
+		// ClothProperty Group
+		GetManagedArrayCollection()->AddAttribute<FString>(KeyName, PropertyGroup);
+		GetManagedArrayCollection()->AddAttribute<FVector3f>(LowValueName, PropertyGroup);
+		GetManagedArrayCollection()->AddAttribute<FVector3f>(HighValueName, PropertyGroup);
+		GetManagedArrayCollection()->AddAttribute<FString>(StringValueName, PropertyGroup);
+		GetManagedArrayCollection()->AddAttribute<uint8>(FlagsName, PropertyGroup);
+	}
+
+	void FPropertyCollectionAdapter::EnableFlag(int32 KeyIndex, EPropertyFlag Flag, bool bEnable)
+	{
+		if (bEnable)
+		{
+			EnumAddFlags(GetFlagsArray()[KeyIndex], (uint8)Flag);
+		}
+		else
+		{
+			EnumRemoveFlags(GetFlagsArray()[KeyIndex], (uint8)Flag);
+		}
+	}
+
+	int32 FPropertyCollectionAdapter::EnableFlag(const FString& Key, EPropertyFlag Flag, bool bEnable)
+	{
+		const int32 KeyIndex = GetKeyIndex(Key);
+		if (KeyIndex != INDEX_NONE)
+		{
+			EnableFlag(KeyIndex, Flag, bEnable);
+		}
+		return KeyIndex;
+	}
+
+	int32 FPropertyCollectionAdapter::AddProperty(const FString& Key, bool bEnabled, bool bAnimatable)
+	{
+		const int32 Index = GetManagedArrayCollection()->AddElements(1, PropertyGroup);
+		const uint8 Flags = (bEnabled ? (uint8)EPropertyFlag::Enabled : 0) | (bAnimatable ? (uint8)EPropertyFlag::Animatable : 0);
+
+		// Setup the new element's default value and enable the property by default
+		GetKeyArray()[Index] = Key;
+		GetLowValueArray()[Index] = GetHighValueArray()[Index] = FVector3f::ZeroVector;
+		GetFlagsArray()[Index] = Flags;
+
+		// Update search map
+		KeyIndices.Emplace(KeyArray[Index], Index);
+
+		// Update the array view count
+		UpdateArrays();
+
+		return Index;
+	}
+
+	int32 FPropertyCollectionAdapter::AddProperties(const TArray<FString>& Keys, bool bEnabled, bool bAnimatable)
+	{
+		if (const int32 NumProperties = Keys.Num())
+		{
+			const int32 StartIndex = GetManagedArrayCollection()->AddElements(NumProperties, PropertyGroup);
+			const uint8 Flags = (bEnabled ? (uint8)EPropertyFlag::Enabled : 0) | (bAnimatable ? (uint8)EPropertyFlag::Animatable : 0);
+
+			for (int32 Index = StartIndex; Index < NumProperties + StartIndex; ++Index)
+			{
+				// Setup the new elements' default value and enable the property by default
+				GetKeyArray()[Index] = Keys[Index - StartIndex];
+				GetLowValueArray()[Index] = GetHighValueArray()[Index] = FVector3f::ZeroVector;
+				GetFlagsArray()[Index] = Flags;
+
+				// Update search map
+				KeyIndices.Emplace(KeyArray[Index], Index);
+			}
+
+			// Update the array view count
+			UpdateArrays();
+
+			return StartIndex;
+		}
+		return INDEX_NONE;
+	}
+}  // End namespace Chaos
diff --git a/Engine/Source/Runtime/Experimental/Chaos/Public/Chaos/PBDBendingConstraintsBase.h b/Engine/Source/Runtime/Experimental/Chaos/Public/Chaos/PBDBendingConstraintsBase.h
--- a/Engine/Source/Runtime/Experimental/Chaos/Public/Chaos/PBDBendingConstraintsBase.h
+++ b/Engine/Source/Runtime/Experimental/Chaos/Public/Chaos/PBDBendingConstraintsBase.h
@@ -5,6 +5,7 @@
 #include "Chaos/PBDSoftsSolverParticles.h"
 #include "Chaos/PBDStiffness.h"
 #include "Chaos/ParticleRule.h"
+#include "Chaos/PropertyCollectionAdapter.h"
 #include "Containers/StaticArray.h"
 
 namespace Chaos::Softs
@@ -13,6 +14,43 @@ namespace Chaos::Softs
 class FPBDBendingConstraintsBase
 {
 public:
+	static bool IsEnabled(const FPropertyCollectionConstAdapter& PropertyCollection)
+	{
+		return IsBendingStiffnessEnabled(PropertyCollection, false);
+	}
+
+	FPBDBendingConstraintsBase(
+		const FSolverParticles& InParticles,
+		int32 ParticleOffset,
+		int32 ParticleCount,
+		TArray<TVec4<int32>>&& InConstraints,
+		const TConstArrayView<FRealSingle>& StiffnessMultipliers,
+		const TConstArrayView<FRealSingle>& BucklingStiffnessMultipliers,
+		const FPropertyCollectionConstAdapter& PropertyCollection,
+		bool bTrimKinematicConstraints = false)
+		: Constraints(bTrimKinematicConstraints ? TrimKinematicConstraints(InConstraints, InParticles): MoveTemp(InConstraints))
+		, ConstraintSharedEdges(ExtractConstraintSharedEdges(Constraints))
+		, Stiffness(FSolverVec2(GetWeightedFloatBendingStiffness(PropertyCollection, 1.f)).ClampAxes((FSolverReal)0., (FSolverReal)1.),
+			StiffnessMultipliers,
+			TConstArrayView<TVec2<int32>>(ConstraintSharedEdges),
+			ParticleOffset,
+			ParticleCount)
+		, BucklingRatio(FMath::Clamp(GetBucklingRatio(PropertyCollection, 0.f), (FSolverReal)0., (FSolverReal)1.))
+		, BucklingStiffness(FSolverVec2(GetWeightedFloatBucklingStiffness(PropertyCollection, 1.f)).ClampAxes((FSolverReal)0., (FSolverReal)1.),
+			BucklingStiffnessMultipliers,
+			TConstArrayView<TVec2<int32>>(ConstraintSharedEdges),
+			ParticleOffset,
+			ParticleCount)
+	{
+		for (const TVec4<int32>& Constraint : Constraints)
+		{
+			const FSolverVec3& P1 = InParticles.X(Constraint[0]);
+			const FSolverVec3& P2 = InParticles.X(Constraint[1]);
+			const FSolverVec3& P3 = InParticles.X(Constraint[2]);
+			const FSolverVec3& P4 = InParticles.X(Constraint[3]);
+			RestAngles.Add(CalcAngle(P1, P2, P3, P4));
+		}
+	}
 
 	FPBDBendingConstraintsBase(const FSolverParticles& InParticles,
 		int32 ParticleOffset,
@@ -37,7 +75,7 @@ public:
 			const FSolverVec3& P3 = InParticles.X(Constraint[2]);
 			const FSolverVec3& P4 = InParticles.X(Constraint[3]);
 			RestAngles.Add(CalcAngle(P1, P2, P3, P4));
-		}	
+		}
 	}
 
 	FPBDBendingConstraintsBase(const FSolverParticles& InParticles, TArray<TVec4<int32>>&& InConstraints, const FSolverReal InStiffness = (FSolverReal)1.)
@@ -59,11 +97,28 @@ public:
 
 	virtual ~FPBDBendingConstraintsBase() {}
 
+	// Set properties using a property collection
+	void SetProperties(const FPropertyCollectionConstAdapter& PropertyCollection)
+	{
+		if (IsBendingStiffnessMutable(PropertyCollection))
+		{
+			Stiffness.SetWeightedValue(FSolverVec2(GetWeightedFloatBendingStiffness(PropertyCollection)).ClampAxes((FSolverReal)0., (FSolverReal)1.));
+		}
+		if (IsBucklingRatioMutable(PropertyCollection))
+		{
+			BucklingRatio = FMath::Clamp((FSolverReal)GetBucklingRatio(PropertyCollection), (FSolverReal)0., (FSolverReal)1.);
+		}
+		if (IsBucklingStiffnessMutable(PropertyCollection))
+		{
+			BucklingStiffness.SetWeightedValue(FSolverVec2(GetWeightedFloatBucklingStiffness(PropertyCollection)).ClampAxes((FSolverReal)0., (FSolverReal)1.));
+		}
+	}
+
 	// Update stiffness values
 	void SetProperties(const FSolverVec2& InStiffness, const FSolverReal InBucklingRatio, const FSolverVec2& InBucklingStiffness)
 	{ 
-		Stiffness.SetWeightedValue(InStiffness.ClampAxes((FSolverReal)0., (FSolverReal)1.)); 
-		BucklingRatio = InBucklingRatio;
+		Stiffness.SetWeightedValue(InStiffness.ClampAxes((FSolverReal)0., (FSolverReal)1.));
+		BucklingRatio = FMath::Clamp(InBucklingRatio, (FSolverReal)0., (FSolverReal)1.);
 		BucklingStiffness.SetWeightedValue(InBucklingStiffness.ClampAxes((FSolverReal)0., (FSolverReal)1.));
 	}
 
@@ -242,6 +297,11 @@ protected:
 
 	TArray<FSolverReal> RestAngles;
 	TArray<bool> IsBuckled;
+
+private:
+	UE_CHAOS_DECLARE_PROPERTYCOLLECTION_NAME(BendingStiffness, float);
+	UE_CHAOS_DECLARE_PROPERTYCOLLECTION_NAME(BucklingRatio, float);
+	UE_CHAOS_DECLARE_PROPERTYCOLLECTION_NAME(BucklingStiffness, float);
 };
 
 }  // End namespace Chaos::Softs
diff --git a/Engine/Source/Runtime/Experimental/Chaos/Public/Chaos/PropertyCollectionAdapter.h b/Engine/Source/Runtime/Experimental/Chaos/Public/Chaos/PropertyCollectionAdapter.h
new file mode 100644
--- /dev/null
+++ b/Engine/Source/Runtime/Experimental/Chaos/Public/Chaos/PropertyCollectionAdapter.h
@@ -0,0 +1,404 @@
+// Copyright Epic Games, Inc. All Rights Reserved.
+
+#pragma once
+
+#include "UObject/NameTypes.h"
+#include "Containers/Map.h"
+#include "Math/Vector.h"
+
+struct FManagedArrayCollection;
+
+namespace Chaos::Softs
+{
+	/** Property flags, whether properties are enabled, animatable, ...etc. */
+	enum class EPropertyFlag : uint8
+	{
+		None,
+		Enabled = 1 << 0,  /** Whether this property is enabled(so that it doesn't have to be removed from the collection when not needed). */
+		Animatable = 1 << 1,  /** Whether this property needs to be set at every frame. */
+		//~ Add new flags above this line
+		Dirty = 1 << 7  /** Whether this property has changed and needs to be updated at the next frame. */
+	};
+	ENUM_CLASS_FLAGS(EPropertyFlag)
+
+	/** Weighted types are all property types that can have a pair of low and high values to be associated with a weight map. */
+	template<typename T> struct TIsWeightedType { static constexpr bool Value = false; };
+	template<> struct TIsWeightedType<bool> { static constexpr bool Value = true; };
+	template<> struct TIsWeightedType<int32> { static constexpr bool Value = true; };
+	template<> struct TIsWeightedType<float> { static constexpr bool Value = true; };
+	template<> struct TIsWeightedType<FVector3f> { static constexpr bool Value = true; };
+
+	/**
+	 * Property collection helper class to access simulation properties from a ManagedArrayCollection.
+	 */
+	class FPropertyCollectionConstAdapter
+	{
+	public:
+		explicit FPropertyCollectionConstAdapter(const TSharedPtr<const FManagedArrayCollection>& InManagedArrayCollection);
+		virtual ~FPropertyCollectionConstAdapter() = default;
+
+		/** return the collection used by this adapter as a property collection. */
+		const TSharedPtr<const FManagedArrayCollection>& GetManagedArrayCollection() const { return ManagedArrayCollection; }
+
+		/** Return the number of cloth properties in this collection. */
+		int32 Num() const { return KeyIndices.Num(); }
+
+		/** Return the property index for the specified key if it exists, or INDEX_NONE otherwise. */
+		int32 GetKeyIndex(const FString& Key) const { const int32* const Index = KeyIndices.Find(Key); return Index ? *Index : INDEX_NONE; }
+
+		//~ Values access per index, fast, no check, index must be valid (0 <= KeyIndex < Num())
+		template<typename T, TEMPLATE_REQUIRES(TIsWeightedType<T>::Value)>
+		T GetLowValue(int32 KeyIndex) const { return GetValue<T>(KeyIndex, LowValueArray); }
+
+		template<typename T, TEMPLATE_REQUIRES(TIsWeightedType<T>::Value)>
+		T GetHighValue(int32 KeyIndex) const { return GetValue<T>(KeyIndex, HighValueArray); }
+
+		template<typename T, TEMPLATE_REQUIRES(TIsWeightedType<T>::Value)>
+		TPair<T, T> GetWeightedValue(int32 KeyIndex) const { return MakeTuple(GetLowValue<T>(KeyIndex), GetHighValue<T>(KeyIndex)); }
+
+		FVector2f GetWeightedFloatValue(int32 KeyIndex) const { return FVector2f(GetLowValue<float>(KeyIndex), GetHighValue<float>(KeyIndex)); }
+
+		template<typename T, TEMPLATE_REQUIRES(TIsWeightedType<T>::Value)>
+		T GetValue(int32 KeyIndex) const { return GetLowValue<T>(KeyIndex); }
+
+		const FString& GetStringValue(int32 KeyIndex, const FString& Default = "") const { return GetValue<const FString&>(KeyIndex, StringValueArray); }
+
+		uint8 GetFlags(int32 KeyIndex) const { return FlagsArray[KeyIndex]; }
+		bool IsEnabled(int32 KeyIndex) const { return EnumHasAnyFlags((EPropertyFlag)FlagsArray[KeyIndex], EPropertyFlag::Enabled); }
+		bool IsAnimatable(int32 KeyIndex) const { return EnumHasAnyFlags((EPropertyFlag)FlagsArray[KeyIndex], EPropertyFlag::Animatable); }
+		bool IsDirty(int32 KeyIndex) const { return EnumHasAnyFlags((EPropertyFlag)FlagsArray[KeyIndex], EPropertyFlag::Dirty); }
+
+		//~ Values access per key
+		template<typename T, TEMPLATE_REQUIRES(TIsWeightedType<T>::Value)>
+		T GetLowValue(const FString& Key, const T& Default = T(0), int32* OutKeyIndex = nullptr) const { return GetValue(Key, LowValueArray, Default, OutKeyIndex); }
+
+		template<typename T, TEMPLATE_REQUIRES(TIsWeightedType<T>::Value)>
+		T GetHighValue(const FString& Key, const T& Default = T(0), int32* OutKeyIndex = nullptr) const { return GetValue(Key, HighValueArray, Default, OutKeyIndex); }
+
+		template<typename T, TEMPLATE_REQUIRES(TIsWeightedType<T>::Value)>
+		inline TPair<T, T> GetWeightedValue(const FString& Key, const T& Default = T(0), int32* OutKeyIndex = nullptr) const;
+
+		inline FVector2f GetWeightedFloatValue(const FString& Key, const float& Default = 0.f, int32* OutKeyIndex = nullptr) const;
+
+		template<typename T, TEMPLATE_REQUIRES(TIsWeightedType<T>::Value)>
+		T GetValue(const FString& Key, const T& Default = T(0), int32* OutKeyIndex = nullptr) const { return GetValue(Key, LowValueArray, Default, OutKeyIndex); }
+
+		FString GetStringValue(const FString& Key, const FString& Default = "", int32* OutKeyIndex = nullptr) const { return GetValue(Key, StringValueArray, Default, OutKeyIndex); }
+
+		uint8 GetFlags(const FString& Key, uint8 Default = 0, int32* OutKeyIndex = nullptr) const { return GetValue(Key, FlagsArray, (uint8)Default, OutKeyIndex); }
+		bool IsEnabled(const FString& Key, bool bDefault = false, int32* OutKeyIndex = nullptr) const { return EnumHasAnyFlags((EPropertyFlag)GetValue(Key, FlagsArray, (uint8)bDefault, OutKeyIndex), EPropertyFlag::Enabled); }
+		bool IsAnimatable(const FString& Key, bool bDefault = false, int32* OutKeyIndex = nullptr) const { return EnumHasAnyFlags((EPropertyFlag)GetValue(Key, FlagsArray, (uint8)bDefault, OutKeyIndex), EPropertyFlag::Animatable); }
+		bool IsDirty(const FString& Key, bool bDefault = false, int32* OutKeyIndex = nullptr) const { return EnumHasAnyFlags((EPropertyFlag)GetValue(Key, FlagsArray, (uint8)bDefault, OutKeyIndex), EPropertyFlag::Dirty); }
+
+	protected:
+		// No init constructor for FPropertyCollectionAdapter
+		FPropertyCollectionConstAdapter(const TSharedPtr<const FManagedArrayCollection>& InManagedArrayCollection, ENoInit);
+
+		// Update the array views
+		void UpdateArrays();
+
+		// Initialize the array views and the search map
+		void Initialize();
+
+		template<typename T, typename ElementType>
+		T GetValue(int32 KeyIndex, const TConstArrayView<ElementType>& ValueArray) const;
+
+		template<typename T, typename ElementType>
+		inline T GetValue(const FString& Key, const TConstArrayView<ElementType>& ValueArray, const T& Default, int32* OutKeyIndex) const;
+
+		template <typename T>
+		TConstArrayView<T> GetArray(const FName& Name) const;
+
+		// Attribute groups, predefined data member of the collection
+		static const FName PropertyGroup;
+		static const FName KeyName;  // Property key, name to look for
+		static const FName LowValueName;  // Boolean, 24 bit integer (max 16777215), float, or vector value, or value of the lowest weight on the weight map if any
+		static const FName HighValueName;  // Boolean, 24 bit integer (max 16777215), float, or vector value of the highest weight on the weight map if any
+		static const FName StringValueName;  // String value, or weight map name, ...etc.
+		static const FName FlagsName;  // Whether this property is enabled, animatable, ...etc.
+
+		// Property Group array views
+		TConstArrayView<FString> KeyArray;
+		TConstArrayView<FVector3f> LowValueArray;
+		TConstArrayView<FVector3f> HighValueArray;
+		TConstArrayView<FString> StringValueArray;
+		TConstArrayView<uint8> FlagsArray;
+
+		// Key to index search map
+		TMap<FString, int32> KeyIndices;
+
+		// Property collection
+		TSharedPtr<const FManagedArrayCollection> ManagedArrayCollection;
+	};
+
+	/**
+	 * Property collection adapter class to use a ManagedArrayCollection property collection and add simulation properties.
+	 * Note: Int properties are limited to 24 bits (-16777215 to 16777215).
+	 */
+	class FPropertyCollectionAdapter final : public FPropertyCollectionConstAdapter
+	{
+	public:
+		explicit FPropertyCollectionAdapter(const TSharedPtr<FManagedArrayCollection>& ManagedArrayCollection);
+		virtual ~FPropertyCollectionAdapter() = default;
+
+		/** Return the collection used by this adapter as a property collection. */
+		TSharedPtr<FManagedArrayCollection> GetManagedArrayCollection() { return ConstCastSharedPtr<FManagedArrayCollection>(ManagedArrayCollection); }
+
+		/** Add a single property, and return its index. */
+		int32 AddProperty(const FString& Key, bool bEnabled = true, bool bAnimatable = false);
+
+		/** Add new properties, and return the index of the first added property. */
+		int32 AddProperties(const TArray<FString>& Keys, bool bEnabled = true, bool bAnimatable = false);
+
+		//~ Values set per index
+		template<typename T, TEMPLATE_REQUIRES(TIsWeightedType<T>::Value)>
+		void SetLowValue(int32 KeyIndex, const T& Value) { GetLowValueArray()[KeyIndex] = FVector3f(Value); }
+
+		template<typename T, TEMPLATE_REQUIRES(TIsWeightedType<T>::Value)>
+		void SetHighValue(int32 KeyIndex, const T& Value) { GetHighValueArray()[KeyIndex] = FVector3f(Value); }
+
+		template<typename T, TEMPLATE_REQUIRES(TIsWeightedType<T>::Value)>
+		void SetWeightedValue(int32 KeyIndex, const T& LowValue, const T& HighValue) { SetLowValue(KeyIndex, LowValue); SetHighValue(KeyIndex, HighValue); }
+
+		void SetWeightedFloatValue(int32 KeyIndex, const FVector2f& Value) { SetLowValue<float>(KeyIndex, Value.X); SetHighValue<float>(KeyIndex, Value.Y); }
+
+		template<typename T, TEMPLATE_REQUIRES(TIsWeightedType<T>::Value)>
+		void SetValue(int32 KeyIndex, const T& Value) { SetWeightedValue(KeyIndex, Value, Value); }
+
+		void SetStringValue(int32 KeyIndex, const FString& Value) { GetStringValueArray()[KeyIndex] = Value; }
+		
+		void SetFlags(int32 KeyIndex, uint8 Flags) { GetFlagsArray()[KeyIndex] = Flags; }
+		void SetEnabled(int32 KeyIndex, bool bEnabled) { EnableFlag(KeyIndex, EPropertyFlag::Enabled, bEnabled); }
+		void SetAnimatable(int32 KeyIndex, bool bAnimatable) { EnableFlag(KeyIndex, EPropertyFlag::Animatable, bAnimatable); }
+		void SetDirty(int32 KeyIndex, bool bDirty) { EnableFlag(KeyIndex, EPropertyFlag::Animatable, bDirty); }
+
+		//~ Values set per key
+		template<typename T, TEMPLATE_REQUIRES(TIsWeightedType<T>::Value)>
+		int32 SetLowValue(const FString& Key, const T& Value) { return SetValue(Key, GetLowValueArray(), FVector3f(Value)); }
+
+		template<typename T, TEMPLATE_REQUIRES(TIsWeightedType<T>::Value)>
+		int32 SetHighValue(const FString& Key, const T& Value) { return SetValue(Key, GetHighValueArray(), FVector3f(Value)); }
+
+		template<typename T, TEMPLATE_REQUIRES(TIsWeightedType<T>::Value)>
+		inline int32 SetWeightedValue(const FString& Key, const T& LowValue, const T& HighValue);
+
+		inline int32 SetWeightedFloatValue(const FString& Key, const FVector2f& Value);
+
+		template<typename T, TEMPLATE_REQUIRES(TIsWeightedType<T>::Value)>
+		int32 SetValue(const FString& Key, const T& Value) { return SetWeightedValue(Key, Value, Value); }
+
+		int32 SetStringValue(const FString& Key, const FString& Value) { return SetValue(Key, GetStringValueArray(), Value); }
+		
+		int32 SetFlags(const FString& Key, uint8 Flags) { return SetValue(Key, GetFlagsArray(), Flags); }
+		int32 SetEnabled(const FString& Key, bool bEnabled) { return EnableFlag(Key, EPropertyFlag::Enabled, bEnabled); }
+		int32 SetAnimatable(const FString& Key, bool bAnimatable) { return EnableFlag(Key, EPropertyFlag::Animatable, bAnimatable); }
+		int32 SetDirty(const FString& Key, bool bDirty) { return EnableFlag(Key, EPropertyFlag::Animatable, bDirty); }
+
+		//~ Add values
+		template<typename T, TEMPLATE_REQUIRES(TIsWeightedType<T>::Value)>
+		inline int32 AddWeightedValue(const FString& Key, const T& LowValue, const T& HighValue, bool bEnabled = true, bool bAnimatable = false);
+
+		inline int32 AddWeightedFloatValue(const FString& Key, const FVector2f& Value, bool bEnabled, bool bAnimatable);
+
+		template<typename T, TEMPLATE_REQUIRES(TIsWeightedType<T>::Value)>
+		int32 AddValue(const FString& Key, const T& Value, bool bEnabled = true, bool bAnimatable = false) { return AddWeightedValue(Value, Value, bEnabled, bAnimatable); }
+
+		inline int32 AddStringValue(const FString& Key, const FString& Value, bool bEnabled = true, bool bAnimatable = false);
+		
+	private:
+		// Add the property group to the target collection
+		void Construct();
+
+		template<typename T>
+		inline int32 SetValue(const FString& Key, const TArrayView<T>& ValueArray, const T& Value);
+
+		void EnableFlag(int32 KeyIndex, EPropertyFlag Flag, bool bEnable);
+		int32 EnableFlag(const FString& Key, EPropertyFlag Flag, bool bEnable);
+
+		const TArrayView<FString>& GetKeyArray() { return reinterpret_cast<TArrayView<FString>&>(KeyArray); }
+		const TArrayView<FVector3f>& GetLowValueArray() { return reinterpret_cast<TArrayView<FVector3f>&>(LowValueArray); }
+		const TArrayView<FVector3f>& GetHighValueArray() { return reinterpret_cast<TArrayView<FVector3f>&>(HighValueArray); }
+		const TArrayView<FString>& GetStringValueArray() { return reinterpret_cast<TArrayView<FString>&>(StringValueArray); }
+		const TArrayView<uint8>& GetFlagsArray() { return reinterpret_cast<const TArrayView<uint8>&>(FlagsArray); }
+	};
+
+	template<typename T, typename TEnableIf<TIsWeightedType<T>::Value, int>::type>
+	inline TPair<T, T> FPropertyCollectionConstAdapter::GetWeightedValue(const FString& Key, const T& Default, int32* OutKeyIndex) const
+	{
+		const int32 KeyIndex = GetKeyIndex(Key);
+		if (OutKeyIndex)
+		{
+			*OutKeyIndex = KeyIndex;
+		}
+		return KeyIndex != INDEX_NONE ? GetWeightedValue<T>(KeyIndex) : MakeTuple(Default, Default);
+	}
+
+	inline FVector2f FPropertyCollectionConstAdapter::GetWeightedFloatValue(const FString& Key, const float& Default, int32* OutKeyIndex) const
+	{
+		const int32 KeyIndex = GetKeyIndex(Key);
+		if (OutKeyIndex)
+		{
+			*OutKeyIndex = KeyIndex;
+		}
+		return KeyIndex != INDEX_NONE ? GetWeightedFloatValue(KeyIndex) : FVector2f(Default, Default);
+	}
+
+	template<typename T, typename ElementType>
+	inline T FPropertyCollectionConstAdapter::GetValue(const FString& Key, const TConstArrayView<ElementType>& ValueArray, const T& Default, int32* OutKeyIndex) const
+	{
+		const int32 KeyIndex = GetKeyIndex(Key);
+		if (OutKeyIndex)
+		{
+			*OutKeyIndex = KeyIndex;
+		}
+		return KeyIndex != INDEX_NONE ? GetValue<T>(KeyIndex, ValueArray) : Default;
+	}
+
+	template<typename T, typename TEnableIf<TIsWeightedType<T>::Value, int>::type>
+	inline int32 FPropertyCollectionAdapter::SetWeightedValue(const FString& Key, const T& LowValue, const T& HighValue)
+	{
+		const int32 KeyIndex = GetKeyIndex(Key);
+		if (KeyIndex != INDEX_NONE)
+		{
+			SetLowValue(KeyIndex, LowValue);
+			SetHighValue(KeyIndex, HighValue);
+		}
+		return KeyIndex;
+	}
+
+	inline int32 FPropertyCollectionAdapter::SetWeightedFloatValue(const FString& Key, const FVector2f& Value)
+	{
+		const int32 KeyIndex = GetKeyIndex(Key);
+		if (KeyIndex != INDEX_NONE)
+		{
+			SetWeightedFloatValue(KeyIndex, Value);
+		}
+		return KeyIndex;
+	}
+
+	template<typename T, typename TEnableIf<TIsWeightedType<T>::Value, int>::type>
+	inline int32 FPropertyCollectionAdapter::AddWeightedValue(const FString& Key, const T& LowValue, const T& HighValue, bool bEnabled, bool bAnimatable)
+	{
+		const int32 KeyIndex = AddProperty(Key, bEnabled, bAnimatable);
+		SetWeightedValue(KeyIndex, LowValue, HighValue);
+		return KeyIndex;
+	}
+
+	inline int32 FPropertyCollectionAdapter::AddWeightedFloatValue(const FString& Key, const FVector2f& Value, bool bEnabled, bool bAnimatable)
+	{
+		const int32 KeyIndex = AddProperty(Key, bEnabled, bAnimatable);
+		SetWeightedFloatValue(KeyIndex, Value);
+		return KeyIndex;
+	}
+
+	inline int32 FPropertyCollectionAdapter::AddStringValue(const FString& Key, const FString& Value, bool bEnabled, bool bAnimatable)
+	{
+		const int32 KeyIndex = AddProperty(Key, bEnabled, bAnimatable);
+		SetStringValue(Key, Value);
+		return KeyIndex;
+	}
+
+	template<typename T>
+	inline int32 FPropertyCollectionAdapter::SetValue(const FString& Key, const TArrayView<T>& ValueArray, const T& Value)
+	{
+		const int32 KeyIndex = GetKeyIndex(Key);
+		if (KeyIndex != INDEX_NONE)
+		{
+			ValueArray[KeyIndex] = Value;
+		}
+		return KeyIndex;
+	}
+}  // End namespace Chaos
+
+// Use this macro to add shorthands for property getters and direct access through the declared key index
+#define UE_CHAOS_DECLARE_PROPERTYCOLLECTION_NAME(PropertyName, Type) \
+	inline static const FName PropertyName##Name = TEXT(#PropertyName); \
+	static FString PropertyName##String() { return PropertyName##Name.ToString(); } \
+	static bool Is##PropertyName##Enabled(const FPropertyCollectionConstAdapter& PropertyCollection, bool bDefault) \
+	{ \
+		return PropertyCollection.IsEnabled(PropertyName##String(), bDefault); \
+	} \
+	Type GetLow##PropertyName(const FPropertyCollectionConstAdapter& PropertyCollection, const Type& Default) \
+	{ \
+		return PropertyCollection.GetLowValue<Type>(PropertyName##String(), Default, &PropertyName##Index); \
+	} \
+	Type GetHigh##PropertyName(const FPropertyCollectionConstAdapter& PropertyCollection, const Type& Default) \
+	{ \
+		return PropertyCollection.GetHighValue<Type>(PropertyName##String(), Default, &PropertyName##Index); \
+	} \
+	TPair<Type, Type> GetWeighted##PropertyName(const FPropertyCollectionConstAdapter& PropertyCollection, const Type& Default) \
+	{ \
+		return PropertyCollection.GetWeightedValue<Type>(PropertyName##String(), Default, &PropertyName##Index); \
+	} \
+	FVector2f GetWeightedFloat##PropertyName(const FPropertyCollectionConstAdapter& PropertyCollection, const float& Default) \
+	{ \
+		return PropertyCollection.GetWeightedFloatValue(PropertyName##String(), Default, &PropertyName##Index); \
+	} \
+	Type Get##PropertyName(const FPropertyCollectionConstAdapter& PropertyCollection, const Type& Default) \
+	{ \
+		return PropertyCollection.GetValue<Type>(PropertyName##String(), Default, &PropertyName##Index); \
+	} \
+	uint8 Get##PropertyName##Flags(const FPropertyCollectionConstAdapter& PropertyCollection, uint8 Default) \
+	{ \
+		return PropertyCollection.GetFlags(PropertyName##String(), Default, &PropertyName##Index); \
+	} \
+	Type GetLow##PropertyName(const FPropertyCollectionConstAdapter& PropertyCollection) const \
+	{ \
+		checkSlow(PropertyName##Index == PropertyCollection.GetKeyIndex(PropertyName##String())); \
+		checkf(PropertyName##Index != INDEX_NONE, TEXT("The default value getter that sets the property index must be called once prior to calling this function.")); \
+		return PropertyCollection.GetLowValue<Type>(PropertyName##Index); \
+	} \
+	Type GetHigh##PropertyName(const FPropertyCollectionConstAdapter& PropertyCollection) const \
+	{ \
+		checkSlow(PropertyName##Index == PropertyCollection.GetKeyIndex(PropertyName##String())); \
+		checkf(PropertyName##Index != INDEX_NONE, TEXT("The default value getter that sets the property index must be called once prior to calling this function.")); \
+		return PropertyCollection.GetHighValue<Type>(PropertyName##Index); \
+	} \
+	TPair<Type, Type> GetWeighted##PropertyName(const FPropertyCollectionConstAdapter& PropertyCollection) const \
+	{ \
+		checkSlow(PropertyName##Index == PropertyCollection.GetKeyIndex(PropertyName##String())); \
+		checkf(PropertyName##Index != INDEX_NONE, TEXT("The default value getter that sets the property index must be called once prior to calling this function.")); \
+		return PropertyCollection.GetWeightedValue<Type>(PropertyName##Index); \
+	} \
+	FVector2f GetWeightedFloat##PropertyName(const FPropertyCollectionConstAdapter& PropertyCollection) const \
+	{ \
+		checkSlow(PropertyName##Index == PropertyCollection.GetKeyIndex(PropertyName##String())); \
+		checkf(PropertyName##Index != INDEX_NONE, TEXT("The default value getter that sets the property index must be called once prior to calling this function.")); \
+		return PropertyCollection.GetWeightedFloatValue(PropertyName##Index); \
+	} \
+	Type Get##PropertyName(const FPropertyCollectionConstAdapter& PropertyCollection) const \
+	{ \
+		checkSlow(PropertyName##Index == PropertyCollection.GetKeyIndex(PropertyName##String())); \
+		checkf(PropertyName##Index != INDEX_NONE, TEXT("The default value getter that sets the property index must be called once prior to calling this function.")); \
+		return PropertyCollection.GetValue<Type>(PropertyName##Index); \
+	} \
+	uint8 Get##PropertyName##Flags(const FPropertyCollectionConstAdapter& PropertyCollection) const \
+	{ \
+		checkSlow(PropertyName##Index == PropertyCollection.GetKeyIndex(PropertyName##String())); \
+		checkf(PropertyName##Index != INDEX_NONE, TEXT("The default value getter that sets the property index must be called once prior to calling this function.")); \
+		return PropertyCollection.GetFlags(PropertyName##Index); \
+	} \
+	bool Is##PropertyName##Enabled(const FPropertyCollectionConstAdapter& PropertyCollection) const \
+	{ \
+		checkSlow(PropertyName##Index == PropertyCollection.GetKeyIndex(PropertyName##String())); \
+		checkf(PropertyName##Index != INDEX_NONE, TEXT("The default value getter that sets the property index must be called once prior to calling this function.")); \
+		return PropertyCollection.IsEnabled(PropertyName##Index); \
+	} \
+	bool Is##PropertyName##Animatable(const FPropertyCollectionConstAdapter& PropertyCollection) const\
+	{ \
+		checkSlow(PropertyName##Index == PropertyCollection.GetKeyIndex(PropertyName##String())); \
+		checkf(PropertyName##Index != INDEX_NONE, TEXT("The default value getter that sets the property index must be called once prior to calling this function.")); \
+		return PropertyCollection.IsAnimatable(PropertyName##Index); \
+	} \
+	bool Is##PropertyName##Dirty(const FPropertyCollectionConstAdapter& PropertyCollection) const \
+	{ \
+		checkSlow(PropertyName##Index == PropertyCollection.GetKeyIndex(PropertyName##String())); \
+		checkf(PropertyName##Index != INDEX_NONE, TEXT("The default value getter that sets the property index must be called once prior to calling this function.")); \
+		return PropertyCollection.IsDirty(PropertyName##Index); \
+	} \
+	bool Is##PropertyName##Mutable(const FPropertyCollectionConstAdapter& PropertyCollection) const \
+	{ \
+		return PropertyName##Index != INDEX_NONE && \
+			EnumHasAllFlags((EPropertyFlag)PropertyCollection.GetFlags(PropertyName##Index), EPropertyFlag::Enabled | EPropertyFlag::Animatable | EPropertyFlag::Dirty); \
+	} \
+	int32 PropertyName##Index = INDEX_NONE;
