Fix Triangle Capsule bug

diff --git a/Engine/Source/Programs/HeadlessChaos/Private/HeadlessChaosTestOverlap.cpp b/Engine/Source/Programs/HeadlessChaos/Private/HeadlessChaosTestOverlap.cpp
--- a/Engine/Source/Programs/HeadlessChaos/Private/HeadlessChaosTestOverlap.cpp
+++ b/Engine/Source/Programs/HeadlessChaos/Private/HeadlessChaosTestOverlap.cpp
@@ -435,6 +435,12 @@ namespace ChaosTest
 					bool bResult = TriangleMesh->OverlapGeom(ScaledCapsule, QueryTM, 0.0, nullptr, TriMeshScale);
 					EXPECT_EQ(bResult, false);
 				}
+				{
+					FRigidTransform3 QueryTM(FVec3(0.0, 0.0, 22.0), FQuat(FVec3(1.0, 0.0, 0.0), 3.1415926 / 2.0));
+					QueryTM = FRigidTransform3(QueryTM.GetLocation() * InvScale, QueryTM.GetRotation());
+					bool bResult = TriangleMesh->OverlapGeom(ScaledCapsule, QueryTM, 0.0, nullptr, TriMeshScale);
+					EXPECT_EQ(bResult, false);
+				}
 				{
 					// Inside Mesh
 					FRigidTransform3 QueryTM(FVec3(0.0, 0.0, 8.0), FQuat::Identity);
@@ -946,12 +952,12 @@ namespace ChaosTest
 					EXPECT_TRUE(bResult);
 				}
 				{
-					const FRigidTransform3 QueryTM(FVec3(5.0, 10.25, 0.0), FQuat::Identity);
+					const FRigidTransform3 QueryTM(FVec3(5.0, -10.25, 0.0), FQuat::Identity);
 					bool bResult = TriangleMesh->OverlapGeom(Capsule, QueryTM, 0.0);
 					EXPECT_TRUE(bResult);
 				}
 				{
-					const FRigidTransform3 QueryTM(FVec3(5.0, 10.6, 0.0), FQuat::Identity);
+					const FRigidTransform3 QueryTM(FVec3(5.0, -10.6, 0.0), FQuat::Identity);
 					bool bResult = TriangleMesh->OverlapGeom(Capsule, QueryTM, 0.0);
 					EXPECT_FALSE(bResult);
 				}
@@ -1003,6 +1009,34 @@ namespace ChaosTest
 					EXPECT_FALSE(bResult);
 				}
 			}
+			{
+				// Capsule test
+				const FVec3 X1 = { 0.0, 0.0, -2.0 };
+				const FVec3 X2 = { 0.0, 0.0, 2.0 };
+				const FReal Radius = 1.0;
+				const FCapsule Capsule = FCapsule(X1, X2, Radius);
+				{
+					const FRigidTransform3 QueryTM(FVec3(0.0, 0.0, 0.0), FQuat::Identity);
+					bool bResult = TriangleMesh->OverlapGeom(Capsule, QueryTM, 0.0);
+					EXPECT_TRUE(bResult);
+				}
+				{
+					const FRigidTransform3 QueryTM(FVec3(3.9, 0.0, 100.0), FQuat::Identity);
+					bool bResult = TriangleMesh->OverlapGeom(Capsule, QueryTM, 0.0);
+					EXPECT_TRUE(bResult);
+				}
+				{
+					const FRigidTransform3 QueryTM(FVec3(4.0, 0.0, 90.0), FQuat::Identity);
+					bool bResult = TriangleMesh->OverlapGeom(Capsule, QueryTM, 0.0);
+					EXPECT_TRUE(bResult);
+				}
+				{
+					const FRigidTransform3 QueryTM(FVec3(4.5, 0.0, 90.0), FQuat::Identity);
+					bool bResult = TriangleMesh->OverlapGeom(Capsule, QueryTM, 0.0);
+					EXPECT_FALSE(bResult);
+				}
+
+			}
 		}
 	}
 }
diff --git a/Engine/Source/Runtime/Experimental/Chaos/Private/Chaos/Collision/TriangleOverlap.cpp b/Engine/Source/Runtime/Experimental/Chaos/Private/Chaos/Collision/TriangleOverlap.cpp
--- a/Engine/Source/Runtime/Experimental/Chaos/Private/Chaos/Collision/TriangleOverlap.cpp
+++ b/Engine/Source/Runtime/Experimental/Chaos/Private/Chaos/Collision/TriangleOverlap.cpp
@@ -35,11 +35,12 @@ namespace Chaos
 		const VectorRegister4Float X2X1 = VectorSubtract(X1, X2);
 		const VectorRegister4Float CA = VectorSubtract(A, C);
 		const VectorRegister4Float CentroidA = VectorSubtract(A, Centroid);
+		const VectorRegister4Float CX1 = VectorSubtract(X1, C);
 		{
 			VectorRegister4Float EdgeSeparationAxis = VectorNormalize(VectorCross(X2X1, CA));
 			EdgeSeparationAxis = VectorSelect(VectorCompareGT(VectorZeroFloat(), VectorDot3(CentroidA, EdgeSeparationAxis)), VectorNegate(EdgeSeparationAxis), EdgeSeparationAxis);
 
-			const FRealSingle EdgeSeparationDist = VectorDot3Scalar(VectorSubtract(X1, C), EdgeSeparationAxis);
+			const FRealSingle EdgeSeparationDist = VectorDot3Scalar(CX1, EdgeSeparationAxis);
 			if (EdgeSeparationDist > Radius)
 			{
 				return false;
@@ -50,24 +51,64 @@ namespace Chaos
 			VectorRegister4Float EdgeSeparationAxis = VectorNormalize(VectorCross(X2X1, AB));
 			EdgeSeparationAxis = VectorSelect(VectorCompareGT(VectorZeroFloat(), VectorDot3(CentroidB, EdgeSeparationAxis)), VectorNegate(EdgeSeparationAxis), EdgeSeparationAxis);
 
-			const FRealSingle EdgeSeparationDist = VectorDot3Scalar(VectorSubtract(X1, A), EdgeSeparationAxis);
+			const FRealSingle EdgeSeparationDist = VectorDot3Scalar(AX1, EdgeSeparationAxis);
 			if (EdgeSeparationDist > Radius)
 			{
 				return false;
 			}
 		}
 		const VectorRegister4Float CentroidC = VectorSubtract(C, Centroid);
+		const VectorRegister4Float BX1 = VectorSubtract(X1, B);
 		{
 			VectorRegister4Float EdgeSeparationAxis = VectorNormalize(VectorCross(X2X1, BC));
 			EdgeSeparationAxis = VectorSelect(VectorCompareGT(VectorZeroFloat(), VectorDot3(CentroidC, EdgeSeparationAxis)), VectorNegate(EdgeSeparationAxis), EdgeSeparationAxis);
 
-			const FRealSingle EdgeSeparationDist = VectorDot3Scalar(VectorSubtract(X1, B), EdgeSeparationAxis);
+			const FRealSingle EdgeSeparationDist = VectorDot3Scalar(BX1, EdgeSeparationAxis);
 			if (EdgeSeparationDist > Radius)
 			{
 				return false;
 			}
 		}
 
+		const VectorRegister4Float RadiusSimd = VectorSetFloat1(Radius);
+		// Edge plane normals and signed distances to each segment point
+		// Mask to compare the distance of the triangle edge plane to the capsule radius.
+		constexpr static int32 AllEdges = 0b1111;
+		{
+			const VectorRegister4Float PlaneCA = VectorNormalize(VectorCross(CA, Normal));
+			const VectorRegister4Float EdgeX1 = VectorDot3(AX1, PlaneCA);
+			const VectorRegister4Float EdgeX2 = VectorDot3(AX2, PlaneCA);
+			const VectorRegister4Float Edge = VectorMoveLh(EdgeX1, EdgeX2);
+			const VectorRegister4Float EdgeGTRadius = VectorCompareGT(Edge, RadiusSimd);
+			if (VectorMaskBits(EdgeGTRadius) == AllEdges)
+			{
+				return false;
+			}
+		}
+
+		{
+			const VectorRegister4Float PlaneAB = VectorNormalize(VectorCross(AB, Normal));
+			const VectorRegister4Float EdgeX1 = VectorDot3(BX1, PlaneAB);
+			const VectorRegister4Float EdgeX2 = VectorDot3(VectorSubtract(X2, B), PlaneAB);
+			const VectorRegister4Float Edge = VectorMoveLh(EdgeX1, EdgeX2);
+			const VectorRegister4Float EdgeGTRadius = VectorCompareGT(Edge, RadiusSimd);
+			if (VectorMaskBits(EdgeGTRadius) == AllEdges)
+			{
+				return false;
+			}
+		}
+		{
+			const VectorRegister4Float PlaneBC = VectorNormalize(VectorCross(BC, Normal));
+			const VectorRegister4Float EdgeX1 = VectorDot3(CX1, PlaneBC);
+			const VectorRegister4Float EdgeX2 = VectorDot3(VectorSubtract(X2, C), PlaneBC);
+			const VectorRegister4Float Edge = VectorMoveLh(EdgeX1, EdgeX2);
+			const VectorRegister4Float EdgeGTRadius = VectorCompareGT(Edge, RadiusSimd);
+			if (VectorMaskBits(EdgeGTRadius) == AllEdges)
+			{
+				return false;
+			}
+		}
+
 		// Triangle Vertices
 		VectorRegister4Float SqrX2X1 = VectorDot3(X2X1, X2X1);
 		VectorRegister4Float ZeroMask = VectorCompareEQ(VectorZeroFloat(), SqrX2X1);
