Chaos Cloth - Add two missing explicit instanciations of the property collection GetValue() template function.

diff --git a/Engine/Source/Runtime/Experimental/Chaos/Private/Chaos/PropertyCollectionAdapter.cpp b/Engine/Source/Runtime/Experimental/Chaos/Private/Chaos/PropertyCollectionAdapter.cpp
--- a/Engine/Source/Runtime/Experimental/Chaos/Private/Chaos/PropertyCollectionAdapter.cpp
+++ b/Engine/Source/Runtime/Experimental/Chaos/Private/Chaos/PropertyCollectionAdapter.cpp
@@ -51,6 +51,8 @@ namespace Chaos::Softs
 	{
 		return ValueArray[KeyIndex];
 	}
+	template FString FPropertyCollectionConstAdapter::GetValue<FString, FString>(int32 KeyIndex, const TConstArrayView<FString>& ValueArray) const;
+	template uint8 FPropertyCollectionConstAdapter::GetValue<uint8, uint8>(int32 KeyIndex, const TConstArrayView<uint8>& ValueArray) const;
 
 	template<>
 	bool FPropertyCollectionConstAdapter::GetValue<bool, FVector3f>(int32 KeyIndex, const TConstArrayView<FVector3f>& ValueArray) const
