OBB in private namespace

diff --git a/Engine/Source/Programs/HeadlessChaos/Private/HeadlessChaosTestObbTree.cpp b/Engine/Source/Programs/HeadlessChaos/Private/HeadlessChaosTestObbTree.cpp
--- a/Engine/Source/Programs/HeadlessChaos/Private/HeadlessChaosTestObbTree.cpp
+++ b/Engine/Source/Programs/HeadlessChaos/Private/HeadlessChaosTestObbTree.cpp
@@ -15,7 +15,7 @@ namespace ChaosTest
 		{
 			// Original points to create OBB // TArray<FVec3f> PointsA{ FVec3f(0, 0, 0), FVec3f(5, 0, 0), FVec3f(0, 5, 0) };
 			FRigidTransform3 Transform = FRigidTransform3(FVec3(1.25, 1.25, 0.0), FQuat(FVec3(0.0, 0.0, 1.0), -3.1415926 / 4.0));
-			FOBBVectorized ObbA(Transform, FVec3f(UE_INV_SQRT_2*5.0, UE_INV_SQRT_2 * 5.0 /2.0, 0.0f));
+			Private::FOBBVectorized ObbA(Transform, FVec3f(UE_INV_SQRT_2*5.0, UE_INV_SQRT_2 * 5.0 /2.0, 0.0f));
 			{
 				TAABB<FRealSingle, 3> Bounds(FVec3f(-2.0f, 1.5f, 0.0f), FVec3f(-1.0f, 2.5f, 0.0f));
 				EXPECT_TRUE(ObbA.IntersectAABB(FAABBVectorized(Bounds)));
@@ -56,7 +56,7 @@ namespace ChaosTest
 		{
 			
 			FRigidTransform3 Transform = FRigidTransform3(FVec3(0.0, 0.0, 12.0), FQuat(FVec3(1.0, 0.0, 0.0), 3.1415926 / 4.0));
-			FOBBVectorized ObbA(Transform.ToMatrixNoScale(), FVec3f(2.0f, 2.0f, 2.0f));
+			Private::FOBBVectorized ObbA(Transform.ToMatrixNoScale(), FVec3f(2.0f, 2.0f, 2.0f));
 			TAABB<FRealSingle, 3> Bounds(FVec3f(-10.0f, -10.0f, 0.0f), FVec3f(10.0f, 10.0f, 10.0f));
 			EXPECT_TRUE(ObbA.IntersectAABB(FAABBVectorized(Bounds)));
 		}
diff --git a/Engine/Source/Runtime/Experimental/Chaos/Private/Chaos/OBBVectorized.cpp b/Engine/Source/Runtime/Experimental/Chaos/Private/Chaos/OBBVectorized.cpp
--- a/Engine/Source/Runtime/Experimental/Chaos/Private/Chaos/OBBVectorized.cpp
+++ b/Engine/Source/Runtime/Experimental/Chaos/Private/Chaos/OBBVectorized.cpp
@@ -8,7 +8,7 @@
 namespace Chaos
 {
 
-	FOBBVectorized::FOBBVectorized(const FRigidTransform3& Transform, const FVec3f& HalfExtentsIn)
+	Private::FOBBVectorized::FOBBVectorized(const FRigidTransform3& Transform, const FVec3f& HalfExtentsIn)
 	{
 		TRotation<FRealSingle, 3> Rotation = Transform.GetRotation();
 		PMatrix<FRealSingle, 3, 3> Matrix = Rotation.ToMatrix();
@@ -20,16 +20,16 @@ namespace Chaos
 		Position = MakeVectorRegisterFloatFromDouble(VectorLoadDouble3(&Translation.X));
 
 		HalfExtents = VectorLoadFloat3(&HalfExtentsIn.X);
-		XHalfExtent = VectorReplicate(HalfExtents, 0);
-		YHalfExtent = VectorReplicate(HalfExtents, 1);
-		ZHalfExtent = VectorReplicate(HalfExtents, 2);
+		const VectorRegister4Float XHalfExtent = VectorReplicate(HalfExtents, 0);
+		const VectorRegister4Float YHalfExtent = VectorReplicate(HalfExtents, 1);
+		const VectorRegister4Float ZHalfExtent = VectorReplicate(HalfExtents, 2);
 
 		VectorRegister4Float FurthestPoint = VectorMultiplyAdd(VectorAbs(XAxis), XHalfExtent, VectorMultiplyAdd(VectorAbs(YAxis), YHalfExtent, VectorMultiply(VectorAbs(ZAxis), ZHalfExtent)));
 		MaxObb = VectorAdd(Position, FurthestPoint);
 		MinObb = VectorSubtract(Position, FurthestPoint);
 	}
 
-	bool FOBBVectorized::IntersectAABB(const FAABBVectorized& Bounds) const
+	bool Private::FOBBVectorized::IntersectAABB(const FAABBVectorized& Bounds) const
 	{
 		VectorRegister4Float HasSeparationAxis = VectorBitwiseOr(VectorCompareGT(MinObb, Bounds.GetMax()), VectorCompareGT(Bounds.GetMin(), MaxObb));
 
diff --git a/Engine/Source/Runtime/Experimental/Chaos/Private/Chaos/TriangleMeshImplicitObject.cpp b/Engine/Source/Runtime/Experimental/Chaos/Private/Chaos/TriangleMeshImplicitObject.cpp
--- a/Engine/Source/Runtime/Experimental/Chaos/Private/Chaos/TriangleMeshImplicitObject.cpp
+++ b/Engine/Source/Runtime/Experimental/Chaos/Private/Chaos/TriangleMeshImplicitObject.cpp
@@ -656,7 +656,7 @@ bool FTrimeshBVH::FindAllIntersectionsNoMTD(const FAABB3& Intersection, const TR
 }
 
 template <typename QueryGeomType>
-bool FTrimeshBVH::FindAllIntersectionsNoMTD(const FOBBVectorized& Intersection, const TRigidTransform<FReal, 3>& Transform, const QueryGeomType& QueryGeom, FReal Thickness, const FVec3& TriMeshScale, const FTriangleMeshImplicitObject* TriMesh) const
+bool FTrimeshBVH::FindAllIntersectionsNoMTD(const Private::FOBBVectorized& Intersection, const TRigidTransform<FReal, 3>& Transform, const QueryGeomType& QueryGeom, FReal Thickness, const FVec3& TriMeshScale, const FTriangleMeshImplicitObject* TriMesh) const
 {
 	FTriangleMeshOverlapVisitorNoMTD Visitor(Transform, QueryGeom, Thickness, TriMeshScale, TriMesh);
 	OverlapOBB(Intersection, Visitor);
@@ -1010,7 +1010,7 @@ bool FTriangleMeshImplicitObject::OverlapGeomImp(const QueryGeomType& QueryGeom,
 		{
 			FAABB3 GeometryAABB = QueryGeom.BoundingBox();
 			GeometryAABB.ThickenSymmetrically(FVec3(Thickness));
-			FOBBVectorized QueryObb(QueryTM, (GeometryAABB.Max() - GeometryAABB.Min()) * 0.5);
+			Private::FOBBVectorized QueryObb(QueryTM, (GeometryAABB.Max() - GeometryAABB.Min()) * 0.5);
 			TRigidTransform<FReal, 3> WorldScaleQueryTM;
 			ScaleTransformHelper(TriMeshScale, QueryTM, WorldScaleQueryTM);
 			return FastBVH.FindAllIntersectionsNoMTD(QueryObb, WorldScaleQueryTM, WorldScaleQueryGeom, Thickness, TriMeshScale, this);
diff --git a/Engine/Source/Runtime/Experimental/Chaos/Public/Chaos/OBBVectorized.h b/Engine/Source/Runtime/Experimental/Chaos/Public/Chaos/OBBVectorized.h
--- a/Engine/Source/Runtime/Experimental/Chaos/Public/Chaos/OBBVectorized.h
+++ b/Engine/Source/Runtime/Experimental/Chaos/Public/Chaos/OBBVectorized.h
@@ -8,25 +8,23 @@ namespace Chaos
 {
 	class FAABBVectorized;
 
-	class FOBBVectorized
+	namespace Private
 	{
-	public:
-		FOBBVectorized(const FRigidTransform3& Transform, const FVec3f& HalfExtentsIn);
-		bool IntersectAABB(const FAABBVectorized& Bounds) const;
+		class FOBBVectorized
+		{
+		public:
+			FOBBVectorized(const FRigidTransform3& Transform, const FVec3f& HalfExtentsIn);
+			bool IntersectAABB(const FAABBVectorized& Bounds) const;
 
-	private:
-		VectorRegister4Float Position;
-		VectorRegister4Float XAxis;
-		VectorRegister4Float YAxis;
-		VectorRegister4Float ZAxis;
-		VectorRegister4Float HalfExtents;
-		VectorRegister4Float XHalfExtent;
-		VectorRegister4Float YHalfExtent;
-		VectorRegister4Float ZHalfExtent;
+		private:
+			VectorRegister4Float Position;
+			VectorRegister4Float XAxis;
+			VectorRegister4Float YAxis;
+			VectorRegister4Float ZAxis;
+			VectorRegister4Float HalfExtents;
 
-		VectorRegister4Float MaxObb;
-		VectorRegister4Float MinObb;
-		
-
-	};
+			VectorRegister4Float MaxObb;
+			VectorRegister4Float MinObb;
+		};
+	}
 }
\ No newline at end of file
diff --git a/Engine/Source/Runtime/Experimental/Chaos/Public/Chaos/TriangleMeshImplicitObject.h b/Engine/Source/Runtime/Experimental/Chaos/Public/Chaos/TriangleMeshImplicitObject.h
--- a/Engine/Source/Runtime/Experimental/Chaos/Public/Chaos/TriangleMeshImplicitObject.h
+++ b/Engine/Source/Runtime/Experimental/Chaos/Public/Chaos/TriangleMeshImplicitObject.h
@@ -184,7 +184,7 @@ namespace Chaos
 		}
 
 		template <typename SQVisitor>
-		FORCEINLINE_DEBUGGABLE void OverlapOBB(const FOBBVectorized& Obb, SQVisitor& Visitor) const
+		FORCEINLINE_DEBUGGABLE void OverlapOBB(const Private::FOBBVectorized& Obb, SQVisitor& Visitor) const
 		{
 			const auto BoundsFilter = [&Obb](const FAABBVectorized& Bounds) -> EFilterResult
 			{
@@ -202,7 +202,7 @@ namespace Chaos
 		}
 
 		template <typename QueryGeomType>
-		bool FindAllIntersectionsNoMTD(const FOBBVectorized& Intersection, const TRigidTransform<FReal, 3>& Transform, const QueryGeomType& QueryGeom, FReal Thickness, const FVec3& TriMeshScale, const FTriangleMeshImplicitObject* TriMesh) const;
+		bool FindAllIntersectionsNoMTD(const Private::FOBBVectorized& Intersection, const TRigidTransform<FReal, 3>& Transform, const QueryGeomType& QueryGeom, FReal Thickness, const FVec3& TriMeshScale, const FTriangleMeshImplicitObject* TriMesh) const;
 		template <typename QueryGeomType>
 		bool FindAllIntersectionsNoMTD(const FAABB3& Intersection, const TRigidTransform<FReal, 3>& Transform, const QueryGeomType& QueryGeom, FReal Thickness, const FVec3& TriMeshScale, const FTriangleMeshImplicitObject* TriMesh) const;
 		TArray<int32> FindAllIntersections(const FAABB3& Intersection) const;
