Overlap Triangle Capsule optimization
Small optimization using heuristic

diff --git a/Engine/Source/Runtime/Experimental/Chaos/Private/Chaos/Collision/TriangleOverlap.cpp b/Engine/Source/Runtime/Experimental/Chaos/Private/Chaos/Collision/TriangleOverlap.cpp
--- a/Engine/Source/Runtime/Experimental/Chaos/Private/Chaos/Collision/TriangleOverlap.cpp
+++ b/Engine/Source/Runtime/Experimental/Chaos/Private/Chaos/Collision/TriangleOverlap.cpp
@@ -6,69 +6,29 @@
 {
 	bool ComputeCapsuleTriangleOverlapSimd(const VectorRegister4Float& A, const VectorRegister4Float& B, const VectorRegister4Float& C, const VectorRegister4Float& X1, const VectorRegister4Float& X2, FRealSingle Radius)
 	{
+
+		const VectorRegister4Float X2X1 = VectorSubtract(X1, X2);
+		const VectorRegister4Float SqrX2X1 = VectorDot3(X2X1, X2X1);
+		const VectorRegister4Float ZeroMask = VectorCompareEQ(VectorZeroFloat(), SqrX2X1);
+
 		const VectorRegister4Float AB = VectorSubtract(B, A);
 		const VectorRegister4Float BC = VectorSubtract(C, B);
+		const VectorRegister4Float CA = VectorSubtract(A, C);
 
-		VectorRegister4Float Normal = VectorNormalize(VectorCross(AB, BC));
+		const VectorRegister4Float Normal = VectorNormalize(VectorCross(AB, BC));
 
-		// Plane Triangle
-		const VectorRegister4Float AX1 = VectorSubtract(X1, A);
-		const VectorRegister4Float AX2 = VectorSubtract(X2, A);
-
-		FRealSingle AX1Dist = VectorDot3Scalar(AX1, Normal);
-		FRealSingle AX2Dist = VectorDot3Scalar(AX2, Normal);
-
-		if (FMath::Sign(AX1Dist) == FMath::Sign(AX2Dist))
-		{
-			FRealSingle ClosestDist = FMath::Min<FRealSingle>(FMath::Abs<FRealSingle>(AX1Dist), FMath::Abs<FRealSingle>(AX2Dist));
-			if (ClosestDist > Radius)
-			{
-				return false;
-			}
-		}
-
-		// Edges
 		constexpr FRealSingle ThirdFloat = 1.0f / 3.0f;
 		constexpr VectorRegister4Float Third = MakeVectorRegisterFloatConstant(ThirdFloat, ThirdFloat, ThirdFloat, ThirdFloat);
 		const VectorRegister4Float Centroid = VectorMultiply(VectorAdd(VectorAdd(A, B), C), Third);
 
-		const VectorRegister4Float X2X1 = VectorSubtract(X1, X2);
-		const VectorRegister4Float CA = VectorSubtract(A, C);
 		const VectorRegister4Float CentroidA = VectorSubtract(A, Centroid);
-		const VectorRegister4Float CX1 = VectorSubtract(X1, C);
-		{
-			VectorRegister4Float EdgeSeparationAxis = VectorNormalize(VectorCross(X2X1, CA));
-			EdgeSeparationAxis = VectorSelect(VectorCompareGT(VectorZeroFloat(), VectorDot3(CentroidA, EdgeSeparationAxis)), VectorNegate(EdgeSeparationAxis), EdgeSeparationAxis);
-
-			const FRealSingle EdgeSeparationDist = VectorDot3Scalar(CX1, EdgeSeparationAxis);
-			if (EdgeSeparationDist > Radius)
-			{
-				return false;
-			}
-		}
 		const VectorRegister4Float CentroidB = VectorSubtract(B, Centroid);
-		{
-			VectorRegister4Float EdgeSeparationAxis = VectorNormalize(VectorCross(X2X1, AB));
-			EdgeSeparationAxis = VectorSelect(VectorCompareGT(VectorZeroFloat(), VectorDot3(CentroidB, EdgeSeparationAxis)), VectorNegate(EdgeSeparationAxis), EdgeSeparationAxis);
-
-			const FRealSingle EdgeSeparationDist = VectorDot3Scalar(AX1, EdgeSeparationAxis);
-			if (EdgeSeparationDist > Radius)
-			{
-				return false;
-			}
-		}
 		const VectorRegister4Float CentroidC = VectorSubtract(C, Centroid);
+
+		const VectorRegister4Float AX1 = VectorSubtract(X1, A);
+		const VectorRegister4Float AX2 = VectorSubtract(X2, A);
 		const VectorRegister4Float BX1 = VectorSubtract(X1, B);
-		{
-			VectorRegister4Float EdgeSeparationAxis = VectorNormalize(VectorCross(X2X1, BC));
-			EdgeSeparationAxis = VectorSelect(VectorCompareGT(VectorZeroFloat(), VectorDot3(CentroidC, EdgeSeparationAxis)), VectorNegate(EdgeSeparationAxis), EdgeSeparationAxis);
+		const VectorRegister4Float CX1 = VectorSubtract(X1, C);
 
-			const FRealSingle EdgeSeparationDist = VectorDot3Scalar(BX1, EdgeSeparationAxis);
-			if (EdgeSeparationDist > Radius)
-			{
-				return false;
-			}
-		}
-
 		const VectorRegister4Float RadiusSimd = VectorSetFloat1(Radius);
 		// Edge plane normals and signed distances to each segment point
@@ -85,7 +45,6 @@
 				return false;
 			}
 		}
-
 		{
 			const VectorRegister4Float PlaneAB = VectorNormalize(VectorCross(AB, Normal));
 			const VectorRegister4Float EdgeX1 = VectorDot3(BX1, PlaneAB);
@@ -109,10 +68,19 @@
 			}
 		}
 
+		// Plane Triangle
+		FRealSingle AX1Dist = VectorDot3Scalar(AX1, Normal);
+		FRealSingle AX2Dist = VectorDot3Scalar(AX2, Normal);
+		if (FMath::Sign(AX1Dist) == FMath::Sign(AX2Dist))
+		{
+			FRealSingle ClosestDist = FMath::Min<FRealSingle>(FMath::Abs<FRealSingle>(AX1Dist), FMath::Abs<FRealSingle>(AX2Dist));
+			if (ClosestDist > Radius)
+			{
+				return false;
+			}
+		}
+
 		// Triangle Vertices
-		VectorRegister4Float SqrX2X1 = VectorDot3(X2X1, X2X1);
-		VectorRegister4Float ZeroMask = VectorCompareEQ(VectorZeroFloat(), SqrX2X1);
-
 		FRealSingle SqrRadius = Radius * Radius;
 		{
 			VectorRegister4Float TimeA = VectorClamp(VectorDivide(VectorDot3(X2X1, AX1), SqrX2X1), VectorZeroFloat(), VectorOneFloat());
@@ -157,6 +125,38 @@
 			}
 		}
 
+		// Edges
+		{
+			VectorRegister4Float EdgeSeparationAxis = VectorNormalize(VectorCross(X2X1, CA));
+			EdgeSeparationAxis = VectorSelect(VectorCompareGT(VectorZeroFloat(), VectorDot3(CentroidA, EdgeSeparationAxis)), VectorNegate(EdgeSeparationAxis), EdgeSeparationAxis);
+
+			const FRealSingle EdgeSeparationDist = VectorDot3Scalar(CX1, EdgeSeparationAxis);
+			if (EdgeSeparationDist > Radius)
+			{
+				return false;
+			}
+		}
+		{
+			VectorRegister4Float EdgeSeparationAxis = VectorNormalize(VectorCross(X2X1, AB));
+			EdgeSeparationAxis = VectorSelect(VectorCompareGT(VectorZeroFloat(), VectorDot3(CentroidB, EdgeSeparationAxis)), VectorNegate(EdgeSeparationAxis), EdgeSeparationAxis);
+
+			const FRealSingle EdgeSeparationDist = VectorDot3Scalar(AX1, EdgeSeparationAxis);
+			if (EdgeSeparationDist > Radius)
+			{
+				return false;
+			}
+		}
+		{
+			VectorRegister4Float EdgeSeparationAxis = VectorNormalize(VectorCross(X2X1, BC));
+			EdgeSeparationAxis = VectorSelect(VectorCompareGT(VectorZeroFloat(), VectorDot3(CentroidC, EdgeSeparationAxis)), VectorNegate(EdgeSeparationAxis), EdgeSeparationAxis);
+
+			const FRealSingle EdgeSeparationDist = VectorDot3Scalar(BX1, EdgeSeparationAxis);
+			if (EdgeSeparationDist > Radius)
+			{
+				return false;
+			}
+		}
+
 		return true;
 	}
 
diff --git a/Engine/Source/Runtime/Experimental/Chaos/Private/Chaos/HeightField.cpp b/Engine/Source/Runtime/Experimental/Chaos/Private/Chaos/HeightField.cpp
--- a/Engine/Source/Runtime/Experimental/Chaos/Private/Chaos/HeightField.cpp
+++ b/Engine/Source/Runtime/Experimental/Chaos/Private/Chaos/HeightField.cpp
@@ -2009,7 +2009,7 @@ namespace Chaos
 					X1 = VectorLoadFloat3(&X1f.X);
 					X2 = VectorLoadFloat3(&X2f.X);
 				}
-				if constexpr (std::is_same<QueryGeomType, Chaos::FSphere>::value)
+				else if constexpr (std::is_same<QueryGeomType, Chaos::FSphere>::value)
 				{
 					Radius = FRealSingle(QueryGeom.GetRadius());
 					const FVec3f X1f = QueryGeom.GetCenter();
