Fix managed array collection accessor copy logic when target attribute wasn't yet added to a collection or not matching the name of the source attribute

diff --git a/Engine/Source/Runtime/Experimental/Chaos/Public/GeometryCollection/ManagedArrayAccessor.h b/Engine/Source/Runtime/Experimental/Chaos/Public/GeometryCollection/ManagedArrayAccessor.h
--- a/Engine/Source/Runtime/Experimental/Chaos/Public/GeometryCollection/ManagedArrayAccessor.h
+++ b/Engine/Source/Runtime/Experimental/Chaos/Public/GeometryCollection/ManagedArrayAccessor.h
@@ -126,7 +126,10 @@ public:
 	void Copy(const TManagedArrayAccessor<T>& FromAttribute )
 	{
 		check(!IsConst());
-		Collection->CopyAttribute(*FromAttribute.Collection, Name, Group);
+		Collection->CopyAttribute(FromAttribute.ConstCollection, FromAttribute.Name, Name, Group);
+		// update attribute array if necessary
+		AttributeArray = Collection->FindAttributeTyped<T>(Name, Group);
+		ConstAttributeArray = AttributeArray;
 	}
 
 	void Remove()
