Chaos - optimize UpdateKinematicTargets

The call to MarkTransientDirtyParticle was updating the particle views every time which for a scene with lots of animated kinematic ragdolls was a major hit.

Don't do the expensive parts of the kinematic update unless the body actually moved.

diff --git a/Engine/Source/Runtime/Experimental/Chaos/Public/Chaos/PBDRigidsEvolution.h b/Engine/Source/Runtime/Experimental/Chaos/Public/Chaos/PBDRigidsEvolution.h
--- a/Engine/Source/Runtime/Experimental/Chaos/Public/Chaos/PBDRigidsEvolution.h
+++ b/Engine/Source/Runtime/Experimental/Chaos/Public/Chaos/PBDRigidsEvolution.h
@@ -736,12 +736,16 @@ public:
 			return Particles.GetActiveKinematicParticlesView();
 		};
 
+		// We do not update the particle views every time we change a particle - we do it once at the end
+		const bool bUpdateParticleViews = false;
+
 		for (auto& Particle : GetKinematicView())
 		{
 			TKinematicTarget<FReal, 3>& KinematicTarget = Particle.KinematicTarget();
 			const FVec3 CurrentX = Particle.X();
 			const FRotation3 CurrentR = Particle.R();
 
+			bool bMoved = false;
 			switch (KinematicTarget.GetMode())
 			{
 			case EKinematicTargetMode::None:
@@ -754,7 +758,7 @@ public:
 				Particle.V() = FVec3(0, 0, 0);
 				Particle.W() = FVec3(0, 0, 0);
 				KinematicTarget.SetMode(EKinematicTargetMode::None);
-				Particles.MarkTransientDirtyParticle(Particle.Handle());
+				Particles.MarkTransientDirtyParticle(Particle.Handle(), bUpdateParticleViews);
 				break;
 			}
 
@@ -777,26 +781,38 @@ public:
 					NewX = FVec3::Lerp(CurrentX, KinematicTarget.GetTarget().GetLocation(), StepFraction);
 					NewR = FRotation3::Slerp(CurrentR, KinematicTarget.GetTarget().GetRotation(), decltype(FQuat::X)(StepFraction));
 				}
-				if (Dt > MinDt)
+
+				bMoved = !FVec3::IsNearlyEqual(NewX, CurrentX, UE_SMALL_NUMBER) || !FRotation3::IsNearlyEqual(NewR, CurrentR, UE_SMALL_NUMBER);
+				FVec3 V = FVec3(0);
+				FVec3 W = FVec3(0);
+				if (bMoved)
 				{
-					FVec3 V = FVec3::CalculateVelocity(CurrentX, NewX, Dt);
-					Particle.V() = V;
+					if (Dt > MinDt)
+					{
+						V = FVec3::CalculateVelocity(CurrentX, NewX, Dt);
+						Particle.V() = V;
+
+						W = FRotation3::CalculateAngularVelocity(CurrentR, NewR, Dt);
+						Particle.W() = W;
+					}
 
-					FVec3 W = FRotation3::CalculateAngularVelocity(CurrentR, NewR, Dt);
-					Particle.W() = W;
+					Particle.X() = NewX;
+					Particle.R() = NewR;
 				}
-				Particle.X() = NewX;
-				Particle.R() = NewR;
-				Particles.MarkTransientDirtyParticle(Particle.Handle());
+				Particle.V() = V;
+				Particle.W() = W;
+				Particles.MarkTransientDirtyParticle(Particle.Handle(), bUpdateParticleViews);
+
 				break;
 			}
 
 			case EKinematicTargetMode::Velocity:
 			{
 				// Move based on velocity
+				bMoved = true;
 				Particle.X() = Particle.X() + Particle.V() * Dt;
 				Particle.R() = FRotation3::IntegrateRotationWithAngularVelocity(Particle.R(), Particle.W(), Dt);
-				Particles.MarkTransientDirtyParticle(Particle.Handle());
+				Particles.MarkTransientDirtyParticle(Particle.Handle(), bUpdateParticleViews);
 				break;
 			}
 			}
@@ -810,13 +826,17 @@ public:
 				Rigid->Q() = Rigid->R();
 				Rigid->PreV() = Rigid->V();
 				Rigid->PreW() = Rigid->W();
-				if (!Rigid->CCDEnabled())
-				{
-					Rigid->UpdateWorldSpaceState(FRigidTransform3(Rigid->P(), Rigid->Q()), FVec3(0));
-				}
-				else
+
+				if (bMoved)
 				{
-					Rigid->UpdateWorldSpaceStateSwept(FRigidTransform3(Rigid->P(), Rigid->Q()), FVec3(0), -Rigid->V() * Dt);
+					if (!Rigid->CCDEnabled())
+					{
+						Rigid->UpdateWorldSpaceState(FRigidTransform3(Rigid->P(), Rigid->Q()), FVec3(0));
+					}
+					else
+					{
+						Rigid->UpdateWorldSpaceStateSwept(FRigidTransform3(Rigid->P(), Rigid->Q()), FVec3(0), -Rigid->V() * Dt);
+					}
 				}
 			}
 		}
@@ -826,6 +846,9 @@ public:
 		{
 			Particles.UpdateAllMovingKinematic();
 		}
+
+		// If we changed any particle state, the views need to be refreshed
+		Particles.UpdateDirtyViews();
 	}
 
 	/** Make a copy of the acceleration structure to allow for external modification.
diff --git a/Engine/Source/Runtime/Experimental/Chaos/Public/Chaos/PBDRigidsSOAs.h b/Engine/Source/Runtime/Experimental/Chaos/Public/Chaos/PBDRigidsSOAs.h
--- a/Engine/Source/Runtime/Experimental/Chaos/Public/Chaos/PBDRigidsSOAs.h
+++ b/Engine/Source/Runtime/Experimental/Chaos/Public/Chaos/PBDRigidsSOAs.h
@@ -229,6 +229,12 @@ public:
 	{
 		check(0);
 	}
+
+	void UpdateDirtyViews()
+	{
+		// @todo(chaos): this should only refresh views that may have changed
+		UpdateViews();
+	}
 	
 	TArray<FGeometryParticleHandle*> CreateStaticParticles(int32 NumParticles, const FUniqueIdx* ExistingIndices = nullptr, const FGeometryParticleParameters& Params = FGeometryParticleParameters())
 	{
@@ -299,7 +305,9 @@ public:
 		TransientDirtyMapArray.Reset();
 	}
 
-	void MarkTransientDirtyParticle(FGeometryParticleHandle* Particle)
+	// @todo(chaos): keep track of which views may be dirty and lazily update them
+	// rather than calling UpdateViews all the time
+	void MarkTransientDirtyParticle(FGeometryParticleHandle* Particle, const bool bUpdateViews = true)
 	{
 		if(FPBDRigidParticleHandle* Rigid =  Particle->CastToRigidParticle())
 		{
@@ -307,7 +315,10 @@ public:
 			if (!ActiveParticlesMapArray.Contains(Rigid))
 			{
 				TransientDirtyMapArray.Insert(Rigid);
-				UpdateViews();
+				if (bUpdateViews)
+				{
+					UpdateViews();
+				}
 			}
 		}
 	}
