Improve breaking event for geometry collections
- Add extents
- Ability to know if the breaking event originates from a crumble event
- break event are now filtered out in the clustering code if the proxy does not opt for it
- overall code imporvement and cleaning

diff --git a/Engine/Source/Runtime/Experimental/Chaos/Private/Chaos/PBDRigidClustering.cpp b/Engine/Source/Runtime/Experimental/Chaos/Private/Chaos/PBDRigidClustering.cpp
--- a/Engine/Source/Runtime/Experimental/Chaos/Private/Chaos/PBDRigidClustering.cpp
+++ b/Engine/Source/Runtime/Experimental/Chaos/Private/Chaos/PBDRigidClustering.cpp
@@ -53,6 +53,35 @@ namespace Chaos
 	float GraphPropagationBasedCollisionFactor = 1;
 	FAutoConsoleVariableRef CVarGraphPropagationBasedCollisionFactor(TEXT("p.GraphPropagationBasedCollisionFactor"), GraphPropagationBasedCollisionFactor, TEXT("when p.GraphPropagationBasedCollisionImpulseProcessing is on, the percentage [0-1] of remaining damage that is distributed to the connected pieces"));
 
+	static FGeometryCollectionPhysicsProxy* GetConcreteProxy(FPBDRigidClusteredParticleHandle* ClusteredParticle)
+	{
+		if (ClusteredParticle)
+		{
+			if (IPhysicsProxyBase* Proxy = ClusteredParticle->PhysicsProxy())
+			{
+				if (Proxy->GetType() == EPhysicsProxyType::GeometryCollectionType)
+				{
+					return static_cast<FGeometryCollectionPhysicsProxy*>(Proxy);
+				}
+			}
+		}
+		return nullptr;
+	}
+
+	static const FGeometryCollectionPhysicsProxy* GetConcreteProxy(const FPBDRigidClusteredParticleHandle* ClusteredParticle)
+	{
+		if (ClusteredParticle)
+		{
+			if (const IPhysicsProxyBase* Proxy = ClusteredParticle->PhysicsProxy())
+			{
+				if (Proxy->GetType() == EPhysicsProxyType::GeometryCollectionType)
+				{
+					return static_cast<const FGeometryCollectionPhysicsProxy*>(Proxy);
+				}
+			}
+		}
+		return nullptr;
+	}
 	
 	//==========================================================================
 	// TPBDRigidClustering
@@ -368,29 +397,33 @@ namespace Chaos
 		return ActivatedChildren;
 	}
 
-	void FRigidClustering::SendBreakingEvent(FPBDRigidClusteredParticleHandle* ClusteredParticle)
+	void FRigidClustering::ResetAllEvents()
 	{
-		check(ClusteredParticle);
-		if (DoGenerateBreakingData)
-		{
-			FBreakingData& ClusterBreak = MAllClusterBreakings.AddDefaulted_GetRef();
-			ClusterBreak.Proxy = ClusteredParticle->PhysicsProxy();
-			ClusterBreak.Location = ClusteredParticle->X();
-			ClusterBreak.Velocity = ClusteredParticle->V();
-			ClusterBreak.AngularVelocity = ClusteredParticle->W();
-			ClusterBreak.Mass = ClusteredParticle->M();
-			if (ClusteredParticle->Geometry() && ClusteredParticle->Geometry()->HasBoundingBox())
-			{
-				ClusterBreak.BoundingBox = ClusteredParticle->Geometry()->BoundingBox();
-			}
-			if (ClusterBreak.Proxy->GetType() == EPhysicsProxyType::GeometryCollectionType)
+		ResetAllClusterBreakings();
+		ResetAllClusterCrumblings();
+		CrumbledSinceLastUpdate.Reset();
+	}
+
+	void FRigidClustering::SendBreakingEvent(FPBDRigidClusteredParticleHandle* ClusteredParticle, bool bFromCrumble)
+	{
+		// only emit break event if the proxy needs it 
+		if (FGeometryCollectionPhysicsProxy* ConcreteProxy = GetConcreteProxy(ClusteredParticle))
+		{
+			const FSimulationParameters& SimParams = ConcreteProxy->GetSimParameters();
+			if (SimParams.bGenerateBreakingData)
 			{
-				const FGeometryCollectionPhysicsProxy* ConcreteProxy = static_cast<FGeometryCollectionPhysicsProxy*>(ClusterBreak.Proxy);
+				FBreakingData& ClusterBreak = MAllClusterBreakings.AddDefaulted_GetRef();
+				ClusterBreak.Proxy = ClusteredParticle->PhysicsProxy();
+				ClusterBreak.Location = ClusteredParticle->X();
+				ClusterBreak.Velocity = ClusteredParticle->V();
+				ClusterBreak.AngularVelocity = ClusteredParticle->W();
+				ClusterBreak.Mass = ClusteredParticle->M();
+				if (ClusteredParticle->Geometry() && ClusteredParticle->Geometry()->HasBoundingBox())
+				{
+					ClusterBreak.BoundingBox = ClusteredParticle->Geometry()->BoundingBox();
+				}
 				ClusterBreak.TransformGroupIndex = ConcreteProxy->GetTransformGroupIndexFromHandle(ClusteredParticle);
-			}
-			else
-			{
-				ClusterBreak.TransformGroupIndex = INDEX_NONE;
+				ClusterBreak.bFromCrumble = bFromCrumble;
 			}
 		}
 	}
@@ -398,33 +431,29 @@ namespace Chaos
 	
 	void FRigidClustering::SendCrumblingEvent(FPBDRigidClusteredParticleHandle* ClusteredParticle)
 	{
-		check(ClusteredParticle);
-		if (IPhysicsProxyBase* Proxy = ClusteredParticle->PhysicsProxy())
+		// only emit crumble events if the proxy needs it 
+		if (FGeometryCollectionPhysicsProxy* ConcreteProxy = GetConcreteProxy(ClusteredParticle))
 		{
-			if (Proxy->GetType() == EPhysicsProxyType::GeometryCollectionType)
+			const FSimulationParameters& SimParams = ConcreteProxy->GetSimParameters(); 
+			if (SimParams.bGenerateCrumblingData)
 			{
-				FGeometryCollectionPhysicsProxy* ConcreteProxy = static_cast<FGeometryCollectionPhysicsProxy*>(Proxy);
-				FSimulationParameters& SimParams = ConcreteProxy->GetSimParameters(); 
-				if (SimParams.bGenerateCrumblingData)
+				FCrumblingData& ClusterCrumbling = MAllClusterCrumblings.AddDefaulted_GetRef();
+				ClusterCrumbling.Proxy = ClusteredParticle->PhysicsProxy();
+				ClusterCrumbling.Location = ClusteredParticle->X();
+				ClusterCrumbling.Orientation = ClusteredParticle->R();
+				ClusterCrumbling.LinearVelocity = ClusteredParticle->V();
+				ClusterCrumbling.AngularVelocity = ClusteredParticle->W();
+				ClusterCrumbling.Mass = ClusteredParticle->M();
+				if (ClusteredParticle->Geometry() && ClusteredParticle->Geometry()->HasBoundingBox())
 				{
-					FCrumblingData& ClusterCrumbling = MAllClusterCrumblings.AddDefaulted_GetRef();
-					ClusterCrumbling.Proxy = ClusteredParticle->PhysicsProxy();
-					ClusterCrumbling.Location = ClusteredParticle->X();
-					ClusterCrumbling.Orientation = ClusteredParticle->R();
-					ClusterCrumbling.LinearVelocity = ClusteredParticle->V();
-					ClusterCrumbling.AngularVelocity = ClusteredParticle->W();
-					ClusterCrumbling.Mass = ClusteredParticle->M();
-					if (ClusteredParticle->Geometry() && ClusteredParticle->Geometry()->HasBoundingBox())
-					{
-						ClusterCrumbling.LocalBounds = ClusteredParticle->Geometry()->BoundingBox();
-					}
-					if (SimParams.bGenerateCrumblingChildrenData)
+					ClusterCrumbling.LocalBounds = ClusteredParticle->Geometry()->BoundingBox();
+				}
+				if (SimParams.bGenerateCrumblingChildrenData)
+				{
+					// when sending this event, children are still attached
+					if (const FRigidHandleArray* Children = MChildren.Find(ClusteredParticle))
 					{
-						// when sending this event, children are still attached
-						if (const FRigidHandleArray* Children = MChildren.Find(ClusteredParticle))
-						{
-							ConcreteProxy->GetTransformGroupIndicesFromHandles(*Children, ClusterCrumbling.Children);
-						}
+						ConcreteProxy->GetTransformGroupIndicesFromHandles(*Children, ClusterCrumbling.Children);
 					}
 				}
 			}
@@ -621,19 +650,16 @@ namespace Chaos
 		bool bUseDamagePropagation = false;
 		float BreakDamagePropagationFactor = 0.0f;
 		float ShockDamagePropagationFactor = 0.0f;
-		FPBDRigidClusteredParticleHandle* ClusteredParent = ClusteredParticle->CastToClustered();
-		if (const IPhysicsProxyBase* Proxy = ClusteredParent->PhysicsProxy())
+		if (const FGeometryCollectionPhysicsProxy* ConcreteProxy = GetConcreteProxy(ClusteredParticle))
 		{
-			if (Proxy->GetType() == EPhysicsProxyType::GeometryCollectionType)
-			{
-				const FGeometryCollectionPhysicsProxy* ConcreteProxy = static_cast<const FGeometryCollectionPhysicsProxy*>(Proxy);
-				bUseDamagePropagation = ConcreteProxy->GetSimParameters().bUseDamagePropagation;
-				BreakDamagePropagationFactor = ConcreteProxy->GetSimParameters().BreakDamagePropagationFactor;
-				ShockDamagePropagationFactor = ConcreteProxy->GetSimParameters().ShockDamagePropagationFactor;
-			}
+			const FSimulationParameters& SimParams = ConcreteProxy->GetSimParameters();
+			bUseDamagePropagation = SimParams.bUseDamagePropagation;
+			BreakDamagePropagationFactor = SimParams.BreakDamagePropagationFactor;
+			ShockDamagePropagationFactor = SimParams.ShockDamagePropagationFactor;
 		}
 
 		TArray<FPBDRigidParticleHandle*>& Children = MChildren[ClusteredParticle];
+		const bool bParentCrumbled = CrumbledSinceLastUpdate.Contains(ClusteredParticle);
 
 		// only used for propagation
 		TMap<FPBDRigidParticleHandle*, FReal> AppliedStrains;
@@ -664,7 +690,7 @@ namespace Chaos
 				// 1 entry at a time.
 				Children.RemoveAtSwap(ChildIdx, 1, false);
 
-				SendBreakingEvent(Child);
+				SendBreakingEvent(Child, bParentCrumbled);
 			}
 			if (bUseDamagePropagation)
 			{
@@ -883,8 +909,7 @@ namespace Chaos
 				SCOPE_CYCLE_COUNTER(STAT_UpdateDirtyImpulses);
 				BreakingModel();
 			} // end if MCollisionImpulseArrayDirty
-
-		} // end if MParticles.Size()
+		}
 		Timer.Stop();
 		UE_LOG(LogChaos, Verbose, TEXT("Cluster Break Update Time is %f"), Time);
 	}
@@ -1150,13 +1175,9 @@ namespace Chaos
 				{
 					// gather propagation information from the parent proxy
 					bool bUseDamagePropagation = false;
-					if (const IPhysicsProxyBase* Proxy = Cluster->PhysicsProxy())
+					if (const FGeometryCollectionPhysicsProxy* ConcreteProxy = GetConcreteProxy(Cluster))
 					{
-						if (Proxy->GetType() == EPhysicsProxyType::GeometryCollectionType)
-						{
-							const FGeometryCollectionPhysicsProxy* ConcreteProxy = static_cast<const FGeometryCollectionPhysicsProxy*>(Proxy);
-							bUseDamagePropagation = ConcreteProxy->GetSimParameters().bUseDamagePropagation;
-						}
+						bUseDamagePropagation = ConcreteProxy->GetSimParameters().bUseDamagePropagation;
 					}
 
 					if (bUseDamagePropagation)
@@ -1238,15 +1259,6 @@ namespace Chaos
 		ClusteredParticle->ClusterGroupIndex() = 0;
 	}
 
-	void FRigidClustering::DisableParticleWithBreakEvent(FPBDRigidClusteredParticleHandle* ClusteredParticle)
-	{
-		if (ClusteredParticle)
-		{
-			DisableCluster(ClusteredParticle);
-			SendBreakingEvent(ClusteredParticle);
-		}
-	}
-
 	FPBDRigidClusteredParticleHandle*
 	FRigidClustering::DestroyClusterParticle(
 		FPBDRigidClusteredParticleHandle* ClusteredParticle,
@@ -1333,6 +1345,7 @@ namespace Chaos
 			}
 			if (ChildrenHandles->Num() > 0)
 			{
+				CrumbledSinceLastUpdate.Add(ClusteredParticle);
 				SendCrumblingEvent(ClusteredParticle);
 			}
 			return true;
@@ -1362,6 +1375,7 @@ namespace Chaos
 					}
 					if (Children->Num() > 0)
 					{
+						CrumbledSinceLastUpdate.Add(ClusteredHandle);
 						SendCrumblingEvent(ClusteredHandle);
 					}
 				}
diff --git a/Engine/Source/Runtime/Experimental/Chaos/Private/EventDefaults.cpp b/Engine/Source/Runtime/Experimental/Chaos/Private/EventDefaults.cpp
--- a/Engine/Source/Runtime/Experimental/Chaos/Private/EventDefaults.cpp
+++ b/Engine/Source/Runtime/Experimental/Chaos/Private/EventDefaults.cpp
@@ -316,61 +316,33 @@ namespace Chaos
 			if (!Solver->GetEventFilters()->IsBreakingEventEnabled())
 				return;
 
-			FBreakingDataArray& AllBreakingDataArray = BreakingEventData.BreakingData.AllBreakingsArray;
-			TMap<IPhysicsProxyBase*, TArray<int32>>& AllBreakingIndicesByPhysicsProxy = BreakingEventData.PhysicsProxyToBreakingIndices.PhysicsProxyToIndicesMap;
+			FBreakingDataArray& FilteredBreakingDataArray = BreakingEventData.BreakingData.AllBreakingsArray;
+			TMap<IPhysicsProxyBase*, TArray<int32>>& FilteredBreakingIndicesByPhysicsProxy = BreakingEventData.PhysicsProxyToBreakingIndices.PhysicsProxyToIndicesMap;
 
 			if (bResetData)
 			{
-				AllBreakingDataArray.Reset();
-				AllBreakingIndicesByPhysicsProxy.Reset();
+				FilteredBreakingDataArray.Reset();
+				FilteredBreakingIndicesByPhysicsProxy.Reset();
 			}
 			BreakingEventData.BreakingData.TimeCreated = Solver->MTime;
 
 			const auto* Evolution = Solver->GetEvolution();
 			const FPBDRigidParticles& Particles = Evolution->GetParticles().GetDynamicParticles();
-			const TArray<FBreakingData>& AllBreakingsArray = Evolution->GetRigidClustering().GetAllClusterBreakings();
-			const TArrayCollectionArray<ClusterId>& ClusterIdsArray = Evolution->GetRigidClustering().GetClusterIdsArray();
-
-#if TODO_REIMPLEMENT_RIGID_CLUSTERING
-			const Chaos::FPBDRigidsSolver::FClusteringType::FClusterMap& ParentToChildrenMap = Evolution->GetRigidClustering().GetChildrenMap();
-#endif
-			
-			if (AllBreakingsArray.Num() > 0)
+			const TArray<FBreakingData>& AllClusterBreakings = Evolution->GetRigidClustering().GetAllClusterBreakings();
+	
+			if (AllClusterBreakings.Num() > 0)
 			{
-				for (int32 Idx = 0; Idx < AllBreakingsArray.Num(); ++Idx)
+				const FSolverBreakingEventFilter* SolverBreakingEventFilter = Solver->GetEventFilters()->GetBreakingFilter();
+				for (int32 Idx = 0; Idx < AllClusterBreakings.Num(); ++Idx)
 				{					
-					FBreakingData BreakingData;
-					BreakingData.Location = AllBreakingsArray[Idx].Location;
-					BreakingData.Velocity = AllBreakingsArray[Idx].Velocity;
-					BreakingData.AngularVelocity = AllBreakingsArray[Idx].AngularVelocity;
-					BreakingData.Mass = AllBreakingsArray[Idx].Mass;
-					BreakingData.Proxy = AllBreakingsArray[Idx].Proxy;
-					BreakingData.BoundingBox = AllBreakingsArray[Idx].BoundingBox;
-					BreakingData.TransformGroupIndex = AllBreakingsArray[Idx].TransformGroupIndex;
-
-					const FSolverBreakingEventFilter* SolverBreakingEventFilter = Solver->GetEventFilters()->GetBreakingFilter();
-					if (!SolverBreakingEventFilter->Enabled() || SolverBreakingEventFilter->Pass(BreakingData))
+					const FBreakingData& ClusterBreaking = AllClusterBreakings[Idx];
+					if (!SolverBreakingEventFilter->Enabled() || SolverBreakingEventFilter->Pass(ClusterBreaking))
 					{
-						int32 NewIdx = AllBreakingDataArray.Add(FBreakingData());
-						FBreakingData& BreakingDataArrayItem = AllBreakingDataArray[NewIdx];
-						BreakingDataArrayItem = BreakingData;
-
-						// Add to AllBreakingIndicesByPhysicsProxy
-						AllBreakingIndicesByPhysicsProxy.FindOrAdd(BreakingData.Proxy).Add(FEventManager::EncodeCollisionIndex(NewIdx, false));
-
-#if 0 // #todo
-						// If AllBreakingsArray[Idx].ParticleIndex is a cluster store an index for a mesh in this cluster
-						if (ClusterIdsArray[AllBreakingsArray[Idx].ParticleIndex].NumChildren > 0)
-						{
-							int32 ParticleIndexMesh = GetParticleIndexMesh(ParentToChildrenMap, AllBreakingsArray[Idx].ParticleIndex);
-							ensure(ParticleIndexMesh != INDEX_NONE);
-							BreakingDataArrayItem.ParticleIndexMesh = ParticleIndexMesh;
-						}
-#endif
-					}					
+						const int32 NewIndex = FilteredBreakingDataArray.Emplace(ClusterBreaking);
+						FilteredBreakingIndicesByPhysicsProxy.FindOrAdd(ClusterBreaking.Proxy).Add(FEventManager::EncodeCollisionIndex(NewIndex, false));
+					}
 				}
 			}
-		
 		});
 	}
 
diff --git a/Engine/Source/Runtime/Experimental/Chaos/Private/PBDRigidsSolver.cpp b/Engine/Source/Runtime/Experimental/Chaos/Private/PBDRigidsSolver.cpp
--- a/Engine/Source/Runtime/Experimental/Chaos/Private/PBDRigidsSolver.cpp
+++ b/Engine/Source/Runtime/Experimental/Chaos/Private/PBDRigidsSolver.cpp
@@ -456,11 +456,10 @@ namespace Chaos
 				MSolver->CompleteSceneSimulation();
 			}
 
-			// reset all clustering event needs to be after CompleteSceneSimulation to make sure the cache recording gets them before they get removed
+			// reset all clustering events needs to be after CompleteSceneSimulation to make sure the cache recording gets them before they get removed
 			// they cannot be right after the presolve callback ( as they were before ) because they will cause geometry collection replicated clients to miss them
 			// Todo(chaos) we should probably move all of the solver event reset here in the future
-			MSolver->GetEvolution()->GetRigidClustering().ResetAllClusterBreakings();
-			MSolver->GetEvolution()->GetRigidClustering().ResetAllClusterCrumblings();
+			MSolver->GetEvolution()->GetRigidClustering().ResetAllEvents();
 
 			if (FRewindData* RewindData = MSolver->GetRewindData())
 			{
diff --git a/Engine/Source/Runtime/Experimental/Chaos/Public/Chaos/ExternalCollisionData.h b/Engine/Source/Runtime/Experimental/Chaos/Public/Chaos/ExternalCollisionData.h
--- a/Engine/Source/Runtime/Experimental/Chaos/Public/Chaos/ExternalCollisionData.h
+++ b/Engine/Source/Runtime/Experimental/Chaos/Public/Chaos/ExternalCollisionData.h
@@ -177,6 +177,7 @@ namespace Chaos
 			, Mass((FReal)0.0)
 			, BoundingBox(FAABB3(FVec3((FReal)0.0), FVec3((FReal)0.0)))
 			, TransformGroupIndex(INDEX_NONE)
+			, bFromCrumble(false)
 		{}
 
 		// The pointer to the proxy should be used with caution on the Game Thread.
@@ -192,6 +193,7 @@ namespace Chaos
 		FReal Mass;
 		Chaos::FAABB3 BoundingBox;
 		int32 TransformGroupIndex;
+		bool bFromCrumble;
 	};
 
 	/*
diff --git a/Engine/Source/Runtime/Experimental/Chaos/Public/Chaos/PBDRigidClustering.h b/Engine/Source/Runtime/Experimental/Chaos/Public/Chaos/PBDRigidClustering.h
--- a/Engine/Source/Runtime/Experimental/Chaos/Public/Chaos/PBDRigidClustering.h
+++ b/Engine/Source/Runtime/Experimental/Chaos/Public/Chaos/PBDRigidClustering.h
@@ -255,6 +255,11 @@ public:
 	*/
 	TArrayCollectionArray<int32>& GetClusterGroupIndexArray() { return MParticles.ClusterGroupIndexArray(); }
 
+	/*
+	* Reset all events ( this include breaking, crumbling event and tracking data 
+	*/
+	void ResetAllEvents();
+
 	/*
 	*  Cluster Break Data
 	*     The cluster breaks can be used to seed particle emissions. 
@@ -346,7 +351,6 @@ public:
 	void ComputeStrainFromCollision(const FPBDCollisionConstraints& CollisionRule);
 	void ResetCollisionImpulseArray();
 	void DisableCluster(FPBDRigidClusteredParticleHandle* ClusteredParticle);
-	void DisableParticleWithBreakEvent(FPBDRigidClusteredParticleHandle* ClusteredParticle);
 
 	/*
 	* Connectivity
@@ -369,7 +373,7 @@ public:
 
 	void RemoveChildFromParent(FPBDRigidParticleHandle* Child, const FPBDRigidClusteredParticleHandle* ClusteredParent);
 
-	void SendBreakingEvent(FPBDRigidClusteredParticleHandle* ClusteredParticle);
+	void SendBreakingEvent(FPBDRigidClusteredParticleHandle* ClusteredParticle, bool bFromCrumble);
 	void SendCrumblingEvent(FPBDRigidClusteredParticleHandle* ClusteredParticle);
 
 	TSet<FPBDRigidParticleHandle*> ReleaseClusterParticlesImpl(
@@ -399,7 +403,9 @@ private:
 	TArray<FBreakingData> MAllClusterBreakings;
 
 	TArray<FCrumblingData> MAllClusterCrumblings;
-	
+
+	TSet<FPBDRigidClusteredParticleHandle*> CrumbledSinceLastUpdate;
+
 	FReal MClusterConnectionFactor;
 	FClusterCreationParameters::EConnectionMethod MClusterUnionConnectionType;
 };
diff --git a/Engine/Source/Runtime/Experimental/ChaosSolverEngine/Private/Chaos/ChaosGameplayEventDispatcher.cpp b/Engine/Source/Runtime/Experimental/ChaosSolverEngine/Private/Chaos/ChaosGameplayEventDispatcher.cpp
--- a/Engine/Source/Runtime/Experimental/ChaosSolverEngine/Private/Chaos/ChaosGameplayEventDispatcher.cpp
+++ b/Engine/Source/Runtime/Experimental/ChaosSolverEngine/Private/Chaos/ChaosGameplayEventDispatcher.cpp
@@ -23,10 +23,23 @@ FChaosBreakEvent::FChaosBreakEvent()
 	, Location(FVector::ZeroVector)
 	, Velocity(FVector::ZeroVector)
 	, AngularVelocity(FVector::ZeroVector)
+	, Extents(FVector::ZeroVector)
 	, Mass(0.0f)
 	, Index(INDEX_NONE)
+	, bFromCrumble(false)
 {
+}
 
+FChaosBreakEvent::FChaosBreakEvent(const Chaos::FBreakingData& BreakingData)
+	: Component(nullptr)
+	, Location(BreakingData.Location)
+	, Velocity(BreakingData.Velocity)
+	, AngularVelocity(BreakingData.AngularVelocity)
+	, Extents(BreakingData.BoundingBox.Extents())
+	, Mass(BreakingData.Mass)
+	, Index(BreakingData.TransformGroupIndex)
+	, bFromCrumble(BreakingData.bFromCrumble)
+{
 }
 
 FChaosRemovalEvent::FChaosRemovalEvent()
@@ -482,27 +495,22 @@ void UChaosGameplayEventDispatcher::HandleBreakingEvents(const Chaos::FBreakingE
 	{
 		LastBreakingDataTime = BreakingDataTimestamp;
 
-		Chaos::FBreakingDataArray const& BreakingData = Event.BreakingData.AllBreakingsArray;
+		const Chaos::FBreakingDataArray& BreakingDataArray = Event.BreakingData.AllBreakingsArray;
 
 		// let's assume breaks are very rare, so we will iterate breaks instead of registered components for now
-		const int32 NumBreaks = BreakingData.Num();
+		const int32 NumBreaks = BreakingDataArray.Num();
 		if (NumBreaks > 0)
 		{
-			for (Chaos::FBreakingData const& BreakingDataItem : BreakingData)
+			for (const Chaos::FBreakingData& BreakingData : BreakingDataArray)
 			{	
-				if (BreakingDataItem.Proxy)
+				if (BreakingData.Proxy)
 				{
-					UPrimitiveComponent* const PrimComp = Cast<UPrimitiveComponent>(BreakingDataItem.Proxy->GetOwner());
+					UPrimitiveComponent* const PrimComp = Cast<UPrimitiveComponent>(BreakingData.Proxy->GetOwner());
 					if (PrimComp && BreakEventRegistrations.Contains(PrimComp))
 					{
 						// queue them up so we can release the physics data before trigging BP events
-						FChaosBreakEvent& BreakEvent = PendingBreakEvents.AddZeroed_GetRef();
+						FChaosBreakEvent& BreakEvent = PendingBreakEvents.Emplace_GetRef(BreakingData);
 						BreakEvent.Component = PrimComp;
-						BreakEvent.Location = BreakingDataItem.Location;
-						BreakEvent.Velocity = BreakingDataItem.Velocity;
-						BreakEvent.AngularVelocity = BreakingDataItem.AngularVelocity;
-						BreakEvent.Mass = BreakingDataItem.Mass;
-						BreakEvent.Index = BreakingDataItem.TransformGroupIndex;
 					}
 				}
 			}
diff --git a/Engine/Source/Runtime/Experimental/ChaosSolverEngine/Public/Chaos/ChaosGameplayEventDispatcher.h b/Engine/Source/Runtime/Experimental/ChaosSolverEngine/Public/Chaos/ChaosGameplayEventDispatcher.h
--- a/Engine/Source/Runtime/Experimental/ChaosSolverEngine/Public/Chaos/ChaosGameplayEventDispatcher.h
+++ b/Engine/Source/Runtime/Experimental/ChaosSolverEngine/Public/Chaos/ChaosGameplayEventDispatcher.h
@@ -16,6 +16,7 @@ namespace Chaos
 	struct FSleepingEventData;
 	struct FRemovalEventData;
 	struct FCrumblingEventData;
+	struct FBreakingData;
 }
 
 USTRUCT(BlueprintType)
@@ -26,6 +27,7 @@ struct CHAOSSOLVERENGINE_API FChaosBreakEvent
 public:
 
 	FChaosBreakEvent();
+	FChaosBreakEvent(const Chaos::FBreakingData& BreakingData);
 
 	/** primitive component involved in the break event */
 	UPROPERTY(BlueprintReadOnly, Category = "Break Event")
@@ -43,6 +45,10 @@ public:
 	UPROPERTY(BlueprintReadOnly, Category = "Break Event")
 	FVector AngularVelocity;
 
+	/** Extents of the bounding box */
+	UPROPERTY(BlueprintReadOnly, Category = "Break Event")
+	FVector Extents;
+
 	/** Mass of the breaking particle  */
 	UPROPERTY(BlueprintReadOnly, Category = "Break Event")
 	float Mass;
@@ -50,6 +56,10 @@ public:
 	/** Index of the geometry collection bone if positive */
 	UPROPERTY(BlueprintReadOnly, Category = "Break Event")
 	int32 Index;
+
+	/** Whether the break event originated from a crumble event */
+	UPROPERTY(BlueprintReadOnly, Category = "Break Event")
+	bool bFromCrumble;
 };
 
 USTRUCT(BlueprintType)
