Iris - Don't try to load maps.

diff --git a/Engine/Source/Runtime/Experimental/Iris/Core/Private/Iris/ReplicationSystem/ObjectReferenceCache.cpp b/Engine/Source/Runtime/Experimental/Iris/Core/Private/Iris/ReplicationSystem/ObjectReferenceCache.cpp
--- a/Engine/Source/Runtime/Experimental/Iris/Core/Private/Iris/ReplicationSystem/ObjectReferenceCache.cpp
+++ b/Engine/Source/Runtime/Experimental/Iris/Core/Private/Iris/ReplicationSystem/ObjectReferenceCache.cpp
@@ -118,6 +118,16 @@ bool FObjectReferenceCache::CanClientLoadObjectInternal(const UObject* Object, b
 	{
 		return false;
 	}
+
+#if WITH_EDITOR
+	// For objects using external package, we need to do the test on the package of their outer most object
+	UObject* OutermostObject = Object->GetOutermostObject();
+	if (OutermostObject && OutermostObject->GetPackage()->ContainsMap())
+	{
+		return false;
+	}
+#endif
+
 	// We can load everything else
 	return true;
 }
