Chaos : fix crash where geometruy collection proxy are added to the same array everytime they get dirty , but only removed only once causing dangling pointers
This was also causing InitializeBodiesPT to be called everytime the GC proxy was dirty too , causing performance problem potentially
This is a knock on bug of CL 23243544
- Added a flag to check if initializePT has been alkready called and only calling it and adding the proxy to the array when it's false

diff --git a/Engine/Source/Runtime/Experimental/Chaos/Private/PBDRigidsSolver.cpp b/Engine/Source/Runtime/Experimental/Chaos/Private/PBDRigidsSolver.cpp
--- a/Engine/Source/Runtime/Experimental/Chaos/Private/PBDRigidsSolver.cpp
+++ b/Engine/Source/Runtime/Experimental/Chaos/Private/PBDRigidsSolver.cpp
@@ -1178,9 +1178,12 @@ namespace Chaos
 					case EPhysicsProxyType::GeometryCollectionType:
 					{
 						auto Proxy = static_cast<FGeometryCollectionPhysicsProxy*>(Dirty.Proxy);
-						// Finish registration on the physics thread...
-						Proxy->InitializeBodiesPT(this, GetParticles());
-						GeometryCollectionPhysicsProxies_Internal.Add(Proxy);
+						if (!Proxy->IsInitializedOnPhysicsThread())
+						{
+							// Finish registration on the physics thread...
+							Proxy->InitializeBodiesPT(this, GetParticles());
+							GeometryCollectionPhysicsProxies_Internal.Add(Proxy);
+						}
 						Proxy->PushToPhysicsState();
 						// Currently no push needed for geometry collections and they handle the particle creation internally
 						// #TODO This skips the rewind data push so GC will not be rewindable until resolved.
diff --git a/Engine/Source/Runtime/Experimental/Chaos/Private/PhysicsProxy/GeometryCollectionPhysicsProxy.cpp b/Engine/Source/Runtime/Experimental/Chaos/Private/PhysicsProxy/GeometryCollectionPhysicsProxy.cpp
--- a/Engine/Source/Runtime/Experimental/Chaos/Private/PhysicsProxy/GeometryCollectionPhysicsProxy.cpp
+++ b/Engine/Source/Runtime/Experimental/Chaos/Private/PhysicsProxy/GeometryCollectionPhysicsProxy.cpp
@@ -856,6 +856,8 @@ void FGeometryCollectionPhysicsProxy::InitializeBodiesPT(Chaos::FPBDRigidsSolver
 	const FGeometryCollection* RestCollection = Parameters.RestCollection;
 	const FGeometryDynamicCollection& DynamicCollection = PhysicsThreadCollection;
 
+	bIsInitializedOnPhysicsThread = true;
+
 	if (Parameters.Simulating && RestCollection)
 	{
 		const TManagedArray<int32>& TransformIndex = RestCollection->TransformIndex;
diff --git a/Engine/Source/Runtime/Experimental/Chaos/Public/PhysicsProxy/GeometryCollectionPhysicsProxy.h b/Engine/Source/Runtime/Experimental/Chaos/Public/PhysicsProxy/GeometryCollectionPhysicsProxy.h
--- a/Engine/Source/Runtime/Experimental/Chaos/Public/PhysicsProxy/GeometryCollectionPhysicsProxy.h
+++ b/Engine/Source/Runtime/Experimental/Chaos/Public/PhysicsProxy/GeometryCollectionPhysicsProxy.h
@@ -172,6 +172,8 @@ public:
 	void Initialize(Chaos::FPBDRigidsEvolutionBase* Evolution);
 	void Reset() { }
 
+	bool IsInitializedOnPhysicsThread() const { return bIsInitializedOnPhysicsThread; }
+
 	/** 
 	 * Finish initialization on the physics thread. 
 	 *
@@ -458,6 +460,9 @@ protected:
 	
 private:
 
+	/* set to true once InitializeBodiesPT has been called*/
+	bool bIsInitializedOnPhysicsThread = false;
+
 	FSimulationParameters Parameters;
 	TArray<FFieldSystemCommand> Commands;
 
