Iris - Fixed subobject RPCs not getting sent in due time.

diff --git a/Engine/Source/Runtime/Experimental/Iris/Core/Private/Iris/ReplicationSystem/ReplicationWriter.cpp b/Engine/Source/Runtime/Experimental/Iris/Core/Private/Iris/ReplicationSystem/ReplicationWriter.cpp
--- a/Engine/Source/Runtime/Experimental/Iris/Core/Private/Iris/ReplicationSystem/ReplicationWriter.cpp
+++ b/Engine/Source/Runtime/Experimental/Iris/Core/Private/Iris/ReplicationSystem/ReplicationWriter.cpp
@@ -2034,7 +2034,7 @@ FReplicationWriter::EWriteObjectStatus FReplicationWriter::WriteObjectInBatch(FN
 			if (OutBatchInfo.ObjectInfos.Num() > BatchObjectInfoCount)
 			{
 				const FBatchObjectInfo& SubObjectEntry = OutBatchInfo.ObjectInfos.Last();
-				bHasDirtySubObjects = SubObjectEntry.bHasDirtySubObjects || SubObjectEntry.bHasUnsentAttachments;
+				bHasDirtySubObjects |= SubObjectEntry.bHasDirtySubObjects || SubObjectEntry.bHasUnsentAttachments;
 			}
 		}
 
