Dataflow : Section rendering in the scene proxy.

diff --git a/Engine/Plugins/Experimental/Dataflow/Source/DataflowEditor/Private/Dataflow/DataflowEditorViewport.cpp b/Engine/Plugins/Experimental/Dataflow/Source/DataflowEditor/Private/Dataflow/DataflowEditorViewport.cpp
--- a/Engine/Plugins/Experimental/Dataflow/Source/DataflowEditor/Private/Dataflow/DataflowEditorViewport.cpp
+++ b/Engine/Plugins/Experimental/Dataflow/Source/DataflowEditor/Private/Dataflow/DataflowEditorViewport.cpp
@@ -20,11 +20,6 @@
 #include "SCommonEditorViewportToolbarBase.h"
 #include "Settings/EditorStyleSettings.h"
 
-
-static int32 bEnableDataflowClientDrawing = 0;
-FAutoConsoleVariableRef CVarbEnableDataflowClientDrawing(
-	TEXT("p.Dataflow.Rendering.ClientDraw"), bEnableDataflowClientDrawing,TEXT("Enable the client rendering, will be slower on complex bodies. [default:0]"));
-
 SDataflowEditorViewport::SDataflowEditorViewport()
 {
 	PreviewScene = MakeShareable(new FAdvancedPreviewScene(FPreviewScene::ConstructionValues()));
@@ -176,7 +171,7 @@ void FDataflowEditorViewportClient::ProcessClick(FSceneView& View, HHitProxy* Hi
 		{
 			if (HitProxy && HitProxy->IsA(HDataflowActor::StaticGetType()))
 			{
-					SelectionState.bComponentSelected = !SelectionState.bComponentSelected;
+				SelectionState.bComponentSelected = !SelectionState.bComponentSelected;
 			}
 			else
 			{
@@ -191,120 +186,6 @@ void FDataflowEditorViewportClient::ProcessClick(FSceneView& View, HHitProxy* Hi
 	}
 }
 
-
-void FDataflowEditorViewportClient::Draw(const FSceneView* View, FPrimitiveDrawInterface* PDI)
-{
-	Super::Draw(View, PDI);
-	
-	if (bEnableDataflowClientDrawing)
-	{
-		GeometryCollection::Facades::FRenderingFacade Facade(RenderCollection);
-		if (Facade.IsValid())
-		{
-			if (bRenderMesh)
-			{
-				MeshBuilder.Reset(new FDynamicMeshBuilder(View->GetFeatureLevel()));
-
-				FLinearColor UnselectedColor = GEngine->C_BrushWire;
-				//UnselectedColor.A = .1f;
-
-				FLinearColor SelectedColor = GetDefault<UEditorStyleSettings>()->SelectionColor;
-				SelectedColor.A = .5f;
-
-				// Allocate the material proxy and register it so it can be deleted properly once the rendering is done with it.
-				FDynamicColoredMaterialRenderProxy* SelectedColorInstance = new FDynamicColoredMaterialRenderProxy(GEngine->GeomMaterial->GetRenderProxy(), SelectedColor);
-				PDI->RegisterDynamicResource(SelectedColorInstance);
-
-				FDynamicColoredMaterialRenderProxy* UnselectedColorInstance = new FDynamicColoredMaterialRenderProxy(GEngine->GeomMaterial->GetRenderProxy(), UnselectedColor);
-				PDI->RegisterDynamicResource(UnselectedColorInstance);
-
-				// mesh surface
-				bool bRenderSurface = false;
-				if (bRenderSurface)
-				{
-					MeshBuilder->AddVertices(VertexBuffer);
-					MeshBuilder->AddTriangles(IndexBuffer);
-					MeshBuilder->Draw(PDI, FMatrix::Identity, UnselectedColorInstance, ESceneDepthPriorityGroup::SDPG_World, false);
-				}
-				else
-				{
-					// wireframes
-					auto ToF = [](FVector3f V) { return FVector(V.X, V.Y, V.Z); };
-					for (int i = 0; i < IndexBuffer.Num(); i += 3)
-					{
-						PDI->DrawLine(ToF(VertexBuffer[i].Position), ToF(VertexBuffer[i + 1].Position), SelectedColor, SDPG_World, 1.f);
-						PDI->DrawLine(ToF(VertexBuffer[i].Position), ToF(VertexBuffer[i + 2].Position), SelectedColor, SDPG_World, 1.f);
-						PDI->DrawLine(ToF(VertexBuffer[i + 1].Position), ToF(VertexBuffer[i + 2].Position), SelectedColor, SDPG_World, 1.f);
-					}
-				}
-			}
-		}
-	}
-}
-
-void FDataflowEditorViewportClient::ReleaseRenderStructures()
-{
-	RenderCollection = FManagedArrayCollection();
-
-	// MeshBuilder
-	if (bEnableDataflowClientDrawing)
-	{
-		bRenderMesh = false;
-		IndexBuffer = TArray<uint32>();
-		VertexBuffer = TArray<FDynamicMeshVertex>();
-		MeshBuilder.Reset(nullptr);
-	}
-}
-
-void FDataflowEditorViewportClient::RenderIntoStructures()
-{
-	GeometryCollection::Facades::FRenderingFacade Facade(RenderCollection);
-	if (bEnableDataflowClientDrawing)
-	{
-		if (Facade.IsValid())
-		{
-			// @todo(flesh) : Setup buffers that will get directly copied to 
-			//                the render. This is only called during invalidation
-			//                of the graph, so any state that can be precalculated 
-			//                should be set here and just pushed to the render in
-			//                the ::Draw() method.
-			// 
-			//	Note: Draw will still be very slow for complicated geometry, 
-			//        we still need to rely on Component rendering through
-			//        a scene proxy as the primary rendering path. 
-			{
-				const TManagedArray< FIntVector >& Indices = Facade.GetIndices();
-				const TManagedArray< FVector3f >& Vertices = Facade.GetVertices();
-				if (Vertices.Num() && Indices.Num())
-				{
-					IndexBuffer.SetNum(Indices.Num() * 3);
-					for (int Idx = 0; Idx < Indices.Num() * 3; Idx++)
-					{
-						IndexBuffer[Idx] = Idx;
-					}
-
-					int Vdx = 0;
-					VertexBuffer.SetNum(Indices.Num() * 3);
-					for (const FIntVector& Face : Indices)
-					{
-						FIntVector Tri = FIntVector(Face[0], Face[1], Face[2]);
-						const FVector3f Pos0(Vertices[Tri[0]]), Pos1(Vertices[Tri[1]]), Pos2(Vertices[Tri[2]]);
-						const FVector3f Normal = FVector3f::CrossProduct(Pos2 - Pos0, Pos1 - Pos0).GetSafeNormal();
-						const FVector3f Tangent = ((Pos1 + Pos2) * 0.5f - Pos0).GetSafeNormal();
-
-						VertexBuffer[Vdx++] = FDynamicMeshVertex(Pos0, Tangent, Normal, FVector2f(0.f, 0.f), FColor::White);
-						VertexBuffer[Vdx++] = FDynamicMeshVertex(Pos1, Tangent, Normal, FVector2f(0.f, 1.f), FColor::White);
-						VertexBuffer[Vdx++] = FDynamicMeshVertex(Pos2, Tangent, Normal, FVector2f(1.f, 1.f), FColor::White);
-					}
-
-					bRenderMesh = true;
-				}
-			}
-		}
-	}
-}
-
-
 void FDataflowEditorViewportClient::Tick(float DeltaSeconds)
 {
 	Super::Tick(DeltaSeconds);
@@ -322,19 +203,8 @@ void FDataflowEditorViewportClient::Tick(float DeltaSeconds)
 					Dataflow::FTimestamp SystemTimestamp = LatestTimestamp(Dataflow, Context.Get());
 					if (SystemTimestamp >= LastModifiedTimestamp)
 					{
-						ReleaseRenderStructures();
-
 						if (Dataflow->GetRenderTargets().Num())
 						{
-							GeometryCollection::Facades::FRenderingFacade Facade(RenderCollection);
-							for (const UDataflowEdNode* Target : Dataflow->GetRenderTargets())
-							{
-								Target->Render(Facade, Context);
-							}
-
-							// Client View Rendering
-							RenderIntoStructures();
-
 							// Component Object Rendering
 							DataflowComponent->ResetRenderTargets();
 							DataflowComponent->SetDataflow(Dataflow);
diff --git a/Engine/Plugins/Experimental/Dataflow/Source/DataflowEditor/Public/Dataflow/DataflowEditorViewport.h b/Engine/Plugins/Experimental/Dataflow/Source/DataflowEditor/Public/Dataflow/DataflowEditorViewport.h
--- a/Engine/Plugins/Experimental/Dataflow/Source/DataflowEditor/Public/Dataflow/DataflowEditorViewport.h
+++ b/Engine/Plugins/Experimental/Dataflow/Source/DataflowEditor/Public/Dataflow/DataflowEditorViewport.h
@@ -119,7 +119,6 @@ public:
 
 	// FEditorViewportClient interface
 	virtual void ProcessClick(FSceneView& View, HHitProxy* HitProxy, FKey Key, EInputEvent Event, uint32 HitX, uint32 HitY) override;
-	virtual void Draw(const FSceneView* View, FPrimitiveDrawInterface* PDI) override;
 	virtual void Tick(float DeltaSeconds) override;
 	// End of FEditorViewportClient
 
@@ -128,23 +127,14 @@ public:
 	virtual FString GetReferencerName() const override { return TEXT("FDataflowEditorViewportClient"); }
 
 private:
-	void RenderIntoStructures();
-	void ReleaseRenderStructures();
 
 	TWeakPtr<FDataflowEditorToolkit> DataflowEditorToolkitPtr = nullptr;
 	Dataflow::FTimestamp LastModifiedTimestamp = Dataflow::FTimestamp::Invalid;
 	ADataflowActor* DataflowActor = nullptr;
 
-	bool bIsSelected = false;
-	FManagedArrayCollection RenderCollection;
-
-	// Renderables
-	bool bRenderMesh = false;
-	TArray<uint32> IndexBuffer;
-	TArray<FDynamicMeshVertex> VertexBuffer;
-	TUniquePtr<FDynamicMeshBuilder> MeshBuilder;
-
-	// State
+	//
+	// Selection
+	//
 	FDataflowSelectionState SelectionState;
 	FDataflowSelectionState::EMode SelectionMode;
 
diff --git a/Engine/Plugins/Experimental/Dataflow/Source/DataflowEnginePlugin/Private/Dataflow/DataflowEngineSceneProxy.cpp b/Engine/Plugins/Experimental/Dataflow/Source/DataflowEnginePlugin/Private/Dataflow/DataflowEngineSceneProxy.cpp
--- a/Engine/Plugins/Experimental/Dataflow/Source/DataflowEnginePlugin/Private/Dataflow/DataflowEngineSceneProxy.cpp
+++ b/Engine/Plugins/Experimental/Dataflow/Source/DataflowEnginePlugin/Private/Dataflow/DataflowEngineSceneProxy.cpp
@@ -30,6 +30,8 @@ FDataflowEngineSceneProxy::~FDataflowEngineSceneProxy()
 
 void FDataflowEngineSceneProxy::CreateRenderThreadResources()
 {
+	static int foobar = INDEX_NONE;
+
 	check(ConstantData);
 	check(IsInRenderingThread());
 
@@ -50,24 +52,32 @@ void FDataflowEngineSceneProxy::CreateRenderThreadResources()
 	HitProxyIdBuffer.Init(TotalNumVertices);
 #endif
 
-	// Initialize points.
+	// Initialize
 	// Triangles are represented as two tris, all of whose vertices are
 	// coincident. The material then offsets them according to the signs of the
 	// vertex normals in a camera facing orientation. Size of the point is given
 	// by U0.
 	if (Facade.NumTriangles() > 0)
 	{
-		MeshBatchDatas.Emplace();
-		FDataflowTriangleSetMeshBatchData& MeshBatchData = MeshBatchDatas.Last();
-		MeshBatchData.MinVertexIndex = 0;
-		MeshBatchData.MaxVertexIndex = NumTriangleVertices - 1;
-		MeshBatchData.StartIndex = 0;
-		MeshBatchData.NumPrimitives = Facade.NumTriangles();
-		if (!RenderMaterial)
+		int32 NumGeom = Facade.NumGeometry();
+		for (int i = 0; i<NumGeom; i++)
 		{
-			RenderMaterial = UMaterial::GetDefaultMaterial(MD_Surface);
+			MeshBatchDatas.Emplace();
+			FDataflowTriangleSetMeshBatchData& MeshBatchData = MeshBatchDatas.Last();
+
+			MeshBatchData.MinVertexIndex = Facade.GetIndicesStartAttribute()[i] * 3;
+			MeshBatchData.MaxVertexIndex = MeshBatchData.MinVertexIndex + Facade.GetIndicesCountAttribute()[i] * 3 - 1;
+			MeshBatchData.FirstTriangleIndex = Facade.GetIndicesStartAttribute()[i] * 3;
+			MeshBatchData.NumTriangles = Facade.GetIndicesCountAttribute()[i];
+			MeshBatchData.GeomIndex = i;
+
+			if (!RenderMaterial)
+			{
+				RenderMaterial = UMaterial::GetDefaultMaterial(MD_Surface);
+			}
+			MeshBatchData.MaterialProxy = RenderMaterial->GetRenderProxy();
 		}
-		MeshBatchData.MaterialProxy = RenderMaterial->GetRenderProxy();
+
 
 		const TManagedArray<FIntVector>& Indices = Facade.GetIndices();
 		const TManagedArray<FVector3f>& Vertex = Facade.GetVertices();
@@ -119,6 +129,7 @@ void FDataflowEngineSceneProxy::CreateRenderThreadResources()
 		HitProxyIdBuffer.VertexColor(i) = DefaultHitProxy->Id.GetColor();
 #endif
 		});
+
 	}
 #if WITH_EDITOR
 	SetUsedMaterialForVerification({RenderMaterial});
@@ -162,7 +173,19 @@ void FDataflowEngineSceneProxy::DestroyRenderThreadResources()
 HHitProxy* FDataflowEngineSceneProxy::CreateHitProxies(UPrimitiveComponent* Component, TArray<TRefCountPtr<HHitProxy> >& OutHitProxies)
 {
 	check(!IsInRenderingThread());
-	OutHitProxies.Add(DefaultHitProxy);
+
+	GeometryCollection::Facades::FRenderingFacade Facade(*ConstantData);
+	if (Facade.NumTriangles() > 0)
+	{
+		int32 NumGeom = Facade.NumGeometry();
+		for (int i = 0; i < NumGeom; i++)
+		{
+			HDataflowActor * NodeProxy = new HDataflowActor(Component->GetOwner(), Component, i, INDEX_NONE);
+			LocalHitProxies.Add(NodeProxy);
+			Facade.ModifyHitProxyIndexAttribute()[i] = LocalHitProxies.Num() - 1;
+			OutHitProxies.Add(NodeProxy);
+		}
+	}
 	return FPrimitiveSceneProxy::CreateHitProxies(Component, OutHitProxies);
 }
 
@@ -181,6 +204,7 @@ void FDataflowEngineSceneProxy::GetDynamicMeshElements(
 {
 	QUICK_SCOPE_CYCLE_COUNTER(STAT_OverlaySceneProxy_GetDynamicMeshElements);
 	check(IsInRenderingThread());
+	const GeometryCollection::Facades::FRenderingFacade Facade(*ConstantData);
 
 	for (int32 ViewIndex = 0; ViewIndex < Views.Num(); ViewIndex++)
 	{
@@ -200,8 +224,8 @@ void FDataflowEngineSceneProxy::GetDynamicMeshElements(
 				BatchElement.PrimitiveUniformBufferResource = &DynamicPrimitiveUniformBuffer.UniformBuffer;
 				//BatchElement.PrimitiveUniformBuffer = GetUniformBuffer();
 
-				BatchElement.FirstIndex = MeshBatchData.StartIndex;
-				BatchElement.NumPrimitives = MeshBatchData.NumPrimitives;
+				BatchElement.FirstIndex = MeshBatchData.FirstTriangleIndex;
+				BatchElement.NumPrimitives = MeshBatchData.NumTriangles;
 				BatchElement.MinVertexIndex = MeshBatchData.MinVertexIndex;
 				BatchElement.MaxVertexIndex = MeshBatchData.MaxVertexIndex;
 				Mesh.ReverseCulling = IsLocalToWorldDeterminantNegative();
@@ -209,7 +233,7 @@ void FDataflowEngineSceneProxy::GetDynamicMeshElements(
 				Mesh.DepthPriorityGroup = SDPG_World;
 				Mesh.bCanApplyViewModeOverrides = true;
 #if WITH_EDITOR
-				Mesh.BatchHitProxyId = DefaultHitProxy->Id;
+				Mesh.BatchHitProxyId = LocalHitProxies[Facade.GetHitProxyIndexAttribute()[MeshBatchData.GeomIndex]]->Id;
 #endif // WITH_EDITOR
 
 				Collector.AddMesh(ViewIndex, Mesh);
diff --git a/Engine/Plugins/Experimental/Dataflow/Source/DataflowEnginePlugin/Private/Dataflow/DataflowEngineSceneProxy.h b/Engine/Plugins/Experimental/Dataflow/Source/DataflowEnginePlugin/Private/Dataflow/DataflowEngineSceneProxy.h
--- a/Engine/Plugins/Experimental/Dataflow/Source/DataflowEnginePlugin/Private/Dataflow/DataflowEngineSceneProxy.h
+++ b/Engine/Plugins/Experimental/Dataflow/Source/DataflowEnginePlugin/Private/Dataflow/DataflowEngineSceneProxy.h
@@ -14,10 +14,11 @@ struct FManagedArrayCollection;
 struct FDataflowTriangleSetMeshBatchData
 {
 	FMaterialRenderProxy* MaterialProxy = nullptr;
-	int32 StartIndex = -1;
-	int32 NumPrimitives = -1;
-	int32 MinVertexIndex = -1;
-	int32 MaxVertexIndex = -1;
+	int32 FirstTriangleIndex = INDEX_NONE;
+	int32 NumTriangles = INDEX_NONE;
+	int32 MinVertexIndex = INDEX_NONE;
+	int32 MaxVertexIndex = INDEX_NONE;
+	int32 GeomIndex = INDEX_NONE;
 };
 
 class FDataflowEngineSceneProxy final : public FPrimitiveSceneProxy
@@ -60,6 +61,7 @@ private:
 	FDynamicMeshIndexBuffer32 IndexBuffer;
 #if WITH_EDITOR
 	FColorVertexBuffer HitProxyIdBuffer;
+	TArray<TRefCountPtr<HHitProxy> > LocalHitProxies;
 #endif // WITH_EDITOR
 
 	TRefCountPtr<HHitProxy> DefaultHitProxy;
diff --git a/Engine/Source/Runtime/Experimental/Chaos/Private/GeometryCollection/Facades/CollectionRenderingFacade.cpp b/Engine/Source/Runtime/Experimental/Chaos/Private/GeometryCollection/Facades/CollectionRenderingFacade.cpp
--- a/Engine/Source/Runtime/Experimental/Chaos/Private/GeometryCollection/Facades/CollectionRenderingFacade.cpp
+++ b/Engine/Source/Runtime/Experimental/Chaos/Private/GeometryCollection/Facades/CollectionRenderingFacade.cpp
@@ -14,18 +14,33 @@ namespace GeometryCollection::Facades
 		: ConstCollection(InCollection)
 		, Collection(&InCollection)
 		, VertexAttribute(InCollection, "Vertex", FGeometryCollection::VerticesGroup)
+		, VertexToGeometryIndexAttribute(InCollection, "GeometryIndex", FGeometryCollection::VerticesGroup, FGeometryCollection::GeometryGroup)
 		, IndicesAttribute(InCollection, "Indices", FGeometryCollection::FacesGroup, FGeometryCollection::VerticesGroup)
 		, MaterialIDAttribute(InCollection, "MaterialID", FGeometryCollection::FacesGroup)
 		, TriangleSectionAttribute(InCollection, "Sections", FGeometryCollection::MaterialGroup)
+		, GeometryNameAttribute(InCollection, "Name", FGeometryCollection::GeometryGroup)
+		, HitProxyIndexAttribute(InCollection, "HitIndex", FGeometryCollection::GeometryGroup)
+		, VertixStartAttribute(InCollection, "VertexStart", FGeometryCollection::GeometryGroup)
+		, VertixCountAttribute(InCollection, "VertexCount", FGeometryCollection::GeometryGroup)
+		, IndicesStartAttribute(InCollection, "IndicesStart", FGeometryCollection::GeometryGroup)
+		, IndicesCountAttribute(InCollection, "IndicesCount", FGeometryCollection::GeometryGroup)
+
 	{}
 
 	FRenderingFacade::FRenderingFacade(const FManagedArrayCollection& InCollection)
 		: ConstCollection(InCollection)
 		, Collection(nullptr)
 		, VertexAttribute(InCollection, "Vertex", FGeometryCollection::VerticesGroup)
+		, VertexToGeometryIndexAttribute(InCollection, "GeometryIndex", FGeometryCollection::VerticesGroup, FGeometryCollection::GeometryGroup)
 		, IndicesAttribute(InCollection, "Indices", FGeometryCollection::FacesGroup, FGeometryCollection::VerticesGroup)
 		, MaterialIDAttribute(InCollection, "MaterialID", FGeometryCollection::FacesGroup)
 		, TriangleSectionAttribute(InCollection, "Sections", FGeometryCollection::MaterialGroup)
+		, GeometryNameAttribute(InCollection, "Name", FGeometryCollection::GeometryGroup)
+		, HitProxyIndexAttribute(InCollection, "HitIndex", FGeometryCollection::GeometryGroup)
+		, VertixStartAttribute(InCollection, "VertexStart", FGeometryCollection::GeometryGroup, FGeometryCollection::VerticesGroup)
+		, VertixCountAttribute(InCollection, "VertexCount", FGeometryCollection::GeometryGroup)
+		, IndicesStartAttribute(InCollection, "IndicesStart", FGeometryCollection::GeometryGroup, FGeometryCollection::FacesGroup)
+		, IndicesCountAttribute(InCollection, "IndicesCount", FGeometryCollection::GeometryGroup)
 	{}
 
 	//
@@ -36,9 +51,16 @@ namespace GeometryCollection::Facades
 	{
 		check(!IsConst());
 		VertexAttribute.Add();
+		VertexToGeometryIndexAttribute.Add();
 		IndicesAttribute.Add();
 		MaterialIDAttribute.Add();
 		TriangleSectionAttribute.Add();
+		GeometryNameAttribute.Add();
+		HitProxyIndexAttribute.Add();
+		VertixStartAttribute.Add();
+		VertixCountAttribute.Add();
+		IndicesStartAttribute.Add();
+		IndicesCountAttribute.Add();
 	}
 
 	bool FRenderingFacade::CanRenderSurface( ) const
@@ -48,8 +70,12 @@ namespace GeometryCollection::Facades
 
 	bool FRenderingFacade::IsValid( ) const
 	{
-		return VertexAttribute.IsValid() && IndicesAttribute.IsValid() &&
-			MaterialIDAttribute.IsValid() && TriangleSectionAttribute.IsValid();
+		return VertexAttribute.IsValid() && VertexToGeometryIndexAttribute.IsValid() &&
+			IndicesAttribute.IsValid() &&
+			MaterialIDAttribute.IsValid() && TriangleSectionAttribute.IsValid() &&
+			GeometryNameAttribute.IsValid() && HitProxyIndexAttribute.IsValid() &&
+			VertixStartAttribute.IsValid() && VertixCountAttribute.IsValid() && 
+			IndicesStartAttribute.IsValid() && IndicesCountAttribute.IsValid();
 	}
 
 	int32 FRenderingFacade::NumTriangles() const
@@ -97,9 +123,16 @@ namespace GeometryCollection::Facades
 			int32 IndicesStart = IndicesAttribute.AddElements(InIndices.Num());
 			int32 VertexStart = VertexAttribute.AddElements(InVertices.Num());
 			
-			const FIntVector * IndiciesDest = Indices.GetData() + IndicesStart;
+			FIntVector * IndiciesDest = Indices.GetData() + IndicesStart;
 			FMemory::Memmove((void*)IndiciesDest, InIndices.GetData(), sizeof(FIntVector) * InIndices.Num());
 
+			for (int i = IndicesStart; i < IndicesStart + InIndices.Num(); i++)
+			{
+				Indices[i][0] += VertexStart;
+				Indices[i][1] += VertexStart;
+				Indices[i][2] += VertexStart;
+			}
+
 			const FVector3f * VerticesDest = Vertices.GetData() + VertexStart;
 			FMemory::Memmove((void*)VerticesDest, InVertices.GetData(), sizeof(FVector3f) * InVertices.Num());
 		}
@@ -113,6 +146,58 @@ namespace GeometryCollection::Facades
 		return FGeometryCollectionSection::BuildMeshSections(ConstCollection, InputIndices, BaseMeshOriginalIndicesIndex, RetIndices);
 	}
 
+
+	int32 FRenderingFacade::StartGeometryGroup(FString InName)
+	{
+		check(!IsConst());
+
+		int32 GeomIndex = INDEX_NONE;
+		if (IsValid())
+		{
+			GeomIndex = GeometryNameAttribute.AddElements(1);
+			GeometryNameAttribute.Modify()[GeomIndex] = InName;
+
+			VertixStartAttribute.Modify()[GeomIndex] = VertexAttribute.Num();
+			VertixCountAttribute.Modify()[GeomIndex] = 0;
+			IndicesStartAttribute.Modify()[GeomIndex] = IndicesAttribute.Num();
+			IndicesCountAttribute.Modify()[GeomIndex] = 0;
+		}
+		return GeomIndex;
+	}
+
+	void FRenderingFacade::EndGeometryGroup(int32 InGeomIndex)
+	{
+		check(!IsConst());
+		if (IsValid())
+		{
+			check( GeometryNameAttribute.Num()-1 == InGeomIndex );
+
+			if (VertixStartAttribute.Get()[InGeomIndex] < VertexAttribute.Num())
+			{
+				VertixCountAttribute.Modify()[InGeomIndex] = VertexAttribute.Num() - VertixStartAttribute.Get()[InGeomIndex];
+
+				TManagedArray<int32>& GeomIndexAttr = VertexToGeometryIndexAttribute.Modify();
+				for (int i = VertixStartAttribute.Get()[InGeomIndex]; i < VertexAttribute.Num(); i++)
+				{
+					GeomIndexAttr[i] = InGeomIndex;
+				}
+			}
+			else
+			{
+				VertixStartAttribute.Modify()[InGeomIndex] = VertexAttribute.Num();
+			}
+
+			if (IndicesStartAttribute.Get()[InGeomIndex] < IndicesAttribute.Num())
+			{
+				IndicesCountAttribute.Modify()[InGeomIndex] = IndicesAttribute.Num() - IndicesStartAttribute.Get()[InGeomIndex];
+			}
+			else
+			{
+				IndicesStartAttribute.Modify()[InGeomIndex] = IndicesAttribute.Num();
+			}
+		}
+	}
+
 }; // GeometryCollection::Facades
 
 
diff --git a/Engine/Source/Runtime/Experimental/Chaos/Public/GeometryCollection/Facades/CollectionRenderingFacade.h b/Engine/Source/Runtime/Experimental/Chaos/Public/GeometryCollection/Facades/CollectionRenderingFacade.h
--- a/Engine/Source/Runtime/Experimental/Chaos/Public/GeometryCollection/Facades/CollectionRenderingFacade.h
+++ b/Engine/Source/Runtime/Experimental/Chaos/Public/GeometryCollection/Facades/CollectionRenderingFacade.h
@@ -46,7 +46,7 @@ namespace GeometryCollection::Facades
 		// Facade API
 		//
 
-		/**Add a triangle to the rendering view.*/
+		/**Number of triangles to render.*/
 		int32 NumTriangles() const;
 
 		/**Add a triangle to the rendering view.*/
@@ -61,6 +61,9 @@ namespace GeometryCollection::Facades
 		/** GetVertices */
 		const TManagedArray< FVector3f >& GetVertices() const { return VertexAttribute.Get(); }
 
+		/** GetVertexToGeometryIndexAttribute */
+		const TManagedArray< int32 >& GetVertexToGeometryIndexAttribute() const { return VertexToGeometryIndexAttribute.Get(); }
+
 		/** GetMaterialID */
 		const TManagedArray< int32 >& GetMaterialID() const { return MaterialIDAttribute.Get(); }
 
@@ -70,15 +73,51 @@ namespace GeometryCollection::Facades
 		/** BuildMeshSections */
 		TArray<FTriangleSection> BuildMeshSections(const TArray<FIntVector>& Indices, TArray<int32> BaseMeshOriginalIndicesIndex, TArray<FIntVector>& RetIndices) const;
 		
+		/** Geometry Group Start : */
+		int32 StartGeometryGroup(FString InName);
+
+		/** Geometry Group End : */
+		void EndGeometryGroup(int32 InGeometryGroupIndex);
+
+		int32 NumGeometry() const { return GeometryNameAttribute.Num(); }
+
+		/** GeometryNameAttribute */
+		const TManagedArray< FString >& GetGeometryNameAttribute() const { return GeometryNameAttribute.Get(); }
+
+		/** HitProxyIDAttribute */
+		const TManagedArray< int32>& GetHitProxyIndexAttribute() const { return HitProxyIndexAttribute.Get(); }
+		TManagedArray< int32 >& ModifyHitProxyIndexAttribute() {check(!IsConst());return HitProxyIndexAttribute.Modify(); }
+
+		/** VertixStartAttribute */
+		const TManagedArray< int32 >& GetVertixStartAttribute() const { return VertixStartAttribute.Get(); }
+
+		/** VertixCountAttribute */
+		const TManagedArray< int32 >& GetVertixCountAttribute() const { return VertixCountAttribute.Get(); }
+
+		/** IndicesStartAttribute */
+		const TManagedArray< int32 >& GetIndicesStartAttribute() const { return IndicesStartAttribute.Get(); }
+
+		/** IndicesCountAttribute */
+		const TManagedArray< int32 >& GetIndicesCountAttribute() const { return IndicesCountAttribute.Get(); }
+
 	private : 
 		const FManagedArrayCollection& ConstCollection;
 		FManagedArrayCollection* Collection = nullptr;
 
 		TManagedArrayAccessor<FVector3f> VertexAttribute;
+		TManagedArrayAccessor<int32> VertexToGeometryIndexAttribute;
 		TManagedArrayAccessor<FIntVector> IndicesAttribute;
 		TManagedArrayAccessor<int32> MaterialIDAttribute;
 		TManagedArrayAccessor<FTriangleSection> TriangleSectionAttribute;
 
+		TManagedArrayAccessor<FString> GeometryNameAttribute;
+		TManagedArrayAccessor<int32> HitProxyIndexAttribute;
+		TManagedArrayAccessor<int32> VertixStartAttribute;
+		TManagedArrayAccessor<int32> VertixCountAttribute;
+		TManagedArrayAccessor<int32> IndicesStartAttribute;
+		TManagedArrayAccessor<int32> IndicesCountAttribute;
+
+
 	};
 
 }
\ No newline at end of file
diff --git a/Engine/Source/Runtime/Experimental/Dataflow/Engine/Private/Dataflow/DataflowEdNode.cpp b/Engine/Source/Runtime/Experimental/Dataflow/Engine/Private/Dataflow/DataflowEdNode.cpp
--- a/Engine/Source/Runtime/Experimental/Dataflow/Engine/Private/Dataflow/DataflowEdNode.cpp
+++ b/Engine/Source/Runtime/Experimental/Dataflow/Engine/Private/Dataflow/DataflowEdNode.cpp
@@ -271,10 +271,15 @@ bool UDataflowEdNode::Render(GeometryCollection::Facades::FRenderingFacade& Rend
 		{
 			if (Dataflow::FRenderingFactory* Factory = Dataflow::FRenderingFactory::GetInstance())
 			{
-				for (Dataflow::FRenderingParameter& Parameter : GetRenderParameters())
+				if (GetRenderParameters().Num())
 				{
-					Factory->RenderNodeOutput(RenderData, { NodeTarget.Get(), Parameter, *Context });
-					bNeedsRefresh = true;
+					int32 GeometryIndex = RenderData.StartGeometryGroup(GetDataflowNodeGuid().ToString());
+					for (Dataflow::FRenderingParameter& Parameter : GetRenderParameters())
+					{
+						Factory->RenderNodeOutput(RenderData, {NodeTarget.Get(), Parameter, *Context});
+						bNeedsRefresh = true;
+					}
+					RenderData.EndGeometryGroup(GeometryIndex);
 				}
 			}
 		}
