[SingleParticlePhysicsProxy] Adding trailing return type for both linear and angular drag functions

diff --git a/Engine/Source/Runtime/Experimental/Chaos/Public/PhysicsProxy/SingleParticlePhysicsProxy.h b/Engine/Source/Runtime/Experimental/Chaos/Public/PhysicsProxy/SingleParticlePhysicsProxy.h
--- a/Engine/Source/Runtime/Experimental/Chaos/Public/PhysicsProxy/SingleParticlePhysicsProxy.h
+++ b/Engine/Source/Runtime/Experimental/Chaos/Public/PhysicsProxy/SingleParticlePhysicsProxy.h
@@ -811,7 +811,7 @@ public:
 
 	const FReal LinearEtherDrag() const
 	{
-		return Read([](auto* Particle)
+		return Read([](auto* Particle) -> FReal
 		{
 			if (auto Rigid = Particle->CastToRigidParticle())
 			{
@@ -835,7 +835,7 @@ public:
 
 	const FReal AngularEtherDrag() const
 	{
-		return Read([](auto* Particle)
+		return Read([](auto* Particle) -> FReal
 		{
 			if (auto Rigid = Particle->CastToRigidParticle())
 			{
@@ -1338,4 +1338,4 @@ inline FSingleParticlePhysicsProxy* FSingleParticlePhysicsProxy::Create(TUniqueP
 	auto Proxy = new FSingleParticlePhysicsProxy(MoveTemp(Particle), nullptr);
 	return Proxy;
 }
-}
\ No newline at end of file
+}
