Chaos Cloth - Add dirty state to the constrait weight maps apply function to avoid re-updating the table when no changes have occured.

#rnx

diff --git a/Engine/Source/Runtime/Experimental/Chaos/Public/Chaos/PBDStiffness.h b/Engine/Source/Runtime/Experimental/Chaos/Public/Chaos/PBDStiffness.h
--- a/Engine/Source/Runtime/Experimental/Chaos/Public/Chaos/PBDStiffness.h
+++ b/Engine/Source/Runtime/Experimental/Chaos/Public/Chaos/PBDStiffness.h
@@ -3,10 +3,10 @@
 
 #include "Chaos/Core.h"
 #include "Chaos/PBDSoftsEvolutionFwd.h"
+#include "Chaos/PBDWeightMap.h"
 #include "Containers/ArrayView.h"
 #include "Containers/ContainersFwd.h"
 #include "ChaosStats.h"
-#include "PBDWeightMap.h"
 
 DECLARE_CYCLE_STAT(TEXT("Chaos PBD Stiffness Apply Values"), STAT_PBD_StiffnessApplyValues, STATGROUP_Chaos);
 
@@ -57,6 +57,9 @@ public:
 private:
 	const FSolverReal ParameterFitBase;
 	const FSolverReal ParameterFitLogBase;
+	FSolverReal PrevDtOrMaxStiffness = (FSolverReal)0.;
+	int32 PrevNumIterations = 0;
+	using FPBDWeightMap::bIsDirty;
 };
 
 FPBDStiffness::FPBDStiffness(
@@ -105,36 +108,35 @@ void FPBDStiffness::ApplyValues(const FSolverReal Dt, const int32 NumIterations)
 
 	// Define the stiffness mapping function
 	auto SimulationValue = [this, Exponent](const FSolverReal InValue)->FSolverReal
-	{
-		// If InValue is 1 then LogValue = -inf and the output becomes -inf as well,
-		// in order for this function to be continuous, we want the output to be 1 when InValue = 1
-		// and need the stiffness to be exactly 0 when the input is 0
-		if (InValue <= (FSolverReal)UE_DOUBLE_SMALL_NUMBER)
-		{
-			return (FSolverReal)0.;
-		}
-		if (InValue >= (FSolverReal)(1. - UE_DOUBLE_SMALL_NUMBER))
 		{
-			return (FSolverReal)1.;
-		}
-		const FSolverReal ParameterFit = CalcExponentialParameterFit(ParameterFitBase, ParameterFitLogBase, InValue);
-
-		// Use simulation dependent stiffness exponent to alleviate the variations in effect when Dt and NumIterations change
-		// This is based on the Position-Based Simulation Methods paper (page 8),
-		// but uses the delta time in addition of the number of iterations in the calculation of the error term.
-		const FSolverReal LogValue = FMath::Loge((FSolverReal)1. - ParameterFit);
-		return (FSolverReal)1. - FMath::Exp(LogValue * Exponent);
-	};
-
-	const FSolverReal Offset = WeightedValue[0];
-	const FSolverReal Range = WeightedValue[1] - WeightedValue[0];
-	const int32 TableSize = Table.Num();
-	const FSolverReal WeightIncrement = (TableSize > 1) ? (FSolverReal)1. / (FSolverReal)(TableSize - 1) : (FSolverReal)1.; // Must allow full range from 0 to 1 included
-	for (int32 Index = 0; Index < TableSize; ++Index)
+			// If InValue is 1 then LogValue = -inf and the output becomes -inf as well,
+			// in order for this function to be continuous, we want the output to be 1 when InValue = 1
+			// and need the stiffness to be exactly 0 when the input is 0
+			if (InValue <= (FSolverReal)UE_DOUBLE_SMALL_NUMBER)
+			{
+				return (FSolverReal)0.;
+			}
+			if (InValue >= (FSolverReal)(1. - UE_DOUBLE_SMALL_NUMBER))
+			{
+				return (FSolverReal)1.;
+			}
+			const FSolverReal ParameterFit = CalcExponentialParameterFit(ParameterFitBase, ParameterFitLogBase, InValue);
+
+			// Use simulation dependent stiffness exponent to alleviate the variations in effect when Dt and NumIterations change
+			// This is based on the Position-Based Simulation Methods paper (page 8),
+			// but uses the delta time in addition of the number of iterations in the calculation of the error term.
+			const FSolverReal LogValue = FMath::Loge((FSolverReal)1. - ParameterFit);
+			return (FSolverReal)1. - FMath::Exp(LogValue * Exponent);
+		};
+	
+	if (Dt != PrevDtOrMaxStiffness || NumIterations != PrevNumIterations)
 	{
-		const FSolverReal Weight = (FSolverReal)Index * WeightIncrement;
-		Table[Index] = SimulationValue(Offset + Weight * Range);
+		PrevDtOrMaxStiffness = Dt;
+		PrevNumIterations = NumIterations;
+		bIsDirty = true;
 	}
+
+	FPBDWeightMap::ApplyValues(SimulationValue);  // Note, this still needs to be called even when not dirty for dealing with any change in values
 }
 
 void FPBDStiffness::ApplyXPBDValues(const FSolverReal MaxStiffness)
@@ -142,20 +144,18 @@ void FPBDStiffness::ApplyXPBDValues(const FSolverReal MaxStiffness)
 	SCOPE_CYCLE_COUNTER(STAT_PBD_StiffnessApplyValues);
 	
 	auto SimulationValue = [this, MaxStiffness](const FSolverReal InValue)->FSolverReal
+		{
+			// Do not apply exponential to XPBDStiffnesses. They are authored in terms of true stiffness values (e.g., kg/s^2), not exponential compliance
+			return FMath::Clamp(InValue, (FSolverReal)0., MaxStiffness);
+		};
+
+	if (MaxStiffness != PrevDtOrMaxStiffness)
 	{
-		// Do not apply exponential to XPBDStiffnesses. They are authored in terms of true stiffness values (e.g., kg/s^2), not exponential compliance
-		return FMath::Clamp(InValue, (FSolverReal)0., MaxStiffness);
-	};
-
-	const FSolverReal Offset = WeightedValue[0];
-	const FSolverReal Range = WeightedValue[1] - WeightedValue[0];
-	const int32 TableSize = Table.Num();
-	const FSolverReal WeightIncrement = (TableSize > 1) ? (FSolverReal)1. / (FSolverReal)(TableSize - 1) : (FSolverReal)1.; // Must allow full range from 0 to 1 included
-	for (int32 Index = 0; Index < TableSize; ++Index)
-	{
-		const FSolverReal Weight = (FSolverReal)Index * WeightIncrement;
-		Table[Index] = SimulationValue(Offset + Weight * Range);
+		PrevDtOrMaxStiffness = MaxStiffness;
+		bIsDirty = true;
 	}
+
+	FPBDWeightMap::ApplyValues(SimulationValue);  // Note, this still needs to be called even when not dirty for dealing with any change in values
 }
 
 }  // End namespace Chaos::Softs
diff --git a/Engine/Source/Runtime/Experimental/Chaos/Public/Chaos/PBDWeightMap.h b/Engine/Source/Runtime/Experimental/Chaos/Public/Chaos/PBDWeightMap.h
--- a/Engine/Source/Runtime/Experimental/Chaos/Public/Chaos/PBDWeightMap.h
+++ b/Engine/Source/Runtime/Experimental/Chaos/Public/Chaos/PBDWeightMap.h
@@ -54,14 +54,14 @@ public:
 	 * The weight map table only gets updated after ApplyValues is called.
 	 * Low and high values are clamped between [0,1]
 	 */
-	void SetWeightedValue(const FSolverVec2& InWeightedValue) { WeightedValue = InWeightedValue.ClampAxes((FSolverReal)0., (FSolverReal)1.); }
+	void SetWeightedValue(const FSolverVec2& InWeightedValue) { WeightedValue = InWeightedValue.ClampAxes((FSolverReal)0., (FSolverReal)1.); bIsDirty = true; }
 
 	/**
 	 * Set the low and high values of the weight map.
 	 * The weight map table only gets updated after ApplyValues is called.
 	 * Low and high values are not clamped. Commonly used for XPBD Stiffness values which are not [0,1]
 	 */
-	void SetWeightedValueUnclamped(const FSolverVec2& InWeightedValue) { WeightedValue = InWeightedValue; }
+	void SetWeightedValueUnclamped(const FSolverVec2& InWeightedValue) { WeightedValue = InWeightedValue; bIsDirty = true; }
 
 	/**
 	 * Return the low and high values set for this weight map.
@@ -70,7 +70,7 @@ public:
 	const FSolverVec2& GetWeightedValue() const { return WeightedValue; }
 
 	/** Update the weight map table with the current simulation parameters. */
-	inline void ApplyValues();
+	void ApplyValues() { ApplyValues([](FSolverReal Value)->FSolverReal { return Value; }); }
 
 	/**
 	 * Lookup for the exponential weighted value at the specified weight map index.
@@ -97,12 +97,16 @@ public:
 	inline void ReorderIndices(const TArray<int32>& OrigToReorderedIndices);
 
 protected:
+	template<typename FunctorType>
+	inline void ApplyValues(FunctorType&& MappingFunction);
+
 	TArray<uint8> Indices; // Per particle/constraints array of index to the stiffness table
 	TArray<FSolverReal> Table;  // Fixed lookup table of stiffness values, use uint8 indexation
 	FSolverVec2 WeightedValue;
+	bool bIsDirty = true;  // Whether the values have changed and the table needs updating
 };
 
-FPBDWeightMap::FPBDWeightMap(
+inline FPBDWeightMap::FPBDWeightMap(
 	const FSolverVec2& InWeightedValue,
 	const TConstArrayView<FRealSingle>& Multipliers,
 	int32 ParticleCount,
@@ -135,7 +139,7 @@ FPBDWeightMap::FPBDWeightMap(
 }
 
 template<int32 Valence>
-FPBDWeightMap::FPBDWeightMap(
+inline FPBDWeightMap::FPBDWeightMap(
 	const FSolverVec2& InWeightedValue,
 	const TConstArrayView<FRealSingle>& Multipliers,
 	const TConstArrayView<TVector<int32, Valence>>& Constraints,
@@ -181,22 +185,27 @@ FPBDWeightMap::FPBDWeightMap(
 	}
 }
 
-void FPBDWeightMap::ApplyValues()
+template<typename FunctorType>
+inline void FPBDWeightMap::ApplyValues(FunctorType&& MappingFunction)
 {
 	SCOPE_CYCLE_COUNTER(STAT_PBD_WeightMapApplyValues);
-
-	const FSolverReal Offset = WeightedValue[0];
-	const FSolverReal Range = WeightedValue[1] - WeightedValue[0];
-	const int32 TableSize = Table.Num();
-	const FSolverReal WeightIncrement = (TableSize > 1) ? (FSolverReal)1. / (FSolverReal)(TableSize - 1) : (FSolverReal)1.; // Must allow full range from 0 to 1 included
-	for (int32 Index = 0; Index < TableSize; ++Index)
+	if (bIsDirty)
 	{
-		const FSolverReal Weight = (FSolverReal)Index * WeightIncrement;
-		Table[Index] = Offset + Weight * Range;
+		const FSolverReal Offset = WeightedValue[0];
+		const FSolverReal Range = WeightedValue[1] - WeightedValue[0];
+		const int32 TableSize = Table.Num();
+		const FSolverReal WeightIncrement = (TableSize > 1) ? (FSolverReal)1. / (FSolverReal)(TableSize - 1) : (FSolverReal)1.; // Must allow full range from 0 to 1 included
+		for (int32 Index = 0; Index < TableSize; ++Index)
+		{
+			const FSolverReal Weight = (FSolverReal)Index * WeightIncrement;
+			Table[Index] = MappingFunction(Offset + Weight * Range);
+		}
+
+		bIsDirty = false;
 	}
 }
 
-void FPBDWeightMap::ReorderIndices(const TArray<int32>& OrigToReorderedConstraintIndices)
+inline void FPBDWeightMap::ReorderIndices(const TArray<int32>& OrigToReorderedConstraintIndices)
 {
 	if (Indices.Num() == OrigToReorderedConstraintIndices.Num())
 	{
