Chaos Cloth:
- Make PBDLongRangeConstraintsBase to use the FWeightMap class for the tether scale property.
- Add missing clamps onto the PBDLongRangeConstraintsBase parameters.

#rnx

diff --git a/Engine/Source/Runtime/Experimental/Chaos/Private/Chaos/PBDLongRangeConstraints.cpp b/Engine/Source/Runtime/Experimental/Chaos/Private/Chaos/PBDLongRangeConstraints.cpp
--- a/Engine/Source/Runtime/Experimental/Chaos/Private/Chaos/PBDLongRangeConstraints.cpp
+++ b/Engine/Source/Runtime/Experimental/Chaos/Private/Chaos/PBDLongRangeConstraints.cpp
@@ -24,9 +24,9 @@ void FPBDLongRangeConstraints::Apply(FSolverParticles& Particles, const FSolverR
 	if (!Stiffness.HasWeightMap())
 	{
 		const FSolverReal ExpStiffnessValue = (FSolverReal)Stiffness;
-		if (!HasScaleWeightMap())
+		if (!TetherScale.HasWeightMap())
 		{
-			const FSolverReal ScaleValue = ScaleTable[0];
+			const FSolverReal ScaleValue = (FSolverReal)TetherScale;
 #if INTEL_ISPC
 			if (bRealTypeCompatibleWithISPC && bChaos_LongRange_ISPC_Enabled)
 			{
@@ -68,8 +68,8 @@ void FPBDLongRangeConstraints::Apply(FSolverParticles& Particles, const FSolverR
 						(ispc::FVector4f*)Particles.GetPAndInvM().GetData(),
 						(const ispc::FTether*)TetherBatch.GetData(),
 						ExpStiffnessValue,
-						ScaleIndices.GetData(),
-						ScaleTable.GetData(),
+						TetherScale.GetIndices().GetData(),
+						TetherScale.GetTable().GetData(),
 						TetherBatch.Num(),
 						ParticleOffset);
 				}
@@ -84,7 +84,7 @@ void FPBDLongRangeConstraints::Apply(FSolverParticles& Particles, const FSolverR
 						{
 							const FTether& Tether = TetherBatch[Index];
 							const int32 LocalParticleIndex = GetEndIndex(Tether);
-							const FSolverReal ScaleValue = ScaleTable[ScaleIndices[LocalParticleIndex]];
+							const FSolverReal ScaleValue = TetherScale[LocalParticleIndex];
 							Particles.P(ParticleOffset + LocalParticleIndex) += ExpStiffnessValue * GetDelta(Particles, Tether, ScaleValue);
 						}, TetherBatch.Num() < MinParallelSize);
 				}
@@ -93,9 +93,9 @@ void FPBDLongRangeConstraints::Apply(FSolverParticles& Particles, const FSolverR
 	}
 	else  // HasStiffnessWeighmap
 	{
-		if (!HasScaleWeightMap())
+		if (!TetherScale.HasWeightMap())
 		{
-			const FSolverReal ScaleValue = ScaleTable[0];
+			const FSolverReal ScaleValue = (FSolverReal)TetherScale;
 
 #if INTEL_ISPC
 			if (bRealTypeCompatibleWithISPC && bChaos_LongRange_ISPC_Enabled)
@@ -142,8 +142,8 @@ void FPBDLongRangeConstraints::Apply(FSolverParticles& Particles, const FSolverR
 						(const ispc::FTether*)TetherBatch.GetData(),
 						Stiffness.GetIndices().GetData(),
 						Stiffness.GetTable().GetData(),
-						ScaleIndices.GetData(),
-						ScaleTable.GetData(),
+						TetherScale.GetIndices().GetData(),
+						TetherScale.GetTable().GetData(),
 						TetherBatch.Num(),
 						ParticleOffset);
 				}
@@ -159,7 +159,7 @@ void FPBDLongRangeConstraints::Apply(FSolverParticles& Particles, const FSolverR
 							const FTether& Tether = TetherBatch[Index];
 							const int32 LocalParticleIndex = GetEndIndex(Tether);
 							const FSolverReal ExpStiffnessValue = Stiffness[LocalParticleIndex];
-							const FSolverReal ScaleValue = ScaleTable[ScaleIndices[LocalParticleIndex]];
+							const FSolverReal ScaleValue = TetherScale[LocalParticleIndex];
 							Particles.P(ParticleOffset + LocalParticleIndex) += ExpStiffnessValue * GetDelta(Particles, Tether, ScaleValue);
 						}, TetherBatch.Num() < MinParallelSize);
 				}
diff --git a/Engine/Source/Runtime/Experimental/Chaos/Private/Chaos/PBDLongRangeConstraintsBase.cpp b/Engine/Source/Runtime/Experimental/Chaos/Private/Chaos/PBDLongRangeConstraintsBase.cpp
--- a/Engine/Source/Runtime/Experimental/Chaos/Private/Chaos/PBDLongRangeConstraintsBase.cpp
+++ b/Engine/Source/Runtime/Experimental/Chaos/Private/Chaos/PBDLongRangeConstraintsBase.cpp
@@ -35,48 +35,12 @@ FPBDLongRangeConstraintsBase::FPBDLongRangeConstraintsBase(
 	const FSolverVec2& InScale)
 	: Tethers(InTethers)
 	, Stiffness(InStiffness, StiffnessMultipliers, InParticleCount)
-	, Scale(InScale)
+	, TetherScale(InScale, ScaleMultipliers, InParticleCount)
 	, ParticleOffset(InParticleOffset)
-	, ParticleCount(InParticleCount)
 {
-	if (ScaleMultipliers.Num() == ParticleCount && ParticleCount > 0)
-	{
-		// Convert the weight maps into an array of lookup indices to the scale table
-		ScaleIndices.AddUninitialized(ParticleCount);
-
-		const FRealSingle TableScale = (FRealSingle)(TableSize - 1);
-
-		for (int32 Index = 0; Index < ParticleCount; ++Index)
-		{
-			ScaleIndices[Index] = (uint8)(FMath::Clamp(ScaleMultipliers[Index], (FRealSingle)0., (FRealSingle)1.) * TableScale);
-		}
-
-		// Initialize empty table until ApplyValues is called
-		ScaleTable.AddZeroed(TableSize);
-	}
-	else
-	{
-		// Initialize with a one element table until ApplyValues is called
-		ScaleIndices.AddZeroed(1);
-		ScaleTable.AddZeroed(1);
-	}
-
 	// Apply default properties
 	ApplyProperties((FSolverReal)(1. / FPBDStiffness::ParameterFrequency), 1);
 }
 
-void FPBDLongRangeConstraintsBase::ApplyScale()
-{
-	const FSolverReal Offset = Scale[0];
-	const FSolverReal Range = Scale[1] - Scale[0];
-	const int32 ScaleTableSize = ScaleTable.Num();
-	const FSolverReal WeightIncrement = (ScaleTableSize > 1) ? (FSolverReal)1. / (FSolverReal)(ScaleTableSize - 1) : (FSolverReal)1.; // Must allow full range from 0 to 1 included
-	for (int32 Index = 0; Index < ScaleTableSize; ++Index)
-	{
-		const FSolverReal Weight = (FSolverReal)Index * WeightIncrement;
-		ScaleTable[Index] = Offset + Weight * Range;
-	}
-}
-
 }  // End namespace Chaos::Softs
 
diff --git a/Engine/Source/Runtime/Experimental/Chaos/Public/Chaos/PBDLongRangeConstraints.h b/Engine/Source/Runtime/Experimental/Chaos/Public/Chaos/PBDLongRangeConstraints.h
--- a/Engine/Source/Runtime/Experimental/Chaos/Public/Chaos/PBDLongRangeConstraints.h
+++ b/Engine/Source/Runtime/Experimental/Chaos/Public/Chaos/PBDLongRangeConstraints.h
@@ -33,9 +33,8 @@ public:
 private:
 	using Base::Tethers;
 	using Base::Stiffness;
+	using Base::TetherScale;
 	using Base::ParticleOffset;
-	using Base::ScaleTable;
-	using Base::ScaleIndices;
 };
 
 }  // End namespace Chaos::Softs
diff --git a/Engine/Source/Runtime/Experimental/Chaos/Public/Chaos/PBDLongRangeConstraintsBase.h b/Engine/Source/Runtime/Experimental/Chaos/Public/Chaos/PBDLongRangeConstraintsBase.h
--- a/Engine/Source/Runtime/Experimental/Chaos/Public/Chaos/PBDLongRangeConstraintsBase.h
+++ b/Engine/Source/Runtime/Experimental/Chaos/Public/Chaos/PBDLongRangeConstraintsBase.h
@@ -38,23 +38,25 @@ public:
 		const FSolverVec2& InStiffness = FSolverVec2::UnitVector,
 		const FSolverVec2& InScale = FSolverVec2::UnitVector);
 
+PRAGMA_DISABLE_DEPRECATION_WARNINGS  // For ScaleIndices and ScaleTable
 	virtual ~FPBDLongRangeConstraintsBase() {}
+PRAGMA_ENABLE_DEPRECATION_WARNINGS
 
 	// Return the stiffness input values used by the constraint
 	FSolverVec2 GetStiffness() const { return Stiffness.GetWeightedValue(); }
 
 	// Set the stiffness input values used by the constraint
-	void SetStiffness(const FSolverVec2& InStiffness) { Stiffness.SetWeightedValue(InStiffness); }
+	void SetStiffness(const FSolverVec2& InStiffness) { Stiffness.SetWeightedValue(InStiffness.ClampAxes((FSolverReal)0., (FSolverReal)1.)); }
 
 	// Set the scale low and high value of the scale weight map
-	void SetScale(const FSolverVec2& InScale) { Scale = FSolverVec2(FMath::Clamp((FSolverReal)InScale.X, (FSolverReal)0.01, (FSolverReal)10.), FMath::Clamp((FSolverReal)InScale.Y, (FSolverReal)0.01, (FSolverReal)10.)); }  // TODO: Fix FSolverVec2 double to float conversion error on some platforms as to remove extra cast to FSolverReal
+	void SetScale(const FSolverVec2& InScale) { TetherScale.SetWeightedValue(InScale.ClampAxes((FSolverReal)0.01, (FSolverReal)10.)); }
 
 public:
 	// Set stiffness offset and range, as well as the simulation stiffness exponent
 	void ApplyProperties(const FSolverReal Dt, const int32 NumIterations)
 	{
 		Stiffness.ApplyValues(Dt, NumIterations);
-		ApplyScale();
+		TetherScale.ApplyValues();
 	}
 
 	// Return the tethers, organized in concurent friendly batches
@@ -76,7 +78,7 @@ public:
 	FSolverReal GetRefLength(const FTether& Tether) const { return (FSolverReal)Tether.Get<2>(); }
 
 	// Return the Tether scale for the specified tether
-	FSolverReal GetScale(const FTether& Tether) const { return HasScaleWeightMap() ? ScaleTable[ScaleIndices[GetEndIndex(Tether)]] : ScaleTable[0]; }
+	FSolverReal GetScale(const FTether& Tether) const { return TetherScale.HasWeightMap() ? TetherScale[GetEndIndex(Tether)] : (FSolverReal)TetherScale; }
 
 	// Return the target length of the specified tether (= RefLength * Scale)
 	FSolverReal GetTargetLength(const FTether& Tether) const { return GetRefLength(Tether) * GetScale(Tether); }
@@ -86,11 +88,12 @@ protected:
 	static int32 GetMinParallelBatchSize();
 
 	// Return whether the constraint has been setup with a weightmap to interpolate between two low and high values of scales
-	bool HasScaleWeightMap() const { return ScaleTable.Num() > 1; }
+	UE_DEPRECATED(5.2, "Use TetherScale.HasWeightMap() instead")
+	bool HasScaleWeightMap() const { return TetherScale.HasWeightMap(); }
 
 	// Part of ApplyProperties to update ScaleTable.
-	void ApplyScale();
-
+	UE_DEPRECATED(5.2, "Use TetherScale.ApplyValues() instead")
+	void ApplyScale() { TetherScale.ApplyValues(); }
 
 	// Return a vector representing the amount of segment required for the tether to shrink back to its maximum target length constraint, or zero if the constraint is already met
 	inline FSolverVec3 GetDelta(const FSolverParticles& Particles, const FTether& Tether, const FSolverReal InScale) const
@@ -122,11 +125,17 @@ protected:
 protected:
 	static constexpr int32 TableSize = 16;  // The size of the weightmaps lookup table
 	const TArray<TConstArrayView<FTether>>& Tethers;  // Array view on the tether provided to this constraint
-	FPBDStiffness Stiffness;  // Stiffness weightmap lookup table 
-	TArray<uint8> ScaleIndices;  // Per particle array of index to the scale lookup table
-	TArray<FSolverReal> ScaleTable;  // Fixed lookup table of scale values
-	FSolverVec2 Scale;  // High and low values of tether target length scale used by the the lookup tables interpolation
+	FPBDStiffness Stiffness;  // Stiffness weightmap lookup table
+	FPBDWeightMap TetherScale;  // Scale weightmap lookup table
 	const int32 ParticleOffset;  // Index of the first usable particle
-	const int32 ParticleCount;  // Number of particles available to this constraint
+
+	UE_DEPRECATED(5.2, "Use TetherScale instead")
+	TArray<uint8> ScaleIndices;
+	UE_DEPRECATED(5.2, "Use TetherScale instead")
+	TArray<FSolverReal> ScaleTable;
+	UE_DEPRECATED(5.2, "Use TetherScale instead")
+	FSolverVec2 Scale;
+	UE_DEPRECATED(5.2, "Use TetherScale.Num() or Stiffness.Num(), instead")
+	const int32 ParticleCount = 0;
 };
 }  // End namespace Chaos::Softs
diff --git a/Engine/Source/Runtime/Experimental/Chaos/Public/Chaos/XPBDLongRangeConstraints.h b/Engine/Source/Runtime/Experimental/Chaos/Public/Chaos/XPBDLongRangeConstraints.h
--- a/Engine/Source/Runtime/Experimental/Chaos/Public/Chaos/XPBDLongRangeConstraints.h
+++ b/Engine/Source/Runtime/Experimental/Chaos/Public/Chaos/XPBDLongRangeConstraints.h
@@ -41,7 +41,11 @@ public:
 	virtual ~FXPBDLongRangeConstraints() override {}
 
 	// Set stiffness offset and range, as well as the simulation stiffness exponent
-	void ApplyProperties(const FSolverReal Dt, const int32 NumIterations) { Stiffness.ApplyXPBDValues(XPBDLongRangeMaxStiffness); ApplyScale(); }
+	void ApplyProperties(const FSolverReal Dt, const int32 NumIterations)
+	{
+		Stiffness.ApplyXPBDValues(XPBDLongRangeMaxStiffness);
+		TetherScale.ApplyValues();
+	}
 
 	void Init() const
 	{
@@ -57,7 +61,7 @@ public:
 
 		if (Stiffness.HasWeightMap())
 		{
-			if (HasScaleWeightMap())
+			if (TetherScale.HasWeightMap())
 			{
 				int32 ConstraintOffset = 0;
 				for (const TConstArrayView<FTether>& TetherBatch : Tethers)
@@ -66,7 +70,7 @@ public:
 						{
 							const FTether& Tether = TetherBatch[Index];
 							const int32 LocalParticleIndex = GetEndIndex(Tether);
-							const FSolverReal Scale = ScaleTable[ScaleIndices[LocalParticleIndex]];
+							const FSolverReal Scale = TetherScale[LocalParticleIndex];
 							const FSolverReal ExpStiffnessValue = Stiffness[LocalParticleIndex];
 							Apply(Particles, Dt, Tether, ConstraintOffset + Index, ExpStiffnessValue, Scale);
 						}, TetherBatch.Num() < MinParallelSize);
@@ -75,7 +79,7 @@ public:
 			}
 			else
 			{
-				const FSolverReal ScaleValue = ScaleTable[0];
+				const FSolverReal ScaleValue = (FSolverReal)TetherScale;
 				int32 ConstraintOffset = 0;
 				for (const TConstArrayView<FTether>& TetherBatch : Tethers)
 				{
@@ -94,7 +98,7 @@ public:
 		{
 			const FSolverReal ExpStiffnessValue = (FSolverReal)Stiffness;
 
-			if (HasScaleWeightMap())
+			if (TetherScale.HasWeightMap())
 			{
 				int32 ConstraintOffset = 0;
 				for (const TConstArrayView<FTether>& TetherBatch : Tethers)
@@ -103,7 +107,7 @@ public:
 						{
 							const FTether& Tether = TetherBatch[Index];
 							const int32 LocalParticleIndex = GetEndIndex(Tether);
-							const FSolverReal Scale = ScaleTable[ScaleIndices[LocalParticleIndex]];
+							const FSolverReal Scale = TetherScale[LocalParticleIndex];
 							Apply(Particles, Dt, Tether, ConstraintOffset + Index, ExpStiffnessValue, Scale);
 						}, TetherBatch.Num() < MinParallelSize);
 					ConstraintOffset += TetherBatch.Num();
@@ -111,7 +115,7 @@ public:
 			}
 			else
 			{
-				const FSolverReal ScaleValue = ScaleTable[0];
+				const FSolverReal ScaleValue = (FSolverReal)TetherScale;
 				int32 ConstraintOffset = 0;
 				for (const TConstArrayView<FTether>& TetherBatch : Tethers)
 				{
@@ -148,9 +152,7 @@ private:
 private:
 	using Base::Tethers;
 	using Base::Stiffness;
-	using Base::ParticleOffset;
-	using Base::ScaleTable;
-	using Base::ScaleIndices;
+	using Base::TetherScale;
 
 	mutable TArray<FSolverReal> Lambdas;
 	int32 NumTethers;
