Chaos Cloth:
- Fixed clamping on tether scale broken with the previous change (CL 23303565).
- Removed redundant clamping on tether stiffness, already done in the FPBDStiffness function.
(Note: More clamping cleanup is needed and will hopefully be under way soon.)

#rnx

diff --git a/Engine/Source/Runtime/Experimental/Chaos/Public/Chaos/PBDLongRangeConstraintsBase.h b/Engine/Source/Runtime/Experimental/Chaos/Public/Chaos/PBDLongRangeConstraintsBase.h
--- a/Engine/Source/Runtime/Experimental/Chaos/Public/Chaos/PBDLongRangeConstraintsBase.h
+++ b/Engine/Source/Runtime/Experimental/Chaos/Public/Chaos/PBDLongRangeConstraintsBase.h
@@ -46,10 +46,10 @@ PRAGMA_ENABLE_DEPRECATION_WARNINGS
 	FSolverVec2 GetStiffness() const { return Stiffness.GetWeightedValue(); }
 
 	// Set the stiffness input values used by the constraint
-	void SetStiffness(const FSolverVec2& InStiffness) { Stiffness.SetWeightedValue(InStiffness.ClampAxes((FSolverReal)0., (FSolverReal)1.)); }
+	void SetStiffness(const FSolverVec2& InStiffness) { Stiffness.SetWeightedValue(InStiffness); }
 
 	// Set the scale low and high value of the scale weight map
-	void SetScale(const FSolverVec2& InScale) { TetherScale.SetWeightedValue(InScale.ClampAxes((FSolverReal)0.01, (FSolverReal)10.)); }
+	void SetScale(const FSolverVec2& InScale) { TetherScale.SetWeightedValueUnclamped(InScale.ClampAxes((FSolverReal)0.01, (FSolverReal)10.)); }
 
 public:
 	// Set stiffness offset and range, as well as the simulation stiffness exponent
