Iris
- Added new trait UseSerializerIsEqual to NetSerializer to indicate that it should use serializers IsEqual method instead of using property compare. One example of where this is required is when
- Fixed bug in GameplayAbilityRepAnimMontageNetSerializer inverting bSkipPositionCorrection

diff --git a/Engine/Plugins/Runtime/GameplayAbilities/Source/GameplayAbilities/Private/Serialization/GameplayAbilityRepAnimMontageNetSerializer.cpp b/Engine/Plugins/Runtime/GameplayAbilities/Source/GameplayAbilities/Private/Serialization/GameplayAbilityRepAnimMontageNetSerializer.cpp
--- a/Engine/Plugins/Runtime/GameplayAbilities/Source/GameplayAbilities/Private/Serialization/GameplayAbilityRepAnimMontageNetSerializer.cpp
+++ b/Engine/Plugins/Runtime/GameplayAbilities/Source/GameplayAbilities/Private/Serialization/GameplayAbilityRepAnimMontageNetSerializer.cpp
@@ -30,6 +30,7 @@ struct FGameplayAbilityRepAnimMontageNetSerializer
 	static constexpr bool bIsForwardingSerializer = true; // Triggers asserts if a function is missing
 	static constexpr bool bHasConnectionSpecificSerialization = true; // One of our members requires connection specific serialization
 	static constexpr bool bHasCustomNetReference = true; // We have object references that are not directly accessible
+	static constexpr bool bUseSerializerIsEqual = true; // Since FGameplayAbilityRepAnimMontageNetSerializer conditionally replicates position and SectionId we need to use the IsEqual function provided by the NetSerializer when comparing the property
 
 	// Types
 	enum EReplicationFlags : uint8
@@ -167,8 +168,6 @@ void FGameplayAbilityRepAnimMontageNetSerializer::Deserialize(FNetSerializationC
 
 void FGameplayAbilityRepAnimMontageNetSerializer::SerializeDelta(FNetSerializationContext& Context, const FNetSerializeDeltaArgs& Args)
 {
-	//NetSerializeDeltaDefault<Serialize>(Context, Args);
-
 	const QuantizedType& Value = *reinterpret_cast<QuantizedType*>(Args.Source);
 	const QuantizedType& PrevValue = *reinterpret_cast<QuantizedType*>(Args.Prev);
 
@@ -306,13 +305,13 @@ void FGameplayAbilityRepAnimMontageNetSerializer::Dequantize(FNetSerializationCo
 	{
 		TargetValue.Position = SourceValue.Position;
 		TargetValue.SectionIdToPlay = 0;
-		TargetValue.SkipPositionCorrection = 1;
+		TargetValue.SkipPositionCorrection = 0;
 	}
 	else
 	{
 		TargetValue.Position = 0.f;
 		TargetValue.SectionIdToPlay = SourceValue.SectionIdToPlay;
-		TargetValue.SkipPositionCorrection = 0;
+		TargetValue.SkipPositionCorrection = 1;
 	}
 }
 
@@ -339,7 +338,7 @@ bool FGameplayAbilityRepAnimMontageNetSerializer::IsEqual(FNetSerializationConte
 			return false;
 		}
 
-		if (Value0.ReplicationFlags != Value1.ReplicationFlags || Value0.Position != Value1.Position || Value0.SectionIdToPlay != Value1.SectionIdToPlay)
+		if (Value0.Position != Value1.Position || Value0.SectionIdToPlay != Value1.SectionIdToPlay)
 		{
 			return false;
 		}
diff --git a/Engine/Source/Runtime/Experimental/Iris/Core/Private/Iris/ReplicationState/InternalPropertyReplicationState.cpp b/Engine/Source/Runtime/Experimental/Iris/Core/Private/Iris/ReplicationState/InternalPropertyReplicationState.cpp
--- a/Engine/Source/Runtime/Experimental/Iris/Core/Private/Iris/ReplicationState/InternalPropertyReplicationState.cpp
+++ b/Engine/Source/Runtime/Experimental/Iris/Core/Private/Iris/ReplicationState/InternalPropertyReplicationState.cpp
@@ -133,7 +133,19 @@ void CopyPropertyReplicationState(uint8* RESTRICT DstStateBuffer, uint8* RESTRIC
 bool InternalCompareMember(const FReplicationStateDescriptor* Descriptor, uint32 MemberIndex, const void* RESTRICT ValueA, const void* RESTRICT ValueB)
 {
 	const FReplicationStateMemberSerializerDescriptor& MemberSerializerDescriptor = Descriptor->MemberSerializerDescriptors[MemberIndex];
-	if (IsUsingStructNetSerializer(MemberSerializerDescriptor))
+
+	if (EnumHasAnyFlags(MemberSerializerDescriptor.Serializer->Traits, ENetSerializerTraits::UseSerializerIsEqual))
+	{
+		FNetSerializationContext Context;
+		FNetIsEqualArgs Args;
+		Args.NetSerializerConfig = MemberSerializerDescriptor.SerializerConfig;
+		Args.Source0 = NetSerializerValuePointer(ValueA);
+		Args.Source1 = NetSerializerValuePointer(ValueB);
+		Args.bStateIsQuantized = false;
+
+		return MemberSerializerDescriptor.Serializer->IsEqual(Context, Args);
+	}
+	else if (IsUsingStructNetSerializer(MemberSerializerDescriptor))
 	{
 		const FStructNetSerializerConfig* StructConfig = static_cast<const FStructNetSerializerConfig*>(MemberSerializerDescriptor.SerializerConfig);
 		const FReplicationStateDescriptor* StructDescriptor = StructConfig->StateDescriptor;
@@ -175,7 +187,7 @@ bool InternalCompareMember(const FReplicationStateDescriptor* Descriptor, uint32
 		}
 	}
 
-	// Default handling for all properties except for structs and arrays that have some non-replicated members
+	// Default handling
 	const FProperty* Property = Descriptor->MemberProperties[MemberIndex];
 	return Property->Identical(ValueA, ValueB);
 }
diff --git a/Engine/Source/Runtime/Experimental/Iris/Core/Private/Iris/ReplicationState/ReplicationStateDescriptorBuilder.cpp b/Engine/Source/Runtime/Experimental/Iris/Core/Private/Iris/ReplicationState/ReplicationStateDescriptorBuilder.cpp
--- a/Engine/Source/Runtime/Experimental/Iris/Core/Private/Iris/ReplicationState/ReplicationStateDescriptorBuilder.cpp
+++ b/Engine/Source/Runtime/Experimental/Iris/Core/Private/Iris/ReplicationState/ReplicationStateDescriptorBuilder.cpp
@@ -77,6 +77,7 @@ enum class EMemberPropertyTraits : uint32
 	IsFastArrayItem						= IsNativeFastArray << 1U,
 	HasConnectionSpecificSerialization	= IsFastArrayItem << 1U,
 	HasPushBasedDirtiness				= HasConnectionSpecificSerialization << 1U,
+	UseSerializerIsEqual				= HasPushBasedDirtiness << 1U,
 };
 ENUM_CLASS_FLAGS(EMemberPropertyTraits);
 
@@ -508,6 +509,11 @@ void FPropertyReplicationStateDescriptorBuilder::BuildMemberCache(FBuilderContex
 				Member.Traits |= EMemberPropertyTraits::HasConnectionSpecificSerialization;
 			}
 
+			if (EnumHasAnyFlags(Descriptor->Traits, EReplicationStateTraits::UseSerializerIsEqual))
+			{
+				Member.Traits |= EMemberPropertyTraits::UseSerializerIsEqual;
+			}
+
 			if (!EnumHasAnyFlags(Descriptor->Traits, EReplicationStateTraits::AllMembersAreReplicated))
 			{
 				// If we are building a struct containing a struct that isn't fully replicated then this struct isn't fully replicated either.
@@ -531,11 +537,17 @@ void FPropertyReplicationStateDescriptorBuilder::BuildMemberCache(FBuilderContex
 			const FReplicationStateDescriptor* Descriptor = CurrentCacheEntry->Descriptor.GetReference();
 
 			// Propagate some traits. Doing this earlier would be quite expensive.
+
 			if (EnumHasAnyFlags(Descriptor->Traits, EReplicationStateTraits::HasObjectReference))
 			{
 				Member.Traits |= EMemberPropertyTraits::HasObjectReference;
 			}
 
+			if (EnumHasAnyFlags(Descriptor->Traits, EReplicationStateTraits::UseSerializerIsEqual))
+			{
+				Member.Traits |= EMemberPropertyTraits::UseSerializerIsEqual;
+			}
+
 			// We propagate this up so that a struct wrapping an array that is not fully replicated will also have this set
 			if (!EnumHasAnyFlags(Descriptor->Traits, EReplicationStateTraits::AllMembersAreReplicated))
 			{
@@ -967,6 +979,7 @@ void FPropertyReplicationStateDescriptorBuilder::BuildMemberTraitsDescriptors(FR
 		CurrentMemberTraitsDescriptor->Traits |= (EnumHasAnyFlags(Info.Traits, EMemberPropertyTraits::HasDynamicState) ? EReplicationStateMemberTraits::HasDynamicState : EReplicationStateMemberTraits::None);
 		CurrentMemberTraitsDescriptor->Traits |= (EnumHasAnyFlags(Info.Traits, EMemberPropertyTraits::HasObjectReference) ? EReplicationStateMemberTraits::HasObjectReference : EReplicationStateMemberTraits::None);
 		CurrentMemberTraitsDescriptor->Traits |= (EnumHasAnyFlags(Info.Traits, EMemberPropertyTraits::RepNotifyAlways) ? EReplicationStateMemberTraits::HasRepNotifyAlways : EReplicationStateMemberTraits::None);
+		CurrentMemberTraitsDescriptor->Traits |= (EnumHasAnyFlags(Info.Traits, EMemberPropertyTraits::UseSerializerIsEqual) ? EReplicationStateMemberTraits::UseSerializerIsEqual : EReplicationStateMemberTraits::None);
 		++CurrentMemberTraitsDescriptor;
 	}
 
@@ -1364,6 +1377,11 @@ EReplicationStateTraits FPropertyReplicationStateDescriptorBuilder::BuildReplica
 		Traits |= EReplicationStateTraits::HasConnectionSpecificSerialization;
 	}
 
+	if (EnumHasAnyFlags(PropertyTraits, EMemberPropertyTraits::UseSerializerIsEqual))
+	{
+		Traits |= EReplicationStateTraits::UseSerializerIsEqual;
+	}
+
 	// Only set the PushBasedDirtiness trait if all properties are push based. An alternative could be to split into multiple states.
 	if (EnumHasAnyFlags(SharedPropertyTraits, EMemberPropertyTraits::HasPushBasedDirtiness))
 	{
@@ -1814,6 +1832,7 @@ void FPropertyReplicationStateDescriptorBuilder::GetSerializerTraits(FMemberProp
 	OutMemberProperty.Traits |= GetHasObjectReferenceTraits(NetSerializer);
 	OutMemberProperty.Traits |= EnumHasAnyFlags(NetSerializer->Traits, ENetSerializerTraits::HasDynamicState) ? EMemberPropertyTraits::HasDynamicState : EMemberPropertyTraits::None;
 	OutMemberProperty.Traits |= EnumHasAnyFlags(NetSerializer->Traits, ENetSerializerTraits::HasConnectionSpecificSerialization) ? EMemberPropertyTraits::HasConnectionSpecificSerialization : EMemberPropertyTraits::None;
+	OutMemberProperty.Traits |= EnumHasAnyFlags(NetSerializer->Traits, ENetSerializerTraits::UseSerializerIsEqual) ? EMemberPropertyTraits::UseSerializerIsEqual : EMemberPropertyTraits::None;
 	OutMemberProperty.Traits |= GetFastArrayPropertyTraits(NetSerializer, Property);
 }
 
diff --git a/Engine/Source/Runtime/Experimental/Iris/Core/Public/Iris/ReplicationState/ReplicationStateDescriptor.h b/Engine/Source/Runtime/Experimental/Iris/Core/Public/Iris/ReplicationState/ReplicationStateDescriptor.h
--- a/Engine/Source/Runtime/Experimental/Iris/Core/Public/Iris/ReplicationState/ReplicationStateDescriptor.h
+++ b/Engine/Source/Runtime/Experimental/Iris/Core/Public/Iris/ReplicationState/ReplicationStateDescriptor.h
@@ -47,6 +47,7 @@ enum class EReplicationStateMemberTraits : uint16
 	HasObjectReference = HasDynamicState << 1U,
 	HasConnectionSpecificSerialization = HasObjectReference << 1U,
 	HasRepNotifyAlways = HasConnectionSpecificSerialization << 1U,
+	UseSerializerIsEqual = HasRepNotifyAlways << 1U,
 };
 ENUM_CLASS_FLAGS(EReplicationStateMemberTraits);
 
@@ -177,6 +178,7 @@ enum class EReplicationStateTraits : uint32
 	HasPushBasedDirtiness				= HasConnectionSpecificSerialization << 1U,
 	// Whether delta compression is supported or not
 	SupportsDeltaCompression			= HasPushBasedDirtiness << 1U,
+	UseSerializerIsEqual				= SupportsDeltaCompression << 1U,	
 };
 ENUM_CLASS_FLAGS(EReplicationStateTraits);
 
diff --git a/Engine/Source/Runtime/Experimental/Iris/Core/Public/Iris/Serialization/NetSerializer.h b/Engine/Source/Runtime/Experimental/Iris/Core/Public/Iris/Serialization/NetSerializer.h
--- a/Engine/Source/Runtime/Experimental/Iris/Core/Public/Iris/Serialization/NetSerializer.h
+++ b/Engine/Source/Runtime/Experimental/Iris/Core/Public/Iris/Serialization/NetSerializer.h
@@ -346,6 +346,9 @@ enum class ENetSerializerTraits : uint32
 	HasConnectionSpecificSerialization = HasDynamicState << 1U,
 	/** There are net references that need to be gathered via calls to CollectNetReferences. */
 	HasCustomNetReference = HasConnectionSpecificSerialization << 1U,
+
+	/** Data replicated using this serializer should use the IsEqual implementation in order to determine whether the data is dirty or not. */
+	UseSerializerIsEqual = HasCustomNetReference << 1U,
 };
 ENUM_CLASS_FLAGS(ENetSerializerTraits);
 
diff --git a/Engine/Source/Runtime/Experimental/Iris/Core/Public/Iris/Serialization/NetSerializerBuilder.inl b/Engine/Source/Runtime/Experimental/Iris/Core/Public/Iris/Serialization/NetSerializerBuilder.inl
--- a/Engine/Source/Runtime/Experimental/Iris/Core/Public/Iris/Serialization/NetSerializerBuilder.inl
+++ b/Engine/Source/Runtime/Experimental/Iris/Core/Public/Iris/Serialization/NetSerializerBuilder.inl
@@ -124,6 +124,7 @@ private:
 		static constexpr bool bHasCustomNetReference = true;
 		static constexpr bool bHasDynamicState = true;
 		static constexpr bool bUseDefaultDelta = true;
+		static constexpr bool bUseSerializerIsEqual = false;
 	};
 
 	template<typename U, U> struct FSignatureCheck;
@@ -140,6 +141,12 @@ private:
 	template<typename U> static ETrueType TestHasCustomNetReferenceIsBool(typename TEnableIf<TIsSame<decltype(&FTraits::bHasCustomNetReference), decltype(&U::bHasCustomNetReference)>::Value>::Type*);
 	template<typename> static EFalseType TestHasCustomNetReferenceIsBool(...);
 
+	template<typename U> static ETrueType TestUseSerializerIsEqualIsPresent(FTypeCheck<decltype(&U::bUseSerializerIsEqual)>*);
+	template<typename> static EFalseType TestUseSerializerIsEqualIsPresent(...);
+
+	template<typename U> static ETrueType TestUseSerializerIsEqualIsBool(typename TEnableIf<TIsSame<decltype(&FTraits::bUseSerializerIsEqual), decltype(&U::bUseSerializerIsEqual)>::Value>::Type*);
+	template<typename> static EFalseType TestUseSerializerIsEqualIsBool(...);
+
 	template<typename U> static ETrueType TestIsForwardingSerializerIsPresent(FTypeCheck<decltype(&U::bIsForwardingSerializer)>*);
 	template<typename> static EFalseType TestIsForwardingSerializerIsPresent(...);
 
@@ -225,6 +232,9 @@ private:
 		HasCustomNetReferenceIsPresent = unsigned(decltype(TestHasCustomNetReferenceIsPresent<NetSerializerImpl>(nullptr))::Value),
 		HasCustomNetReferenceIsBool = unsigned(decltype(TestHasCustomNetReferenceIsBool<NetSerializerImpl>(nullptr))::Value),
 
+		UseSerializerIsEqualIsPresent = unsigned(decltype(TestUseSerializerIsEqualIsPresent<NetSerializerImpl>(nullptr))::Value),
+		UseSerializerIsEqualIsBool = unsigned(decltype(TestUseSerializerIsEqualIsBool<NetSerializerImpl>(nullptr))::Value),
+
 		IsForwardingSerializerIsPresent = unsigned(decltype(TestIsForwardingSerializerIsPresent<NetSerializerImpl>(nullptr))::Value),
 		IsForwardingSerializerIsBool = unsigned(decltype(TestIsForwardingSerializerIsBool<NetSerializerImpl>(nullptr))::Value),
 		HasConnectionSpecificSerializationIsPresent = unsigned(decltype(TestHasConnectionSpecificSerializationIsPresent<NetSerializerImpl>(nullptr))::Value),
@@ -270,6 +280,12 @@ public:
 	template<typename T = void, typename U = typename TEnableIf<!HasCustomNetReferenceIsBool, T>::Type, char V = 0>
 	static constexpr bool HasCustomNetReference() { return false; }
 
+	template<typename T = void, typename U = typename TEnableIf<UseSerializerIsEqualIsBool && HasIsEqual, T>::Type, bool V = true>
+	static constexpr bool UseSerializerIsEqual() { return NetSerializerImpl::bUseSerializerIsEqual; }
+
+	template<typename T = void, typename U = typename TEnableIf<!(UseSerializerIsEqualIsBool && HasIsEqual), T>::Type, char V = 0>
+	static constexpr bool UseSerializerIsEqual() { return false; }
+
 	template<typename T = void, typename U = typename TEnableIf<IsForwardingSerializerIsBool, T>::Type, bool V = true>
 	static constexpr bool IsForwardingSerializer() { return NetSerializerImpl::bIsForwardingSerializer; }
 
@@ -426,6 +442,8 @@ public:
 		Traits |= (HasConnectionSpecificSerialization() ? ENetSerializerTraits::HasConnectionSpecificSerialization : ENetSerializerTraits::None);
 		Traits |= (HasCustomNetReference() ? ENetSerializerTraits::HasCustomNetReference : ENetSerializerTraits::None);
 		Traits |= (HasDynamicState() ? ENetSerializerTraits::HasDynamicState : ENetSerializerTraits::None);
+		Traits |= (UseSerializerIsEqual() ? ENetSerializerTraits::UseSerializerIsEqual : ENetSerializerTraits::None);
+
 		return Traits;
 	}
 
