Fix Chaos/Transform.h's TransformNormalUnsafe to apply scale and rotation in the correct order, and add a safe version of TransformNormal

diff --git a/Engine/Source/Runtime/Experimental/Chaos/Public/Chaos/Transform.h b/Engine/Source/Runtime/Experimental/Chaos/Public/Chaos/Transform.h
--- a/Engine/Source/Runtime/Experimental/Chaos/Public/Chaos/Transform.h
+++ b/Engine/Source/Runtime/Experimental/Chaos/Public/Chaos/Transform.h
@@ -205,16 +205,36 @@ namespace Chaos
 		// Transform the normal when scale may be non-unitary. Assumes no scale components are zero.
 		TVector<FRealSingle, 3> TransformNormalUnsafe(const TVector<FRealSingle, 3>& Normal) const
 		{
-			const TVector<FRealSingle, 3> RotatedNormal = TransformNormalNoScale(Normal);
-			const TVector<FRealSingle, 3> ScaledNormal = RotatedNormal / GetScale3D();
+			const TVector<FRealSingle, 3> ScaledNormal = Normal / GetScale3D();
 			const FRealSingle ScaledNormal2 = ScaledNormal.SizeSquared();
 			if (ScaledNormal2 > UE_SMALL_NUMBER)
 			{
-				return ScaledNormal * FMath::InvSqrt(ScaledNormal2);
+				return TransformNormalNoScale(ScaledNormal * FMath::InvSqrt(ScaledNormal2));
 			}
 			else
 			{
-				return RotatedNormal;
+				return TransformNormalNoScale(Normal);
+			}
+		}
+
+		// Transform the normal when scale may be non-unitary.
+		TVector<FRealSingle, 3> TransformNormal(const TVector<FRealSingle, 3>& Normal) const
+		{
+			// Apply inverse scaling without a per-component divide (conceptually, by scaling the normal by an extra det*sign(det))
+			const TVector<FRealSingle, 3>& S = GetScale3D();
+			FRealSingle DetSign = (S.X * S.Y * S.Z) < 0 ? (FRealSingle)-1 : (FRealSingle)1;
+			const TVector<FRealSingle, 3> SafeInvS(S.Y * S.Z * DetSign, S.X * S.Z * DetSign, S.X * S.Y * DetSign);
+			const TVector<FRealSingle, 3> ScaledNormal = SafeInvS * Normal;
+			const FRealSingle ScaledNormalLengthSq = ScaledNormal.SizeSquared();
+			// If inverse scaling would not zero the normal, normalize it and rotate
+			if (ScaledNormalLengthSq > UE_SMALL_NUMBER)
+			{
+				return TransformVectorNoScale(ScaledNormal * FMath::InvSqrt(ScaledNormalLengthSq));
+			}
+			// Otherwise just rotate without scaling
+			else
+			{
+				return TransformVectorNoScale(Normal);
 			}
 		}
 
@@ -314,16 +334,36 @@ namespace Chaos
 		// Transform the normal when scale may be non-unitary. Assumes no scale components are zero.
 		TVector<FRealDouble, 3> TransformNormalUnsafe(const TVector<FRealDouble, 3>& Normal) const
 		{
-			const TVector<FRealDouble, 3> RotatedNormal = TransformNormalNoScale(Normal);
-			const TVector<FRealDouble, 3> ScaledNormal = RotatedNormal / GetScale3D();
+			const TVector<FRealDouble, 3> ScaledNormal = Normal / GetScale3D();
 			const FRealDouble ScaledNormal2 = ScaledNormal.SizeSquared();
 			if (ScaledNormal2 > UE_SMALL_NUMBER)
 			{
-				return ScaledNormal * FMath::InvSqrt(ScaledNormal2);
+				return TransformNormalNoScale(ScaledNormal * FMath::InvSqrt(ScaledNormal2));
+			}
+			else
+			{
+				return TransformNormalNoScale(Normal);
+			}
+		}
+
+		// Transform the normal when scale may be non-unitary.
+		TVector<FRealDouble, 3> TransformNormal(const TVector<FRealDouble, 3>& Normal) const
+		{
+			// Apply inverse scaling without a per-component divide (conceptually, by scaling the normal by an extra det*sign(det))
+			const TVector<FRealDouble, 3>& S = GetScale3D();
+			FRealDouble DetSign = (S.X * S.Y * S.Z) < 0 ? (FRealDouble)-1 : (FRealDouble)1;
+			const TVector<FRealDouble, 3> SafeInvS(S.Y * S.Z * DetSign, S.X * S.Z * DetSign, S.X * S.Y * DetSign);
+			const TVector<FRealDouble, 3> ScaledNormal = SafeInvS * Normal;
+			const FRealDouble ScaledNormalLengthSq = ScaledNormal.SizeSquared();
+			// If inverse scaling would not zero the normal, normalize it and rotate
+			if (ScaledNormalLengthSq > UE_SMALL_NUMBER)
+			{
+				return TransformVectorNoScale(ScaledNormal * FMath::InvSqrt(ScaledNormalLengthSq));
 			}
+			// Otherwise just rotate without scaling
 			else
 			{
-				return RotatedNormal;
+				return TransformVectorNoScale(Normal);
 			}
 		}
 
