Overlap Box optimization with Height Field and Triangle Mesh

diff --git a/Engine/Source/Programs/HeadlessChaos/Private/HeadlessChaosTestOverlap.cpp b/Engine/Source/Programs/HeadlessChaos/Private/HeadlessChaosTestOverlap.cpp
--- a/Engine/Source/Programs/HeadlessChaos/Private/HeadlessChaosTestOverlap.cpp
+++ b/Engine/Source/Programs/HeadlessChaos/Private/HeadlessChaosTestOverlap.cpp
@@ -606,6 +606,45 @@ namespace ChaosTest
 					bool bResult = TriangleMesh->OverlapGeom(Box, QueryTM, 0.0);
 					EXPECT_EQ(bResult, false);
 				}
+				{
+					const FRigidTransform3 QueryTM(FVec3(0.0, 0.0, -2.5), FQuat(FVec3(1.0, 0.0, 0.0), 3.1415926 / 4.0));
+					bool bResult = TriangleMesh->OverlapGeom(Box, QueryTM, 0.0);
+					EXPECT_EQ(bResult, false);
+				}
+				{
+					const FRigidTransform3 QueryTM(FVec3(5.0, 5.0, -2.5), FQuat(FVec3(0.0, 1.0, 0.0), 3.1415926 / 6.0));
+					bool bResult = TriangleMesh->OverlapGeom(Box, QueryTM, 0.0);
+					EXPECT_EQ(bResult, false);
+				}
+				{
+					const FRigidTransform3 QueryTM(FVec3(-5.0, 5.0, -2.5), FQuat(FVec3(1.0, 0.0, 0.0), 3.1415926 / 4.0));
+					bool bResult = TriangleMesh->OverlapGeom(Box, QueryTM, 0.0);
+					EXPECT_EQ(bResult, false);
+				}
+				{
+					const FRigidTransform3 QueryTM(FVec3(-5.0, -5.0, -2.5), FQuat(FVec3(1.0, 0.0, 0.0), -3.1415926 / 6.0));
+					bool bResult = TriangleMesh->OverlapGeom(Box, QueryTM, 0.0);
+					EXPECT_EQ(bResult, false);
+				}
+				{
+					const FRigidTransform3 QueryTM(FVec3(0.0, 0.0, -2.5), FQuat(FVec3(1.0, 0.0, 0.0), 3.1415926 / 4.0));
+					bool bResult = TriangleMesh->OverlapGeom(Box, QueryTM, 0.0);
+					EXPECT_EQ(bResult, false);
+				}
+				{
+					const FRigidTransform3 QueryTM(FVec3(7.0, 0.0, 10.0), FQuat::Identity);
+					bool bResult = TriangleMesh->OverlapGeom(Box, QueryTM, 0.0);
+					EXPECT_EQ(bResult, false);
+				}
+			}
+			{
+				// Box test
+				const TBox<FReal, 3> Box = TBox<FReal, 3>({ -1.0, -1.0, -1.0 }, { 1.0, 1.0, 1.0 });
+				{
+					const FRigidTransform3 QueryTM(FVec3(7.0, 0.0, 10.0), FQuat::Identity);
+					bool bResult = TriangleMesh->OverlapGeom(Box, QueryTM, 0.0);
+					EXPECT_EQ(bResult, false);
+				}
 			}
 
 			// Box Scaled
diff --git a/Engine/Source/Runtime/Experimental/Chaos/Private/Chaos/Collision/TriangleOverlap.cpp b/Engine/Source/Runtime/Experimental/Chaos/Private/Chaos/Collision/TriangleOverlap.cpp
--- a/Engine/Source/Runtime/Experimental/Chaos/Private/Chaos/Collision/TriangleOverlap.cpp
+++ b/Engine/Source/Runtime/Experimental/Chaos/Private/Chaos/Collision/TriangleOverlap.cpp
@@ -1,12 +1,13 @@
 // Copyright Epic Games, Inc. All Rights Reserved.
 #include "Chaos/Collision/TriangleOverlap.h"
+#include "Chaos/Box.h"
+#include "Chaos/ImplicitObjectScaled.h"
 #include "Chaos/VectorUtility.h"
 
 namespace Chaos
 {
 	bool ComputeCapsuleTriangleOverlapSimd(const VectorRegister4Float& A, const VectorRegister4Float& B, const VectorRegister4Float& C, const VectorRegister4Float& X1, const VectorRegister4Float& X2, FRealSingle Radius)
 	{
-
 		const VectorRegister4Float X2X1 = VectorSubtract(X1, X2);
 		const VectorRegister4Float SqrX2X1 = VectorDot3(X2X1, X2X1);
 		const VectorRegister4Float ZeroMask = VectorCompareEQ(VectorZeroFloat(), SqrX2X1);
@@ -162,7 +163,6 @@ namespace Chaos
 
 	bool ComputeSphereTriangleOverlapSimd(const VectorRegister4Float& A, const VectorRegister4Float& B, const VectorRegister4Float& C, const VectorRegister4Float& X, FRealSingle Radius)
 	{
-
 		const VectorRegister4Float AB = VectorSubtract(B, A);
 		const VectorRegister4Float BC = VectorSubtract(C, B);
 
@@ -242,4 +242,390 @@ namespace Chaos
 		}
 		return true;
 	}
+
+	FBoxSimd::FBoxSimd()
+	{
+	}
+
+	FBoxSimd::FBoxSimd(const FRigidTransform3& WorldScaleQueryTM, const TBox<FReal, 3>& InQueryGeom)
+	{
+		FRigidTransform3 Transform(WorldScaleQueryTM);
+		Transform.SetTranslation(Transform.GetTranslation() + InQueryGeom.GetCenter());
+		FVec3f HalfExtentsf = InQueryGeom.Extents() * 0.5;
+		Initialize(Transform, HalfExtentsf);
+	}
+
+	FBoxSimd::FBoxSimd(const FRigidTransform3& QueryTM, const TImplicitObjectScaled< TBox<FReal, 3> >& QueryGeom)
+	{
+		FRigidTransform3 Transform(QueryTM);
+		Transform.TransformPosition(QueryGeom.GetUnscaledObject()->GetCenter() * QueryGeom.GetScale());
+		FVec3f HalfExtentsf = QueryGeom.GetUnscaledObject()->Extents() * 0.5 * (QueryGeom.GetScale() + UE_SMALL_NUMBER);
+		Initialize(Transform, HalfExtentsf);
+	}
+
+	FORCEINLINE void FBoxSimd::Initialize(const FRigidTransform3& Transform, const FVec3f& HalfExtentsf)
+	{
+		TRotation<FRealSingle, 3> Rotation = Transform.GetRotation();
+		PMatrix<FRealSingle, 3, 3> Matrix = Rotation.ToMatrix();
+		XAxis = VectorLoadFloat3(Matrix.M[0]);
+		YAxis = VectorLoadFloat3(Matrix.M[1]);
+		ZAxis = VectorLoadFloat3(Matrix.M[2]);
+
+		FVec3 Translation = Transform.GetTranslation();
+		Position = MakeVectorRegisterFloatFromDouble(VectorLoadDouble3(&Translation.X));
+
+		const VectorRegister4Float HalfExtents = VectorLoadFloat3(&HalfExtentsf.X);
+		XHalfExtent = VectorReplicate(HalfExtents, 0);
+		YHalfExtent = VectorReplicate(HalfExtents, 1);
+		ZHalfExtent = VectorReplicate(HalfExtents, 2);
+
+		XAxisHalfExtent = VectorMultiply(XAxis, XHalfExtent);
+		YAxisHalfExtent = VectorMultiply(YAxis, YHalfExtent);
+		ZAxisHalfExtent = VectorMultiply(ZAxis, ZHalfExtent);
+
+	}
+	namespace
+	{
+		FORCEINLINE_DEBUGGABLE bool HasToComputeEdge(const VectorRegister4Float& BoxNormal1, const VectorRegister4Float& BoxNormal2, const VectorRegister4Float& BoxEdge, const VectorRegister4Float& TriNormal, const VectorRegister4Float& TriEdge)
+		{
+			const FRealSingle TriNormBoxEdge = VectorDot3Scalar(TriNormal, BoxEdge);		// TriNormal | BoxEdge
+			const FRealSingle BoxNorm1TriEdge = VectorDot3Scalar(BoxNormal1, TriEdge);		// BoxNormalA | TriEdge
+			const FRealSingle BoxNorm2TriEdge = VectorDot3Scalar(BoxNormal2, TriEdge);		// BoxNormalB | TriEdge
+
+			return ((BoxNorm1TriEdge * BoxNorm2TriEdge) < 0.0f) && (((TriNormBoxEdge * BoxNorm2TriEdge) > 0.0f || (TriNormBoxEdge * BoxNorm1TriEdge) > 0.0f));
+		}
+	}
+
+	FORCEINLINE_DEBUGGABLE bool FBoxSimd::ComputeEdge(const VectorRegister4Float& Normal, const VectorRegister4Float& PlaneVertex, const VectorRegister4Float& Edge, const VectorRegister4Float& Centroid) const
+	{
+		const VectorRegister4Float PlaneNormal = VectorNormalize(VectorCross(Normal, Edge));
+		const VectorRegister4Float XComp = VectorSelect(VectorCompareGT(VectorDot3(PlaneNormal, XAxisHalfExtent), VectorZeroFloat()), XAxisHalfExtent, VectorNegate(XAxisHalfExtent));
+		const VectorRegister4Float YComp = VectorSelect(VectorCompareGT(VectorDot3(PlaneNormal, YAxisHalfExtent), VectorZeroFloat()), YAxisHalfExtent, VectorNegate(YAxisHalfExtent));
+		const VectorRegister4Float ZComp = VectorSelect(VectorCompareGT(VectorDot3(PlaneNormal, ZAxisHalfExtent), VectorZeroFloat()), ZAxisHalfExtent, VectorNegate(ZAxisHalfExtent));
+
+		const VectorRegister4Float LocalClosest = VectorAdd(VectorAdd(XComp, YComp), ZComp);
+		const VectorRegister4Float ClosestPoint = VectorAdd(LocalClosest, Position);
+
+		FRealSingle Dist = VectorDot3Scalar(VectorSubtract(PlaneVertex, ClosestPoint), PlaneNormal);
+		if (Dist > 0.0)
+		{
+			return false;
+		}
+
+		// Triangle edge vs box edges 
+		VectorRegister4Float OtherClosests[3];
+		OtherClosests[0] = VectorAdd(VectorAdd(VectorNegate(XComp), YComp), ZComp);
+		OtherClosests[1] = VectorAdd(VectorAdd(XComp, VectorNegate(YComp)), ZComp);
+		OtherClosests[2] = VectorAdd(VectorAdd(XComp, YComp), VectorNegate(ZComp));
+
+		VectorRegister4Float BoxEdges[3];
+		VectorRegister4Float BoxEdgeNormals[3];
+		for (int32 i = 0; i < 3; i++)
+		{
+			const VectorRegister4Float OtherClosest = VectorAdd(OtherClosests[i], Position);
+			BoxEdges[i] = VectorSubtract(OtherClosest, ClosestPoint);
+			BoxEdgeNormals[i] = VectorNegate(VectorNormalize(BoxEdges[i]));
+		}
+
+		for (int32 i = 0; i < 3; i++)
+		{
+			if (HasToComputeEdge(BoxEdgeNormals[(i + 1) % 3], BoxEdgeNormals[(i + 2) % 3], BoxEdges[i], Normal, Edge))
+			{
+				VectorRegister4Float Axis = VectorNormalize(VectorCross(BoxEdges[i], Edge));
+				const FRealSingle Sign = VectorDot3Scalar(VectorSubtract(PlaneVertex, Centroid), Axis);
+				if (Sign < 0.0f)
+				{
+					Axis = VectorNegate(Axis);
+				}
+				const FRealSingle Separation = VectorDot3Scalar(VectorSubtract(ClosestPoint, PlaneVertex), Axis);
+				if (Separation > 0.0f)
+				{
+					return false;
+				}
+			}
+		}
+		return true;
+	}
+
+	FORCEINLINE_DEBUGGABLE bool FBoxSimd::ComputeBoxPlane(const VectorRegister4Float& DistCenter, const VectorRegister4Float& AxisHalfExtent) const
+	{
+		const VectorRegister4Float SignMask = GlobalVectorConstants::SignMask();
+		VectorRegister4Float AllPosSignVector = VectorBitwiseAnd(DistCenter, SignMask);
+		VectorRegister4Float AllNegSignVector = VectorNegate(AllPosSignVector);
+
+		constexpr static int32 TrueMask = 0b1111;
+		if (VectorMaskBits(VectorCompareEQ(DistCenter, AllPosSignVector)) == TrueMask || VectorMaskBits(VectorCompareEQ(DistCenter, AllNegSignVector)) == TrueMask)
+		{
+			VectorRegister4Float IsSepX = VectorCompareGT(VectorAbs(DistCenter), AxisHalfExtent);
+			if (VectorMaskBits(IsSepX) == TrueMask)
+			{
+				return false;
+			}
+		}
+		return true;
+	}
+
+	bool FBoxSimd::OverlapTriangle(const VectorRegister4Float& A, const VectorRegister4Float& B, const VectorRegister4Float& C) const
+	{
+		const VectorRegister4Float PA = VectorSubtract(A, Position);
+		const VectorRegister4Float PB = VectorSubtract(B, Position);
+		const VectorRegister4Float PC = VectorSubtract(C, Position);
+
+		// Box planes
+		const VectorRegister4Float DistACenterX = VectorDot3(XAxis, PA);
+		const VectorRegister4Float DistBCenterX = VectorDot3(XAxis, PB);
+		const VectorRegister4Float DistCCenterX = VectorDot3(XAxis, PC);
+
+		VectorRegister4Float DistCenterX = VectorUnpackLo(DistACenterX, DistBCenterX);
+		DistCenterX = VectorMoveLh(DistCenterX, DistCCenterX);
+
+		if (!ComputeBoxPlane(DistCenterX, XHalfExtent))
+		{
+			return false;
+		}
+
+		const VectorRegister4Float DistACenterY = VectorDot3(YAxis, PA);
+		const VectorRegister4Float DistBCenterY = VectorDot3(YAxis, PB);
+		const VectorRegister4Float DistCCenterY = VectorDot3(YAxis, PC);
+
+		VectorRegister4Float DistCenterY = VectorUnpackLo(DistACenterY, DistBCenterY);
+		DistCenterY = VectorMoveLh(DistCenterY, DistCCenterY);
+
+		if (!ComputeBoxPlane(DistCenterY, YHalfExtent))
+		{
+			return false;
+		}
+
+		const VectorRegister4Float DistACenterZ = VectorDot3(ZAxis, PA);
+		const VectorRegister4Float DistBCenterZ = VectorDot3(ZAxis, PB);
+		const VectorRegister4Float DistCCenterZ = VectorDot3(ZAxis, PC);
+
+		VectorRegister4Float DistCenterZ = VectorUnpackLo(DistACenterZ, DistBCenterZ);
+		DistCenterZ = VectorMoveLh(DistCenterZ, DistCCenterZ);
+
+		if (!ComputeBoxPlane(DistCenterZ, ZHalfExtent))
+		{
+			return false;
+		}
+
+		// Triangles Edges
+		const VectorRegister4Float AB = VectorSubtract(B, A);
+		const VectorRegister4Float BC = VectorSubtract(C, B);
+		const VectorRegister4Float CA = VectorSubtract(A, C);
+
+		const VectorRegister4Float Normal = VectorNormalize(VectorCross(AB, BC));
+		constexpr FRealSingle ThirdFloat = 1.0f / 3.0f;
+		constexpr VectorRegister4Float Third = MakeVectorRegisterFloatConstant(ThirdFloat, ThirdFloat, ThirdFloat, ThirdFloat);
+		const VectorRegister4Float Centroid = VectorMultiply(VectorAdd(VectorAdd(A, B), C), Third);
+
+		if (!ComputeEdge(Normal, A, AB, Centroid))
+		{
+			return false;
+		}
+		if (!ComputeEdge(Normal, B, BC, Centroid))
+		{
+			return false;
+		}
+		if (!ComputeEdge(Normal, C, CA, Centroid))
+		{
+			return false;
+		}
+
+		// Triangle Plane 
+		VectorRegister4Float XCompNorm = VectorMultiply(XAxis, XHalfExtent);
+		XCompNorm = VectorSelect(VectorCompareGT(VectorDot3(Normal, XCompNorm), VectorZeroFloat()), XCompNorm, VectorNegate(XCompNorm));
+		VectorRegister4Float YCompNorm = VectorMultiply(YAxis, YHalfExtent);
+		YCompNorm = VectorSelect(VectorCompareGT(VectorDot3(Normal, YCompNorm), VectorZeroFloat()), YCompNorm, VectorNegate(YCompNorm));
+		VectorRegister4Float ZCompNorm = VectorMultiply(ZAxis, ZHalfExtent);
+		ZCompNorm = VectorSelect(VectorCompareGT(VectorDot3(Normal, ZCompNorm), VectorZeroFloat()), ZCompNorm, VectorNegate(ZCompNorm));
+
+		const VectorRegister4Float LoacalClosestNorm = VectorAdd(VectorAdd(XCompNorm, YCompNorm), ZCompNorm);
+		const VectorRegister4Float ClosestPointNorm = VectorAdd(LoacalClosestNorm, Position);
+		const VectorRegister4Float FurthestPointNorm = VectorSubtract(Position, LoacalClosestNorm);
+
+		VectorRegister4Float ClosCent = VectorSubtract(ClosestPointNorm, Centroid);
+		VectorRegister4Float FurtCent = VectorSubtract(FurthestPointNorm, Centroid);
+
+		FRealSingle ClosestDist = VectorDot3Scalar(Normal, VectorNormalize(ClosCent));
+		FRealSingle FurthestDist = VectorDot3Scalar(Normal, VectorNormalize(FurtCent));
+
+		return FMath::Sign(ClosestDist) != FMath::Sign(FurthestDist);
+	}
+
+	const VectorRegister4Float FAABBSimd::SignBit = GlobalVectorConstants::SignBit();
+	const VectorRegister4Float FAABBSimd::SignX = MakeVectorRegisterFloatMask(0x80000000, 0x00000000, 0x00000000, 0x00000000);
+	const VectorRegister4Float FAABBSimd::SignY = MakeVectorRegisterFloatMask(0x00000000, 0x80000000, 0x00000000, 0x00000000);
+	const VectorRegister4Float FAABBSimd::SignZ = MakeVectorRegisterFloatMask(0x00000000, 0x00000000, 0x80000000, 0x00000000);
+
+
+	FAABBSimd::FAABBSimd()
+		: Position(VectorZeroFloat())
+	{
+	}
+	
+	FAABBSimd::FAABBSimd(const TBox<FReal, 3>& QueryGeom)
+	{
+		FVec3f Positionf = QueryGeom.GetCenter();
+		Position = VectorLoadFloat3(&Positionf.X);
+		FVec3f HalfExtentsf = QueryGeom.Extents() * 0.5;
+		HalfExtents = VectorLoadFloat3(&HalfExtentsf.X);
+	}
+
+	FAABBSimd::FAABBSimd(const TImplicitObjectScaled< TBox<FReal, 3> >& QueryGeom)
+	{
+		FVec3f Positionf = QueryGeom.GetUnscaledObject()->GetCenter()* QueryGeom.GetScale();
+		Position = VectorLoadFloat3(&Positionf.X);
+		FVec3f HalfExtentsf = QueryGeom.GetUnscaledObject()->Extents() * 0.5 * (QueryGeom.GetScale() + UE_SMALL_NUMBER);
+		HalfExtents = VectorLoadFloat3(&HalfExtentsf.X);
+	}
+
+	FAABBSimd::FAABBSimd(const FRigidTransform3& WorldScaleQueryTM, const TBox<FReal, 3>& InQueryGeom)
+	{
+		TRigidTransform<FReal, 3> Transform(WorldScaleQueryTM);
+		Transform.SetTranslation(Transform.GetTranslation() + InQueryGeom.GetCenter());
+		FVec3f HalfExtentsf = InQueryGeom.Extents() * 0.5;
+		Initialize(Transform, HalfExtentsf);
+	}
+
+	FAABBSimd::FAABBSimd(const FRigidTransform3& QueryTM, const TImplicitObjectScaled< TBox<FReal, 3> >& QueryGeom)
+	{
+		TRigidTransform<FReal, 3> Transform(QueryTM);
+		Transform.TransformPosition(QueryGeom.GetUnscaledObject()->GetCenter() * QueryGeom.GetScale());
+		FVec3f HalfExtentsf = QueryGeom.GetUnscaledObject()->Extents() * 0.5 * (QueryGeom.GetScale() + UE_SMALL_NUMBER);
+		Initialize(Transform, HalfExtentsf);
+	}
+
+	FORCEINLINE void FAABBSimd::Initialize(const FRigidTransform3& Transform, const FVec3f& HalfExtentsf)
+	{
+		FVec3 Translation = Transform.GetTranslation();
+		Position = MakeVectorRegisterFloatFromDouble(VectorLoadDouble3(&Translation.X));
+
+		HalfExtents = VectorLoadFloat3(&HalfExtentsf.X);
+	}
+
+	FORCEINLINE bool FAABBSimd::ComputeEdgeOverlap(const VectorRegister4Float& TriangleEdge, const VectorRegister4Float& TriangleVertex, const VectorRegister4Float& Centroid, const VectorRegister4Float& Normal, const VectorRegister4Float& LocalClosest) const
+	{
+		// Triangle edge vs box edges 
+		const VectorRegister4Float ClosestPoint = VectorAdd(LocalClosest, Position);
+
+		VectorRegister4Float OtherClosests[3];
+		OtherClosests[0] = VectorBitwiseXor(LocalClosest, SignX);
+		OtherClosests[1] = VectorBitwiseXor(LocalClosest, SignY);
+		OtherClosests[2] = VectorBitwiseXor(LocalClosest, SignZ);
+		VectorRegister4Float BoxEdges[3];
+		VectorRegister4Float BoxEdgeNormals[3];
+		for (int32 i = 0; i < 3; i++)
+		{
+			const VectorRegister4Float OtherClosest = VectorAdd(OtherClosests[i], Position);
+			BoxEdges[i] = VectorSubtract(OtherClosest, ClosestPoint);
+			BoxEdgeNormals[i] = VectorNegate(VectorNormalize(BoxEdges[i]));
+		}
+
+		for (int32 i = 0; i < 3; i++)
+		{
+			if (HasToComputeEdge(BoxEdgeNormals[(i + 1) % 3], BoxEdgeNormals[(i + 2) % 3], BoxEdges[i], Normal, TriangleEdge))
+			{
+				VectorRegister4Float Axis = VectorNormalize(VectorCross(BoxEdges[i], TriangleEdge));
+				const FRealSingle Sign = VectorDot3Scalar(VectorSubtract(TriangleVertex, Centroid), Axis);
+				if (Sign < 0.0f)
+				{
+					Axis = VectorNegate(Axis);
+				}
+				const FRealSingle ScaledSeparation = VectorDot3Scalar(VectorSubtract(ClosestPoint, TriangleVertex), Axis);
+				if (ScaledSeparation > 0.0f)
+				{
+					return false;
+				}
+			}
+		}
+		return true;
+	}
+
+	bool FAABBSimd::OverlapTriangle(const VectorRegister4Float& A, const VectorRegister4Float& B, const VectorRegister4Float& C) const
+	{
+		const VectorRegister4Float PA = VectorSubtract(A, Position);
+		const VectorRegister4Float PB = VectorSubtract(B, Position);
+		const VectorRegister4Float PC = VectorSubtract(C, Position);
+
+		// AABB planes
+		VectorRegister4Float MinAbsDist = VectorMin(VectorMin(VectorAbs(PA), VectorAbs(PB)), VectorAbs(PC));
+
+		const VectorRegister4Float MinDist = VectorMin(VectorMin(PA, PB), PC);
+		const VectorRegister4Float MaxDist = VectorMax(VectorMax(PA, PB), PC);
+		const VectorRegister4Float MinSign = VectorCastIntToFloat(VectorShiftRightImmArithmetic(VectorCastFloatToInt(VectorBitwiseAnd(SignBit, MinDist)), 1));
+		const VectorRegister4Float MaxSign = VectorCastIntToFloat(VectorShiftRightImmArithmetic(VectorCastFloatToInt(VectorBitwiseAnd(SignBit, MaxDist)), 1));
+
+		if (VectorMaskBits(VectorBitwiseAnd(VectorCompareGT(MinAbsDist, HalfExtents), VectorCompareEQ(MinSign, MaxSign))) != 0)
+		{
+			return false;
+		}
+
+		// Triangles Edges
+		const VectorRegister4Float AB = VectorSubtract(B, A);
+		const VectorRegister4Float BC = VectorSubtract(C, B);
+		const VectorRegister4Float CA = VectorSubtract(A, C);
+
+		const VectorRegister4Float Normal = VectorNormalize(VectorCross(AB, BC));
+		constexpr FRealSingle ThirdFloat = 1.0f / 3.0f;
+		constexpr VectorRegister4Float Third = MakeVectorRegisterFloatConstant(ThirdFloat, ThirdFloat, ThirdFloat, ThirdFloat);
+		const VectorRegister4Float Centroid = VectorMultiply(VectorAdd(VectorAdd(A, B), C), Third);
+		
+
+		const VectorRegister4Float ABPlane = VectorNormalize(VectorCross(Normal, AB));
+		const VectorRegister4Float ABSigns = VectorBitwiseAnd(ABPlane, SignBit);
+		const VectorRegister4Float ABLocalClosest = VectorBitwiseOr(ABSigns, HalfExtents);
+		const VectorRegister4Float ABClosestPoint = VectorAdd(ABLocalClosest, Position);
+
+		if (VectorDot3Scalar(VectorSubtract(A, ABClosestPoint), ABPlane) > 0.0)
+		{
+			return false;
+		}
+
+		const VectorRegister4Float BCPlane = VectorNormalize(VectorCross(Normal, BC));
+		const VectorRegister4Float BCSigns = VectorBitwiseAnd(BCPlane, SignBit);
+		const VectorRegister4Float BCLocalClosest = VectorBitwiseOr(BCSigns, HalfExtents);
+		const VectorRegister4Float BCClosestPoint = VectorAdd(BCLocalClosest, Position);
+
+		if (VectorDot3Scalar(VectorSubtract(B, BCClosestPoint), BCPlane) > 0.0)
+		{
+			return false;
+		}
+
+		const VectorRegister4Float CAPlane = VectorNormalize(VectorCross(Normal, CA));
+		const VectorRegister4Float CASigns = VectorBitwiseAnd(CAPlane, SignBit);
+		const VectorRegister4Float CALocalClosest = VectorBitwiseOr(CASigns, HalfExtents);
+		const VectorRegister4Float CAClosestPoint = VectorAdd(CALocalClosest, Position);
+
+		if (VectorDot3Scalar(VectorSubtract(C, CAClosestPoint), CAPlane) > 0.0)
+		{
+			return false;
+		}
+
+		// Triangle Plane 
+		const VectorRegister4Float Signs = VectorBitwiseAnd(Normal, SignBit);
+		const VectorRegister4Float LocalClosest = VectorBitwiseOr(Signs, HalfExtents);
+		const VectorRegister4Float ClosestPoint = VectorAdd(LocalClosest, Position);
+		const VectorRegister4Float FurthestPoint = VectorSubtract(Position, LocalClosest);
+
+		VectorRegister4Float ClosCent = VectorSubtract(ClosestPoint, Centroid);
+		VectorRegister4Float FurtCent = VectorSubtract(FurthestPoint, Centroid);
+
+		FRealSingle ClosestDist = VectorDot3Scalar(Normal, VectorNormalize(ClosCent));
+		FRealSingle FurthestDist = VectorDot3Scalar(Normal, VectorNormalize(FurtCent));
+
+		if (FMath::Sign(ClosestDist) == FMath::Sign(FurthestDist))
+		{
+			return false;
+		}
+
+		if (!ComputeEdgeOverlap(AB, A, Centroid, Normal, ABLocalClosest))
+		{
+			return false;
+		}
+		if (!ComputeEdgeOverlap(BC, B, Centroid, Normal, BCLocalClosest))
+		{
+			return false;
+		}
+		return ComputeEdgeOverlap(CA, C, Centroid, Normal, CALocalClosest);
+	}
 }
diff --git a/Engine/Source/Runtime/Experimental/Chaos/Private/Chaos/Collision/TriangleOverlap.h b/Engine/Source/Runtime/Experimental/Chaos/Private/Chaos/Collision/TriangleOverlap.h
--- a/Engine/Source/Runtime/Experimental/Chaos/Private/Chaos/Collision/TriangleOverlap.h
+++ b/Engine/Source/Runtime/Experimental/Chaos/Private/Chaos/Collision/TriangleOverlap.h
@@ -5,7 +5,62 @@
 
 namespace Chaos
 {
+	template<class T, int d>
+	class TBox;
+
+	template<typename TConcrete, bool bInstanced>
+	class TImplicitObjectScaled;
+
+
 	bool ComputeCapsuleTriangleOverlapSimd(const VectorRegister4Float& A, const VectorRegister4Float& B, const VectorRegister4Float& C, const VectorRegister4Float& X1, const VectorRegister4Float& X2, FRealSingle Radius);
 	bool ComputeSphereTriangleOverlapSimd(const VectorRegister4Float& A, const VectorRegister4Float& B, const VectorRegister4Float& C, const VectorRegister4Float& X, FRealSingle Radius);
+
+	class FBoxSimd
+	{
+	public:
+		FBoxSimd();
+		FBoxSimd(const FRigidTransform3& WorldScaleQueryTM, const TBox<FReal, 3>& InQueryGeom);
+		FBoxSimd(const FRigidTransform3& QueryTM, const TImplicitObjectScaled< TBox<FReal, 3>, true >& QueryGeom);
+
+		bool OverlapTriangle(const VectorRegister4Float& A, const VectorRegister4Float& B, const VectorRegister4Float& C) const;
+
+	private:
+		void Initialize(const FRigidTransform3& Transform, const FVec3f& HalfExtentsIn);
+		bool ComputeEdge(const VectorRegister4Float& PlaneNormal, const VectorRegister4Float& PlaneVertex, const VectorRegister4Float& Edge, const VectorRegister4Float& Centroid) const;
+		bool ComputeBoxPlane(const VectorRegister4Float& DistCenter, const VectorRegister4Float& AxisHalfExtent) const;
+
+		VectorRegister4Float Position;
+		VectorRegister4Float XAxis;
+		VectorRegister4Float YAxis;
+		VectorRegister4Float ZAxis;
+		VectorRegister4Float XHalfExtent;
+		VectorRegister4Float YHalfExtent;
+		VectorRegister4Float ZHalfExtent;
+		VectorRegister4Float XAxisHalfExtent;
+		VectorRegister4Float YAxisHalfExtent;
+		VectorRegister4Float ZAxisHalfExtent;
+	};
+
+	class FAABBSimd
+	{
+	public:
+		FAABBSimd();
+		FAABBSimd(const TBox<FReal, 3>& InQueryGeom);
+		FAABBSimd(const TImplicitObjectScaled< TBox<FReal, 3>, true >& QueryGeom);
+		FAABBSimd(const FRigidTransform3& WorldScaleQueryTM, const TBox<FReal, 3>& InQueryGeom);
+		FAABBSimd(const FRigidTransform3& QueryTM, const TImplicitObjectScaled< TBox<FReal, 3>, true >& QueryGeom);
+		bool OverlapTriangle(const VectorRegister4Float& A, const VectorRegister4Float& B, const VectorRegister4Float& C) const;
+
+	private:
+		void Initialize(const FRigidTransform3& Transform, const FVec3f& HalfExtentsf);
+		bool ComputeEdgeOverlap(const VectorRegister4Float& TriangleEdge, const VectorRegister4Float& TriangleVertex, const VectorRegister4Float& Centroid, const VectorRegister4Float& Normal, const VectorRegister4Float& LocalClosest) const;
+
+		VectorRegister4Float Position;
+		VectorRegister4Float HalfExtents;
+		const static VectorRegister4Float SignBit;
+		const static VectorRegister4Float SignX;
+		const static VectorRegister4Float SignY;
+		const static VectorRegister4Float SignZ;
+	};
 }
 
diff --git a/Engine/Source/Runtime/Experimental/Chaos/Private/Chaos/HeightField.cpp b/Engine/Source/Runtime/Experimental/Chaos/Private/Chaos/HeightField.cpp
--- a/Engine/Source/Runtime/Experimental/Chaos/Private/Chaos/HeightField.cpp
+++ b/Engine/Source/Runtime/Experimental/Chaos/Private/Chaos/HeightField.cpp
@@ -1993,6 +1993,7 @@ namespace Chaos
 				[[maybe_unused]] VectorRegister4Float X1;
 				[[maybe_unused]] VectorRegister4Float X2;
 				[[maybe_unused]] FRealSingle Radius;
+				[[maybe_unused]] FAABBSimd AABBSimd;
 				if constexpr (std::is_same<QueryGeomType, FCapsule>::value)
 				{
 					Radius = FRealSingle(QueryGeom.GetRadius());
@@ -2021,6 +2022,10 @@ namespace Chaos
 					const FVec3f X1f = QueryGeom.GetUnscaledObject()->GetCenter() * QueryGeom.GetScale();
 					X1 = VectorLoadFloat3(&X1f.X);
 				}
+				else if constexpr (std::is_same<QueryGeomType, TBox<FReal, 3>>::value || std::is_same<QueryGeomType, TImplicitObjectScaled<TBox<FReal, 3>>>::value)
+				{
+					AABBSimd = FAABBSimd(QueryGeom);
+				}
 
 				VectorRegister4Float Points[4];
 				const UE::Math::TQuat<FReal>& RotationDouble = QueryTM.GetRotation();
@@ -2068,6 +2073,17 @@ namespace Chaos
 								return true;
 							}
 						}
+						else if constexpr (std::is_same<QueryGeomType, TBox<FReal, 3>>::value || std::is_same<QueryGeomType, TImplicitObjectScaled<TBox<FReal, 3>>>::value)
+						{
+							if (AABBSimd.OverlapTriangle(Points[0], Points[1], Points[3]))
+							{
+								return true;
+							}
+							if (AABBSimd.OverlapTriangle(Points[0], Points[3], Points[2]))
+							{
+								return true;
+							}
+						}
 						else
 						{
 							if (OverlapTriangleNoMTD(Points[0], Points[1], Points[3]))
@@ -2078,7 +2094,7 @@ namespace Chaos
 							{
 								return true;
 							}
-						}	
+						}
 					}
 				}
 			}
diff --git a/Engine/Source/Runtime/Experimental/Chaos/Private/Chaos/TriangleMeshImplicitObject.cpp b/Engine/Source/Runtime/Experimental/Chaos/Private/Chaos/TriangleMeshImplicitObject.cpp
--- a/Engine/Source/Runtime/Experimental/Chaos/Private/Chaos/TriangleMeshImplicitObject.cpp
+++ b/Engine/Source/Runtime/Experimental/Chaos/Private/Chaos/TriangleMeshImplicitObject.cpp
@@ -646,6 +646,135 @@ struct FTriangleMeshOverlapVisitorNoMTD<TImplicitObjectScaled<Chaos::FSphere>>
 	bool bFoundIntersection;
 };
 
+template <>
+struct FTriangleMeshOverlapVisitorNoMTD<TBox<FReal, 3>>
+{
+	FTriangleMeshOverlapVisitorNoMTD(const TRigidTransform<FReal, 3>& WorldScaleQueryTM, const TBox<FReal, 3>& InQueryGeom, FReal InThickness, const FVec3f& InTriMeshScale, const FTriangleMeshImplicitObject* InTriMesh)
+		: bIsAxisAligned(FMath::Abs(WorldScaleQueryTM.GetRotation().W) == 1.0)
+		, TriMesh(InTriMesh)
+		, bFoundIntersection(false)
+	{
+		if (bIsAxisAligned)
+		{
+			Box.AABBSimd = FAABBSimd(WorldScaleQueryTM, InQueryGeom);
+		}
+		else
+		{
+			Box.BoxSimd = FBoxSimd(WorldScaleQueryTM, InQueryGeom);
+		}
+		
+		TriMeshScale = VectorLoadFloat3(&InTriMeshScale.X);
+	}
+
+	bool VisitOverlap(int32 TriIdx)
+	{
+		VectorRegister4Float A, B, C;
+		if (TriMesh->MElements.RequiresLargeIndices())
+		{
+			TriangleMeshTransformVertsHelperSimd(TriMeshScale, TriIdx, TriMesh->MParticles, TriMesh->MElements.GetLargeIndexBuffer(), A, B, C);
+		}
+		else
+		{
+			TriangleMeshTransformVertsHelperSimd(TriMeshScale, TriIdx, TriMesh->MParticles, TriMesh->MElements.GetSmallIndexBuffer(), A, B, C);
+		}
+
+		if (bIsAxisAligned)
+		{
+			bFoundIntersection = Box.AABBSimd.OverlapTriangle(A, B, C);
+		}
+		else
+		{
+			bFoundIntersection = Box.BoxSimd.OverlapTriangle(A, B, C);
+		}
+		return !bFoundIntersection;
+	}
+
+	bool VisitSweep(int32 Instance, FRealSingle& CurDataLength)
+	{
+		check(false);
+		return true;
+	}
+	bool VisitRaycast(int32 Instance, FRealSingle& CurDataLength)
+	{
+		check(false);
+		return true;
+	}
+
+	const bool bIsAxisAligned;
+	struct {
+		FBoxSimd  BoxSimd;
+		FAABBSimd AABBSimd;
+	} Box;
+	VectorRegister4Float TriMeshScale;
+	const FTriangleMeshImplicitObject* TriMesh;
+	bool bFoundIntersection;
+};
+
+
+template <>
+struct FTriangleMeshOverlapVisitorNoMTD<TImplicitObjectScaled<TBox<FReal, 3>>>
+{
+	FTriangleMeshOverlapVisitorNoMTD(const TRigidTransform<FReal, 3>& WorldScaleQueryTM, const TImplicitObjectScaled< TBox<FReal, 3> >& InQueryGeom, FReal InThickness, const FVec3f& InTriMeshScale, const FTriangleMeshImplicitObject* InTriMesh)
+		: bIsAxisAligned(FMath::Abs(WorldScaleQueryTM.GetRotation().W) == 1.0)
+		, TriMesh(InTriMesh)
+		, bFoundIntersection(false)
+	{
+		if (bIsAxisAligned)
+		{
+			Box.AABBSimd = FAABBSimd(WorldScaleQueryTM, InQueryGeom);
+		}
+		else
+		{
+			Box.BoxSimd = FBoxSimd(WorldScaleQueryTM, InQueryGeom);
+		}
+
+		TriMeshScale = VectorLoadFloat3(&InTriMeshScale.X);
+	}
+
+	bool VisitOverlap(int32 TriIdx)
+	{
+		VectorRegister4Float A, B, C;
+		if (TriMesh->MElements.RequiresLargeIndices())
+		{
+			TriangleMeshTransformVertsHelperSimd(TriMeshScale, TriIdx, TriMesh->MParticles, TriMesh->MElements.GetLargeIndexBuffer(), A, B, C);
+		}
+		else
+		{
+			TriangleMeshTransformVertsHelperSimd(TriMeshScale, TriIdx, TriMesh->MParticles, TriMesh->MElements.GetSmallIndexBuffer(), A, B, C);
+		}
+
+		if (bIsAxisAligned)
+		{
+			bFoundIntersection = Box.AABBSimd.OverlapTriangle(A, B, C);
+		}
+		else
+		{
+			bFoundIntersection = Box.BoxSimd.OverlapTriangle(A, B, C);
+		}
+		return !bFoundIntersection;
+	}
+
+	bool VisitSweep(int32 Instance, FRealSingle& CurDataLength)
+	{
+		check(false);
+		return true;
+	}
+	bool VisitRaycast(int32 Instance, FRealSingle& CurDataLength)
+	{
+		check(false);
+		return true;
+	}
+
+	const bool bIsAxisAligned;
+	struct {
+		FBoxSimd  BoxSimd;
+		FAABBSimd AABBSimd;
+	} Box;
+	VectorRegister4Float TriMeshScale;
+	const FTriangleMeshImplicitObject* TriMesh;
+	bool bFoundIntersection;
+};
+
 template <typename QueryGeomType>
 bool FTrimeshBVH::FindAllIntersectionsNoMTD(const FAABB3& Intersection, const TRigidTransform<FReal, 3>& Transform, const QueryGeomType& QueryGeom, FReal Thickness, const FVec3& TriMeshScale, const FTriangleMeshImplicitObject* TriMesh) const
 {
@@ -1874,7 +2003,6 @@ void FTriangleMeshImplicitObject::UpdateVertices(const TArray<FVector>& NewPosit
 			MParticles.X(InternalIdx) = Chaos::FVec3(NewPositions[i]);
 		}
 	}
-
 	RebuildFastBVH();
 }
 
