Geometry collection : optimize InitConstantData
On local testing, average performance of the function went from 27ms to 11ms
- Changed UVs data layout in FGeometryCollectionConstantData to be an array of UVChannels instead of an array of vertices containing each multiple channels, this reduces allocations quite a lot
- Avoid memcpy by passing array by reference to BuildMeshSections
- Use batched parallelFor for copying bone colors and UVChannels
- Rewrite indices remap to be O(N + M) instead of O(N * M)

Most of the remaining cost is in copying the arrays now
Future work could include making a copy per thread ( but may have to keep allocation on the calling thread )

diff --git a/Engine/Source/Runtime/Experimental/Chaos/Private/GeometryCollection/GeometryCollection.cpp b/Engine/Source/Runtime/Experimental/Chaos/Private/GeometryCollection/GeometryCollection.cpp
--- a/Engine/Source/Runtime/Experimental/Chaos/Private/GeometryCollection/GeometryCollection.cpp
+++ b/Engine/Source/Runtime/Experimental/Chaos/Private/GeometryCollection/GeometryCollection.cpp
@@ -509,7 +509,7 @@ void FGeometryCollection::ReindexMaterials(FManagedArrayCollection& InCollection
 	}
 }
 
-TArray<FGeometryCollectionSection> FGeometryCollection::BuildMeshSections(const TArray<FIntVector>& InputIndices, TArray<int32> BaseMeshOriginalIndicesIndex, TArray<FIntVector>& RetIndices) const
+TArray<FGeometryCollectionSection> FGeometryCollection::BuildMeshSections(const TArray<FIntVector>& InputIndices, const TArray<int32>& BaseMeshOriginalIndicesIndex, TArray<FIntVector>& RetIndices) const
 {
 
 	return FGeometryCollectionSection::BuildMeshSections(*this, InputIndices, BaseMeshOriginalIndicesIndex, RetIndices);
diff --git a/Engine/Source/Runtime/Experimental/Chaos/Private/GeometryCollection/GeometryCollectionSection.cpp b/Engine/Source/Runtime/Experimental/Chaos/Private/GeometryCollection/GeometryCollectionSection.cpp
--- a/Engine/Source/Runtime/Experimental/Chaos/Private/GeometryCollection/GeometryCollectionSection.cpp
+++ b/Engine/Source/Runtime/Experimental/Chaos/Private/GeometryCollection/GeometryCollectionSection.cpp
@@ -6,7 +6,7 @@
 
 
 TArray<FGeometryCollectionSection>
-FGeometryCollectionSection::BuildMeshSections(const FManagedArrayCollection& InCollection, const TArray<FIntVector>& InputIndices, TArray<int32> BaseMeshOriginalIndicesIndex, TArray<FIntVector>& RetIndices)
+FGeometryCollectionSection::BuildMeshSections(const FManagedArrayCollection& InCollection, const TArray<FIntVector>& InputIndices, const TArray<int32>& BaseMeshOriginalIndicesIndex, TArray<FIntVector>& RetIndices)
 {
 	TArray<FGeometryCollectionSection> RetSections;
 	TArray<FGeometryCollectionSection> TmpSections;
@@ -53,18 +53,23 @@ FGeometryCollectionSection::BuildMeshSections(const FManagedArrayCollection& InC
 
 		// remap indices so the materials appear to be grouped
 		RetIndices.AddUninitialized(InputIndices.Num());
-		int Idx = 0;
+
+		// since we know the number of trinagle per section we can distribute the faces in RetIndices
+		// this avoid nested loop that result in N * M operations (N=sections M=faces) and we can process it in (N + M) operations instead 
+		TArray<int32> OffsetPerSection;
+		OffsetPerSection.AddUninitialized(TmpSections.Num());
+		int32 SectionOffset = 0;
 		for (int Section = 0; Section < TmpSections.Num(); Section++)
 		{
-			for (int FaceElement = 0; FaceElement < InputIndices.Num(); FaceElement++)
-			{
-				int32 ID = (MaterialID)[BaseMeshOriginalIndicesIndex[FaceElement]];
+			OffsetPerSection[Section] = SectionOffset;
+			SectionOffset += TmpSections[Section].NumTriangles;
+		}
 
-				if (Section == ID)
-				{
-					RetIndices[Idx++] = InputIndices[FaceElement];
-				}
-			}
+		for (int FaceElement = 0; FaceElement < InputIndices.Num(); FaceElement++)
+		{
+			const int32 SectionID = (MaterialID)[BaseMeshOriginalIndicesIndex[FaceElement]];
+			int32& SectionOffsetRef = OffsetPerSection[SectionID];
+			RetIndices[SectionOffsetRef++] = InputIndices[FaceElement];
 		}
 
 		// if a material group no longer has any triangles in it then add material section for removal
diff --git a/Engine/Source/Runtime/Experimental/Chaos/Public/GeometryCollection/GeometryCollection.h b/Engine/Source/Runtime/Experimental/Chaos/Public/GeometryCollection/GeometryCollection.h
--- a/Engine/Source/Runtime/Experimental/Chaos/Public/GeometryCollection/GeometryCollection.h
+++ b/Engine/Source/Runtime/Experimental/Chaos/Public/GeometryCollection/GeometryCollection.h
@@ -219,7 +219,7 @@ public:
 	* Currently, this call assumes that the indices are ordered by MaterialID
 	* #todo(dmp): Refactor this and ReindexMaterials to share code
 	*/
-	TArray<FGeometryCollectionSection> BuildMeshSections(const TArray<FIntVector> &Indices, TArray<int32> BaseMeshOriginalIndicesIndex, TArray<FIntVector> &RetIndices) const;
+	TArray<FGeometryCollectionSection> BuildMeshSections(const TArray<FIntVector> &Indices, const TArray<int32>& BaseMeshOriginalIndicesIndex, TArray<FIntVector> &RetIndices) const;
 	//
 	//
 	//
diff --git a/Engine/Source/Runtime/Experimental/Chaos/Public/GeometryCollection/GeometryCollectionSection.h b/Engine/Source/Runtime/Experimental/Chaos/Public/GeometryCollection/GeometryCollectionSection.h
--- a/Engine/Source/Runtime/Experimental/Chaos/Public/GeometryCollection/GeometryCollectionSection.h
+++ b/Engine/Source/Runtime/Experimental/Chaos/Public/GeometryCollection/GeometryCollectionSection.h
@@ -34,7 +34,7 @@ struct CHAOS_API FGeometryCollectionSection
 	}
 
 	static TArray<FGeometryCollectionSection>
-	BuildMeshSections(const FManagedArrayCollection& InCollection, const TArray<FIntVector>& InputIndices, TArray<int32> BaseMeshOriginalIndicesIndex, TArray<FIntVector>& RetIndices);
+	BuildMeshSections(const FManagedArrayCollection& InCollection, const TArray<FIntVector>& InputIndices, const TArray<int32>& BaseMeshOriginalIndicesIndex, TArray<FIntVector>& RetIndices);
 
 	/** The index of the material with which to render this section. */
 	int32 MaterialID;
diff --git a/Engine/Source/Runtime/Experimental/GeometryCollectionEngine/Private/GeometryCollection/GeometryCollectionComponent.cpp b/Engine/Source/Runtime/Experimental/GeometryCollectionEngine/Private/GeometryCollection/GeometryCollectionComponent.cpp
--- a/Engine/Source/Runtime/Experimental/GeometryCollectionEngine/Private/GeometryCollection/GeometryCollectionComponent.cpp
+++ b/Engine/Source/Runtime/Experimental/GeometryCollectionEngine/Private/GeometryCollection/GeometryCollectionComponent.cpp
@@ -100,6 +100,12 @@ FAutoConsoleVariableRef CVarChaosGCUseISMPool(TEXT("p.Chaos.GC.UseISMPool"), bCh
 bool bChaos_GC_UseISMPoolForNonFracturedParts = true;
 FAutoConsoleVariableRef CVarChaosGCUseISMPoolForNonFracturedParts(TEXT("p.Chaos.GC.UseISMPoolForNonFracturedParts"), bChaos_GC_UseISMPoolForNonFracturedParts, TEXT("When enabled, non fractured part will use the ISM pool if specified"));
 
+bool bChaos_GC_InitConstantDataUseParallelFor = true;
+FAutoConsoleVariableRef CVarChaosGCInitConstantDataUseParallelFor(TEXT("p.Chaos.GC.InitConstantDataUseParallelFor"), bChaos_GC_InitConstantDataUseParallelFor, TEXT("When enabled, InitConstant data will use parallelFor for copying some of the data"));
+
+int32 bChaos_GC_InitConstantDataParallelForBatchSize = 5000;
+FAutoConsoleVariableRef CVarChaosGCInitConstantDataParallelForBatchSize(TEXT("p.Chaos.GC.InitConstantDataParallelForBatchSize"), bChaos_GC_InitConstantDataParallelForBatchSize, TEXT("When parallelFor is used in InitConstantData, defined the minimium size of a batch of vertex "));
+
 DEFINE_LOG_CATEGORY_STATIC(UGCC_LOG, Error, All);
 
 extern FGeometryCollectionDynamicDataPool GDynamicDataPool;
@@ -1825,16 +1831,48 @@ void UGeometryCollectionComponent::InitConstantData(FGeometryCollectionConstantD
 		ConstantData->TangentU = TArray<FVector3f>(TangentU.GetData(), TangentU.Num());
 		ConstantData->TangentV = TArray<FVector3f>(TangentV.GetData(), TangentV.Num());
 		ConstantData->Normals = TArray<FVector3f>(Normal.GetData(), Normal.Num());
-		ConstantData->UVs = TArray<TArray<FVector2f>>(UVs.GetData(), UVs.Num());
 		ConstantData->Colors = TArray<FLinearColor>(Color.GetData(), Color.Num());
 
+		// use the first point to know the number of channels
+		const int32 NumUVChannels = Collection->UVs.Num() ? Collection->UVs[0].Num() : 1;
+		ConstantData->UVChannels.SetNum(NumUVChannels);
+		for (int32 UVChannelIndex = 0; UVChannelIndex < NumUVChannels; UVChannelIndex++)
+		{
+			ConstantData->UVChannels[UVChannelIndex].AddUninitialized(NumPoints);
+		}
+
 		ConstantData->BoneColors.AddUninitialized(NumPoints);
 
-		ParallelFor(NumPoints, [&](const int32 InPointIndex)
+		if (bChaos_GC_InitConstantDataUseParallelFor)
+		{
+			ParallelFor(TEXT("GC:InitConstantData"), NumPoints, bChaos_GC_InitConstantDataParallelForBatchSize,
+				[&](const int32 InPointIndex)
+				{
+					const int32 BoneIndex = ConstantData->BoneMap[InPointIndex];
+					ConstantData->BoneColors[InPointIndex] = BoneColors[BoneIndex];
+
+					for (int32 UVChannelIndex = 0; UVChannelIndex < NumUVChannels; UVChannelIndex++)
+					{
+						ConstantData->UVChannels[UVChannelIndex][InPointIndex] = Collection->UVs[InPointIndex][UVChannelIndex];
+					}
+				});
+		}
+		else
+		{
+			for (int32 PointIndex = 0; PointIndex < NumPoints; PointIndex++)
+			{
+				const int32 BoneIndex = ConstantData->BoneMap[PointIndex];
+				ConstantData->BoneColors[PointIndex] = BoneColors[BoneIndex];
+			}
+
+			for (int32 UVChannelIndex = 0; UVChannelIndex < NumUVChannels; UVChannelIndex++)
 			{
-				const int32 BoneIndex = ConstantData->BoneMap[InPointIndex];
-				ConstantData->BoneColors[InPointIndex] = BoneColors[BoneIndex];
-			});
+				for (int32 PointIndex = 0; PointIndex < NumPoints; PointIndex++)
+				{
+					ConstantData->UVChannels[UVChannelIndex][PointIndex] = Collection->UVs[PointIndex][UVChannelIndex];
+				}
+			}
+		}
 
 		int32 NumIndices = 0;
 		const TManagedArray<FIntVector>& Indices = Collection->Indices;
diff --git a/Engine/Source/Runtime/Experimental/GeometryCollectionEngine/Private/GeometryCollection/GeometryCollectionSceneProxy.cpp b/Engine/Source/Runtime/Experimental/GeometryCollectionEngine/Private/GeometryCollection/GeometryCollectionSceneProxy.cpp
--- a/Engine/Source/Runtime/Experimental/GeometryCollectionEngine/Private/GeometryCollection/GeometryCollectionSceneProxy.cpp
+++ b/Engine/Source/Runtime/Experimental/GeometryCollectionEngine/Private/GeometryCollection/GeometryCollectionSceneProxy.cpp
@@ -476,16 +476,18 @@ void FGeometryCollectionSceneProxy::BuildGeometry( const FGeometryCollectionCons
 		OutVertices[PointIdx] =
 			FDynamicMeshVertex(
 				ConstantDataIn->Vertices[PointIdx],
-				ConstantDataIn->UVs[PointIdx][0],
+				ConstantDataIn->UVChannels[0][PointIdx],
 				GetColorFunc(ConstantDataIn, PointIdx, bCapturedShowBoneColors)
 			);
 		OutVertices[PointIdx].SetTangents(ConstantDataIn->TangentU[PointIdx], ConstantDataIn->TangentV[PointIdx], ConstantDataIn->Normals[PointIdx]);
 
-		if (ConstantDataIn->UVs[PointIdx].Num() > 1)
+		const int32 NumUvChannels = ConstantDataIn->UVChannels.Num();
+		if (ConstantDataIn->UVChannels.Num() > 1)
 		{
-			for (int32 UVLayerIdx = 1; UVLayerIdx < ConstantDataIn->UVs[PointIdx].Num(); ++UVLayerIdx)
+			for (int32 UVLayerIdx = 1; UVLayerIdx < NumUvChannels; ++UVLayerIdx)
 			{
-				OutVertices[PointIdx].TextureCoordinate[UVLayerIdx] = ConstantDataIn->UVs[PointIdx][UVLayerIdx];
+				const FGeometryCollectionConstantData::FUVChannel& UVChannel = ConstantDataIn->UVChannels[UVLayerIdx];
+				OutVertices[PointIdx].TextureCoordinate[UVLayerIdx] = UVChannel[PointIdx];
 			}
 		}
 	};
diff --git a/Engine/Source/Runtime/Experimental/GeometryCollectionEngine/Private/GeometryCollection/GeometryCollectionSceneProxy.h b/Engine/Source/Runtime/Experimental/GeometryCollectionEngine/Private/GeometryCollection/GeometryCollectionSceneProxy.h
--- a/Engine/Source/Runtime/Experimental/GeometryCollectionEngine/Private/GeometryCollection/GeometryCollectionSceneProxy.h
+++ b/Engine/Source/Runtime/Experimental/GeometryCollectionEngine/Private/GeometryCollection/GeometryCollectionSceneProxy.h
@@ -94,12 +94,14 @@ inline void CopyTransformsWithConversionWhenNeeded(TArray<FMatrix44f>& DstTransf
 /** Immutable rendering data (kind of) */
 struct FGeometryCollectionConstantData
 {
+	using FUVChannel = TArray<FVector2f>;
+
 	TArray<FVector3f> Vertices;
 	TArray<FIntVector> Indices;
 	TArray<FVector3f> Normals;
 	TArray<FVector3f> TangentU;
 	TArray<FVector3f> TangentV;
-	TArray<TArray<FVector2f>> UVs;
+	TArray<FUVChannel> UVChannels;
 	TArray<FLinearColor> Colors;
 	TArray<int32> BoneMap;
 	TArray<FLinearColor> BoneColors;
