Iris - Delayed resetting of dirty changemasks until all systems have polled. Global dirty tracking via FReplicationStateHeader.

diff --git a/Engine/Plugins/Runtime/ReplicationSystemTestPlugin/Source/Private/Tests/ReplicationState/TestGeneratedReplicationState.cpp b/Engine/Plugins/Runtime/ReplicationSystemTestPlugin/Source/Private/Tests/ReplicationState/TestGeneratedReplicationState.cpp
--- a/Engine/Plugins/Runtime/ReplicationSystemTestPlugin/Source/Private/Tests/ReplicationState/TestGeneratedReplicationState.cpp
+++ b/Engine/Plugins/Runtime/ReplicationSystemTestPlugin/Source/Private/Tests/ReplicationState/TestGeneratedReplicationState.cpp
@@ -94,10 +94,9 @@ UE_NET_TEST(FTestTypedReplicationState, ValidateState)
 	// so test by casting as it always should be the first member
 	const UE::Net::FReplicationStateHeader& Header = UE::Net::Private::GetReplicationStateHeader(reinterpret_cast<uint8*>(&State), State.GetReplicationStateDescriptor());
 
-	UE_NET_ASSERT_EQ(0u, UE::Net::Private::FReplicationStateHeaderAccessor::GetReplicationIndex(Header));
-	UE_NET_ASSERT_EQ(0u, UE::Net::Private::FReplicationStateHeaderAccessor::GetReplicationSystemId(Header));
-	UE_NET_ASSERT_EQ(0u, UE::Net::Private::FReplicationStateHeaderAccessor::GetReplicationFlags(Header));
-	UE_NET_ASSERT_EQ(false, UE::Net::Private::FReplicationStateHeaderAccessor::GetIsInitStateDirty(Header));
+	UE_NET_ASSERT_EQ(FReplicationStateHeaderAccessor::GetNetHandleId(Header), 0U);
+	UE_NET_ASSERT_FALSE(FReplicationStateHeaderAccessor::GetIsInitStateDirty(Header));
+	UE_NET_ASSERT_FALSE(FReplicationStateHeaderAccessor::GetIsStateDirty(Header));
 
 	// Verify defaults
 	const int32 ExpectedDefaultValue = 0;
diff --git a/Engine/Plugins/Runtime/ReplicationSystemTestPlugin/Source/Private/Tests/ReplicationState/TestPropertyReplicationState.cpp b/Engine/Plugins/Runtime/ReplicationSystemTestPlugin/Source/Private/Tests/ReplicationState/TestPropertyReplicationState.cpp
--- a/Engine/Plugins/Runtime/ReplicationSystemTestPlugin/Source/Private/Tests/ReplicationState/TestPropertyReplicationState.cpp
+++ b/Engine/Plugins/Runtime/ReplicationSystemTestPlugin/Source/Private/Tests/ReplicationState/TestPropertyReplicationState.cpp
@@ -11,6 +11,7 @@
 #include "Misc/EnumClassFlags.h"
 #include "UObject/StrongObjectPtr.h"
 #include "Net/UnrealNetwork.h"
+#include "Net/Core/NetHandle/NetHandleManager.h"
 
 
 void UTestPropertyReplicationState_TestClass::GetLifetimeReplicatedProps( TArray< class FLifetimeProperty > & OutLifetimeProps ) const
@@ -42,8 +43,9 @@ void UTestPropertyReplicationState_TestClassWithInitAndCArrays::GetLifetimeRepli
 namespace UE::Net::Private
 {
 
-struct FTestPropertyReplicationStateContext : public FNetworkAutomationTestSuiteFixture
+class FTestPropertyReplicationStateContext : public FNetworkAutomationTestSuiteFixture
 {
+protected:
 	FReplicationStateDescriptorBuilder::FResult Descriptors;
 	const FReplicationStateDescriptor* Descriptor;
 
@@ -53,7 +55,14 @@ struct FTestPropertyReplicationStateContext : public FNetworkAutomationTestSuite
 
 		Descriptor = Descriptors[0];
 
-		UE_NET_ASSERT_TRUE(Descriptor != nullptr);
+		UE_NET_ASSERT_NE(Descriptor, nullptr);
+	}
+
+	FNetHandle GenerateNetHandle()
+	{
+		UObject* Object = NewObject<UTestPropertyReplicationState_TestClass>();
+		FNetHandle NetHandle = FNetHandleManager::GetOrCreateNetHandle(Object);
+		return NetHandle;
 	}
 };
 
@@ -222,7 +231,8 @@ UE_NET_TEST_FIXTURE(FTestPropertyReplicationStateContext, PropertyReplicationSta
 	BoundState.MarkDirty(1);
 
 	// Fake that we are bound
-	UE::Net::Private::FReplicationStateHeaderAccessor::SetReplicationIndex(UE::Net::Private::GetReplicationStateHeader(BoundState.GetStateBuffer(), Descriptor), 1U, 0U);
+	const FNetHandle NetHandle = GenerateNetHandle();
+	FReplicationStateHeaderAccessor::SetNetHandleId(GetReplicationStateHeader(BoundState.GetStateBuffer(), Descriptor), NetHandle);
 
 	// Create other state
 	FPropertyReplicationState OtherState(Descriptor);
@@ -236,8 +246,8 @@ UE_NET_TEST_FIXTURE(FTestPropertyReplicationStateContext, PropertyReplicationSta
 	// Assign to the bound state, this should result in all states being dirty
 	BoundState = OtherState;
 
-	// Verify that we did not overwrite the internal index
-	UE_NET_ASSERT_EQ(1u, UE::Net::Private::FReplicationStateHeaderAccessor::GetReplicationIndex(UE::Net::Private::GetReplicationStateHeader(BoundState.GetStateBuffer(), Descriptor)));
+	// Verify that we did not overwrite the NetHandle
+	UE_NET_ASSERT_EQ(Private::FReplicationStateHeaderAccessor::GetNetHandleId(GetReplicationStateHeader(BoundState.GetStateBuffer(), Descriptor)), NetHandle.GetId());
 
 	// check that original states are still dirty and that we also have dirtied the value that was modified in OtherState
 	UE_NET_ASSERT_TRUE(BoundState.IsDirty(0));
@@ -257,7 +267,8 @@ UE_NET_TEST_FIXTURE(FTestPropertyReplicationStateContext, PropertyReplicationSta
 	BoundState.MarkDirty(1);
 
 	// Fake that we are bound
-	UE::Net::Private::FReplicationStateHeaderAccessor::SetReplicationIndex(UE::Net::Private::GetReplicationStateHeader(BoundState.GetStateBuffer(), Descriptor), 1U, 0U);
+	const FNetHandle NetHandle = GenerateNetHandle();
+	FReplicationStateHeaderAccessor::SetNetHandleId(GetReplicationStateHeader(BoundState.GetStateBuffer(), Descriptor), NetHandle);
 
 	// Create other state
 	FPropertyReplicationState OtherState(Descriptor);
@@ -271,8 +282,8 @@ UE_NET_TEST_FIXTURE(FTestPropertyReplicationStateContext, PropertyReplicationSta
 	// copy the bound state, this should overwrite all local changes but not copy the internalindex
 	OtherState = BoundState;
 
-	// Verify that we did not overwrite the internal index
-	UE_NET_ASSERT_EQ(0u, UE::Net::Private::FReplicationStateHeaderAccessor::GetReplicationIndex(UE::Net::Private::GetReplicationStateHeader(OtherState.GetStateBuffer(), Descriptor)));
+	// Verify that we did not overwrite the NetHandle
+	UE_NET_ASSERT_EQ(FReplicationStateHeaderAccessor::GetNetHandleId(GetReplicationStateHeader(OtherState.GetStateBuffer(), Descriptor)), 0U);
 
 	// check that original states are still dirty and that we also have dirtied the value that was modified in OtherState
 	UE_NET_ASSERT_TRUE(OtherState.IsDirty(0));
diff --git a/Engine/Plugins/Runtime/ReplicationSystemTestPlugin/Source/Private/Tests/ReplicationSystem/TestMultiReplicationSystems.cpp b/Engine/Plugins/Runtime/ReplicationSystemTestPlugin/Source/Private/Tests/ReplicationSystem/TestMultiReplicationSystems.cpp
--- a/Engine/Plugins/Runtime/ReplicationSystemTestPlugin/Source/Private/Tests/ReplicationSystem/TestMultiReplicationSystems.cpp
+++ b/Engine/Plugins/Runtime/ReplicationSystemTestPlugin/Source/Private/Tests/ReplicationSystem/TestMultiReplicationSystems.cpp
@@ -96,7 +96,7 @@ UE_NET_TEST_FIXTURE(FMultiReplicationSystemsTestFixture, ObjectIsReplicatedToAll
 	// Create clients for all systems
 	for (FReplicationSystemTestServer* Server : GetAllServers())
 	{
-		constexpr unsigned ClientCountPerServer = 3;
+		constexpr unsigned ClientCountPerServer = 3U;
 		for (unsigned ClientIt = 0, ClientEndIt = ClientCountPerServer; ClientIt != ClientEndIt; ++ClientIt)
 		{
 			CreateClientForServer(Server);
@@ -210,5 +210,41 @@ UE_NET_TEST_FIXTURE(FMultiReplicationSystemsTestFixture, ReplicatedObjectCanStop
 	}
 }
 
+UE_NET_TEST_FIXTURE(FMultiReplicationSystemsTestFixture, ModifiedObjectIsReplicatedToAllClientsOnAllSystems)
+{
+	CreateSomeServers();
+
+	// Create clients for all systems
+	for (FReplicationSystemTestServer* Server : GetAllServers())
+	{
+		constexpr unsigned ClientCountPerServer = 1U;
+		for (unsigned ClientIt = 0, ClientEndIt = ClientCountPerServer; ClientIt != ClientEndIt; ++ClientIt)
+		{
+			CreateClientForServer(Server);
+		}
+	}
+
+	UTestReplicatedIrisObject* ServerObject = CreateObject(UTestReplicatedIrisObject::FComponents{});
+	BeginReplication(ServerObject);
+
+	FullSendAndDeliverUpdate();
+
+	ServerObject->IntA ^= 4711;
+	const int32 ExpectedIntAVAlue = ServerObject->IntA;
+
+	FullSendAndDeliverUpdate();
+
+	// Verify the object has the updated value on all clients
+	for (FReplicationSystemTestServer* Server : GetAllServers())
+	{
+		FNetRefHandle RefHandleOnServer = Server->GetReplicationBridge()->GetReplicatedRefHandle(ServerObject);
+		for (FReplicationSystemTestClient* Client : GetClients(Server))
+		{
+			UTestReplicatedIrisObject* ClientObject = Cast<UTestReplicatedIrisObject>(Client->GetReplicationBridge()->GetReplicatedObject(RefHandleOnServer));
+			CA_ASSUME(ClientObject != nullptr);
+			UE_NET_ASSERT_EQ(ClientObject->IntA, ExpectedIntAVAlue);
+		}
+	}
+}
 
 }
diff --git a/Engine/Source/Runtime/Experimental/Iris/Core/Private/Iris/ReplicationState/PropertyReplicationState.cpp b/Engine/Source/Runtime/Experimental/Iris/Core/Private/Iris/ReplicationState/PropertyReplicationState.cpp
--- a/Engine/Source/Runtime/Experimental/Iris/Core/Private/Iris/ReplicationState/PropertyReplicationState.cpp
+++ b/Engine/Source/Runtime/Experimental/Iris/Core/Private/Iris/ReplicationState/PropertyReplicationState.cpp
@@ -77,6 +77,17 @@ bool FPropertyReplicationState::IsInitState() const
 	return IsValid() && ReplicationStateDescriptor->IsInitState();
 }
 
+bool FPropertyReplicationState::IsDirty() const
+{
+	if (IsValid())
+	{
+		const FReplicationStateHeader& Header = Private::GetReplicationStateHeader(StateBuffer, ReplicationStateDescriptor);
+		return Private::FReplicationStateHeaderAccessor::GetIsStateDirty(Header) || Private::FReplicationStateHeaderAccessor::GetIsInitStateDirty(Header);
+	}
+
+	return false;
+}
+
 void FPropertyReplicationState::ConstructStateInternal()
 {
 	// allocate memory
@@ -111,8 +122,7 @@ void FPropertyReplicationState::InjectState(const FReplicationStateDescriptor* D
 	StateBuffer = InStateBuffer;
 }
 
-void
-FPropertyReplicationState::Set(const FPropertyReplicationState& Other)
+void FPropertyReplicationState::Set(const FPropertyReplicationState& Other)
 {
 	if (IsValid())
 	{
@@ -204,7 +214,7 @@ bool FPropertyReplicationState::IsDirty(uint32 Index) const
 	}
 }
 
-void FPropertyReplicationState::PollPropertyReplicationState(const void* RESTRICT SrcStateData)
+bool FPropertyReplicationState::PollPropertyReplicationState(const void* RESTRICT SrcStateData)
 {
 	if (IsValid())
 	{
@@ -226,6 +236,8 @@ void FPropertyReplicationState::PollPropertyReplicationState(const void* RESTRIC
 			SetPropertyValue(MemberIt, SrcBuffer + Property->GetOffset_ForGC() + Property->ElementSize*MemberPropertyDescriptor.ArrayIndex);
 		}
 	}
+
+	return IsDirty();
 }
 
 void FPropertyReplicationState::PushPropertyReplicationState(void* RESTRICT DstData, bool bInPushAll) const
@@ -259,7 +271,7 @@ void FPropertyReplicationState::PushPropertyReplicationState(void* RESTRICT DstD
 	}
 }
 
-void FPropertyReplicationState::PollObjectReferences(const void* RESTRICT SrcStateData)
+bool FPropertyReplicationState::PollObjectReferences(const void* RESTRICT SrcStateData)
 {
 	if (IsValid())
 	{
@@ -284,6 +296,8 @@ void FPropertyReplicationState::PollObjectReferences(const void* RESTRICT SrcSta
 			}
 		}
 	}
+
+	return IsDirty();
 }
 
 void FPropertyReplicationState::CallRepNotifies(void* RESTRICT DstData, const FPropertyReplicationState* PreviousState, bool bIsInit) const
diff --git a/Engine/Source/Runtime/Experimental/Iris/Core/Private/Iris/ReplicationSystem/Conditionals/ReplicationConditionals.cpp b/Engine/Source/Runtime/Experimental/Iris/Core/Private/Iris/ReplicationSystem/Conditionals/ReplicationConditionals.cpp
--- a/Engine/Source/Runtime/Experimental/Iris/Core/Private/Iris/ReplicationSystem/Conditionals/ReplicationConditionals.cpp
+++ b/Engine/Source/Runtime/Experimental/Iris/Core/Private/Iris/ReplicationSystem/Conditionals/ReplicationConditionals.cpp
@@ -233,7 +233,7 @@ bool FReplicationConditionals::SetPropertyCustomCondition(FInternalNetRefIndex O
 
 			// If a condition is enabled we also mark the corresponding regular changemask as dirty.
 			FNetBitArrayView MemberChangeMask = UE::Net::Private::GetMemberChangeMask(Fragment.ExternalSrcBuffer, StateDescriptor);
-			const FReplicationStateHeader& ReplicationStateHeader = UE::Net::Private::GetReplicationStateHeader(Fragment.ExternalSrcBuffer, StateDescriptor);
+			FReplicationStateHeader& ReplicationStateHeader = UE::Net::Private::GetReplicationStateHeader(Fragment.ExternalSrcBuffer, StateDescriptor);
 			MarkDirty(ReplicationStateHeader, MemberChangeMask, ChangeMaskDescriptor);
 
 			// Enabled conditions causes new properties to be replicated which most likely have incorrect values at the receiving end.
@@ -298,7 +298,7 @@ bool FReplicationConditionals::SetPropertyCustomCondition(FInternalNetRefIndex O
 
 					// If a condition is enabled we also mark the corresponding regular changemask as dirty.
 					FNetBitArrayView MemberChangeMask = GetMemberChangeMask(Fragment.ExternalSrcBuffer, StateDescriptor);
-					const FReplicationStateHeader& Header = GetReplicationStateHeader(Fragment.ExternalSrcBuffer, StateDescriptor);
+					FReplicationStateHeader& Header = GetReplicationStateHeader(Fragment.ExternalSrcBuffer, StateDescriptor);
 					MarkDirty(Header, MemberChangeMask, ChangeMaskDescriptor);
 
 					// Enabled conditions causes new properties to be replicated which most likely have incorrect values at the receiving end.
diff --git a/Engine/Source/Runtime/Experimental/Iris/Core/Private/Iris/ReplicationSystem/ObjectReplicationBridge.cpp b/Engine/Source/Runtime/Experimental/Iris/Core/Private/Iris/ReplicationSystem/ObjectReplicationBridge.cpp
--- a/Engine/Source/Runtime/Experimental/Iris/Core/Private/Iris/ReplicationSystem/ObjectReplicationBridge.cpp
+++ b/Engine/Source/Runtime/Experimental/Iris/Core/Private/Iris/ReplicationSystem/ObjectReplicationBridge.cpp
@@ -270,9 +270,9 @@ UE::Net::FNetRefHandle UObjectReplicationBridge::BeginReplication(UObject* Insta
 		FNetRefHandle RefHandle = InternalCreateNetObject(AllocatedRefHandle, NetHandle, ReplicationProtocol);
 		if (RefHandle.IsValid())
 		{
-			// Attach the instance and Bind the Instance protocol to dirty tracking
+			// Attach the instance and bind the instance protocol to dirty tracking
 			constexpr bool bBindInstanceProtocol = true;
-			InternalAttachInstanceToNetHandle(RefHandle, bBindInstanceProtocol, InstanceProtocol.Get(), Instance);
+			InternalAttachInstanceToNetRefHandle(RefHandle, bBindInstanceProtocol, InstanceProtocol.Get(), Instance, NetHandle);
 #if WITH_PUSH_MODEL
 			SetNetPushIdOnInstance(InstanceProtocol.Get(), NetHandle);
 #endif
@@ -486,30 +486,30 @@ void UObjectReplicationBridge::DetachInstance(FNetRefHandle Handle)
 	UnregisterRemoteInstance(Handle, bTearOff, bShouldDestroyInstance);
 }
 
-void UObjectReplicationBridge::RegisterRemoteInstance(FNetRefHandle Handle, UObject* Instance, const UE::Net::FReplicationProtocol* Protocol, UE::Net::FReplicationInstanceProtocol* InstanceProtocol, const FCreationHeader* Header, uint32 ConnectionId)
+void UObjectReplicationBridge::RegisterRemoteInstance(FNetRefHandle RefHandle, UObject* Instance, const UE::Net::FReplicationProtocol* Protocol, UE::Net::FReplicationInstanceProtocol* InstanceProtocol, const FCreationHeader* Header, uint32 ConnectionId)
 {
-	// Attach the instance protocol and instance to the Handle
+	// Attach the instance protocol and instance to the handle
 	constexpr bool bBindInstanceProtocol = false;
-	InternalAttachInstanceToNetHandle(Handle, bBindInstanceProtocol, InstanceProtocol, Instance);
+	InternalAttachInstanceToNetRefHandle(RefHandle, bBindInstanceProtocol, InstanceProtocol, Instance, FNetHandle());
 
 	// Dynamic references needs to be promoted to find the instantiated object
-	if (Handle.IsDynamic())
+	if (RefHandle.IsDynamic())
 	{
-		GetObjectReferenceCache()->AddRemoteReference(Handle, Instance);
+		GetObjectReferenceCache()->AddRemoteReference(RefHandle, Instance);
 	}
 
-	UE_LOG_OBJECTREPLICATIONBRIDGE(Verbose, TEXT("RegisterRemoteInstance %s %s with ProtocolId:0x%" UINT64_x_FMT), *Handle.ToString(), ToCStr(Instance->GetName()), Protocol->ProtocolIdentifier);
+	UE_LOG_OBJECTREPLICATIONBRIDGE(Verbose, TEXT("RegisterRemoteInstance %s %s with ProtocolId:0x%" UINT64_x_FMT), *RefHandle.ToString(), ToCStr(Instance->GetName()), Protocol->ProtocolIdentifier);
 }
 
-void UObjectReplicationBridge::UnregisterRemoteInstance(FNetRefHandle Handle, bool bTearOff, bool bShouldDestroyInstance)
+void UObjectReplicationBridge::UnregisterRemoteInstance(FNetRefHandle RefHandle, bool bTearOff, bool bShouldDestroyInstance)
 {
 	// Lookup the instance and remove it	
-	UObject* Instance = GetObjectFromReferenceHandle(Handle);
+	UObject* Instance = GetObjectFromReferenceHandle(RefHandle);
 	
 	// Try to remove any references to dynamic objects
-	if (Handle.IsDynamic())
+	if (RefHandle.IsDynamic())
 	{
-		GetObjectReferenceCache()->RemoveReference(Handle, Instance);
+		GetObjectReferenceCache()->RemoveReference(RefHandle, Instance);
 	}
 
 	// Destroy instance if we should	
@@ -609,18 +609,18 @@ void UObjectReplicationBridge::PostApplyInitialState(FNetRefHandle Handle)
 	EndInstantiateFromRemote(Handle);
 }
 
-void UObjectReplicationBridge::PreSendUpdateSingleHandle(FNetRefHandle Handle)
+void UObjectReplicationBridge::PreSendUpdateSingleHandle(FNetRefHandle RefHandle)
 {
-	PreUpdateAndPollImpl(Handle);
+	PreUpdateAndPollImpl(RefHandle);
 }
 
 void UObjectReplicationBridge::PreSendUpdate()
 {
 	IRIS_PROFILER_SCOPE(UObjectReplicationBridge_OnPreSendUpdate);
 
-	// Invalid Handle means update all objects
-	FNetRefHandle Handle;
-	PreUpdateAndPollImpl(Handle);
+	// Invalid/default handle means update all objects
+	FNetRefHandle RefHandle;
+	PreUpdateAndPollImpl(RefHandle);
 }
 
 
@@ -691,7 +691,7 @@ void UObjectReplicationBridge::PreUpdateAndPollImpl(FNetRefHandle Handle)
 	const FNetRefHandleManager& LocalNetRefHandleManager = ReplicationSystemInternal->GetNetRefHandleManager();
 	const TArray<UObject*>& ReplicatedInstances = LocalNetRefHandleManager.GetReplicatedInstances();
 	const bool bIsUsingPushModel = IsIrisPushModelEnabled();
-	const FNetBitArrayView DirtyObjects = ReplicationSystemInternal->GetDirtyNetObjectTracker().GetDirtyNetObjects();
+	FNetBitArrayView DirtyObjects = ReplicationSystemInternal->GetDirtyNetObjectTracker().GetDirtyNetObjects();
 
 	struct FPreUpdateAndPollStats
 	{
@@ -702,7 +702,7 @@ void UObjectReplicationBridge::PreUpdateAndPollImpl(FNetRefHandle Handle)
 
 	FPreUpdateAndPollStats Stats;
 
-	auto UpdateAndPollFunction = [&ReplicatedInstances, this, &LocalNetRefHandleManager, &Stats](uint32 InternalObjectIndex)
+	auto UpdateAndPollFunction = [&ReplicatedInstances, this, &LocalNetRefHandleManager, &DirtyObjects, &Stats](uint32 InternalObjectIndex)
 	{
 		const FNetRefHandleManager::FReplicatedObjectData& ObjectData = LocalNetRefHandleManager.GetReplicatedObjectDataNoCheck(InternalObjectIndex);
 		if (ObjectData.InstanceProtocol)
@@ -733,7 +733,11 @@ void UObjectReplicationBridge::PreUpdateAndPollImpl(FNetRefHandle Handle)
 				EReplicationFragmentPollFlags PollOptions = EReplicationFragmentPollFlags::PollAllState;
 				PollOptions |= bIsGCAffectedObject ? EReplicationFragmentPollFlags::ForceRefreshCachedObjectReferencesAfterGC : EReplicationFragmentPollFlags::None;
 
-				FReplicationInstanceOperations::PollAndRefreshCachedPropertyData(ObjectData.InstanceProtocol, PollOptions);
+				const bool bWasMarkedDirty = FReplicationInstanceOperations::PollAndRefreshCachedPropertyData(ObjectData.InstanceProtocol, PollOptions);
+				if (bWasMarkedDirty)
+				{
+					DirtyObjects.SetBit(InternalObjectIndex);
+				}
 				++Stats.PolledObjectCount;
 			}
 		}
@@ -775,6 +779,7 @@ void UObjectReplicationBridge::PreUpdateAndPollImpl(FNetRefHandle Handle)
 #endif
 
 			// If the object is fully push model we only need to poll it if it's dirty, unless it's a new object or was garbage collected.
+			bool bWasMarkedDirty = false;
 			if (EnumHasAnyFlags(InstanceTraits, EReplicationInstanceProtocolTraits::HasFullPushBasedDirtiness))
 			{
 				if (bIsDirtyObject | bIsNewInScope)
@@ -782,14 +787,14 @@ void UObjectReplicationBridge::PreUpdateAndPollImpl(FNetRefHandle Handle)
 					// We need to do a poll if object is marked as dirty
 					EReplicationFragmentPollFlags PollOptions = EReplicationFragmentPollFlags::PollAllState;
 					PollOptions |= bIsGCAffectedObject ? EReplicationFragmentPollFlags::ForceRefreshCachedObjectReferencesAfterGC : EReplicationFragmentPollFlags::None;
-					FReplicationInstanceOperations::PollAndRefreshCachedPropertyData(InstanceProtocol, EReplicationFragmentTraits::None, PollOptions);
+					bWasMarkedDirty = FReplicationInstanceOperations::PollAndRefreshCachedPropertyData(InstanceProtocol, EReplicationFragmentTraits::None, PollOptions);
 					++Stats.PolledObjectCount;
 				}
 				else if (bIsGCAffectedObject)
 				{
 					// If this object might have been affected by GC, only refresh cached references
 					const EReplicationFragmentTraits RequiredTraits = EReplicationFragmentTraits::HasPushBasedDirtiness;
-					FReplicationInstanceOperations::PollAndRefreshCachedObjectReferences(InstanceProtocol, RequiredTraits);
+					bWasMarkedDirty = FReplicationInstanceOperations::PollAndRefreshCachedObjectReferences(InstanceProtocol, RequiredTraits);
 					++Stats.PolledReferencesObjectCount;
 				}
 			}
@@ -803,7 +808,7 @@ void UObjectReplicationBridge::PreUpdateAndPollImpl(FNetRefHandle Handle)
 				{
 					// Only states which has push based dirtiness need to be updated as the other states will be polled in full anyway.
 					const EReplicationFragmentTraits RequiredTraits = EReplicationFragmentTraits::HasPushBasedDirtiness;
-					FReplicationInstanceOperations::PollAndRefreshCachedObjectReferences(InstanceProtocol, RequiredTraits);
+					bWasMarkedDirty = FReplicationInstanceOperations::PollAndRefreshCachedObjectReferences(InstanceProtocol, RequiredTraits);
 					++Stats.PolledReferencesObjectCount;
 				}
 
@@ -813,9 +818,14 @@ void UObjectReplicationBridge::PreUpdateAndPollImpl(FNetRefHandle Handle)
 
 				// If the object is not new or dirty at this point we only need to poll non-push based fragments as we know that pushed based states have not been modified
 				const EReplicationFragmentTraits ExcludeTraits = (bIsDirtyObject | bIsNewInScope) ? EReplicationFragmentTraits::None : EReplicationFragmentTraits::HasPushBasedDirtiness;
-				FReplicationInstanceOperations::PollAndRefreshCachedPropertyData(InstanceProtocol, ExcludeTraits, PollOptions);
+				bWasMarkedDirty |= FReplicationInstanceOperations::PollAndRefreshCachedPropertyData(InstanceProtocol, ExcludeTraits, PollOptions);
 				++Stats.PolledObjectCount;
 			}
+
+			if (bWasMarkedDirty)
+			{
+				DirtyObjects.SetBit(InternalObjectIndex);
+			}
 		}
 	};
 
diff --git a/Engine/Source/Runtime/Experimental/Iris/Core/Private/Iris/ReplicationSystem/PropertyReplicationFragment.cpp b/Engine/Source/Runtime/Experimental/Iris/Core/Private/Iris/ReplicationSystem/PropertyReplicationFragment.cpp
--- a/Engine/Source/Runtime/Experimental/Iris/Core/Private/Iris/ReplicationSystem/PropertyReplicationFragment.cpp
+++ b/Engine/Source/Runtime/Experimental/Iris/Core/Private/Iris/ReplicationSystem/PropertyReplicationFragment.cpp
@@ -110,20 +110,22 @@ void FPropertyReplicationFragment::ApplyReplicatedState(FReplicationStateApplyCo
 	ReceivedState.PushPropertyReplicationState(Owner);
 }
 
-void FPropertyReplicationFragment::PollReplicatedState(EReplicationFragmentPollFlags PollOption)
+bool FPropertyReplicationFragment::PollReplicatedState(EReplicationFragmentPollFlags PollOption)
 {
 	// Since PollPropertyReplicationState always copies the new value we do not need to do anything special
 	// with object references if the EReplicationFragmentPollOption::All flag is set
 	if (EnumHasAnyFlags(PollOption, EReplicationFragmentPollFlags::PollAllState))
 	{
 		const void* ExternalSourceState = static_cast<void*>(Owner);
-		SrcReplicationState->PollPropertyReplicationState(ExternalSourceState);
+		return SrcReplicationState->PollPropertyReplicationState(ExternalSourceState);
 	}
 	else if (EnumHasAnyFlags(PollOption, EReplicationFragmentPollFlags::ForceRefreshCachedObjectReferencesAfterGC))
 	{
 		const void* ExternalSourceState = static_cast<void*>(Owner);
-		SrcReplicationState->PollObjectReferences(ExternalSourceState);
+		return SrcReplicationState->PollObjectReferences(ExternalSourceState);
 	}
+
+	return false;
 }
 
 void FPropertyReplicationFragment::ReplicatedStateToString(FStringBuilderBase& StringBuilder, FReplicationStateApplyContext& Context, EReplicationStateToStringFlags Flags) const
diff --git a/Engine/Source/Runtime/Experimental/Iris/Core/Private/Iris/ReplicationSystem/ReplicationBridge.cpp b/Engine/Source/Runtime/Experimental/Iris/Core/Private/Iris/ReplicationSystem/ReplicationBridge.cpp
--- a/Engine/Source/Runtime/Experimental/Iris/Core/Private/Iris/ReplicationSystem/ReplicationBridge.cpp
+++ b/Engine/Source/Runtime/Experimental/Iris/Core/Private/Iris/ReplicationSystem/ReplicationBridge.cpp
@@ -325,31 +325,31 @@ UE::Net::FNetRefHandle UReplicationBridge::InternalCreateNetObjectFromRemote(FNe
 	return Handle;
 }
 
-void UReplicationBridge::InternalAttachInstanceToNetHandle(FNetRefHandle Handle, bool bBindInstanceProtocol, UE::Net::FReplicationInstanceProtocol* InstanceProtocol, UObject* Instance)
+void UReplicationBridge::InternalAttachInstanceToNetRefHandle(FNetRefHandle RefHandle, bool bBindInstanceProtocol, UE::Net::FReplicationInstanceProtocol* InstanceProtocol, UObject* Instance, FNetHandle NetHandle)
 {
 	using namespace UE::Net;
 	using namespace UE::Net::Private;
 
-	const uint32 ReplicationSystemId = Handle.GetReplicationSystemId();
-	const FInternalNetRefIndex InternalReplicationIndex = NetRefHandleManager->GetInternalIndex(Handle);
+	const uint32 ReplicationSystemId = RefHandle.GetReplicationSystemId();
+	const FInternalNetRefIndex InternalReplicationIndex = NetRefHandleManager->GetInternalIndex(RefHandle);
 
 	NetRefHandleManager->AttachInstanceProtocol(InternalReplicationIndex, InstanceProtocol, Instance);
-	UE_LOG_REPLICATIONBRIDGE(Verbose, TEXT("InternalAttachInstanceToNetHandle Attached: %s %s to ( InternalIndex: %u )"), *Instance->GetName(), *Handle.ToString(), InternalReplicationIndex);
+	UE_LOG_REPLICATIONBRIDGE(Verbose, TEXT("InternalAttachInstanceToNetHandle Attached: %s %s to ( InternalIndex: %u )"), *Instance->GetName(), *RefHandle.ToString(), InternalReplicationIndex);
 
 	// Bind instance protocol to dirty state tracking
 	if (bBindInstanceProtocol)
 	{
-		FReplicationInstanceOperationsInternal::BindInstanceProtocol(ReplicationSystemId, InternalReplicationIndex, InstanceProtocol, NetRefHandleManager->GetReplicatedObjectDataNoCheck(InternalReplicationIndex).Protocol);
-		MarkNetObjectStateDirty(Handle.GetReplicationSystemId(), InternalReplicationIndex);
+		FReplicationInstanceOperationsInternal::BindInstanceProtocol(NetHandle, InstanceProtocol, NetRefHandleManager->GetReplicatedObjectDataNoCheck(InternalReplicationIndex).Protocol);
+		MarkNetObjectStateDirty(ReplicationSystemId, InternalReplicationIndex);
 	}
 }
 
-void UReplicationBridge::InternalDetachInstanceFromNetRefHandle(FNetRefHandle Handle)
+void UReplicationBridge::InternalDetachInstanceFromNetRefHandle(FNetRefHandle RefHandle)
 {
 	using namespace UE::Net;
 	using namespace UE::Net::Private;
 
-	const FInternalNetRefIndex InternalReplicationIndex = NetRefHandleManager->GetInternalIndex(Handle);
+	const FInternalNetRefIndex InternalReplicationIndex = NetRefHandleManager->GetInternalIndex(RefHandle);
 
 	if (FReplicationInstanceProtocol* InstanceProtocol = const_cast<FReplicationInstanceProtocol*>(NetRefHandleManager->DetachInstanceProtocol(InternalReplicationIndex)))
 	{
diff --git a/Engine/Source/Runtime/Experimental/Iris/Core/Private/Iris/ReplicationSystem/ReplicationOperations.cpp b/Engine/Source/Runtime/Experimental/Iris/Core/Private/Iris/ReplicationSystem/ReplicationOperations.cpp
--- a/Engine/Source/Runtime/Experimental/Iris/Core/Private/Iris/ReplicationSystem/ReplicationOperations.cpp
+++ b/Engine/Source/Runtime/Experimental/Iris/Core/Private/Iris/ReplicationSystem/ReplicationOperations.cpp
@@ -46,7 +46,12 @@ static FNetDebugName s_LifetimeConditionDebugNames[] =
 	{TEXT("SkipReplay"), 0},
 	{TEXT("ILLEGAL"), 0},
 	{TEXT("Never"), 0},
+	{TEXT("NetGroup"), 0},
 };
+static_assert(COND_Max == 17, "s_LifetimeConditionDebugNames may need updating.");
+
+static void AppendMemberChangeMasks(FNetBitStreamWriter* ChangeMaskWriter, FNetBitStreamWriter* ConditionalChangeMaskWriter, uint8* ExternalStateBuffer, const FReplicationStateDescriptor* Descriptor);
+static void ResetMemberChangeMasks(uint8* ExternalStateBuffer, const FReplicationStateDescriptor* Descriptor);
 
 }
 
@@ -458,29 +463,33 @@ void FReplicationStateOperations::DeserializeDeltaWithMask(FNetSerializationCont
 	}
 }
 
-void FReplicationInstanceOperations::PollAndRefreshCachedPropertyData(const FReplicationInstanceProtocol* InstanceProtocol, EReplicationFragmentTraits ExcludeTraits, EReplicationFragmentPollFlags PollOptions)
+bool FReplicationInstanceOperations::PollAndRefreshCachedPropertyData(const FReplicationInstanceProtocol* InstanceProtocol, EReplicationFragmentTraits ExcludeTraits, EReplicationFragmentPollFlags PollOptions)
 {
+	bool bIsStateDirty = false;
 	FReplicationFragment* const * Fragments = InstanceProtocol->Fragments;
 
-	// $IRIS, Cache traits to avoid touching fragments
+	// @TODO Cache traits to avoid touching fragments
 	const uint32 FragmentCount = InstanceProtocol->FragmentCount;
 	for (uint32 StateIt = 0; StateIt < FragmentCount; ++StateIt)
 	{
 		const EReplicationFragmentTraits MaskedFragmentTraits = Fragments[StateIt]->GetTraits() & (EReplicationFragmentTraits::NeedsPoll | ExcludeTraits);
 		if (MaskedFragmentTraits == EReplicationFragmentTraits::NeedsPoll)
 		{			
-			Fragments[StateIt]->PollReplicatedState(PollOptions);
+			bIsStateDirty |= Fragments[StateIt]->PollReplicatedState(PollOptions);
 		}
 	}
+
+	return bIsStateDirty;
 }
 
-void FReplicationInstanceOperations::PollAndRefreshCachedPropertyData(const FReplicationInstanceProtocol* InstanceProtocol, EReplicationFragmentPollFlags PollOptions)
+bool FReplicationInstanceOperations::PollAndRefreshCachedPropertyData(const FReplicationInstanceProtocol* InstanceProtocol, EReplicationFragmentPollFlags PollOptions)
 {
-	PollAndRefreshCachedPropertyData(InstanceProtocol, EReplicationFragmentTraits::None, PollOptions);
+	return PollAndRefreshCachedPropertyData(InstanceProtocol, EReplicationFragmentTraits::None, PollOptions);
 }
 
-void FReplicationInstanceOperations::PollAndRefreshCachedObjectReferences(const FReplicationInstanceProtocol* InstanceProtocol, EReplicationFragmentTraits RequiredTraits)
+bool FReplicationInstanceOperations::PollAndRefreshCachedObjectReferences(const FReplicationInstanceProtocol* InstanceProtocol, EReplicationFragmentTraits RequiredTraits)
 {
+	bool bIsStateDirty = false;
 	const EReplicationFragmentTraits ExpectedTraits = EReplicationFragmentTraits::HasObjectReference | RequiredTraits;
 
 	FReplicationFragment* const * Fragments = InstanceProtocol->Fragments;
@@ -490,61 +499,13 @@ void FReplicationInstanceOperations::PollAndRefreshCachedObjectReferences(const
 		const EReplicationFragmentTraits MaskedFragmentTraits = Fragments[StateIt]->GetTraits() & ExpectedTraits;
 		if (MaskedFragmentTraits == ExpectedTraits)
 		{			
-			Fragments[StateIt]->PollReplicatedState(EReplicationFragmentPollFlags::ForceRefreshCachedObjectReferencesAfterGC);
+			bIsStateDirty |= Fragments[StateIt]->PollReplicatedState(EReplicationFragmentPollFlags::ForceRefreshCachedObjectReferencesAfterGC);
 		}
 	}
-}
-
-static void AppendAndResetMemberChangeMasks(FNetBitStreamWriter* ChangeMaskWriter, FNetBitStreamWriter* ConditionalChangeMaskWriter, uint8* StateBuffer, const FReplicationStateDescriptor* Descriptor)
-{
-	FNetBitArrayView::StorageWordType* ChangeMaskStorage = reinterpret_cast<FNetBitArrayView::StorageWordType*>(StateBuffer + Descriptor->GetChangeMaskOffset());
-	const uint32 BitCount = Descriptor->ChangeMaskBitCount;
-
-	// Append change mask bits
-	if (BitCount <= 32u)
-	{
-		ChangeMaskWriter->WriteBits(*ChangeMaskStorage, BitCount);
-	}
-	else
-	{
-		ChangeMaskWriter->WriteBitStream(ChangeMaskStorage, 0, BitCount);
-	}
 
-	// Reset the local changemask
-	FNetBitArrayView MemberChangeMask(ChangeMaskStorage, BitCount, FNetBitArrayView::NoResetNoValidate);
-	MemberChangeMask.Reset();
-
-	// Also reset init state dirtiness which doesn't have a changemask
-	FReplicationStateHeader& ReplicationStateHeader = Private::GetReplicationStateHeader(StateBuffer, Descriptor);
-	Private::FReplicationStateHeaderAccessor::ClearInitStateIsDirty(ReplicationStateHeader);
-
-	if (ConditionalChangeMaskWriter != nullptr)
-	{
-		if (EnumHasAnyFlags(Descriptor->Traits, EReplicationStateTraits::HasLifetimeConditionals))
-		{
-			FNetBitArrayView::StorageWordType* ConditionalChangeMaskStorage = reinterpret_cast<FNetBitArrayView::StorageWordType*>(StateBuffer + Descriptor->GetConditionalChangeMaskOffset());
-
-			// Append conditional change mask bits.
-			// The conditionals are on or off rather than tracking dirtiness so we do no reset them.
-			if (BitCount <= 32u)
-			{
-				ConditionalChangeMaskWriter->WriteBits(*ConditionalChangeMaskStorage, BitCount);
-			}
-			else
-			{
-				constexpr uint32 BitOffset = 0;
-				ConditionalChangeMaskWriter->WriteBitStream(ConditionalChangeMaskStorage, BitOffset, BitCount);
-			}
-		}
-		else
-		{
-			// Skip past our non-existing conditional mask.
-			ConditionalChangeMaskWriter->Seek(ConditionalChangeMaskWriter->GetPosBits() + BitCount);
-		}
-	}
+	return bIsStateDirty;
 }
 
-
 void FReplicationInstanceOperations::CopyAndQuantize(FNetSerializationContext& Context, uint8* DstObjectStateBuffer, FNetBitStreamWriter* OutChangeMaskWriter, const FReplicationInstanceProtocol* InstanceProtocol, const FReplicationProtocol* Protocol)
 {
 	check(InstanceProtocol && InstanceProtocol->FragmentCount == Protocol->ReplicationStateCount);
@@ -570,8 +531,8 @@ void FReplicationInstanceOperations::CopyAndQuantize(FNetSerializationContext& C
 		
 		FReplicationStateOperations::Quantize(Context, CurrentInternalStateBuffer, FragmentData[StateIt].ExternalSrcBuffer, CurrentDescriptor);
 
-		// Append and reset changemask data
-		AppendAndResetMemberChangeMasks(OutChangeMaskWriter, OutConditionalChangeMaskWriter, FragmentData[StateIt].ExternalSrcBuffer, CurrentDescriptor);
+		// Append changemask data to bitstreams
+		Private::AppendMemberChangeMasks(OutChangeMaskWriter, OutConditionalChangeMaskWriter, FragmentData[StateIt].ExternalSrcBuffer, CurrentDescriptor);
 		
 		CurrentInternalStateBuffer += CurrentDescriptor->InternalSize;
 	}
@@ -583,6 +544,19 @@ void FReplicationInstanceOperations::CopyAndQuantize(FNetSerializationContext& C
 	}
 }
 
+void FReplicationInstanceOperations::ResetDirtiness(const FReplicationInstanceProtocol* InstanceProtocol, const FReplicationProtocol* Protocol)
+{
+	check(InstanceProtocol && InstanceProtocol->FragmentCount == Protocol->ReplicationStateCount);
+
+	FReplicationInstanceProtocol::FFragmentData* FragmentData = InstanceProtocol->FragmentData;
+	const FReplicationStateDescriptor** ReplicationStateDescriptors = Protocol->ReplicationStateDescriptors;
+
+	for (uint32 StateIt = 0, StateEndIt = Protocol->ReplicationStateCount; StateIt < StateEndIt; ++StateIt)
+	{
+		const FReplicationStateDescriptor* CurrentDescriptor = ReplicationStateDescriptors[StateIt];
+		Private::ResetMemberChangeMasks(FragmentData[StateIt].ExternalSrcBuffer, CurrentDescriptor);
+	}
+}
 
 void FReplicationInstanceOperations::DequantizeAndApply(FNetSerializationContext& NetSerializationContext, const FDequantizeAndApplyParameters& Parameters)
 {
@@ -1136,3 +1110,62 @@ bool FReplicationProtocolOperations::IsEqualQuantizedState(FNetSerializationCont
 }
 
 }
+
+namespace UE::Net::Private
+{
+
+static void AppendMemberChangeMasks(FNetBitStreamWriter* ChangeMaskWriter, FNetBitStreamWriter* ConditionalChangeMaskWriter, uint8* StateBuffer, const FReplicationStateDescriptor* Descriptor)
+{
+	FNetBitArrayView::StorageWordType* ChangeMaskStorage = reinterpret_cast<FNetBitArrayView::StorageWordType*>(StateBuffer + Descriptor->GetChangeMaskOffset());
+	const uint32 BitCount = Descriptor->ChangeMaskBitCount;
+
+	// Append change mask bits
+	if (BitCount <= 32u)
+	{
+		ChangeMaskWriter->WriteBits(*ChangeMaskStorage, BitCount);
+	}
+	else
+	{
+		ChangeMaskWriter->WriteBitStream(ChangeMaskStorage, 0, BitCount);
+	}
+
+	if (ConditionalChangeMaskWriter != nullptr)
+	{
+		if (EnumHasAnyFlags(Descriptor->Traits, EReplicationStateTraits::HasLifetimeConditionals))
+		{
+			FNetBitArrayView::StorageWordType* ConditionalChangeMaskStorage = reinterpret_cast<FNetBitArrayView::StorageWordType*>(StateBuffer + Descriptor->GetConditionalChangeMaskOffset());
+
+			// Append conditional change mask bits.
+			// The conditionals are on or off rather than tracking dirtiness so we do not reset them.
+			if (BitCount <= 32u)
+			{
+				ConditionalChangeMaskWriter->WriteBits(*ConditionalChangeMaskStorage, BitCount);
+			}
+			else
+			{
+				constexpr uint32 BitOffset = 0;
+				ConditionalChangeMaskWriter->WriteBitStream(ConditionalChangeMaskStorage, BitOffset, BitCount);
+			}
+		}
+		else
+		{
+			// Skip past our non-existing conditional mask.
+			ConditionalChangeMaskWriter->Seek(ConditionalChangeMaskWriter->GetPosBits() + BitCount);
+		}
+	}
+}
+
+static void ResetMemberChangeMasks(uint8* ExternalStateBuffer, const FReplicationStateDescriptor* Descriptor)
+{
+	FNetBitArrayView::StorageWordType* ChangeMaskStorage = reinterpret_cast<FNetBitArrayView::StorageWordType*>(ExternalStateBuffer + Descriptor->GetChangeMaskOffset());
+	const uint32 BitCount = Descriptor->ChangeMaskBitCount;
+
+	// Reset the changemask. We assume it's fine to clear any padding bits.
+	FNetBitArrayView MemberChangeMask(ChangeMaskStorage, BitCount, FNetBitArrayView::ResetOnInit);
+
+	// Reset state dirtiness.
+	FReplicationStateHeader& ReplicationStateHeader = Private::GetReplicationStateHeader(ExternalStateBuffer, Descriptor);
+	Private::FReplicationStateHeaderAccessor::ClearAllStateIsDirty(ReplicationStateHeader);
+}
+
+}
diff --git a/Engine/Source/Runtime/Experimental/Iris/Core/Private/Iris/ReplicationSystem/ReplicationOperationsInternal.cpp b/Engine/Source/Runtime/Experimental/Iris/Core/Private/Iris/ReplicationSystem/ReplicationOperationsInternal.cpp
--- a/Engine/Source/Runtime/Experimental/Iris/Core/Private/Iris/ReplicationSystem/ReplicationOperationsInternal.cpp
+++ b/Engine/Source/Runtime/Experimental/Iris/Core/Private/Iris/ReplicationSystem/ReplicationOperationsInternal.cpp
@@ -23,7 +23,7 @@
 namespace UE::Net::Private
 {
 
-void FReplicationInstanceOperationsInternal::BindInstanceProtocol(uint32 ReplicationSystemId, uint32 InternalReplicationIndex, FReplicationInstanceProtocol* InstanceProtocol, const FReplicationProtocol* Protocol)
+void FReplicationInstanceOperationsInternal::BindInstanceProtocol(FNetHandle NetHandle, FReplicationInstanceProtocol* InstanceProtocol, const FReplicationProtocol* Protocol)
 {
 	FReplicationInstanceProtocol::FFragmentData* FragmentData = InstanceProtocol->FragmentData;
 	const FReplicationStateDescriptor** Descriptors = Protocol->ReplicationStateDescriptors;
@@ -35,21 +35,21 @@ void FReplicationInstanceOperationsInternal::BindInstanceProtocol(uint32 Replica
 		{
 			UE::Net::FReplicationStateHeader& ReplicationStateHeader = UE::Net::Private::GetReplicationStateHeader(FragmentData[It].ExternalSrcBuffer, Descriptors[It]);
 
-			// should be initialized to 0 when we bind it
-			check(ReplicationStateHeader.IsBound() == false || InternalReplicationIndex == 0U);
+			// Can't overwrite a bound header with a valid NetHandle.
+			check(!ReplicationStateHeader.IsBound() || !NetHandle.IsValid());
 
-			UE::Net::Private::FReplicationStateHeaderAccessor::SetReplicationIndex(ReplicationStateHeader, InternalReplicationIndex, ReplicationSystemId);
+			FReplicationStateHeaderAccessor::SetNetHandleId(ReplicationStateHeader, NetHandle);
 		}
 	}
 
-	InstanceProtocol->InstanceTraits = InternalReplicationIndex != 0U ? InstanceProtocol->InstanceTraits | EReplicationInstanceProtocolTraits::IsBound : InstanceProtocol->InstanceTraits &= ~EReplicationInstanceProtocolTraits::IsBound;
+	InstanceProtocol->InstanceTraits = NetHandle.IsValid() ? InstanceProtocol->InstanceTraits | EReplicationInstanceProtocolTraits::IsBound : InstanceProtocol->InstanceTraits &= ~EReplicationInstanceProtocolTraits::IsBound;
 }
 
 void FReplicationInstanceOperationsInternal::UnbindInstanceProtocol(FReplicationInstanceProtocol* InstanceProtocol, const FReplicationProtocol* Protocol)
 { 
 	if (EnumHasAnyFlags(InstanceProtocol->InstanceTraits, EReplicationInstanceProtocolTraits::IsBound))
 	{
-		BindInstanceProtocol(0, 0, InstanceProtocol, Protocol);
+		BindInstanceProtocol(FNetHandle(), InstanceProtocol, Protocol);
 	}
 }
 
@@ -128,6 +128,19 @@ uint32 FReplicationInstanceOperationsInternal::CopyObjectStateData(FNetBitStream
 	return 0U;
 }
 
+void FReplicationInstanceOperationsInternal::ResetObjectStateDirtiness(FNetRefHandleManager& NetRefHandleManager, uint32 InternalIndex)
+{
+	if (NetRefHandleManager.IsScopableIndex(InternalIndex))
+	{
+		const FNetRefHandleManager::FReplicatedObjectData& Object = NetRefHandleManager.GetReplicatedObjectDataNoCheck(InternalIndex);
+		// Only instance protocols with state date can have dirtiness.
+		if (Object.InstanceProtocol && Object.Protocol->InternalTotalSize > 0U)
+		{
+			FReplicationInstanceOperations::ResetDirtiness(Object.InstanceProtocol, Object.Protocol);
+		}
+	}
+}
+
 void FReplicationStateOperationsInternal::CloneDynamicState(FNetSerializationContext& Context, uint8* RESTRICT DstInternalBuffer, const uint8* RESTRICT SrcInternalBuffer, const FReplicationStateDescriptor* Descriptor)
 {
 	const FReplicationStateMemberDescriptor* MemberDescriptors = Descriptor->MemberDescriptors;
diff --git a/Engine/Source/Runtime/Experimental/Iris/Core/Private/Iris/ReplicationSystem/ReplicationOperationsInternal.h b/Engine/Source/Runtime/Experimental/Iris/Core/Private/Iris/ReplicationSystem/ReplicationOperationsInternal.h
--- a/Engine/Source/Runtime/Experimental/Iris/Core/Private/Iris/ReplicationSystem/ReplicationOperationsInternal.h
+++ b/Engine/Source/Runtime/Experimental/Iris/Core/Private/Iris/ReplicationSystem/ReplicationOperationsInternal.h
@@ -3,8 +3,8 @@
 #pragma once
 
 #include "CoreTypes.h"
-#include "Iris/ReplicationSystem/NetHandle.h"
 #include "Iris/Core/NetObjectReference.h"
+#include "Net/Core/NetHandle/NetHandle.h"
 
 
 namespace UE::Net 
@@ -30,13 +30,16 @@ namespace UE::Net::Private
 struct FReplicationInstanceOperationsInternal
 {
 	/** A call to this function will inject the index of the handle into the external statebuffer contained in the Instance Protocol */
-	static IRISCORE_API void BindInstanceProtocol(uint32 ReplicationSystemId, uint32 InternalReplicationIndex, FReplicationInstanceProtocol* InstanceProtocol, const FReplicationProtocol* Protocol);
+	static IRISCORE_API void BindInstanceProtocol(FNetHandle NetHandle, FReplicationInstanceProtocol* InstanceProtocol, const FReplicationProtocol* Protocol);
 
 	/** A call to this function will clear the index of the handle into the external statebuffer contained in the Instance Protocol */
 	static IRISCORE_API void UnbindInstanceProtocol(FReplicationInstanceProtocol* InstanceProtocol, const FReplicationProtocol* Protocol);
 
 	/** Copy state data for a single instance */
 	static IRISCORE_API uint32 CopyObjectStateData(FNetBitStreamWriter& ChangeMaskWriter, FChangeMaskCache& Cache, FNetRefHandleManager& NetRefHandleManager, FNetSerializationContext& SerializationContext, uint32 InternalIndex);
+
+	/** Reset instance state dirtiness. */
+	static IRISCORE_API void ResetObjectStateDirtiness(FNetRefHandleManager& NetRefHandleManager, uint32 InternalIndex);
 };
 
 struct FReplicationStateOperationsInternal
diff --git a/Engine/Source/Runtime/Experimental/Iris/Core/Private/Iris/ReplicationSystem/ReplicationSystem.cpp b/Engine/Source/Runtime/Experimental/Iris/Core/Private/Iris/ReplicationSystem/ReplicationSystem.cpp
--- a/Engine/Source/Runtime/Experimental/Iris/Core/Private/Iris/ReplicationSystem/ReplicationSystem.cpp
+++ b/Engine/Source/Runtime/Experimental/Iris/Core/Private/Iris/ReplicationSystem/ReplicationSystem.cpp
@@ -290,7 +290,6 @@ public:
 		FNetRefHandleManager& NetRefHandleManager = ReplicationSystemInternal.GetNetRefHandleManager();
 		FChangeMaskCache& Cache = ReplicationSystemInternal.GetChangeMaskCache();
 
-		const uint32 ReplicationSystemId = ReplicationSystem->GetId();
 		uint32 CopiedObjectCount = 0;
 		
 		// Prepare cache
@@ -307,7 +306,7 @@ public:
 
 		SerializationContext.SetInternalContext(&InternalContext);
 
-		auto CopyFunction = [&ChangeMaskWriter, &Cache, &NetRefHandleManager, &CopiedObjectCount, &SerializationContext, ReplicationSystemId](uint32 DirtyIndex)
+		auto CopyFunction = [&ChangeMaskWriter, &Cache, &NetRefHandleManager, &CopiedObjectCount, &SerializationContext](uint32 DirtyIndex)
 		{
 			CopiedObjectCount += FReplicationInstanceOperationsInternal::CopyObjectStateData(ChangeMaskWriter, Cache, NetRefHandleManager, SerializationContext, DirtyIndex);
 		};
@@ -320,9 +319,28 @@ public:
 		// Copy all ReplicatedObjects with dirty state data
 		DirtyObjects.ForAllSetBits(CopyFunction);
 
+		const uint32 ReplicationSystemId = ReplicationSystem->GetId();
 		UE_NET_TRACE_FRAME_STATSCOUNTER(ReplicationSystemId, ReplicationSystem.CopiedObjectCount, CopiedObjectCount, ENetTraceVerbosity::Trace);
 	}
 
+	void ResetObjectStateDirtiness()
+	{
+		IRIS_PROFILER_SCOPE(FReplicationSystem_ResetObjectStateDirtiness);
+
+		FNetRefHandleManager& NetRefHandleManager = ReplicationSystemInternal.GetNetRefHandleManager();
+		FDirtyNetObjectTracker& DirtyNetObjectTracker = ReplicationSystemInternal.GetDirtyNetObjectTracker();
+
+		auto ResetDirtinessFunction = [&NetRefHandleManager](uint32 DirtyIndex)
+		{
+			FReplicationInstanceOperationsInternal::ResetObjectStateDirtiness(NetRefHandleManager, DirtyIndex);
+		};
+
+		FNetBitArrayView DirtyObjects = DirtyNetObjectTracker.GetDirtyNetObjects();
+		DirtyObjects.ForAllSetBits(ResetDirtinessFunction);
+
+		DirtyNetObjectTracker.ClearDirtyNetObjects();
+	}
+
 	void ProcessNetObjectAttachmentSendQueue(FNetBlobManager::EProcessMode ProcessMode)
 	{
 		IRIS_PROFILER_SCOPE(FReplicationSystem_ProcessNetObjectAttachmentSendQueue);
@@ -588,7 +606,7 @@ void UReplicationSystem::PostSendUpdate()
 	
 	FReplicationSystemInternal& InternalSys = Impl->ReplicationSystemInternal;
 
-	InternalSys.GetDirtyNetObjectTracker().ClearDirtyNetObjects();
+	Impl->ResetObjectStateDirtiness();
 
 	InternalSys.GetChangeMaskCache().ResetCache();
 
diff --git a/Engine/Source/Runtime/Experimental/Iris/Core/Public/Iris/ReplicationState/PropertyReplicationState.h b/Engine/Source/Runtime/Experimental/Iris/Core/Public/Iris/ReplicationState/PropertyReplicationState.h
--- a/Engine/Source/Runtime/Experimental/Iris/Core/Public/Iris/ReplicationState/PropertyReplicationState.h
+++ b/Engine/Source/Runtime/Experimental/Iris/Core/Public/Iris/ReplicationState/PropertyReplicationState.h
@@ -83,13 +83,14 @@ public:
 
 	/** $IRIS TODO: Move Poll/Push/CallRepNotifies out of PropertyReplicationState as loose functions */
 
-	/** Poll src data from properties in SrcData, Note: SrcData is a UClass/UStruct containing properties
-		Compare and update representation in DstStateBuffer and update ChangeMask 
-		This is currently used to track and update dirtiness for the properties
-		$TODO: THIS requires some thought as the property offset we get in all are expressed from the base, but we do generate descriptors for all parts of the inheritance chain
-		do we carry around the base pointer or should we recalculate the offsets when we build the descriptor
+	/**
+	 * Poll src data from properties in SrcData, where SrcData is a UClass/UStruct containing properties.
+	 * Compare and update representation in DstStateBuffer and update ChangeMask. This is currently used to track and update dirtiness for the properties.
+	 * @return True if any state is dirty, false if not.
+	 * @todo This requires some thought as the property offset we get in all are expressed from the base, but we do generate descriptors for all parts of the inheritance chain.
+	 * Do we carry around the base pointer or should we recalculate the offsets when we build the descriptor
 	*/
-	IRISCORE_API void PollPropertyReplicationState(const void* RESTRICT SrcData);
+	IRISCORE_API bool PollPropertyReplicationState(const void* RESTRICT SrcData);
 
 	/** Push received state data to properties in DstData buffer, Note: DstData is a UClass/UStruct containing properties
 		Compare and update representation in DstStateBuffer and update ChangeMask 
@@ -98,13 +99,13 @@ public:
 
 
 	/**
-	  * Poll data from object referencing properties in SrcData, Note: SrcData is a UObject containing properties.
-	  * Behaves like PollPropertyReplicationState when it comes to dirtiness and such.
-	*/
-	IRISCORE_API void PollObjectReferences(const void* RESTRICT SrcData);
+	 * Poll data from object referencing properties in SrcData, Note: SrcData is a UObject containing properties.
+	 * Behaves like PollPropertyReplicationState when it comes to dirtiness and such.
+	 * @return True if any state is dirty, false if not.
+	 */
+	IRISCORE_API bool PollObjectReferences(const void* RESTRICT SrcData);
 
-	/** Invoke repnotifies for all dirty members
-	*/
+	/** Invoke repnotifies for all dirty members */
 	IRISCORE_API void CallRepNotifies(void* RESTRICT DstData, const FPropertyReplicationState* PreviousState, bool bIsInit) const;
 
 	/** Debug output state to FString */
@@ -130,6 +131,7 @@ public:
 
 private:
 	bool IsInitState() const;
+	bool IsDirty() const;
 	void ConstructStateInternal();
 	void DestructStateInternal();
 	void InjectState(const FReplicationStateDescriptor* Descriptor, uint8* InStateBuffer);
diff --git a/Engine/Source/Runtime/Experimental/Iris/Core/Public/Iris/ReplicationState/ReplicationStateFwd.h b/Engine/Source/Runtime/Experimental/Iris/Core/Public/Iris/ReplicationState/ReplicationStateFwd.h
--- a/Engine/Source/Runtime/Experimental/Iris/Core/Public/Iris/ReplicationState/ReplicationStateFwd.h
+++ b/Engine/Source/Runtime/Experimental/Iris/Core/Public/Iris/ReplicationState/ReplicationStateFwd.h
@@ -2,7 +2,8 @@
 
 #pragma once
 
-#include "CoreTypes.h"
+#include "HAL/Platform.h"
+#include "Net/Core/NetHandle/NetHandle.h"
 
 namespace UE::Net
 {
@@ -17,22 +18,20 @@ namespace Private
  */
 struct FReplicationStateHeader
 {
-	inline FReplicationStateHeader() : ReplicationIndex(0), ReplicationSystemId(0), bInitStateIsDirty(0), ReplicationFlags(0) {};
+	FReplicationStateHeader() : NetHandleId(0), bInitStateIsDirty(0), bStateIsDirty(0) {}
 
 	/** Returns true if the state is bound to the dirty tracking system */
-	bool IsBound() const { return ReplicationIndex != 0U; }
+	bool IsBound() const { return NetHandleId != 0; }
 
 private:
 	friend Private::FReplicationStateHeaderAccessor;
 
-	// All replication states that are bound by a instance protocol as assigned an index used for dirty state tracking
-	uint32 ReplicationIndex : 24;
-	// In order to support PIE we need to be able to manage multiple replication systems
-	uint32 ReplicationSystemId : 4;
+	// All replication states that are bound by an instance protocol is assigned a NetHandle for dirty state tracking
+	uint32 NetHandleId : 30;
 	// Init state doesn't use changemasks, instead we have a reserved bit here
 	uint32 bInitStateIsDirty : 1;
-	// Unused so far
-	uint32 ReplicationFlags : 3;
+	// Track whether any state is dirty.
+	uint32 bStateIsDirty : 1;
 };
 
 static_assert(sizeof(FReplicationStateHeader) == sizeof(uint32) && alignof(FReplicationStateHeader) == alignof(uint32), "FReplicationStateHeader must currently have the same size and alignment as uint32");
@@ -45,14 +44,21 @@ namespace Private
  */
 struct FReplicationStateHeaderAccessor
 {
-	static uint32 GetReplicationIndex(const FReplicationStateHeader& Header) { return Header.ReplicationIndex; }
-	static uint32 GetReplicationSystemId(const FReplicationStateHeader& Header) { return Header.ReplicationSystemId; }
+	static uint32 GetNetHandleId(const FReplicationStateHeader& Header) { return Header.NetHandleId; }
 	static bool GetIsInitStateDirty(const FReplicationStateHeader& Header) { return Header.bInitStateIsDirty; }
-	static uint32 GetReplicationFlags(const FReplicationStateHeader& Header) { return Header.ReplicationFlags; }
+	static bool GetIsStateDirty(const FReplicationStateHeader& Header) { return Header.bStateIsDirty; }
 
 	static void MarkInitStateDirty(FReplicationStateHeader& Header) { Header.bInitStateIsDirty = true; }
-	static void ClearInitStateIsDirty(FReplicationStateHeader& Header) { Header.bInitStateIsDirty = false; }
-	static void SetReplicationIndex(FReplicationStateHeader& Header, uint32 Index, uint32 ReplicationSystemId) { Header.ReplicationIndex = Index; Header.ReplicationSystemId = ReplicationSystemId; }
+	static void MarkStateDirty(FReplicationStateHeader& Header) { Header.bStateIsDirty = true; }
+
+	/** Clears both state and init state dirtiness. */
+	static void ClearAllStateIsDirty(FReplicationStateHeader& Header)
+	{ 
+		Header.bInitStateIsDirty = false;
+		Header.bStateIsDirty = false;
+	}
+
+	static void SetNetHandleId(FReplicationStateHeader& Header, FNetHandle NetHandle) { Header.NetHandleId = NetHandle.GetId(); }
 };
 
 }
diff --git a/Engine/Source/Runtime/Experimental/Iris/Core/Public/Iris/ReplicationState/ReplicationStateUtil.h b/Engine/Source/Runtime/Experimental/Iris/Core/Public/Iris/ReplicationState/ReplicationStateUtil.h
--- a/Engine/Source/Runtime/Experimental/Iris/Core/Public/Iris/ReplicationState/ReplicationStateUtil.h
+++ b/Engine/Source/Runtime/Experimental/Iris/Core/Public/Iris/ReplicationState/ReplicationStateUtil.h
@@ -2,8 +2,10 @@
 
 #pragma once
 
-#include "CoreTypes.h"
+#include "HAL/Platform.h"
 #include "Net/Core/NetBitArray.h"
+#include "Net/Core/NetHandle/NetHandleManager.h"
+#include "Net/Core/DirtyNetObjectTracker/GlobalDirtyNetObjectTracker.h"
 #include "Iris/ReplicationState/ReplicationStateDescriptor.h"
 #include "Iris/ReplicationState/ReplicationStateFwd.h"
 
@@ -14,10 +16,13 @@ namespace UE::Net::Private
 
 namespace UE::Net
 {
-/**
- * Mark the a NetObject/Handle as dirty which means it requires to be copied when we update the Replication system
- */
-inline void MarkNetObjectStateDirty(const FReplicationStateHeader& Header) { Private::MarkNetObjectStateDirty(Private::FReplicationStateHeaderAccessor::GetReplicationSystemId(Header), Private::FReplicationStateHeaderAccessor::GetReplicationIndex(Header)); }
+/** Mark the a NetObject/Handle as dirty which means it requires to be copied when we update the Replication system */
+inline void MarkNetObjectStateDirty(FReplicationStateHeader& Header)
+{
+	MarkNetObjectStateDirty(FNetHandleManager::MakeNetHandleFromId(Private::FReplicationStateHeaderAccessor::GetNetHandleId(Header)));
+	Private::FReplicationStateHeaderAccessor::MarkStateDirty(Header);
+}
+
 }
 
 namespace UE::Net::Private
@@ -65,13 +70,14 @@ inline FNetBitArrayView GetMemberConditionalChangeMask(uint8* StateBuffer, const
  * Mark specific member dirty, if this is the first bit marked as dirty in the local MemberChangeMask and the state is bound, mark owning object as dirty as well.
  * Note that all bits described by the ChangeMaskInfo will be dirtied by the call. 
  */
-inline void MarkDirty(const UE::Net::FReplicationStateHeader& InternalState, FNetBitArrayView& MemberChangeMask, const FReplicationStateMemberChangeMaskDescriptor& ChangeMaskInfo)
+inline void MarkDirty(UE::Net::FReplicationStateHeader& InternalState, FNetBitArrayView& MemberChangeMask, const FReplicationStateMemberChangeMaskDescriptor& ChangeMaskInfo)
 {
 	// If this state is bound to a replicated object, notify the replication system that we have data to copy
 	// Note that we only check the first bit of the changemask, as this is used to indicate whether the property is dirty or not
 	// some properties might use additional bits but will only be treated as dirty if the parent bit is set.
 	if (InternalState.IsBound() && !MemberChangeMask.GetBit(ChangeMaskInfo.BitOffset))
 	{
+		FReplicationStateHeaderAccessor::MarkStateDirty(InternalState);
 		MarkNetObjectStateDirty(InternalState);
 	}
 
diff --git a/Engine/Source/Runtime/Experimental/Iris/Core/Public/Iris/ReplicationSystem/FastArrayReplicationFragment.h b/Engine/Source/Runtime/Experimental/Iris/Core/Public/Iris/ReplicationSystem/FastArrayReplicationFragment.h
--- a/Engine/Source/Runtime/Experimental/Iris/Core/Public/Iris/ReplicationSystem/FastArrayReplicationFragment.h
+++ b/Engine/Source/Runtime/Experimental/Iris/Core/Public/Iris/ReplicationSystem/FastArrayReplicationFragment.h
@@ -29,9 +29,11 @@ public:
 
 protected:
 	virtual void ApplyReplicatedState(FReplicationStateApplyContext& Context) const override;
-	virtual void PollReplicatedState(EReplicationFragmentPollFlags PollOption) override;
+	virtual bool PollReplicatedState(EReplicationFragmentPollFlags PollOption) override;
 
-	void PollAllState(bool bForceFullCompare = false);
+	bool PollAllState(bool bForceFullCompare = false);
+
+	bool IsDirty() const;
 
 private:
 	// Get the wrapped FastArraySerializerProperty
@@ -60,9 +62,11 @@ protected:
 	virtual void ApplyReplicatedState(FReplicationStateApplyContext& Context) const override;
 	virtual void CallRepNotifies(FReplicationStateApplyContext& Context) override;
 	virtual void ReplicatedStateToString(FStringBuilderBase& StringBuilder, FReplicationStateApplyContext& Context, EReplicationStateToStringFlags Flags) const override;
-	virtual void PollReplicatedState(EReplicationFragmentPollFlags PollOption) override;
+	virtual bool PollReplicatedState(EReplicationFragmentPollFlags PollOption) override;
+
+	bool PollAllState();
 
-	void PollAllState();
+	bool IsDirty() const;
 
 private:
 	// Get the wrapped FastArraySerializerProperty
@@ -103,21 +107,23 @@ FastArrayType* TFastArrayReplicationFragment<FastArrayItemType, FastArrayType>::
 }
 
 template <typename FastArrayItemType, typename FastArrayType>
-void TFastArrayReplicationFragment<FastArrayItemType, FastArrayType>::PollReplicatedState(EReplicationFragmentPollFlags PollOption)
+bool TFastArrayReplicationFragment<FastArrayItemType, FastArrayType>::PollReplicatedState(EReplicationFragmentPollFlags PollOption)
 {
 	// If the ForceObjectReferences flag is set we cannot early out and must always refresh cached data
 	if (EnumHasAnyFlags(PollOption, EReplicationFragmentPollFlags::ForceRefreshCachedObjectReferencesAfterGC))
 	{
-		PollAllState(true);
+		constexpr bool bForceFullCompare = true;
+		return PollAllState(bForceFullCompare);
 	}
 	else
 	{
-		PollAllState(false);
+		constexpr bool bForceFullCompare = false;
+		return PollAllState(bForceFullCompare);
 	}
 }
 
 template <typename FastArrayItemType, typename FastArrayType>
-void TFastArrayReplicationFragment<FastArrayItemType, FastArrayType>::PollAllState(bool bForceFullCompare)
+bool TFastArrayReplicationFragment<FastArrayItemType, FastArrayType>::PollAllState(bool bForceFullCompare)
 {
 	// Lookup source data, we need the actual FastArraySerializer and the Array it is wrapping
 	FastArrayType* SrcArraySerializer = GetWrappedFastArraySerializer();
@@ -130,7 +136,7 @@ void TFastArrayReplicationFragment<FastArrayItemType, FastArrayType>::PollAllSta
 	// Check if we can early out
 	if (!bForceFullCompare && SrcArraySerializer->ArrayReplicationKey == DstArraySerializer->ArrayReplicationKey)
 	{
-		return;
+		return IsDirty();
 	}
 
 	// First we must resize target to match size of source data
@@ -196,6 +202,15 @@ void TFastArrayReplicationFragment<FastArrayItemType, FastArrayType>::PollAllSta
 		MemberChangeMask.SetBit(FIrisFastArraySerializer::IrisFastArrayPropertyBitIndex);
 		MarkNetObjectStateDirty(UE::Net::Private::GetReplicationStateHeader(SrcReplicationState->GetStateBuffer(), ReplicationStateDescriptor));
 	}
+
+	return IsDirty();
+}
+
+template <typename FastArrayItemType, typename FastArrayType>
+bool TFastArrayReplicationFragment<FastArrayItemType, FastArrayType>::IsDirty() const
+{
+	FReplicationStateHeader& ReplicationStateHeader = Private::GetReplicationStateHeader(SrcReplicationState->GetStateBuffer(), ReplicationStateDescriptor);
+	return Private::FReplicationStateHeaderAccessor::GetIsStateDirty(ReplicationStateHeader);
 }
 
 /*
@@ -244,17 +259,19 @@ void TNativeFastArrayReplicationFragment<FastArrayItemType, FastArrayType, Polli
 }
 
 template <typename FastArrayItemType, typename FastArrayType, typename PollingPolicyType>
-void TNativeFastArrayReplicationFragment<FastArrayItemType, FastArrayType, PollingPolicyType>::PollReplicatedState(EReplicationFragmentPollFlags PollOption)
+bool TNativeFastArrayReplicationFragment<FastArrayItemType, FastArrayType, PollingPolicyType>::PollReplicatedState(EReplicationFragmentPollFlags PollOption)
 {
 	// We ignore object references polling. Since the source state will have references cleaned up they will be valid once any affected item is dirtied.
 	if (PollOption == EReplicationFragmentPollFlags::PollAllState)
 	{
-		PollAllState();
+		return PollAllState();
 	}
+
+	return IsDirty();
 }
 
 template <typename FastArrayItemType, typename FastArrayType, typename PollingPolicyType>
-void TNativeFastArrayReplicationFragment<FastArrayItemType, FastArrayType, PollingPolicyType>::PollAllState()
+bool TNativeFastArrayReplicationFragment<FastArrayItemType, FastArrayType, PollingPolicyType>::PollAllState()
 {
 	using FPollingState = FastArrayPollingPolicies::FPollingState;
 	if (FPollingState* PollingState = PollingPolicy.GetPollingState())
@@ -266,7 +283,7 @@ void TNativeFastArrayReplicationFragment<FastArrayItemType, FastArrayType, Polli
 		// Check if we can early out
 		if (SrcArraySerializer->ArrayReplicationKey == PollingState->ArrayReplicationKey)
 		{
-			return;
+			return IsDirty();
 		}
 
 		// First we must resize target to match size of source data
@@ -331,6 +348,16 @@ void TNativeFastArrayReplicationFragment<FastArrayItemType, FastArrayType, Polli
 			}
 		}
 	}
+
+	return IsDirty();
+}
+
+template <typename FastArrayItemType, typename FastArrayType, typename PollingPolicyType>
+bool TNativeFastArrayReplicationFragment<FastArrayItemType, FastArrayType, PollingPolicyType>::IsDirty() const
+{
+	FastArrayType* SrcArraySerializer = GetWrappedFastArraySerializer();
+	const FReplicationStateHeader& ReplicationStateHeader = Private::FIrisFastArraySerializerPrivateAccessor::GetReplicationStateHeader(*SrcArraySerializer);
+	return Private::FReplicationStateHeaderAccessor::GetIsStateDirty(ReplicationStateHeader);
 }
 
 template <typename FastArrayItemType, typename FastArrayType, typename PollingPolicyType>
diff --git a/Engine/Source/Runtime/Experimental/Iris/Core/Public/Iris/ReplicationSystem/ObjectReplicationBridge.h b/Engine/Source/Runtime/Experimental/Iris/Core/Public/Iris/ReplicationSystem/ObjectReplicationBridge.h
--- a/Engine/Source/Runtime/Experimental/Iris/Core/Public/Iris/ReplicationSystem/ObjectReplicationBridge.h
+++ b/Engine/Source/Runtime/Experimental/Iris/Core/Public/Iris/ReplicationSystem/ObjectReplicationBridge.h
@@ -151,7 +151,7 @@ protected:
 
 	IRISCORE_API virtual void Initialize(UReplicationSystem* InReplicationSystem) override;
 	IRISCORE_API virtual void Deinitialize() override;
-	IRISCORE_API virtual void PreSendUpdateSingleHandle(FNetRefHandle Handle) override;	
+	IRISCORE_API virtual void PreSendUpdateSingleHandle(FNetRefHandle RefHandle) override;	
 	IRISCORE_API virtual void PreSendUpdate() override;	
 	IRISCORE_API virtual void UpdateInstancesWorldLocation() override;
 	IRISCORE_API virtual void PruneStaleObjects() override;	
@@ -225,11 +225,11 @@ protected:
 
 private:
 
-	void PreUpdateAndPollImpl(FNetRefHandle Handle);
+	void PreUpdateAndPollImpl(FNetRefHandle RefHandle);
 
-	void RegisterRemoteInstance(FNetRefHandle Handle, UObject* InstancePtr, const UE::Net::FReplicationProtocol* Protocol, UE::Net::FReplicationInstanceProtocol* InstanceProtocol, const FCreationHeader* Header, uint32 ConnectionId);
+	void RegisterRemoteInstance(FNetRefHandle RefHandle, UObject* InstancePtr, const UE::Net::FReplicationProtocol* Protocol, UE::Net::FReplicationInstanceProtocol* InstanceProtocol, const FCreationHeader* Header, uint32 ConnectionId);
 
-	void UnregisterRemoteInstance(FNetRefHandle Handle, bool bTearOff, bool bShouldDestroyInstance);
+	void UnregisterRemoteInstance(FNetRefHandle RefHandle, bool bTearOff, bool bShouldDestroyInstance);
 
 	void SetNetPushIdOnInstance(UE::Net::FReplicationInstanceProtocol* InstanceProtocol, FNetHandle NetHandle);
 
diff --git a/Engine/Source/Runtime/Experimental/Iris/Core/Public/Iris/ReplicationSystem/PropertyReplicationFragment.h b/Engine/Source/Runtime/Experimental/Iris/Core/Public/Iris/ReplicationSystem/PropertyReplicationFragment.h
--- a/Engine/Source/Runtime/Experimental/Iris/Core/Public/Iris/ReplicationSystem/PropertyReplicationFragment.h
+++ b/Engine/Source/Runtime/Experimental/Iris/Core/Public/Iris/ReplicationSystem/PropertyReplicationFragment.h
@@ -47,7 +47,7 @@ protected:
 	
 	virtual void CallRepNotifies(FReplicationStateApplyContext& Context) override;
 	
-	virtual void PollReplicatedState(EReplicationFragmentPollFlags PollOption) override;
+	virtual bool PollReplicatedState(EReplicationFragmentPollFlags PollOption) override;
 
 	virtual void ReplicatedStateToString(FStringBuilderBase& StringBuilder, FReplicationStateApplyContext& Context, EReplicationStateToStringFlags Flags) const override;
 
diff --git a/Engine/Source/Runtime/Experimental/Iris/Core/Public/Iris/ReplicationSystem/ReplicationBridge.h b/Engine/Source/Runtime/Experimental/Iris/Core/Public/Iris/ReplicationSystem/ReplicationBridge.h
--- a/Engine/Source/Runtime/Experimental/Iris/Core/Public/Iris/ReplicationSystem/ReplicationBridge.h
+++ b/Engine/Source/Runtime/Experimental/Iris/Core/Public/Iris/ReplicationSystem/ReplicationBridge.h
@@ -184,7 +184,7 @@ protected:
 	IRISCORE_API FNetRefHandle InternalCreateNetObjectFromRemote(FNetRefHandle WantedNetHandle, const UE::Net::FReplicationProtocol* ReplicationProtocol);
 
 	/** Attach instance to NetRefHandle. */
-	IRISCORE_API void InternalAttachInstanceToNetHandle(FNetRefHandle Handle, bool bBindInstanceProtocol, UE::Net::FReplicationInstanceProtocol* InstanceProtocol, UObject* Instance);
+	IRISCORE_API void InternalAttachInstanceToNetRefHandle(FNetRefHandle RefHandle, bool bBindInstanceProtocol, UE::Net::FReplicationInstanceProtocol* InstanceProtocol, UObject* Instance, FNetHandle NetHandle);
 
 	/** Detach instance from NetRefHandle and destroy the instance protocol. */
 	IRISCORE_API void InternalDetachInstanceFromNetRefHandle(FNetRefHandle Handle);
diff --git a/Engine/Source/Runtime/Experimental/Iris/Core/Public/Iris/ReplicationSystem/ReplicationFragment.h b/Engine/Source/Runtime/Experimental/Iris/Core/Public/Iris/ReplicationSystem/ReplicationFragment.h
--- a/Engine/Source/Runtime/Experimental/Iris/Core/Public/Iris/ReplicationSystem/ReplicationFragment.h
+++ b/Engine/Source/Runtime/Experimental/Iris/Core/Public/Iris/ReplicationSystem/ReplicationFragment.h
@@ -157,9 +157,10 @@ public:
 	virtual void CallRepNotifies(FReplicationStateApplyContext& Context) {};
 
 	/**
-	* Optional Poll method required for backwards compatibility mode which will be invoked for all Fragment with the EReplicationFragmentTraits::NeedsPoll trait set.
-	*/
-	virtual void PollReplicatedState(EReplicationFragmentPollFlags PollOption = EReplicationFragmentPollFlags::PollAllState) {}
+	 * Optional Poll method required for backwards compatibility mode which will be invoked for all Fragment with the EReplicationFragmentTraits::NeedsPoll trait set.
+	 * @return True if the state is dirty, false if not.
+	 */
+	virtual bool PollReplicatedState(EReplicationFragmentPollFlags PollOption = EReplicationFragmentPollFlags::PollAllState) { return false; }
 
 	/**
 	* Optional method to output state data to StringBuilder.
diff --git a/Engine/Source/Runtime/Experimental/Iris/Core/Public/Iris/ReplicationSystem/ReplicationOperations.h b/Engine/Source/Runtime/Experimental/Iris/Core/Public/Iris/ReplicationSystem/ReplicationOperations.h
--- a/Engine/Source/Runtime/Experimental/Iris/Core/Public/Iris/ReplicationSystem/ReplicationOperations.h
+++ b/Engine/Source/Runtime/Experimental/Iris/Core/Public/Iris/ReplicationSystem/ReplicationOperations.h
@@ -68,17 +68,24 @@ struct FReplicationStateOperations
 
 struct FReplicationInstanceOperations
 {
-	/** Update all registered Fragments that updates dirtiness by polling, except for those with any of the ExcludeTraits. */
-	static IRISCORE_API void PollAndRefreshCachedPropertyData(const FReplicationInstanceProtocol* InstanceProtocol, EReplicationFragmentTraits ExcludeTraits, EReplicationFragmentPollFlags PollOptions = EReplicationFragmentPollFlags::PollAllState);
+	/** Update all registered Fragments that updates dirtiness by polling, except for those with any of the ExcludeTraits. Returns true if a polled state is dirty. */
+	static IRISCORE_API bool PollAndRefreshCachedPropertyData(const FReplicationInstanceProtocol* InstanceProtocol, EReplicationFragmentTraits ExcludeTraits, EReplicationFragmentPollFlags PollOptions = EReplicationFragmentPollFlags::PollAllState);
 
-	/** Update all registered Fragments that updates dirtiness by polling. */
-	static IRISCORE_API void PollAndRefreshCachedPropertyData(const FReplicationInstanceProtocol* InstanceProtocol, EReplicationFragmentPollFlags PollOptions = EReplicationFragmentPollFlags::PollAllState);
+	/** Update all registered Fragments that updates dirtiness by polling.  Returns true if a polled state is dirty. */
+	static IRISCORE_API bool PollAndRefreshCachedPropertyData(const FReplicationInstanceProtocol* InstanceProtocol, EReplicationFragmentPollFlags PollOptions = EReplicationFragmentPollFlags::PollAllState);
 
-	/** Update object references in fragments that has object references and additional required traits. */
-	static IRISCORE_API void PollAndRefreshCachedObjectReferences(const FReplicationInstanceProtocol* InstanceProtocol, EReplicationFragmentTraits RequiredTraits);
+	/** Update object references in fragments that has object references and additional required traits. Returns true if a polled state is dirty. */
+	static IRISCORE_API bool PollAndRefreshCachedObjectReferences(const FReplicationInstanceProtocol* InstanceProtocol, EReplicationFragmentTraits RequiredTraits);
 
-	/** Quantize the state for a replicated object with a given InstanceProtocol using the ReplicationProtocol. DstObjectStateBuffer needs to be cleared before calling this function. */
-	static IRISCORE_API void CopyAndQuantize(FNetSerializationContext& Context, uint8* DstObjectStateBuffer, FNetBitStreamWriter* OutChangeMaskWriter, const FReplicationInstanceProtocol* InstanceProtocol, const FReplicationProtocol* Protocol);
+	/**
+	 * Quantize the state for a replicated object with a given InstanceProtocol using the ReplicationProtocol. DstObjectStateBuffer needs to be cleared before calling this function.
+	 * Changemasks will be written to the ChangeMaskWriter. Dirtiness will not be reset.
+	 * @see ResetDirtiness
+	 */
+	static IRISCORE_API void CopyAndQuantize(FNetSerializationContext& Context, uint8* DstObjectStateBuffer, FNetBitStreamWriter* ChangeMaskWriter, const FReplicationInstanceProtocol* InstanceProtocol, const FReplicationProtocol* Protocol);
+
+	/** Resets dirty tracking stored with the protocol, such as changemasks and init state dirtiness. */
+	static IRISCORE_API void ResetDirtiness(const FReplicationInstanceProtocol* InstanceProtocol, const FReplicationProtocol* Protocol);
 
 	/** Dequantize the state for a replicated object with a given protocol. Data will be pushed out by using the ReplicationFragments. */
 	static IRISCORE_API void DequantizeAndApply(FNetSerializationContext& Context, const FDequantizeAndApplyParameters& Parameters);
diff --git a/Engine/Source/Runtime/Net/Core/Private/Net/Core/NetHandle/NetHandleManager.cpp b/Engine/Source/Runtime/Net/Core/Private/Net/Core/NetHandle/NetHandleManager.cpp
--- a/Engine/Source/Runtime/Net/Core/Private/Net/Core/NetHandle/NetHandleManager.cpp
+++ b/Engine/Source/Runtime/Net/Core/Private/Net/Core/NetHandle/NetHandleManager.cpp
@@ -83,6 +83,26 @@ FNetHandle FNetHandleManager::GetOrCreateNetHandle(const UObject* Object)
 	return FNetHandle();
 }
 
+FNetHandle FNetHandleManager::MakeNetHandleFromId(uint32 Id)
+{
+	// There aren't any valid NetHandles unless there's a manager.
+	if (!Instance)
+	{
+		return FNetHandle();
+	}
+
+	if (FUObjectItem* ObjectItem = GUObjectArray.IndexToObject(static_cast<int32>(Id)))
+	{
+		FNetHandle::FInternalValue InternalValue;
+		InternalValue.Id = Id;
+		InternalValue.Epoch = static_cast<uint32>(ObjectItem->GetSerialNumber());
+		FNetHandle Handle(InternalValue);
+		return Handle;
+	}
+
+	return FNetHandle();
+}
+
 void FNetHandleManager::DestroyNetHandle(FNetHandle Handle)
 {
 	if (!Instance)
diff --git a/Engine/Source/Runtime/Net/Core/Public/Net/Core/NetHandle/NetHandle.h b/Engine/Source/Runtime/Net/Core/Public/Net/Core/NetHandle/NetHandle.h
--- a/Engine/Source/Runtime/Net/Core/Public/Net/Core/NetHandle/NetHandle.h
+++ b/Engine/Source/Runtime/Net/Core/Public/Net/Core/NetHandle/NetHandle.h
@@ -64,7 +64,9 @@ private:
 	};
 
 	FNetHandle(FObjectKey ObjectKey);
+	explicit FNetHandle(FInternalValue InInternalValue);
 	explicit FNetHandle(uint64 InInternalValue);
+
 	uint64 GetInternalValue() const;
 
 	union
@@ -85,7 +87,7 @@ inline FNetHandle::FNetHandle(const	FNetHandle& Other)
 }
 
 inline FNetHandle::FNetHandle(FNetHandle&& Other)
-	: Value(Other.Value)
+: Value(Other.Value)
 {
 }
 
@@ -94,9 +96,14 @@ inline FNetHandle::FNetHandle(FObjectKey ObjectKey)
 {
 }
 
+inline FNetHandle::FNetHandle(FInternalValue InInternalValue)
+: Value(BitCast<FObjectKey>(InInternalValue))
+{
+}
+
 inline FNetHandle::FNetHandle(uint64 InInternalValue)
+: Value(BitCast<FObjectKey>(InInternalValue))
 {
-	Value = BitCast<FObjectKey>(InInternalValue);
 }
 
 inline uint64 FNetHandle::GetInternalValue() const
diff --git a/Engine/Source/Runtime/Net/Core/Public/Net/Core/NetHandle/NetHandleManager.h b/Engine/Source/Runtime/Net/Core/Public/Net/Core/NetHandle/NetHandleManager.h
--- a/Engine/Source/Runtime/Net/Core/Public/Net/Core/NetHandle/NetHandleManager.h
+++ b/Engine/Source/Runtime/Net/Core/Public/Net/Core/NetHandle/NetHandleManager.h
@@ -32,6 +32,9 @@ public:
 	/** Returns an existing NetHandle for an object if it's present or creates one if it's not. Assumes Init() has been called. */
 	NETCORE_API static FNetHandle GetOrCreateNetHandle(const UObject*);
 
+	/** Returns a NetHandle given an Id that has previously been extracted from a valid NetHandle. The behavior is implementation defined if the Id was not extracted from a valid NetHandle. */
+	NETCORE_API static FNetHandle MakeNetHandleFromId(uint32 Id);
+
 private:
 	friend FNetHandleDestroyer;
 
