Overlap Box micro optimization

diff --git a/Engine/Source/Runtime/Experimental/Chaos/Private/Chaos/Collision/TriangleOverlap.cpp b/Engine/Source/Runtime/Experimental/Chaos/Private/Chaos/Collision/TriangleOverlap.cpp
--- a/Engine/Source/Runtime/Experimental/Chaos/Private/Chaos/Collision/TriangleOverlap.cpp
+++ b/Engine/Source/Runtime/Experimental/Chaos/Private/Chaos/Collision/TriangleOverlap.cpp
@@ -324,14 +324,14 @@ namespace Chaos
 		{
 			const VectorRegister4Float OtherClosest = VectorAdd(OtherClosests[i], Position);
 			BoxEdges[i] = VectorSubtract(OtherClosest, ClosestPoint);
-			BoxEdgeNormals[i] = VectorNegate(VectorNormalize(BoxEdges[i]));
+			BoxEdgeNormals[i] = VectorNegate(BoxEdges[i]);
 		}
 
 		for (int32 i = 0; i < 3; i++)
 		{
 			if (HasToComputeEdge(BoxEdgeNormals[(i + 1) % 3], BoxEdgeNormals[(i + 2) % 3], BoxEdges[i], Normal, Edge))
 			{
-				VectorRegister4Float Axis = VectorNormalize(VectorCross(BoxEdges[i], Edge));
+				VectorRegister4Float Axis = VectorCross(BoxEdges[i], Edge);
 				const FRealSingle Sign = VectorDot3Scalar(VectorSubtract(PlaneVertex, Centroid), Axis);
 				if (Sign < 0.0f)
 				{
@@ -446,8 +446,8 @@ namespace Chaos
 		VectorRegister4Float ClosCent = VectorSubtract(ClosestPointNorm, Centroid);
 		VectorRegister4Float FurtCent = VectorSubtract(FurthestPointNorm, Centroid);
 
-		FRealSingle ClosestDist = VectorDot3Scalar(Normal, VectorNormalize(ClosCent));
-		FRealSingle FurthestDist = VectorDot3Scalar(Normal, VectorNormalize(FurtCent));
+		FRealSingle ClosestDist = VectorDot3Scalar(Normal, ClosCent);
+		FRealSingle FurthestDist = VectorDot3Scalar(Normal, FurtCent);
 
 		return FMath::Sign(ClosestDist) != FMath::Sign(FurthestDist);
 	}
@@ -518,18 +518,18 @@ namespace Chaos
 		{
 			const VectorRegister4Float OtherClosest = VectorAdd(OtherClosests[i], Position);
 			BoxEdges[i] = VectorSubtract(OtherClosest, ClosestPoint);
-			BoxEdgeNormals[i] = VectorNegate(VectorNormalize(BoxEdges[i]));
+			BoxEdgeNormals[i] = VectorBitwiseXor(SignBit, BoxEdges[i]);
 		}
 
 		for (int32 i = 0; i < 3; i++)
 		{
 			if (HasToComputeEdge(BoxEdgeNormals[(i + 1) % 3], BoxEdgeNormals[(i + 2) % 3], BoxEdges[i], Normal, TriangleEdge))
 			{
-				VectorRegister4Float Axis = VectorNormalize(VectorCross(BoxEdges[i], TriangleEdge));
+				VectorRegister4Float Axis = VectorCross(BoxEdges[i], TriangleEdge);
 				const FRealSingle Sign = VectorDot3Scalar(VectorSubtract(TriangleVertex, Centroid), Axis);
 				if (Sign < 0.0f)
 				{
-					Axis = VectorNegate(Axis);
+					Axis = VectorBitwiseXor(SignBit, Axis);
 				}
 				const FRealSingle ScaledSeparation = VectorDot3Scalar(VectorSubtract(ClosestPoint, TriangleVertex), Axis);
 				if (ScaledSeparation > 0.0f)
@@ -571,7 +571,7 @@ namespace Chaos
 		const VectorRegister4Float Centroid = VectorMultiply(VectorAdd(VectorAdd(A, B), C), Third);
 		
 
-		const VectorRegister4Float ABPlane = VectorNormalize(VectorCross(Normal, AB));
+		const VectorRegister4Float ABPlane = VectorCross(Normal, AB);
 		const VectorRegister4Float ABSigns = VectorBitwiseAnd(ABPlane, SignBit);
 		const VectorRegister4Float ABLocalClosest = VectorBitwiseOr(ABSigns, HalfExtents);
 		const VectorRegister4Float ABClosestPoint = VectorAdd(ABLocalClosest, Position);
@@ -581,7 +581,7 @@ namespace Chaos
 			return false;
 		}
 
-		const VectorRegister4Float BCPlane = VectorNormalize(VectorCross(Normal, BC));
+		const VectorRegister4Float BCPlane = VectorCross(Normal, BC);
 		const VectorRegister4Float BCSigns = VectorBitwiseAnd(BCPlane, SignBit);
 		const VectorRegister4Float BCLocalClosest = VectorBitwiseOr(BCSigns, HalfExtents);
 		const VectorRegister4Float BCClosestPoint = VectorAdd(BCLocalClosest, Position);
@@ -591,7 +591,7 @@ namespace Chaos
 			return false;
 		}
 
-		const VectorRegister4Float CAPlane = VectorNormalize(VectorCross(Normal, CA));
+		const VectorRegister4Float CAPlane = VectorCross(Normal, CA);
 		const VectorRegister4Float CASigns = VectorBitwiseAnd(CAPlane, SignBit);
 		const VectorRegister4Float CALocalClosest = VectorBitwiseOr(CASigns, HalfExtents);
 		const VectorRegister4Float CAClosestPoint = VectorAdd(CALocalClosest, Position);
@@ -610,8 +610,8 @@ namespace Chaos
 		VectorRegister4Float ClosCent = VectorSubtract(ClosestPoint, Centroid);
 		VectorRegister4Float FurtCent = VectorSubtract(FurthestPoint, Centroid);
 
-		FRealSingle ClosestDist = VectorDot3Scalar(Normal, VectorNormalize(ClosCent));
-		FRealSingle FurthestDist = VectorDot3Scalar(Normal, VectorNormalize(FurtCent));
+		FRealSingle ClosestDist = VectorDot3Scalar(Normal, ClosCent);
+		FRealSingle FurthestDist = VectorDot3Scalar(Normal, FurtCent);
 
 		if (FMath::Sign(ClosestDist) == FMath::Sign(FurthestDist))
 		{
