Chaos: fix kienamtically driven GC not properly updating in some situations
- when internal clusters are the kienamtically driven roots
- when all pieces are broken and sleeping , preventing the rendering from proeprly updating the matrices

diff --git a/Engine/Source/Runtime/Experimental/Chaos/Private/PhysicsProxy/GeometryCollectionPhysicsProxy.cpp b/Engine/Source/Runtime/Experimental/Chaos/Private/PhysicsProxy/GeometryCollectionPhysicsProxy.cpp
--- a/Engine/Source/Runtime/Experimental/Chaos/Private/PhysicsProxy/GeometryCollectionPhysicsProxy.cpp
+++ b/Engine/Source/Runtime/Experimental/Chaos/Private/PhysicsProxy/GeometryCollectionPhysicsProxy.cpp
@@ -2215,6 +2215,7 @@ void FGeometryCollectionPhysicsProxy::PushStateOnGameThread(Chaos::FPBDRigidsSol
 	bIsPhysicsThreadWorldTransformDirty = GameThreadPerFrameData.GetIsWorldTransformDirty();
 	if (bIsPhysicsThreadWorldTransformDirty)
 	{
+		Parameters.PrevWorldTransform = Parameters.WorldTransform;
 		Parameters.WorldTransform = GameThreadPerFrameData.GetWorldTransform();
 		GameThreadPerFrameData.ResetIsWorldTransformDirty();
 	}
@@ -2232,69 +2233,83 @@ void FGeometryCollectionPhysicsProxy::PushToPhysicsState()
 {
 	using namespace Chaos;
 	// CONTEXT: PHYSICSTHREAD
-	// because the attached actor can be dynamic, we need to update the kinematic particles properly
-	if (bIsPhysicsThreadWorldTransformDirty || bIsCollisionFilterDataDirty)
-	{
-		const FTransform& ActorToWorld = Parameters.WorldTransform;
-
-		// used to avoid doing the work twice if we have a internalCluster parent 
-		bool InternalClusterParentUpdated = false;
 
-		int32 NumTransformGroupElements = PhysicsThreadCollection.NumElements(FGeometryCollection::TransformGroup);
-		for (int32 TransformGroupIndex = 0; TransformGroupIndex < NumTransformGroupElements; ++TransformGroupIndex)
+	TArray<FClusterHandle*> ProcessedInternalClusters;
+	if (bIsCollisionFilterDataDirty)
+	{
+		if (Chaos::FPhysicsSolver* RigidSolver = GetSolver<Chaos::FPhysicsSolver>())
 		{
-			Chaos::FPBDRigidClusteredParticleHandle* Handle = SolverParticleHandles[TransformGroupIndex];
-			if (Handle)
+			FRigidClustering& RigidClustering = RigidSolver->GetEvolution()->GetRigidClustering();
+
+			const int32 NumTransforms = PhysicsThreadCollection.NumElements(FGeometryCollection::TransformGroup);
+			for (int32 TransformIndex = 0; TransformIndex < NumTransforms; ++TransformIndex)
 			{
-				// Must update our filters before updating an internal cluster parent
-				if (bIsCollisionFilterDataDirty)
+				if (Chaos::FPBDRigidClusteredParticleHandle* Handle = SolverParticleHandles[TransformIndex])
 				{
+					// Must update our filters before updating an internal cluster parent
 					const Chaos::FShapesArray& ShapesArray = Handle->ShapesArray();
 					for (const TUniquePtr<Chaos::FPerShapeData>& Shape : ShapesArray)
 					{
 						Shape->SetQueryData(Parameters.QueryFilterData);
 						Shape->SetSimData(Parameters.SimulationFilterData);
 					}
-				}
 
-				// in the case of cluster union we need to find our Internal Cluster parent and update it
-				if (!InternalClusterParentUpdated)
-				{
 					FClusterHandle* ParentHandle = Handle->Parent();
 					if (ParentHandle && ParentHandle->InternalCluster())
 					{
-						if (bIsPhysicsThreadWorldTransformDirty && Handle->ObjectState() == Chaos::EObjectStateType::Kinematic && ParentHandle->ObjectState() == Chaos::EObjectStateType::Kinematic && !ParentHandle->Disabled())
+						if (!ProcessedInternalClusters.Contains(ParentHandle))
 						{
-							FTransform NewChildWorldTransform = PhysicsThreadCollection.MassToLocal[TransformGroupIndex] * PhysicsThreadCollection.Transform[TransformGroupIndex] * ActorToWorld;
-							Chaos::FRigidTransform3 ParentToChildTransform = Handle->ChildToParent().Inverse();
-							FTransform NewParentWorldTRansform = ParentToChildTransform * NewChildWorldTransform;
-							SetClusteredParticleKinematicTarget_Internal(ParentHandle, NewParentWorldTRansform);
-						}
-
-						if (bIsCollisionFilterDataDirty)
-						{
-							if(Chaos::FPhysicsSolver* RigidSolver = GetSolver<Chaos::FPhysicsSolver>())
+							ProcessedInternalClusters.Add(ParentHandle);
+							// Must update our filters before updating an internal cluster parent
+							FRigidClustering::FRigidHandleArray* ChildArray = RigidClustering.GetChildrenMap().Find(ParentHandle);
+							if (ensure(ChildArray))
 							{
-								FRigidClustering& RigidClustering = RigidSolver->GetEvolution()->GetRigidClustering();
-								FRigidClustering::FRigidHandleArray* ChildArray = RigidClustering.GetChildrenMap().Find(ParentHandle);
-								if (ensure(ChildArray))
-								{
-									// If our filter changed, internal cluster parent's filter may be stale.
-									UpdateClusterFilterDataFromChildren(ParentHandle, *ChildArray);
-								}
-
+								// If our filter changed, internal cluster parent's filter may be stale.
+								UpdateClusterFilterDataFromChildren(ParentHandle, *ChildArray);
 							}
 						}
+					}
+				}
+			}
+		}
+	}
+
+	if (bIsPhysicsThreadWorldTransformDirty)
+	{
+		ProcessedInternalClusters.Reset();
 
+		const FTransform& ActorToWorld = Parameters.WorldTransform;
 
-						InternalClusterParentUpdated = true;
+		const int32 NumTransforms = PhysicsThreadCollection.NumElements(FGeometryCollection::TransformGroup);
+		for (int32 TransformGroupIndex = 0; TransformGroupIndex < NumTransforms; ++TransformGroupIndex)
+		{
+			if (Chaos::FPBDRigidClusteredParticleHandle* Handle = SolverParticleHandles[TransformGroupIndex])
+			{
+				FClusterHandle* KinematicRootHandle = nullptr;
+				if (!Handle->Disabled() && Handle->ObjectState() == Chaos::EObjectStateType::Kinematic)
+				{
+					KinematicRootHandle = Handle;
+				}
+				else 
+				{
+					// is there a internal parent as a kinematic root?
+					FClusterHandle* ParentHandle = Handle->Parent();
+					if (ParentHandle && ParentHandle->InternalCluster() && !ParentHandle->Disabled() && ParentHandle->ObjectState() == Chaos::EObjectStateType::Kinematic)
+					{
+						if (!ProcessedInternalClusters.Contains(ParentHandle))
+						{
+							ProcessedInternalClusters.Add(ParentHandle);
+							KinematicRootHandle = ParentHandle;
+						}
 					}
 				}
-
-				if (bIsPhysicsThreadWorldTransformDirty && !Handle->Disabled() && Handle->ObjectState() == Chaos::EObjectStateType::Kinematic)
+				if (KinematicRootHandle)
 				{
-					FTransform WorldTransform = PhysicsThreadCollection.MassToLocal[TransformGroupIndex] * PhysicsThreadCollection.Transform[TransformGroupIndex] * ActorToWorld;
-					SetClusteredParticleKinematicTarget_Internal(Handle, WorldTransform);
+					const FTransform RootWorldTransform(KinematicRootHandle->R(), KinematicRootHandle->X());
+					const FTransform RootRelativeTransform = RootWorldTransform.GetRelativeTransform(Parameters.PrevWorldTransform);
+					const FTransform WorldTransform = RootRelativeTransform * ActorToWorld;
+
+					SetClusteredParticleKinematicTarget_Internal(KinematicRootHandle, WorldTransform);
 				}
 			}
 		}
@@ -2848,8 +2863,9 @@ bool FGeometryCollectionPhysicsProxy::PullFromPhysicsState(const Chaos::FDirtyGe
 			}
 		}
 
-		//question: why do we need this? Sleeping objects will always have to update GPU
-		if (bIsCollectionDirty)
+		// if physics world transform was dirtied this frame we need to force an update to make sure transforms of 
+		// broken sleeping particles on kinemically driven GCs remained updated and the object feel grounded to the world
+		if (bIsCollectionDirty || bIsPhysicsThreadWorldTransformDirty)
 		{
 			GameThreadCollection.MakeDirty();
 		}
diff --git a/Engine/Source/Runtime/Experimental/Chaos/Public/GeometryCollection/GeometryCollectionSimulationCoreTypes.h b/Engine/Source/Runtime/Experimental/Chaos/Public/GeometryCollection/GeometryCollectionSimulationCoreTypes.h
--- a/Engine/Source/Runtime/Experimental/Chaos/Public/GeometryCollection/GeometryCollectionSimulationCoreTypes.h
+++ b/Engine/Source/Runtime/Experimental/Chaos/Public/GeometryCollection/GeometryCollectionSimulationCoreTypes.h
@@ -281,6 +281,7 @@ struct FSimulationParameters
 	bool Simulating;
 
 	FTransform WorldTransform;
+	FTransform PrevWorldTransform;
 
 	bool EnableClustering;
 	int32 ClusterGroupIndex;
