geometry collection : fixed material consolidation and make sure instanced mesh array is properly generated when building a GC in dataflow
- Improved InitializeMaterials to be able to handle material list only made of external materials
- Added CollectionInstancedMeshFacade and update code using the attribute directly to use the facade instead
- Added new input / output to Geometry Collection Terminal node ( InstancedMeshes )
- Simplified FindOrAddAutoInstanceMesh methods

diff --git a/Engine/Plugins/Experimental/GeometryCollectionPlugin/Source/GeometryCollectionNodes/Private/Dataflow/Nodes/GeometryCollectionAssetNodes.cpp b/Engine/Plugins/Experimental/GeometryCollectionPlugin/Source/GeometryCollectionNodes/Private/Dataflow/Nodes/GeometryCollectionAssetNodes.cpp
--- a/Engine/Plugins/Experimental/GeometryCollectionPlugin/Source/GeometryCollectionNodes/Private/Dataflow/Nodes/GeometryCollectionAssetNodes.cpp
+++ b/Engine/Plugins/Experimental/GeometryCollectionPlugin/Source/GeometryCollectionNodes/Private/Dataflow/Nodes/GeometryCollectionAssetNodes.cpp
@@ -8,6 +8,7 @@
 #include "GeometryCollection/GeometryCollectionObject.h"
 #include "GeometryCollection/GeometryCollectionClusteringUtility.h"
 #include "GeometryCollection/Facades/CollectionHierarchyFacade.h"
+#include "GeometryCollection/Facades/CollectionInstancedMeshFacade.h"
 
 namespace Dataflow
 {
@@ -32,6 +33,7 @@ void FGeometryCollectionTerminalDataflowNode::SetAssetValue(TObjectPtr<UObject>
 {
 	using FGeometryCollectionPtr = TSharedPtr<FGeometryCollection, ESPMode::ThreadSafe>;
 	using FMaterialArray = TArray<TObjectPtr<UMaterial>>;
+	using FInstancedMeshesArray = TArray<FGeometryCollectionAutoInstanceMesh>;
 
 	if (UGeometryCollection* CollectionAsset = Cast<UGeometryCollection>(Asset.Get()))
 	{
@@ -39,7 +41,11 @@ void FGeometryCollectionTerminalDataflowNode::SetAssetValue(TObjectPtr<UObject>
 		{
 			const FManagedArrayCollection& InCollection = GetValue<FManagedArrayCollection>(Context, &Collection);
 			const FMaterialArray& InMaterials = GetValue<FMaterialArray>(Context, &Materials);
-			CollectionAsset->ResetFrom(InCollection, InMaterials);
+			const FInstancedMeshesArray& InInstancedMeshes = GetValue<FInstancedMeshesArray>(Context, &InstancedMeshes);
+
+			const bool bHasInternalMaterial = false; // with data flow there's no assumption of internal materials
+			CollectionAsset->ResetFrom(InCollection, InMaterials, false);
+			CollectionAsset->AutoInstanceMeshes = InInstancedMeshes;
 		}
 	}
 }
@@ -47,12 +53,15 @@ void FGeometryCollectionTerminalDataflowNode::SetAssetValue(TObjectPtr<UObject>
 void FGeometryCollectionTerminalDataflowNode::Evaluate(Dataflow::FContext& Context) const
 {
 	using FMaterialArray = TArray<TObjectPtr<UMaterial>>;
+	using FInstancedMeshesArray = TArray<FGeometryCollectionAutoInstanceMesh>;
 
 	const FManagedArrayCollection& InCollection = GetValue<FManagedArrayCollection>(Context, &Collection);
 	const FMaterialArray& InMaterials = GetValue<FMaterialArray>(Context, &Materials);
+	const FInstancedMeshesArray& InInstancedMeshes = GetValue<FInstancedMeshesArray>(Context, &InstancedMeshes);
 
 	SetValue<FManagedArrayCollection>(Context, InCollection, &Collection);
 	SetValue<FMaterialArray>(Context, InMaterials, &Materials);
+	SetValue<FInstancedMeshesArray>(Context, InInstancedMeshes, &InstancedMeshes);
 }
 
 // ===========================================================================================================================
@@ -111,28 +120,57 @@ FCreateGeometryCollectionFromSourcesDataflowNode::FCreateGeometryCollectionFromS
 	RegisterInputConnection(&Sources);
 	RegisterOutputConnection(&Collection);
 	RegisterOutputConnection(&Materials);
+	RegisterOutputConnection(&InstancedMeshes);
 }
 
 void FCreateGeometryCollectionFromSourcesDataflowNode::Evaluate(Dataflow::FContext& Context, const FDataflowOutput* Out) const
 {
-	ensure(Out->IsA(&Collection) || Out->IsA(&Materials));
+	ensure(Out->IsA(&Collection) || Out->IsA(&Materials) || Out->IsA(&InstancedMeshes));
 	
 	using FGeometryCollectionSourceArray = TArray<FGeometryCollectionSource>;
 	const FGeometryCollectionSourceArray& InSources = GetValue<FGeometryCollectionSourceArray>(Context, &Sources);
 
 	FGeometryCollection OutCollection;
 	TArray<TObjectPtr<UMaterial>> OutMaterials;
+	TArray<FGeometryCollectionAutoInstanceMesh> OutInstancedMeshes;
+
+	// make sure we have an attribute for instanced meshes
+	GeometryCollection::Facades::FCollectionInstancedMeshFacade InstancedMeshFacade(OutCollection);
+	InstancedMeshFacade.DefineSchema();
+
 	constexpr bool bReindexMaterialsInLoop = false;
 	for (const FGeometryCollectionSource& Source: InSources)
 	{
+		const int32 NumTransformsBeforeAppending = OutCollection.NumElements(FGeometryCollection::TransformGroup);
+
 		// todo: change AppendGeometryCollectionSource to take a FManagedArrayCollection so we could move the collection when assigning it to the output
 		FGeometryCollectionEngineConversion::AppendGeometryCollectionSource(Source, OutCollection, OutMaterials, bReindexMaterialsInLoop);
+
+		// todo(chaos) if the source is a geometry collection this will not work properly 
+		FGeometryCollectionAutoInstanceMesh InstancedMesh;
+		InstancedMesh.StaticMesh = Source.SourceGeometryObject;
+		InstancedMesh.Materials = Source.SourceMaterial;
+		const int32 InstancedMeshIndex = OutInstancedMeshes.AddUnique(InstancedMesh);
+
+		// add the instanced mesh  for all the newly added transforms 
+		const int32 NumTransformsAfterAppending = OutCollection.NumElements(FGeometryCollection::TransformGroup);
+		for (int32 TransformIndex = NumTransformsBeforeAppending; TransformIndex < NumTransformsAfterAppending; TransformIndex++)
+		{
+			InstancedMeshFacade.SetIndex(TransformIndex, InstancedMeshIndex);
+		}
 	}
 	if (bReindexMaterialsInLoop == false)
 	{
 		OutCollection.ReindexMaterials();
 	}
 
+	// add the instanced mesh indices
+
+	const int32 NumTransforms = InstancedMeshFacade.GetNumIndices();
+	for (int32 TransformIndex = 0; TransformIndex < NumTransforms; TransformIndex++)
+	{
+	}
+
 	// make sure we have only one root
 	if (FGeometryCollectionClusteringUtility::ContainsMultipleRootBones(&OutCollection))
 	{
@@ -146,4 +184,6 @@ void FCreateGeometryCollectionFromSourcesDataflowNode::Evaluate(Dataflow::FConte
 	// we have to make a copy since we have generated a FGeometryCollection which is inherited from FManagedArrayCollection
 	SetValue(Context, static_cast<const FManagedArrayCollection&>(OutCollection), &Collection);
 	SetValue(Context, MoveTemp(OutMaterials), &Materials);
+	SetValue(Context, MoveTemp(OutInstancedMeshes), &InstancedMeshes);
+	
 }
diff --git a/Engine/Plugins/Experimental/GeometryCollectionPlugin/Source/GeometryCollectionNodes/Public/Dataflow/Nodes/GeometryCollectionAssetNodes.h b/Engine/Plugins/Experimental/GeometryCollectionPlugin/Source/GeometryCollectionNodes/Public/Dataflow/Nodes/GeometryCollectionAssetNodes.h
--- a/Engine/Plugins/Experimental/GeometryCollectionPlugin/Source/GeometryCollectionNodes/Public/Dataflow/Nodes/GeometryCollectionAssetNodes.h
+++ b/Engine/Plugins/Experimental/GeometryCollectionPlugin/Source/GeometryCollectionNodes/Public/Dataflow/Nodes/GeometryCollectionAssetNodes.h
@@ -27,6 +27,10 @@ public:
 	UPROPERTY(meta = (DataflowInput, DataflowOutput, DataflowPassthrough = "Materials", DisplayName = "Materials"))
 	TArray<TObjectPtr<UMaterial>> Materials;
 
+	/** array of instanced meshes*/
+	UPROPERTY(meta = (DataflowInput, DataflowOutput, DataflowPassthrough = "InstancedMeshes", DisplayName = "InstancedMeshes"))
+	TArray<FGeometryCollectionAutoInstanceMesh> InstancedMeshes;
+
 	FGeometryCollectionTerminalDataflowNode(const Dataflow::FNodeParameters& InParam, FGuid InGuid = FGuid::NewGuid())
 		: FDataflowTerminalNode(InParam, InGuid)
 	{
@@ -34,6 +38,8 @@ public:
 		RegisterOutputConnection(&Collection, &Collection);
 		RegisterInputConnection(&Materials);
 		RegisterOutputConnection(&Materials, &Materials);
+		RegisterInputConnection(&InstancedMeshes);
+		RegisterOutputConnection(&InstancedMeshes, &InstancedMeshes);
 	}
 
 	virtual void Evaluate(Dataflow::FContext& Context) const override;
@@ -108,6 +114,10 @@ public:
 	/** Geometry collection newly created */
 	UPROPERTY(meta = (DataflowOutput, DisplayName = "Materials"))
 	TArray<TObjectPtr<UMaterial>> Materials;
+
+	/** array of instanced meshes*/
+	UPROPERTY(meta = (DataflowOutput, DisplayName = "InstancedMeshes"))
+	TArray<FGeometryCollectionAutoInstanceMesh> InstancedMeshes;
 };
 
 namespace Dataflow
diff --git a/Engine/Source/Runtime/Experimental/Chaos/Private/GeometryCollection/Facades/CollectionInstancedMeshFacade.cpp b/Engine/Source/Runtime/Experimental/Chaos/Private/GeometryCollection/Facades/CollectionInstancedMeshFacade.cpp
new file mode 100644
--- /dev/null
+++ b/Engine/Source/Runtime/Experimental/Chaos/Private/GeometryCollection/Facades/CollectionInstancedMeshFacade.cpp
@@ -0,0 +1,49 @@
+// Copyright Epic Games, Inc. All Rights Reserved.
+
+#include "GeometryCollection/Facades/CollectionInstancedMeshFacade.h"
+#include "GeometryCollection/Facades/CollectionHierarchyFacade.h"
+#include "GeometryCollection/GeometryCollection.h"
+
+namespace GeometryCollection::Facades
+{
+	FCollectionInstancedMeshFacade::FCollectionInstancedMeshFacade(FManagedArrayCollection& InCollection)
+		: InstancedMeshIndexAttribute(InCollection, "AutoInstanceMeshIndex", FGeometryCollection::TransformGroup)
+	{}
+	
+	FCollectionInstancedMeshFacade::FCollectionInstancedMeshFacade(const FManagedArrayCollection& InCollection)
+		: InstancedMeshIndexAttribute(InCollection, "AutoInstanceMeshIndex", FGeometryCollection::TransformGroup)
+	{}
+
+
+	void FCollectionInstancedMeshFacade::DefineSchema()
+	{
+		InstancedMeshIndexAttribute.AddAndFill(INDEX_NONE);
+	}
+
+	bool FCollectionInstancedMeshFacade::IsValid() const
+	{
+		return InstancedMeshIndexAttribute.IsValid();
+	}
+
+	bool FCollectionInstancedMeshFacade::IsConst() const
+	{
+		return InstancedMeshIndexAttribute.IsConst();
+	}
+
+	int32 FCollectionInstancedMeshFacade::GetNumIndices() const
+	{
+		return InstancedMeshIndexAttribute.Get().Num();
+	}
+
+	int32 FCollectionInstancedMeshFacade::GetIndex(int32 TransformIndex) const
+	{
+		return InstancedMeshIndexAttribute.Get()[TransformIndex];
+	}
+
+	void FCollectionInstancedMeshFacade::SetIndex(int32 TransformIndex, int32 InstanceMeshIndex)
+	{
+		check(!IsConst());
+		InstancedMeshIndexAttribute.Modify()[TransformIndex] = InstanceMeshIndex;
+	}
+
+}
\ No newline at end of file
diff --git a/Engine/Source/Runtime/Experimental/Chaos/Public/GeometryCollection/Facades/CollectionInstancedMeshFacade.h b/Engine/Source/Runtime/Experimental/Chaos/Public/GeometryCollection/Facades/CollectionInstancedMeshFacade.h
new file mode 100644
--- /dev/null
+++ b/Engine/Source/Runtime/Experimental/Chaos/Public/GeometryCollection/Facades/CollectionInstancedMeshFacade.h
@@ -0,0 +1,42 @@
+// Copyright Epic Games, Inc. All Rights Reserved.
+
+#pragma once
+
+#include "CoreMinimal.h"
+#include "GeometryCollection/ManagedArrayAccessor.h"
+#include "GeometryCollection/ManagedArrayCollection.h"
+
+namespace GeometryCollection::Facades
+{
+	/**
+	 * Provides an API to read and manipulate hierarchy in a managed array collection
+	 */
+	class CHAOS_API FCollectionInstancedMeshFacade
+	{
+	public:
+		FCollectionInstancedMeshFacade(FManagedArrayCollection& InCollection);
+		FCollectionInstancedMeshFacade(const FManagedArrayCollection& InCollection);
+
+		/** Create the facade attributes. */
+		void DefineSchema();
+
+		/** Valid if parent and children arrays are available */
+		bool IsValid() const;
+
+		/** Is the facade defined constant. */
+		bool IsConst() const;
+
+		/** get the total number of indices */
+		int32 GetNumIndices() const;
+
+		/** get the instance mesh index for specific transform index */
+		int32 GetIndex(int32 TransformIndex) const;
+
+		/** set the instance mesh index for specific transform index */
+		void SetIndex(int32 TransformIndex, int32 InstanceMeshIndex);
+
+
+	private:
+		TManagedArrayAccessor<int32> InstancedMeshIndexAttribute;
+	};
+}
\ No newline at end of file
diff --git a/Engine/Source/Runtime/Experimental/GeometryCollectionEngine/Private/GeometryCollection/GeometryCollectionComponent.cpp b/Engine/Source/Runtime/Experimental/GeometryCollectionEngine/Private/GeometryCollection/GeometryCollectionComponent.cpp
--- a/Engine/Source/Runtime/Experimental/GeometryCollectionEngine/Private/GeometryCollection/GeometryCollectionComponent.cpp
+++ b/Engine/Source/Runtime/Experimental/GeometryCollectionEngine/Private/GeometryCollection/GeometryCollectionComponent.cpp
@@ -54,6 +54,7 @@
 #include "GeometryCollection/GeometryCollectionEngineRemoval.h"
 #include "GeometryCollection/Facades/CollectionAnchoringFacade.h"
 #include "GeometryCollection/Facades/CollectionRemoveOnBreakFacade.h"
+#include "GeometryCollection/Facades/CollectionInstancedMeshFacade.h"
 
 #include UE_INLINE_GENERATED_CPP_BY_NAME(GeometryCollectionComponent)
 
@@ -3973,13 +3974,15 @@ void UGeometryCollectionComponent::RegisterToISMPool()
 						// fisrt count the instance per mesh 
 						TArray<int32> InstanceCounts;
 						InstanceCounts.AddZeroed(RestCollection->AutoInstanceMeshes.Num());
-						const TManagedArray<int32>* AutoInstanceMeshIndices = RestCollection->GetGeometryCollection()->FindAttribute<int32>("AutoInstanceMeshIndex", FGeometryCollection::TransformGroup);
+
 						const TManagedArray<TSet<int32>>& Children = RestCollection->GetGeometryCollection()->Children;
-						if (AutoInstanceMeshIndices)
+
+						const GeometryCollection::Facades::FCollectionInstancedMeshFacade InstancedMeshFacade(*RestCollection->GetGeometryCollection());
+						if (InstancedMeshFacade.IsValid())
 						{
-							for (int32 TransformIndex = 0; TransformIndex < AutoInstanceMeshIndices->Num(); TransformIndex++)
+							for (int32 TransformIndex = 0; TransformIndex < InstancedMeshFacade.GetNumIndices(); TransformIndex++)
 							{
-								const int32 AutoInstanceMeshIndex = (*AutoInstanceMeshIndices)[TransformIndex];
+								const int32 AutoInstanceMeshIndex = InstancedMeshFacade.GetIndex(TransformIndex);
 								if (Children[TransformIndex].Num() == 0)
 								{
 									InstanceCounts[AutoInstanceMeshIndex]++;
@@ -4066,12 +4069,12 @@ void UGeometryCollectionComponent::RefreshISMPoolInstances()
 				{
 					if (RestCollection->GetGeometryCollection())
 					{
-						const TManagedArray<int32>* AutoInstanceMeshIndices = RestCollection->GetGeometryCollection()->FindAttribute<int32>("AutoInstanceMeshIndex", FGeometryCollection::TransformGroup);
 						const TManagedArray<TSet<int32>>& Children = RestCollection->GetGeometryCollection()->Children;
-						if (AutoInstanceMeshIndices)
+
+						const GeometryCollection::Facades::FCollectionInstancedMeshFacade InstancedMeshFacade(*RestCollection->GetGeometryCollection());
+						if (InstancedMeshFacade.IsValid())
 						{
 							const int32 NumTransforms = RestCollection->NumElements(FGeometryCollection::TransformAttribute);
-							ensure(AutoInstanceMeshIndices->Num() == NumTransforms);
 
 							CalculateGlobalMatrices();
 
@@ -4111,7 +4114,7 @@ void UGeometryCollectionComponent::RefreshISMPoolInstances()
 									InstanceTransforms.Reset(NumTransforms); // Allocate for worst case
 									for (int32 TransformIndex = 0; TransformIndex < NumTransforms; TransformIndex++)
 									{
-										const int32 AutoInstanceMeshIndex = (*AutoInstanceMeshIndices)[TransformIndex];
+										const int32 AutoInstanceMeshIndex = InstancedMeshFacade.GetIndex(TransformIndex);
 										if (AutoInstanceMeshIndex == MeshIndex && Children[TransformIndex].Num() == 0)
 										{
 											InstanceTransforms.Add(FTransform(GlobalMatrices[TransformIndex]) * ComponentTransform);
diff --git a/Engine/Source/Runtime/Experimental/GeometryCollectionEngine/Private/GeometryCollection/GeometryCollectionEngineConversion.cpp b/Engine/Source/Runtime/Experimental/GeometryCollectionEngine/Private/GeometryCollection/GeometryCollectionEngineConversion.cpp
--- a/Engine/Source/Runtime/Experimental/GeometryCollectionEngine/Private/GeometryCollection/GeometryCollectionEngineConversion.cpp
+++ b/Engine/Source/Runtime/Experimental/GeometryCollectionEngine/Private/GeometryCollection/GeometryCollectionEngineConversion.cpp
@@ -30,6 +30,7 @@
 #include "PhysicsEngine/BodySetup.h"
 #include "MeshDescriptionBuilder.h"
 #include "VertexConnectedComponents.h"
+#include "GeometryCollection/Facades/CollectionInstancedMeshFacade.h"
 
 
 DEFINE_LOG_CATEGORY_STATIC(UGeometryCollectionConversionLogging, Log, All);
@@ -457,7 +458,30 @@ int32 FGeometryCollectionEngineConversion::AppendMaterials(const TArray<UMateria
 	return MaterialStart;
 }
 
-bool FGeometryCollectionEngineConversion::AppendStaticMesh(const UStaticMesh* StaticMesh, const TArray<UMaterialInterface*>& Materials, 
+void FGeometryCollectionEngineConversion::AppendAutoInstanceMeshIndices(UGeometryCollection* GeometryCollectionObject, int32 FromTransformIndex, const UStaticMesh* StaticMesh, const TArray<UMaterialInterface*>& Materials)
+{
+	TSharedPtr<FGeometryCollection, ESPMode::ThreadSafe> GeometryCollectionPtr = GeometryCollectionObject->GetGeometryCollection();
+	if (GeometryCollectionPtr)
+	{
+		using namespace GeometryCollection::Facades;
+		FCollectionInstancedMeshFacade InstancedMeshFacade(*GeometryCollectionPtr);
+
+		const int32 NewNumOfTransforms = GeometryCollectionPtr->NumElements(FGeometryCollection::TransformGroup);
+		if (NewNumOfTransforms > FromTransformIndex)
+		{
+			// create the schema if necessary
+			InstancedMeshFacade.DefineSchema();
+	
+			const int32 AutoInstanceMeshIndex = GeometryCollectionObject->FindOrAddAutoInstanceMesh(*StaticMesh, Materials);
+			for (int32 TransformIndex = FromTransformIndex; TransformIndex < NewNumOfTransforms; TransformIndex++)
+			{
+				InstancedMeshFacade.SetIndex(TransformIndex, AutoInstanceMeshIndex);
+			}
+		}
+	}
+}
+
+bool FGeometryCollectionEngineConversion::AppendStaticMesh(const UStaticMesh* StaticMesh, const TArray<UMaterialInterface*>& Materials,
 	const FTransform& StaticMeshTransform, UGeometryCollection* GeometryCollectionObject, bool bReindexMaterials,
 	bool bAddInternalMaterials, bool bSplitComponents)
 {
@@ -476,18 +500,7 @@ bool FGeometryCollectionEngineConversion::AppendStaticMesh(const UStaticMesh* St
 	{
 		AppendMaterials(Materials, GeometryCollectionObject, bAddInternalMaterials);
 
-		// add index to the auto instanced meshes array
-		const int32 NewNumOfTransforms = GeometryCollection->NumElements(FGeometryCollection::TransformGroup);
-		if (NewNumOfTransforms > OriginalNumOfTransforms)
-		{
-			TManagedArray<int32>& AutoInstanceMeshIndices = GeometryCollection->AddAttribute<int32>("AutoInstanceMeshIndex", FGeometryCollection::TransformGroup);
-			const int32 AutoInstanceMeshIndex = GeometryCollectionObject->FindOrAddAutoInstanceMesh(*StaticMesh, Materials);
-
-			for (int32 TransformIndex = OriginalNumOfTransforms; TransformIndex < NewNumOfTransforms; TransformIndex++)
-			{
-				AutoInstanceMeshIndices[TransformIndex] = AutoInstanceMeshIndex;
-			}
-		}
+		AppendAutoInstanceMeshIndices(GeometryCollectionObject, OriginalNumOfTransforms, StaticMesh, Materials);
 
 		return true;
 	}
@@ -960,26 +973,30 @@ void FGeometryCollectionEngineConversion::AppendGeometryCollectionInstancedMeshe
 	TSharedPtr<const FGeometryCollection, ESPMode::ThreadSafe> SourceGeometryCollectionPtr = SourceGeometryCollectionObject->GetGeometryCollection();
 	TSharedPtr<FGeometryCollection, ESPMode::ThreadSafe> TargetGeometryCollectionPtr = TargetGeometryCollectionObject->GetGeometryCollection();
 
+	using namespace GeometryCollection::Facades;
+
 	if (SourceGeometryCollectionPtr && TargetGeometryCollectionPtr)
 	{
-		if (const TManagedArray<int32>* SourceAutoInstanceMeshIndices = SourceGeometryCollectionPtr->FindAttribute<int32>("AutoInstanceMeshIndex", FGeometryCollection::TransformGroup))
-		{
-			TManagedArray<int32>& TargetAutoInstanceMeshIndices = TargetGeometryCollectionPtr->AddAttribute<int32>("AutoInstanceMeshIndex", FGeometryCollection::TransformGroup);
+		const FCollectionInstancedMeshFacade SourceInstancedMeshFacade(*SourceGeometryCollectionPtr);
+		FCollectionInstancedMeshFacade TargetInstancedMeshFacade(*TargetGeometryCollectionPtr);
+
+		TargetInstancedMeshFacade.DefineSchema();
 
-			for (int32 SourceTransformIndex = 0; SourceTransformIndex < SourceAutoInstanceMeshIndices->Num(); SourceTransformIndex++)
+		const int32 NumSourceIndices = SourceInstancedMeshFacade.GetNumIndices();
+		for (int32 SourceTransformIndex = 0; SourceTransformIndex < NumSourceIndices; SourceTransformIndex++)
+		{
+			int32 TargetInstancedMeshIndex = INDEX_NONE;
+			if (SourceInstancedMeshFacade.IsValid())
 			{
-				const int32 TargettransformIndex = TargetTransformStartIndex + SourceTransformIndex;
-				TargetAutoInstanceMeshIndices[TargettransformIndex] = INDEX_NONE;
-				if (SourceAutoInstanceMeshIndices)
+				const int32 SourceAutoInstanceIndex = SourceInstancedMeshFacade.GetIndex(SourceTransformIndex);
+				if (SourceAutoInstanceIndex != INDEX_NONE)
 				{
-					const int32 SourceAutoInstanceIndex = (*SourceAutoInstanceMeshIndices)[SourceTransformIndex];
-					if (SourceAutoInstanceIndex != INDEX_NONE)
-					{
-						const FGeometryCollectionAutoInstanceMesh& SourceAutoInstanceMesh = SourceGeometryCollectionObject->GetAutoInstanceMesh(SourceAutoInstanceIndex);
-						TargetAutoInstanceMeshIndices[TargettransformIndex] = TargetGeometryCollectionObject->FindOrAddAutoInstanceMesh(SourceAutoInstanceMesh);
-					}
+					const FGeometryCollectionAutoInstanceMesh& SourceAutoInstanceMesh = SourceGeometryCollectionObject->GetAutoInstanceMesh(SourceAutoInstanceIndex);
+					TargetInstancedMeshIndex = TargetGeometryCollectionObject->FindOrAddAutoInstanceMesh(SourceAutoInstanceMesh);
 				}
 			}
+			const int32 TargetTransformIndex = TargetTransformStartIndex + SourceTransformIndex;
+			TargetInstancedMeshFacade.SetIndex(TargetTransformIndex, TargetInstancedMeshIndex);
 		}
 	}
 }
diff --git a/Engine/Source/Runtime/Experimental/GeometryCollectionEngine/Private/GeometryCollection/GeometryCollectionObject.cpp b/Engine/Source/Runtime/Experimental/GeometryCollectionEngine/Private/GeometryCollection/GeometryCollectionObject.cpp
--- a/Engine/Source/Runtime/Experimental/GeometryCollectionEngine/Private/GeometryCollection/GeometryCollectionObject.cpp
+++ b/Engine/Source/Runtime/Experimental/GeometryCollectionEngine/Private/GeometryCollection/GeometryCollectionObject.cpp
@@ -387,7 +387,7 @@ void UGeometryCollection::Reset()
 	}
 }
 
-void UGeometryCollection::ResetFrom(const FManagedArrayCollection& InCollection, const TArray<UMaterial*>& InMaterials)
+void UGeometryCollection::ResetFrom(const FManagedArrayCollection& InCollection, const TArray<UMaterial*>& InMaterials, bool bHasInternalMaterials)
 {
 	if (GeometryCollection.IsValid())
 	{
@@ -399,7 +399,7 @@ void UGeometryCollection::ResetFrom(const FManagedArrayCollection& InCollection,
 		UpdateConvexGeometry();
 				
 		Materials.Append(InMaterials);
-		InitializeMaterials();
+		InitializeMaterials(bHasInternalMaterials);
 	}
 }
 
@@ -479,82 +479,109 @@ void UGeometryCollection::ReindexMaterialSections()
 	InvalidateCollection();
 }
 
-void UGeometryCollection::InitializeMaterials()
+void UGeometryCollection::InitializeMaterials(bool bHasInternalMaterials)
 {
 	Modify();
 
 	// Last Material is the selection one
 	UMaterialInterface* BoneSelectedMaterial = LoadObject<UMaterialInterface>(nullptr, GetSelectedMaterialPath(), nullptr, LOAD_None, nullptr);
 
-	TManagedArray<int32>& MaterialID = GeometryCollection->MaterialID;
+	TManagedArray<int32>& MaterialIDs = GeometryCollection->MaterialID;
 
 	// normally we filter out instances of the selection material ID, but if it's actually used on any face we have to keep it
 	bool bBoneSelectedMaterialIsUsed = false;
-	for (int32 FaceIdx = 0; FaceIdx < MaterialID.Num(); ++FaceIdx)
+	for (int32 FaceIdx = 0; FaceIdx < MaterialIDs.Num(); ++FaceIdx)
 	{
-		int32 FaceMaterialID = MaterialID[FaceIdx];
+		int32 FaceMaterialID = MaterialIDs[FaceIdx];
 		if (FaceMaterialID < Materials.Num() && Materials[FaceMaterialID] == BoneSelectedMaterial)
 		{
 			bBoneSelectedMaterialIsUsed = true;
 			break;
 		}
 	}
-	
-	// We're assuming that all materials are arranged in pairs, so first we collect these.
-	using FMaterialPair = TPair<UMaterialInterface*, UMaterialInterface*>;
-	TSet<FMaterialPair> MaterialSet;
-	for (int32 MaterialIndex = 0; MaterialIndex < Materials.Num(); ++MaterialIndex)
+
+	TArray<UMaterialInterface*> FinalMaterials;
+	if (bHasInternalMaterials)
 	{
-		UMaterialInterface* ExteriorMaterial = Materials[MaterialIndex];
-		if (ExteriorMaterial == BoneSelectedMaterial && !bBoneSelectedMaterialIsUsed) // skip unused bone selected material
-		{
-			continue;
-		}
-		
-		// If we have an odd number of materials, the last material duplicates itself.
-		UMaterialInterface* InteriorMaterial = Materials[MaterialIndex];
-		while (++MaterialIndex < Materials.Num())
+		// We're assuming that all materials are arranged in pairs, so first we collect these.
+		using FMaterialPair = TPair<UMaterialInterface*, UMaterialInterface*>;
+		TSet<FMaterialPair> MaterialSet;
+		for (int32 MaterialIndex = 0; MaterialIndex < Materials.Num(); ++MaterialIndex)
 		{
-			if (Materials[MaterialIndex] == BoneSelectedMaterial && !bBoneSelectedMaterialIsUsed) // skip bone selected material
+			UMaterialInterface* ExteriorMaterial = Materials[MaterialIndex];
+			if (ExteriorMaterial == BoneSelectedMaterial && !bBoneSelectedMaterialIsUsed) // skip unused bone selected material
 			{
 				continue;
 			}
-			InteriorMaterial = Materials[MaterialIndex];
-			break;
+
+			// If we have an odd number of materials, the last material duplicates itself.
+			UMaterialInterface* InteriorMaterial = Materials[MaterialIndex];
+			while (++MaterialIndex < Materials.Num())
+			{
+				if (Materials[MaterialIndex] == BoneSelectedMaterial && !bBoneSelectedMaterialIsUsed) // skip bone selected material
+				{
+					continue;
+				}
+				InteriorMaterial = Materials[MaterialIndex];
+				break;
+			}
+
+			MaterialSet.Add(FMaterialPair(ExteriorMaterial, InteriorMaterial));
 		}
 
-		MaterialSet.Add(FMaterialPair(ExteriorMaterial, InteriorMaterial));
-	}
-	
-	// create the final material array only containing unique materials
-	// alternating exterior and interior materials
-	TMap<UMaterialInterface*, int32> ExteriorMaterialPtrToArrayIndex;
-	TMap<UMaterialInterface*, int32> InteriorMaterialPtrToArrayIndex;
-	TArray<UMaterialInterface*> FinalMaterials;
-	for (const FMaterialPair& Curr : MaterialSet)
-	{
-		// Add base material
-		TTuple< UMaterialInterface*, int32> BaseTuple(Curr.Key, FinalMaterials.Add(Curr.Key));
-		ExteriorMaterialPtrToArrayIndex.Add(BaseTuple);
+		// create the final material array only containing unique materials
+		// alternating exterior and interior materials
+		TMap<UMaterialInterface*, int32> ExteriorMaterialPtrToArrayIndex;
+		TMap<UMaterialInterface*, int32> InteriorMaterialPtrToArrayIndex;
+		for (const FMaterialPair& Curr : MaterialSet)
+		{
+			// Add base material
+			TTuple< UMaterialInterface*, int32> BaseTuple(Curr.Key, FinalMaterials.Add(Curr.Key));
+			ExteriorMaterialPtrToArrayIndex.Add(BaseTuple);
 
-		// Add interior material
-		TTuple< UMaterialInterface*, int32> InteriorTuple(Curr.Value, FinalMaterials.Add(Curr.Value));
-		InteriorMaterialPtrToArrayIndex.Add(InteriorTuple);
-	}
+			// Add interior material
+			TTuple< UMaterialInterface*, int32> InteriorTuple(Curr.Value, FinalMaterials.Add(Curr.Value));
+			InteriorMaterialPtrToArrayIndex.Add(InteriorTuple);
+		}
 
-	// Reassign material ID for each face given the new consolidated array of materials
-	for (int32 Material = 0; Material < MaterialID.Num(); ++Material)
+		// Reassign material ID for each face given the new consolidated array of materials
+		for (int32 Material = 0; Material < MaterialIDs.Num(); ++Material)
+		{
+			if (MaterialIDs[Material] < Materials.Num())
+			{
+				UMaterialInterface* OldMaterialPtr = Materials[MaterialIDs[Material]];
+				if (MaterialIDs[Material] % 2 == 0)
+				{
+					MaterialIDs[Material] = *ExteriorMaterialPtrToArrayIndex.Find(OldMaterialPtr);
+				}
+				else
+				{
+					MaterialIDs[Material] = *InteriorMaterialPtrToArrayIndex.Find(OldMaterialPtr);
+				}
+			}
+		}
+	}
+	else
 	{
-		if (MaterialID[Material] < Materials.Num())
+		// simple deduping process
+		for (int32 MaterialIndex = 0; MaterialIndex < Materials.Num(); ++MaterialIndex)
 		{
-			UMaterialInterface* OldMaterialPtr = Materials[MaterialID[Material]];
-			if (MaterialID[Material] % 2 == 0)
+			UMaterialInterface* Material = Materials[MaterialIndex];
+			if (Material == BoneSelectedMaterial && !bBoneSelectedMaterialIsUsed) // skip unused bone selected material
 			{
-				MaterialID[Material] = *ExteriorMaterialPtrToArrayIndex.Find(OldMaterialPtr);
+				continue;
 			}
-			else
+			FinalMaterials.AddUnique(Material);
+		}
+
+		// Reassign material ID for each face given the new consolidated array of materials
+		for (int32 MaterialIDIndex = 0; MaterialIDIndex < MaterialIDs.Num(); MaterialIDIndex++)
+		{
+			const int32 OldMaterialID = MaterialIDs[MaterialIDIndex];
+			if (Materials.IsValidIndex(OldMaterialID))
 			{
-				MaterialID[Material] = *InteriorMaterialPtrToArrayIndex.Find(OldMaterialPtr);
+				UMaterialInterface* Material = Materials[OldMaterialID];
+				MaterialIDs[MaterialIDIndex] = FinalMaterials.Find(Material);
 			}
 		}
 	}
@@ -1200,6 +1227,11 @@ void UGeometryCollection::RemoveExemplars(const TArray<int32>& SortedRemovalIndi
 	}
 }
 
+bool FGeometryCollectionAutoInstanceMesh::operator ==(const FGeometryCollectionAutoInstanceMesh& Other) const
+{
+	return (StaticMesh == Other.StaticMesh) && (Materials == Other.Materials);
+}
+
 /** find or add a auto instance mesh and return its index */
 const FGeometryCollectionAutoInstanceMesh& UGeometryCollection::GetAutoInstanceMesh(int32 AutoInstanceMeshIndex) const
 {
@@ -1209,62 +1241,15 @@ const FGeometryCollectionAutoInstanceMesh& UGeometryCollection::GetAutoInstanceM
 /**  find or add a auto instance mesh from another one and return its index */
 int32 UGeometryCollection::FindOrAddAutoInstanceMesh(const FGeometryCollectionAutoInstanceMesh& AutoInstanecMesh)
 {
-	int32 ReturnedIndex = INDEX_NONE;
-
-	for (int32 MeshIndex = 0; MeshIndex < AutoInstanceMeshes.Num(); MeshIndex++)
-	{
-		const FGeometryCollectionAutoInstanceMesh& Mesh = AutoInstanceMeshes[MeshIndex];
-		if (Mesh.StaticMesh == AutoInstanecMesh.StaticMesh && Mesh.Materials == AutoInstanecMesh.Materials)
-		{
-			ReturnedIndex = MeshIndex;
-			break;
-		}
-	}
-	if (ReturnedIndex == INDEX_NONE)
-	{
-		ReturnedIndex = AutoInstanceMeshes.Add(AutoInstanecMesh);
-	}
-	return ReturnedIndex;
+	return AutoInstanceMeshes.AddUnique(AutoInstanecMesh);
 }
 
 int32 UGeometryCollection::FindOrAddAutoInstanceMesh(const UStaticMesh& StaticMesh, const TArray<UMaterialInterface*>& MeshMaterials)
 {
-	int32 ReturnedIndex = INDEX_NONE;
-
-	FSoftObjectPath StaticMeshSoftPath(&StaticMesh);
-
-	for (int32 MeshIndex = 0; MeshIndex < AutoInstanceMeshes.Num(); MeshIndex++)
-	{
-		const FGeometryCollectionAutoInstanceMesh& Mesh = AutoInstanceMeshes[MeshIndex];
-		if (Mesh.StaticMesh == StaticMeshSoftPath)
-		{
-			if (Mesh.Materials.Num() == MeshMaterials.Num())
-			{
-				bool MaterialAreAllTheSame = true;
-				for (int32 MaterialIndex = 0; MaterialIndex < MeshMaterials.Num(); MaterialIndex++)
-				{
-					if (Mesh.Materials[MaterialIndex] != MeshMaterials[MaterialIndex])
-					{
-						MaterialAreAllTheSame = false;
-						break;
-					}
-				}
-				if (MaterialAreAllTheSame)
-				{
-					ReturnedIndex = MeshIndex;
-					break;
-				}
-			}
-		}
-	}
-	if (ReturnedIndex == INDEX_NONE)
-	{
-		FGeometryCollectionAutoInstanceMesh NewMesh;
-		NewMesh.StaticMesh = StaticMeshSoftPath;
-		NewMesh.Materials = MeshMaterials;
-		ReturnedIndex = AutoInstanceMeshes.Emplace(NewMesh);
-	}
-	return ReturnedIndex;
+	FGeometryCollectionAutoInstanceMesh NewMesh;
+	NewMesh.StaticMesh = FSoftObjectPath(&StaticMesh);
+	NewMesh.Materials = MeshMaterials;
+	return AutoInstanceMeshes.AddUnique(NewMesh);
 }
 
 FGuid UGeometryCollection::GetIdGuid() const
diff --git a/Engine/Source/Runtime/Experimental/GeometryCollectionEngine/Public/GeometryCollection/GeometryCollectionEngineConversion.h b/Engine/Source/Runtime/Experimental/GeometryCollectionEngine/Public/GeometryCollection/GeometryCollectionEngineConversion.h
--- a/Engine/Source/Runtime/Experimental/GeometryCollectionEngine/Public/GeometryCollection/GeometryCollectionEngineConversion.h
+++ b/Engine/Source/Runtime/Experimental/GeometryCollectionEngine/Public/GeometryCollection/GeometryCollectionEngineConversion.h
@@ -34,6 +34,15 @@ public:
 	 */
 	static int32 AppendMaterials(const TArray<UMaterialInterface*>& Materials, UGeometryCollection* GeometryCollectionObject, bool bAddInteriorCopy);
 
+	/**
+	 * Appends instanced mesh indices 
+	 * @param GeometryCollectionObject geometry collection to add to 
+	 * @param FromTransformIndex transform index to start from 
+	 * @param StaticMesh  static mesh to add reference to 
+	 * @param Materials materials correcponding to the static mesh instance to get the index from 
+	 */
+	static void AppendAutoInstanceMeshIndices(UGeometryCollection* GeometryCollectionObject, int32 FromTransformIndex, const UStaticMesh* StaticMesh, const TArray<UMaterialInterface*>& Materials);
+
 	/**
 	 * Appends a MeshDescription to a GeometryCollection.
 	 * @param MeshDescription : Const mesh description to read vertex/normals/index data from
diff --git a/Engine/Source/Runtime/Experimental/GeometryCollectionEngine/Public/GeometryCollection/GeometryCollectionObject.h b/Engine/Source/Runtime/Experimental/GeometryCollectionEngine/Public/GeometryCollection/GeometryCollectionObject.h
--- a/Engine/Source/Runtime/Experimental/GeometryCollectionEngine/Public/GeometryCollection/GeometryCollectionObject.h
+++ b/Engine/Source/Runtime/Experimental/GeometryCollectionEngine/Public/GeometryCollection/GeometryCollectionObject.h
@@ -56,6 +56,8 @@ struct GEOMETRYCOLLECTIONENGINE_API FGeometryCollectionAutoInstanceMesh
 
 	UPROPERTY(BlueprintReadWrite, EditAnywhere, Category = "AutoInstance")
 	TArray<TObjectPtr<UMaterialInterface>> Materials;
+
+	bool operator ==(const FGeometryCollectionAutoInstanceMesh& Other) const;
 };
 
 USTRUCT(BlueprintType)
@@ -359,7 +361,7 @@ public:
 	void Reset();
 
 	/** Reset the collection from another set of attributes and materials. */
-	void ResetFrom(const FManagedArrayCollection& InCollection, const TArray<UMaterial*>& InMaterials);
+	void ResetFrom(const FManagedArrayCollection& InCollection, const TArray<UMaterial*>& InMaterials, bool bHasInternalMaterials);
 	
 	int32 AppendGeometry(const UGeometryCollection & Element, bool ReindexAllMaterials = false, const FTransform& TransformRoot = FTransform::Identity);
 	int32 NumElements(const FName& Group) const;
@@ -398,7 +400,7 @@ public:
 	void ReindexMaterialSections();
 
 	/** appends the standard materials to this UObject */
-	void InitializeMaterials();
+	void InitializeMaterials(bool bHasInternalMaterials = true);
 
 
 	/** Returns true if there is anything to render */
