Chaos Cloth:
- Removed the clamping from the weight map code, since it is constraint dependent.
- Added a MaxStiffness argument to the SetValue functions in order to apply the correct clamping between PBD and XPBD (defaults to 1.0 for PBD).
- Added missing clamping when the constraints are constructed, since it could be the only place where they are setup.

#rnx

diff --git a/Engine/Plugins/Experimental/ChaosCloth/Source/Chaos/Private/ChaosCloth/ChaosClothConstraints.cpp b/Engine/Plugins/Experimental/ChaosCloth/Source/Chaos/Private/ChaosCloth/ChaosClothConstraints.cpp
--- a/Engine/Plugins/Experimental/ChaosCloth/Source/Chaos/Private/ChaosCloth/ChaosClothConstraints.cpp
+++ b/Engine/Plugins/Experimental/ChaosCloth/Source/Chaos/Private/ChaosCloth/ChaosClothConstraints.cpp
@@ -701,8 +701,7 @@ void FClothConstraints::SetLongRangeAttachmentProperties(const Softs::FSolverVec
 {
 	if (LongRangeConstraints)
 	{
-		LongRangeConstraints->SetStiffness(TetherStiffness);
-		LongRangeConstraints->SetScale(TetherScale);
+		LongRangeConstraints->SetProperties(TetherStiffness, TetherScale);
 	}
 }
 
diff --git a/Engine/Source/Runtime/Experimental/Chaos/Private/Chaos/PBDLongRangeConstraintsBase.cpp b/Engine/Source/Runtime/Experimental/Chaos/Private/Chaos/PBDLongRangeConstraintsBase.cpp
--- a/Engine/Source/Runtime/Experimental/Chaos/Private/Chaos/PBDLongRangeConstraintsBase.cpp
+++ b/Engine/Source/Runtime/Experimental/Chaos/Private/Chaos/PBDLongRangeConstraintsBase.cpp
@@ -32,9 +32,16 @@ FPBDLongRangeConstraintsBase::FPBDLongRangeConstraintsBase(
 	const TConstArrayView<FRealSingle>& StiffnessMultipliers,
 	const TConstArrayView<FRealSingle>& ScaleMultipliers,
 	const FSolverVec2& InStiffness,
-	const FSolverVec2& InScale)
+	const FSolverVec2& InScale,
+	FSolverReal MaxStiffness)
 	: Tethers(InTethers)
-	, Stiffness(InStiffness, StiffnessMultipliers, InParticleCount)
+	, Stiffness(
+		InStiffness,
+		StiffnessMultipliers,
+		InParticleCount,
+		FPBDStiffness::DefaultTableSize,
+		FPBDStiffness::DefaultParameterFitBase,
+		MaxStiffness)
 	, TetherScale(InScale, ScaleMultipliers, InParticleCount)
 	, ParticleOffset(InParticleOffset)
 {
diff --git a/Engine/Source/Runtime/Experimental/Chaos/Public/Chaos/PBDAnimDriveConstraint.h b/Engine/Source/Runtime/Experimental/Chaos/Public/Chaos/PBDAnimDriveConstraint.h
--- a/Engine/Source/Runtime/Experimental/Chaos/Public/Chaos/PBDAnimDriveConstraint.h
+++ b/Engine/Source/Runtime/Experimental/Chaos/Public/Chaos/PBDAnimDriveConstraint.h
@@ -71,8 +71,8 @@ namespace Chaos::Softs
 		// Set stiffness offset and range, as well as the simulation stiffness exponent
 		inline void ApplyProperties(const FSolverReal Dt, const int32 NumIterations)
 		{
-			Stiffness.ApplyValues(Dt, NumIterations);
-			Damping.ApplyValues(Dt, NumIterations);
+			Stiffness.ApplyPBDValues(Dt, NumIterations);
+			Damping.ApplyPBDValues(Dt, NumIterations);
 		}
 
 		inline void Apply(FSolverParticles& InParticles, const FSolverReal Dt) const
diff --git a/Engine/Source/Runtime/Experimental/Chaos/Public/Chaos/PBDAxialSpringConstraintsBase.h b/Engine/Source/Runtime/Experimental/Chaos/Public/Chaos/PBDAxialSpringConstraintsBase.h
--- a/Engine/Source/Runtime/Experimental/Chaos/Public/Chaos/PBDAxialSpringConstraintsBase.h
+++ b/Engine/Source/Runtime/Experimental/Chaos/Public/Chaos/PBDAxialSpringConstraintsBase.h
@@ -19,13 +19,22 @@ public:
 		const TArray<TVec3<int32>>& InConstraints,
 		const TConstArrayView<FRealSingle>& StiffnessMultipliers,
 		const FSolverVec2& InStiffness,
-		bool bTrimKinematicConstraints)
-		: Constraints(TrimConstraints(InConstraints, 
+		bool bTrimKinematicConstraints,
+		FSolverReal MaxStiffness = FPBDStiffness::DefaultPBDMaxStiffness)
+		: Constraints(TrimConstraints(InConstraints,
 			[&Particles, bTrimKinematicConstraints](int32 Index0, int32 Index1, int32 Index2)
 			{
 				return bTrimKinematicConstraints && Particles.InvM(Index0) == (FSolverReal)0. && Particles.InvM(Index1) == (FSolverReal)0. && Particles.InvM(Index2) == (FSolverReal)0.;
 			}))
-		, Stiffness(InStiffness, StiffnessMultipliers, TConstArrayView<TVec3<int32>>(Constraints), ParticleOffset, ParticleCount)
+		, Stiffness(
+			InStiffness,
+			StiffnessMultipliers,
+			TConstArrayView<TVec3<int32>>(Constraints),
+			ParticleOffset,
+			ParticleCount,
+			FPBDStiffness::DefaultTableSize,
+			FPBDStiffness::DefaultParameterFitBase,
+			MaxStiffness)
 	{
 		Init(Particles);
 	}
@@ -34,7 +43,7 @@ public:
 
 	void SetProperties(const FSolverVec2& InStiffness) { Stiffness.SetWeightedValue(InStiffness); }
 
-	void ApplyProperties(const FSolverReal Dt, const int32 NumIterations) { Stiffness.ApplyValues(Dt, NumIterations); }
+	void ApplyProperties(const FSolverReal Dt, const int32 NumIterations) { Stiffness.ApplyPBDValues(Dt, NumIterations); }
 
 protected:
 	inline FSolverVec3 GetDelta(const FSolverParticles& Particles, const int32 ConstraintIndex, const FSolverReal ExpStiffnessValue) const
diff --git a/Engine/Source/Runtime/Experimental/Chaos/Public/Chaos/PBDBendingConstraintsBase.h b/Engine/Source/Runtime/Experimental/Chaos/Public/Chaos/PBDBendingConstraintsBase.h
--- a/Engine/Source/Runtime/Experimental/Chaos/Public/Chaos/PBDBendingConstraintsBase.h
+++ b/Engine/Source/Runtime/Experimental/Chaos/Public/Chaos/PBDBendingConstraintsBase.h
@@ -27,20 +27,28 @@ public:
 		const TConstArrayView<FRealSingle>& StiffnessMultipliers,
 		const TConstArrayView<FRealSingle>& BucklingStiffnessMultipliers,
 		const FPropertyCollectionConstAdapter& PropertyCollection,
-		bool bTrimKinematicConstraints = false)
+		bool bTrimKinematicConstraints = false,
+		FSolverReal MaxStiffness = FPBDStiffness::DefaultPBDMaxStiffness)
 		: Constraints(bTrimKinematicConstraints ? TrimKinematicConstraints(InConstraints, InParticles): MoveTemp(InConstraints))
 		, ConstraintSharedEdges(ExtractConstraintSharedEdges(Constraints))
-		, Stiffness(FSolverVec2(GetWeightedFloatBendingStiffness(PropertyCollection, 1.f)).ClampAxes((FSolverReal)0., (FSolverReal)1.),
+		, Stiffness(FSolverVec2(GetWeightedFloatBendingStiffness(PropertyCollection, 1.f)),
 			StiffnessMultipliers,
 			TConstArrayView<TVec2<int32>>(ConstraintSharedEdges),
 			ParticleOffset,
-			ParticleCount)
-		, BucklingRatio(FMath::Clamp(GetBucklingRatio(PropertyCollection, 0.f), (FSolverReal)0., (FSolverReal)1.))
-		, BucklingStiffness(FSolverVec2(GetWeightedFloatBucklingStiffness(PropertyCollection, 1.f)).ClampAxes((FSolverReal)0., (FSolverReal)1.),
+			ParticleCount,
+			FPBDStiffness::DefaultTableSize,
+			FPBDStiffness::DefaultParameterFitBase,
+			MaxStiffness)
+		, BucklingRatio((FSolverReal)FMath::Clamp(GetBucklingRatio(PropertyCollection, 0.f), 0.f, 1.f))
+		, BucklingStiffness(
+			FSolverVec2(GetWeightedFloatBucklingStiffness(PropertyCollection, 1.f)),
 			BucklingStiffnessMultipliers,
 			TConstArrayView<TVec2<int32>>(ConstraintSharedEdges),
 			ParticleOffset,
-			ParticleCount)
+			ParticleCount,
+			FPBDStiffness::DefaultTableSize,
+			FPBDStiffness::DefaultParameterFitBase,
+			MaxStiffness)
 	{
 		for (const TVec4<int32>& Constraint : Constraints)
 		{
@@ -61,12 +69,29 @@ public:
 		const FSolverVec2& InStiffness,
 		const FSolverReal InBucklingRatio,
 		const FSolverVec2& InBucklingStiffness,
-		bool bTrimKinematicConstraints = false)
+		bool bTrimKinematicConstraints = false,
+		FSolverReal MaxStiffness = FPBDStiffness::DefaultPBDMaxStiffness)
 		: Constraints(bTrimKinematicConstraints ? TrimKinematicConstraints(InConstraints, InParticles): MoveTemp(InConstraints))
 		, ConstraintSharedEdges(ExtractConstraintSharedEdges(Constraints))
-		, Stiffness(InStiffness, StiffnessMultipliers, TConstArrayView<TVec2<int32>>(ConstraintSharedEdges), ParticleOffset, ParticleCount)
+		, Stiffness(
+			InStiffness,
+			StiffnessMultipliers,
+			TConstArrayView<TVec2<int32>>(ConstraintSharedEdges),
+			ParticleOffset,
+			ParticleCount,
+			FPBDStiffness::DefaultTableSize,
+			FPBDStiffness::DefaultParameterFitBase,
+			MaxStiffness)
 		, BucklingRatio(InBucklingRatio)
-		, BucklingStiffness(InBucklingStiffness, BucklingStiffnessMultipliers, TConstArrayView<TVec2<int32>>(ConstraintSharedEdges), ParticleOffset, ParticleCount)
+		, BucklingStiffness(
+			InBucklingStiffness,
+			BucklingStiffnessMultipliers,
+			TConstArrayView<TVec2<int32>>(ConstraintSharedEdges),
+			ParticleOffset,
+			ParticleCount,
+			FPBDStiffness::DefaultTableSize,
+			FPBDStiffness::DefaultParameterFitBase,
+			MaxStiffness)
 	{
 		for (const TVec4<int32>& Constraint : Constraints)
 		{
@@ -102,7 +127,7 @@ public:
 	{
 		if (IsBendingStiffnessMutable(PropertyCollection))
 		{
-			Stiffness.SetWeightedValue(FSolverVec2(GetWeightedFloatBendingStiffness(PropertyCollection)).ClampAxes((FSolverReal)0., (FSolverReal)1.));
+			Stiffness.SetWeightedValue(FSolverVec2(GetWeightedFloatBendingStiffness(PropertyCollection)));
 		}
 		if (IsBucklingRatioMutable(PropertyCollection))
 		{
@@ -110,20 +135,24 @@ public:
 		}
 		if (IsBucklingStiffnessMutable(PropertyCollection))
 		{
-			BucklingStiffness.SetWeightedValue(FSolverVec2(GetWeightedFloatBucklingStiffness(PropertyCollection)).ClampAxes((FSolverReal)0., (FSolverReal)1.));
+			BucklingStiffness.SetWeightedValue(FSolverVec2(GetWeightedFloatBucklingStiffness(PropertyCollection)));
 		}
 	}
 
 	// Update stiffness values
 	void SetProperties(const FSolverVec2& InStiffness, const FSolverReal InBucklingRatio, const FSolverVec2& InBucklingStiffness)
 	{ 
-		Stiffness.SetWeightedValue(InStiffness.ClampAxes((FSolverReal)0., (FSolverReal)1.));
+		Stiffness.SetWeightedValue(InStiffness);
 		BucklingRatio = FMath::Clamp(InBucklingRatio, (FSolverReal)0., (FSolverReal)1.);
-		BucklingStiffness.SetWeightedValue(InBucklingStiffness.ClampAxes((FSolverReal)0., (FSolverReal)1.));
+		BucklingStiffness.SetWeightedValue(InBucklingStiffness);
 	}
 
 	// Update stiffness table, as well as the simulation stiffness exponent
-	void ApplyProperties(const FSolverReal Dt, const int32 NumIterations) { Stiffness.ApplyValues(Dt, NumIterations); BucklingStiffness.ApplyValues(Dt, NumIterations); }
+	void ApplyProperties(const FSolverReal Dt, const int32 NumIterations)
+	{
+		Stiffness.ApplyPBDValues(Dt, NumIterations);
+		BucklingStiffness.ApplyPBDValues(Dt, NumIterations);
+	}
 
 	UE_DEPRECATED(5.1, "Use SetProperties instead.")
 	void SetStiffness(FSolverReal InStiffness) { SetProperties(FSolverVec2(InStiffness), 0.f, 1.f); }
diff --git a/Engine/Source/Runtime/Experimental/Chaos/Public/Chaos/PBDLongRangeConstraintsBase.h b/Engine/Source/Runtime/Experimental/Chaos/Public/Chaos/PBDLongRangeConstraintsBase.h
--- a/Engine/Source/Runtime/Experimental/Chaos/Public/Chaos/PBDLongRangeConstraintsBase.h
+++ b/Engine/Source/Runtime/Experimental/Chaos/Public/Chaos/PBDLongRangeConstraintsBase.h
@@ -15,6 +15,9 @@ namespace Chaos::Softs
 class CHAOS_API FPBDLongRangeConstraintsBase
 {
 public:
+	static constexpr FSolverReal MinTetherScale = (FSolverReal)0.01;
+	static constexpr FSolverReal MaxTetherScale = (FSolverReal)10.;
+
 	enum class EMode : uint8
 	{
 		Euclidean,
@@ -36,7 +39,8 @@ public:
 		const TConstArrayView<FRealSingle>& StiffnessMultipliers,
 		const TConstArrayView<FRealSingle>& ScaleMultipliers,
 		const FSolverVec2& InStiffness = FSolverVec2::UnitVector,
-		const FSolverVec2& InScale = FSolverVec2::UnitVector);
+		const FSolverVec2& InScale = FSolverVec2::UnitVector,
+		FSolverReal MaxStiffness = FPBDStiffness::DefaultPBDMaxStiffness);
 
 PRAGMA_DISABLE_DEPRECATION_WARNINGS  // For ScaleIndices and ScaleTable
 	virtual ~FPBDLongRangeConstraintsBase() {}
@@ -45,17 +49,25 @@ PRAGMA_ENABLE_DEPRECATION_WARNINGS
 	// Return the stiffness input values used by the constraint
 	FSolverVec2 GetStiffness() const { return Stiffness.GetWeightedValue(); }
 
+	// Set the stiffness and scale values used by the constraint
+	void SetProperties(const FSolverVec2& InStiffness, const FSolverVec2& InScale)
+	{
+		Stiffness.SetWeightedValue(InStiffness);
+		TetherScale.SetWeightedValue(InScale.ClampAxes(MinTetherScale, MaxTetherScale));
+	}
+
 	// Set the stiffness input values used by the constraint
+	UE_DEPRECATED(5.2, "Use SetProperties instead.")
 	void SetStiffness(const FSolverVec2& InStiffness) { Stiffness.SetWeightedValue(InStiffness); }
 
 	// Set the scale low and high value of the scale weight map
-	void SetScale(const FSolverVec2& InScale) { TetherScale.SetWeightedValueUnclamped(InScale.ClampAxes((FSolverReal)0.01, (FSolverReal)10.)); }
+	UE_DEPRECATED(5.2, "Use SetProperties instead.")
+	void SetScale(const FSolverVec2& InScale) { TetherScale.SetWeightedValue(InScale.ClampAxes(MinTetherScale, MaxTetherScale)); }
 
-public:
 	// Set stiffness offset and range, as well as the simulation stiffness exponent
 	void ApplyProperties(const FSolverReal Dt, const int32 NumIterations)
 	{
-		Stiffness.ApplyValues(Dt, NumIterations);
+		Stiffness.ApplyPBDValues(Dt, NumIterations);
 		TetherScale.ApplyValues();
 	}
 
diff --git a/Engine/Source/Runtime/Experimental/Chaos/Public/Chaos/PBDSpringConstraintsBase.h b/Engine/Source/Runtime/Experimental/Chaos/Public/Chaos/PBDSpringConstraintsBase.h
--- a/Engine/Source/Runtime/Experimental/Chaos/Public/Chaos/PBDSpringConstraintsBase.h
+++ b/Engine/Source/Runtime/Experimental/Chaos/Public/Chaos/PBDSpringConstraintsBase.h
@@ -15,7 +15,7 @@ namespace Chaos::Softs
 class FPBDSpringConstraintsBase
 {
 public:
-	template<int32 Valence>
+	template<int32 Valence, TEMPLATE_REQUIRES(Valence >= 2 && Valence <= 4)>
 	FPBDSpringConstraintsBase(
 		const FSolverParticles& Particles,
 		int32 ParticleOffset,
@@ -24,13 +24,21 @@ public:
 		const TConstArrayView<FRealSingle>& StiffnessMultipliers,
 		const FSolverVec2& InStiffness,
 		bool bTrimKinematicConstraints = false,
-		typename TEnableIf<Valence >= 2 && Valence <= 4>::Type* = nullptr)
+		FSolverReal MaxStiffness = FPBDStiffness::DefaultPBDMaxStiffness)
 		: Constraints(TrimConstraints(InConstraints, 
 			[&Particles, bTrimKinematicConstraints](int32 Index0, int32 Index1)
 			{
 				return bTrimKinematicConstraints && Particles.InvM(Index0) == (FSolverReal)0. && Particles.InvM(Index1) == (FSolverReal)0.;
 			}))
-		, Stiffness(InStiffness, StiffnessMultipliers, TConstArrayView<TVec2<int32>>(Constraints), ParticleOffset, ParticleCount)
+		, Stiffness(
+			InStiffness,
+			StiffnessMultipliers,
+			TConstArrayView<TVec2<int32>>(Constraints),
+			ParticleOffset,
+			ParticleCount,
+			FPBDStiffness::DefaultTableSize,
+			FPBDStiffness::DefaultParameterFitBase,
+			MaxStiffness)
 	{
 		// Update distances
 		Dists.Reset(Constraints.Num());
@@ -46,10 +54,10 @@ public:
 	{}
 
 	// Update stiffness values
-	void SetProperties(const FSolverVec2& InStiffness) { Stiffness.SetWeightedValue(InStiffness.ClampAxes((FSolverReal)0., (FSolverReal)1.)); }
+	void SetProperties(const FSolverVec2& InStiffness) { Stiffness.SetWeightedValue(InStiffness); }
 
 	// Update stiffness table, as well as the simulation stiffness exponent
-	inline void ApplyProperties(const FSolverReal Dt, const int32 NumIterations) { Stiffness.ApplyValues(Dt, NumIterations); }
+	inline void ApplyProperties(const FSolverReal Dt, const int32 NumIterations) { Stiffness.ApplyPBDValues(Dt, NumIterations); }
 
 	const TArray<TVec2<int32>>& GetConstraints() const { return Constraints; }
 
diff --git a/Engine/Source/Runtime/Experimental/Chaos/Public/Chaos/PBDStiffness.h b/Engine/Source/Runtime/Experimental/Chaos/Public/Chaos/PBDStiffness.h
--- a/Engine/Source/Runtime/Experimental/Chaos/Public/Chaos/PBDStiffness.h
+++ b/Engine/Source/Runtime/Experimental/Chaos/Public/Chaos/PBDStiffness.h
@@ -21,6 +21,9 @@ class FPBDStiffness final : public FPBDWeightMap
 {
 public:
 	static constexpr FSolverReal ParameterFrequency = (FSolverReal)120.;  // 60Hz @ 2 iterations as a root for all stiffness values TODO: Make this a global solver parameter
+	static constexpr FSolverReal DefaultPBDMaxStiffness = (FSolverReal)1.;
+	static constexpr FSolverReal DefaultParameterFitBase = (FSolverReal)1.e3;
+	static constexpr int32 DefaultTableSize = 16;
 
 	/**
 	 * Weightmap particle constructor. 
@@ -29,27 +32,42 @@ public:
 		const FSolverVec2& InWeightedValue,
 		const TConstArrayView<FRealSingle>& Multipliers = TConstArrayView<FRealSingle>(),
 		int32 ParticleCount = 0,  // A value of 0 also disables the map
-		int32 TableSize = 16,  // Size of the lookup table, can't be more than 256 values, the larger the table the longer it takes to apply changes to the stiffness values
-		FSolverReal InParameterFitBase = (FSolverReal)1.e3);  // Logarithm base to use in the stiffness parameter fit function
+		int32 TableSize = DefaultTableSize,  // Size of the lookup table, can't be more than 256 values, the larger the table the longer it takes to apply changes to the stiffness values
+		FSolverReal InParameterFitBase = DefaultParameterFitBase,  // Logarithm base to use in the PBD stiffness parameter fit function
+		FSolverReal InMaxStiffness = DefaultPBDMaxStiffness);
 
 	/**
 	 * Weightmap constraint constructor. 
 	 */
-	template<int32 Valence>
+	template<int32 Valence, TEMPLATE_REQUIRES(Valence >= 2 && Valence <= 4)>
 	inline FPBDStiffness(
 		const FSolverVec2& InWeightedValue,
 		const TConstArrayView<FRealSingle>& Multipliers = TConstArrayView<FRealSingle>(),
 		const TConstArrayView<TVector<int32, Valence>>& Constraints = TConstArrayView<TVector<int32, Valence>>(),
 		int32 ParticleOffset = INDEX_NONE, // Constraints have usually a particle offset added to them compared to the weight maps that always starts at index 0
 		int32 ParticleCount = 0,  // A value of 0 also disables the map
-		int32 TableSize = 16,  // Size of the lookup table, can't be more than 256 values, the larger the table the longer it takes to apply changes to the stiffness values
-		FSolverReal InParameterFitBase = (FSolverReal)1.e3,  // Logarithm base to use in the stiffness parameter fit function
-		typename TEnableIf<Valence >= 2 && Valence <= 4>::Type* = nullptr);  // Prevents incorrect valence, the value is actually unused
+		int32 TableSize = DefaultTableSize,  // Size of the lookup table, can't be more than 256 values, the larger the table the longer it takes to apply changes to the stiffness values
+		FSolverReal InParameterFitBase = DefaultParameterFitBase,  // Logarithm base to use in the PBD stiffness parameter fit function
+		FSolverReal MaxStiffness = DefaultPBDMaxStiffness);
 
 	~FPBDStiffness() {}
 
+	/**
+	 * Set the low and high values of the weight map.
+	 * The weight map table only gets updated after ApplyValues is called.
+	 * Low and high values are clamped between [0, MaxStiffness]
+	 */
+	void SetWeightedValue(const FSolverVec2& InWeightedValue, FSolverReal MaxStiffness = DefaultPBDMaxStiffness)
+	{
+		FPBDWeightMap::SetWeightedValue(InWeightedValue.ClampAxes((FSolverReal)0., MaxStiffness));
+	}
+
 	/** Update the weight map table with the current simulation parameters. */
-	inline void ApplyValues(const FSolverReal Dt, const int32 NumIterations);
+	inline void ApplyPBDValues(const FSolverReal Dt, const int32 NumIterations);
+
+	/** Update the weight map table with the current simulation parameters. */
+	UE_DEPRECATED(5.2, "Use ApplyPBDValues() instead.")
+	void ApplyValues(const FSolverReal Dt, const int32 NumIterations) { ApplyPBDValues(Dt, NumIterations); }
 
 	/** Update the weight map table with the current simulation parameters. */
 	inline void ApplyXPBDValues(const FSolverReal MaxStiffnesss);
@@ -67,14 +85,16 @@ FPBDStiffness::FPBDStiffness(
 	const TConstArrayView<FRealSingle>& Multipliers,
 	int32 ParticleCount,
 	int32 TableSize,
-	FSolverReal InParameterFitBase)
-	: FPBDWeightMap(InWeightedValue, Multipliers, ParticleCount, TableSize)
+	FSolverReal InParameterFitBase,
+	FSolverReal MaxStiffness)
+	: FPBDWeightMap(FSolverVec2::ZeroVector, Multipliers, ParticleCount, TableSize)
 	, ParameterFitBase(InParameterFitBase)
 	, ParameterFitLogBase(FMath::Loge(InParameterFitBase))
 {
+	SetWeightedValue(InWeightedValue, MaxStiffness);
 }
 
-template<int32 Valence>
+template<int32 Valence, typename TEnableIf<Valence >= 2 && Valence <= 4, int>::type>
 FPBDStiffness::FPBDStiffness(
 	const FSolverVec2& InWeightedValue,
 	const TConstArrayView<FRealSingle>& Multipliers,
@@ -83,11 +103,12 @@ FPBDStiffness::FPBDStiffness(
 	int32 ParticleCount,
 	int32 TableSize,
 	FSolverReal InParameterFitBase,
-	typename TEnableIf<Valence >= 2 && Valence <= 4>::Type*)
-	: FPBDWeightMap(InWeightedValue, Multipliers, Constraints, ParticleOffset, ParticleCount, TableSize)
+	FSolverReal MaxStiffness)
+	: FPBDWeightMap(FSolverVec2::ZeroVector, Multipliers, Constraints, ParticleOffset, ParticleCount, TableSize)
 	, ParameterFitBase(InParameterFitBase)
 	, ParameterFitLogBase(FMath::Loge(InParameterFitBase))
 {
+	SetWeightedValue(InWeightedValue, MaxStiffness);
 }
 
 static inline FSolverReal CalcExponentialParameterFit(const FSolverReal ParameterFitBase, const FSolverReal ParameterFitLogBase, const FSolverReal InValue)
@@ -99,7 +120,7 @@ static inline FSolverReal CalcExponentialParameterFit(const FSolverReal Paramete
 	return (FMath::Exp(ParameterFitLogBase * InValue) - (FSolverReal)1.) / (ParameterFitBase - (FSolverReal)1.);
 }
 
-void FPBDStiffness::ApplyValues(const FSolverReal Dt, const int32 NumIterations)
+void FPBDStiffness::ApplyPBDValues(const FSolverReal Dt, const int32 NumIterations)
 {
 	SCOPE_CYCLE_COUNTER(STAT_PBD_StiffnessApplyValues);
 
@@ -136,13 +157,13 @@ void FPBDStiffness::ApplyValues(const FSolverReal Dt, const int32 NumIterations)
 		bIsDirty = true;
 	}
 
-	FPBDWeightMap::ApplyValues(SimulationValue);  // Note, this still needs to be called even when not dirty for dealing with any change in values
+	FPBDWeightMap::ApplyValues(SimulationValue);  // Note, this still needs to be called even when not dirty for dealing with any change in weights
 }
 
 void FPBDStiffness::ApplyXPBDValues(const FSolverReal MaxStiffness)
 {
 	SCOPE_CYCLE_COUNTER(STAT_PBD_StiffnessApplyValues);
-	
+
 	auto SimulationValue = [this, MaxStiffness](const FSolverReal InValue)->FSolverReal
 		{
 			// Do not apply exponential to XPBDStiffnesses. They are authored in terms of true stiffness values (e.g., kg/s^2), not exponential compliance
@@ -155,7 +176,7 @@ void FPBDStiffness::ApplyXPBDValues(const FSolverReal MaxStiffness)
 		bIsDirty = true;
 	}
 
-	FPBDWeightMap::ApplyValues(SimulationValue);  // Note, this still needs to be called even when not dirty for dealing with any change in values
+	FPBDWeightMap::ApplyValues(SimulationValue);  // Note, this still needs to be called even when not dirty for dealing with any change in weights
 }
 
 }  // End namespace Chaos::Softs
diff --git a/Engine/Source/Runtime/Experimental/Chaos/Public/Chaos/PBDWeightMap.h b/Engine/Source/Runtime/Experimental/Chaos/Public/Chaos/PBDWeightMap.h
--- a/Engine/Source/Runtime/Experimental/Chaos/Public/Chaos/PBDWeightMap.h
+++ b/Engine/Source/Runtime/Experimental/Chaos/Public/Chaos/PBDWeightMap.h
@@ -54,13 +54,14 @@ public:
 	 * The weight map table only gets updated after ApplyValues is called.
 	 * Low and high values are clamped between [0,1]
 	 */
-	void SetWeightedValue(const FSolverVec2& InWeightedValue) { WeightedValue = InWeightedValue.ClampAxes((FSolverReal)0., (FSolverReal)1.); bIsDirty = true; }
+	void SetWeightedValue(const FSolverVec2& InWeightedValue) { WeightedValue = InWeightedValue; bIsDirty = true; }
 
 	/**
 	 * Set the low and high values of the weight map.
 	 * The weight map table only gets updated after ApplyValues is called.
 	 * Low and high values are not clamped. Commonly used for XPBD Stiffness values which are not [0,1]
 	 */
+	UE_DEPRECATED(5.2, "Use SetWeightedValue.")
 	void SetWeightedValueUnclamped(const FSolverVec2& InWeightedValue) { WeightedValue = InWeightedValue; bIsDirty = true; }
 
 	/**
@@ -111,7 +112,7 @@ inline FPBDWeightMap::FPBDWeightMap(
 	const TConstArrayView<FRealSingle>& Multipliers,
 	int32 ParticleCount,
 	int32 TableSize)
-	: WeightedValue(InWeightedValue.ClampAxes((FSolverReal)0., (FSolverReal)1.))
+	: WeightedValue(InWeightedValue)
 {
 	check(TableSize > 0 && TableSize < 256);  // The Stiffness lookup table is restricted to uint8 sized indices
 
@@ -147,7 +148,7 @@ inline FPBDWeightMap::FPBDWeightMap(
 	int32 ParticleCount,
 	int32 TableSize,
 	typename TEnableIf<Valence >= 2 && Valence <= 4>::Type*)
-	: WeightedValue(InWeightedValue.ClampAxes((FSolverReal)0., (FSolverReal)1.))
+	: WeightedValue(InWeightedValue)
 {
 	check(TableSize > 0 && TableSize < 256);  // The Stiffness lookup table is restricted to uint8 sized indices
 
diff --git a/Engine/Source/Runtime/Experimental/Chaos/Public/Chaos/XPBDAxialSpringConstraints.h b/Engine/Source/Runtime/Experimental/Chaos/Public/Chaos/XPBDAxialSpringConstraints.h
--- a/Engine/Source/Runtime/Experimental/Chaos/Public/Chaos/XPBDAxialSpringConstraints.h
+++ b/Engine/Source/Runtime/Experimental/Chaos/Public/Chaos/XPBDAxialSpringConstraints.h
@@ -30,14 +30,24 @@ public:
 		const TConstArrayView<FRealSingle>& StiffnessMultipliers,
 		const FSolverVec2& InStiffness,
 		bool bTrimKinematicConstraints)
-		: Base(Particles, ParticleOffset, ParticleCount, InConstraints, StiffnessMultipliers, InStiffness, bTrimKinematicConstraints)
+		: Base(
+			Particles,
+			ParticleOffset,
+			ParticleCount,
+			InConstraints,
+			StiffnessMultipliers,
+			InStiffness,
+			bTrimKinematicConstraints,
+			XPBDAxialSpringMaxStiffness)
 	{
 		Lambdas.Init(0.f, Constraints.Num());
 	}
 
 	virtual ~FXPBDAxialSpringConstraints() override {}
 
-	void ApplyProperties(const FSolverReal Dt, const int32 NumIterations) { Stiffness.ApplyXPBDValues(XPBDAxialSpringMaxStiffness); }
+	void SetProperties(const FSolverVec2& InStiffness) { Stiffness.SetWeightedValue(InStiffness, XPBDAxialSpringMaxStiffness); }
+
+	void ApplyProperties(const FSolverReal /*Dt*/, const int32 /*NumIterations*/) { Stiffness.ApplyXPBDValues(XPBDAxialSpringMaxStiffness); }
 
 	void Init() const { for (FSolverReal& Lambda : Lambdas) { Lambda = (FSolverReal)0.; } }
 
diff --git a/Engine/Source/Runtime/Experimental/Chaos/Public/Chaos/XPBDBendingConstraints.h b/Engine/Source/Runtime/Experimental/Chaos/Public/Chaos/XPBDBendingConstraints.h
--- a/Engine/Source/Runtime/Experimental/Chaos/Public/Chaos/XPBDBendingConstraints.h
+++ b/Engine/Source/Runtime/Experimental/Chaos/Public/Chaos/XPBDBendingConstraints.h
@@ -10,7 +10,7 @@ namespace Chaos::Softs
 
 // Stiffness is in kg cm^2 / rad^2 s^2
 static const FSolverReal XPBDBendMinStiffness = (FSolverReal)1e-4; // Stiffness below this will be considered 0 since all of our calculations are actually based on 1 / stiffness.
-static const FSolverReal XPBDBendMaxStiffness = (FSolverReal)1e7; 
+static const FSolverReal XPBDBendMaxStiffness = (FSolverReal)1e7;
 
 class CHAOS_API FXPBDBendingConstraints final : public FPBDBendingConstraintsBase
 {
@@ -31,7 +31,18 @@ public:
 		const FSolverReal InBucklingRatio,
 		const FSolverVec2& InBucklingStiffness,
 		bool bTrimKinematicConstraints = false)
-		: Base(InParticles, ParticleOffset, ParticleCount, MoveTemp(InConstraints), StiffnessMultipliers, BucklingStiffnessMultipliers, InStiffness, InBucklingRatio, InBucklingStiffness, bTrimKinematicConstraints)
+		: Base(
+			InParticles,
+			ParticleOffset,
+			ParticleCount,
+			MoveTemp(InConstraints),
+			StiffnessMultipliers,
+			BucklingStiffnessMultipliers,
+			InStiffness,
+			InBucklingRatio,
+			InBucklingStiffness,
+			bTrimKinematicConstraints,
+			XPBDBendMaxStiffness)
 		, DampingRatio(FSolverVec2::ZeroVector)
 	{
 		Lambdas.Init((FSolverReal)0., Constraints.Num());
@@ -50,7 +61,18 @@ public:
 		const FSolverVec2& InBucklingStiffness,
 		const FSolverVec2& InDampingRatio,
 		bool bTrimKinematicConstraints = false)
-		: Base(InParticles, ParticleOffset, ParticleCount, MoveTemp(InConstraints), StiffnessMultipliers, BucklingStiffnessMultipliers, InStiffness, InBucklingRatio, InBucklingStiffness, bTrimKinematicConstraints)
+		: Base(
+			InParticles,
+			ParticleOffset,
+			ParticleCount,
+			MoveTemp(InConstraints),
+			StiffnessMultipliers,
+			BucklingStiffnessMultipliers,
+			InStiffness,
+			InBucklingRatio,
+			InBucklingStiffness,
+			bTrimKinematicConstraints,
+			XPBDBendMaxStiffness)
 		, DampingRatio(InDampingRatio, DampingMultipliers, TConstArrayView<TVec2<int32>>(ConstraintSharedEdges), ParticleOffset, ParticleCount)
 	{
 		Lambdas.Init((FSolverReal)0., Constraints.Num());
@@ -69,14 +91,18 @@ public:
 	// Update stiffness values
 	void SetProperties(const FSolverVec2& InStiffness, const FSolverReal InBucklingRatio, const FSolverVec2& InBucklingStiffness, const FSolverVec2& InDampingRatio = FSolverVec2::ZeroVector)
 	{
-		Stiffness.SetWeightedValueUnclamped(InStiffness);
+		Stiffness.SetWeightedValue(InStiffness, XPBDBendMaxStiffness);
 		BucklingRatio = InBucklingRatio;
-		BucklingStiffness.SetWeightedValueUnclamped(InBucklingStiffness);
-		DampingRatio.SetWeightedValueUnclamped(InDampingRatio);
+		BucklingStiffness.SetWeightedValue(InBucklingStiffness, XPBDBendMaxStiffness);
+		DampingRatio.SetWeightedValue(InDampingRatio);
 	}
 
 	// Update stiffness table, as well as the simulation stiffness exponent
-	void ApplyProperties(const FSolverReal Dt, const int32 NumIterations) { Stiffness.ApplyXPBDValues(XPBDBendMaxStiffness); BucklingStiffness.ApplyXPBDValues(XPBDBendMaxStiffness); DampingRatio.ApplyValues(); }
+	void ApplyProperties(const FSolverReal /*Dt*/, const int32 /*NumIterations*/)
+	{
+		Stiffness.ApplyXPBDValues(XPBDBendMaxStiffness);
+		BucklingStiffness.ApplyXPBDValues(XPBDBendMaxStiffness);
+		DampingRatio.ApplyValues(); }
 
 	void Apply(FSolverParticles& Particles, const FSolverReal Dt) const;
 
diff --git a/Engine/Source/Runtime/Experimental/Chaos/Public/Chaos/XPBDLongRangeConstraints.h b/Engine/Source/Runtime/Experimental/Chaos/Public/Chaos/XPBDLongRangeConstraints.h
--- a/Engine/Source/Runtime/Experimental/Chaos/Public/Chaos/XPBDLongRangeConstraints.h
+++ b/Engine/Source/Runtime/Experimental/Chaos/Public/Chaos/XPBDLongRangeConstraints.h
@@ -28,7 +28,16 @@ public:
 		const TConstArrayView<FRealSingle>& ScaleMultipliers,
 		const FSolverVec2& InStiffness = FSolverVec2::UnitVector,
 		const FSolverVec2& InScale = FSolverVec2::UnitVector)
-	    : FPBDLongRangeConstraintsBase(Particles, InParticleOffset, InParticleCount, InTethers, StiffnessMultipliers, ScaleMultipliers, InStiffness, InScale)
+	    : FPBDLongRangeConstraintsBase(
+			Particles,
+			InParticleOffset,
+			InParticleCount,
+			InTethers,
+			StiffnessMultipliers,
+			ScaleMultipliers,
+			InStiffness,
+			InScale,
+			XPBDLongRangeMaxStiffness)
 	{
 		NumTethers = 0;
 		for (const TConstArrayView<FTether>& TetherBatch : Tethers)
@@ -40,8 +49,15 @@ public:
 
 	virtual ~FXPBDLongRangeConstraints() override {}
 
+	// Set the stiffness and scale values used by the constraint
+	void SetProperties(const FSolverVec2& InStiffness, const FSolverVec2& InScale)
+	{
+		Stiffness.SetWeightedValue(InStiffness, XPBDLongRangeMaxStiffness);
+		TetherScale.SetWeightedValue(InScale.ClampAxes(Base::MinTetherScale, Base::MaxTetherScale));
+	}
+
 	// Set stiffness offset and range, as well as the simulation stiffness exponent
-	void ApplyProperties(const FSolverReal Dt, const int32 NumIterations)
+	void ApplyProperties(const FSolverReal /*Dt*/, const int32 /*NumIterations*/)
 	{
 		Stiffness.ApplyXPBDValues(XPBDLongRangeMaxStiffness);
 		TetherScale.ApplyValues();
diff --git a/Engine/Source/Runtime/Experimental/Chaos/Public/Chaos/XPBDSpringConstraints.h b/Engine/Source/Runtime/Experimental/Chaos/Public/Chaos/XPBDSpringConstraints.h
--- a/Engine/Source/Runtime/Experimental/Chaos/Public/Chaos/XPBDSpringConstraints.h
+++ b/Engine/Source/Runtime/Experimental/Chaos/Public/Chaos/XPBDSpringConstraints.h
@@ -10,7 +10,7 @@ namespace Chaos::Softs
 
 // Stiffness is in kg/s^2
 static const FSolverReal XPBDSpringMinStiffness = (FSolverReal)1e-4; // Stiffness below this will be considered 0 since all of our calculations are actually based on 1 / stiffness.
-static const FSolverReal XPBDSpringMaxStiffness = (FSolverReal)1e7; 
+static const FSolverReal XPBDSpringMaxStiffness = (FSolverReal)1e7;
 
 class CHAOS_API FXPBDSpringConstraints final : public FPBDSpringConstraintsBase
 {
@@ -20,7 +20,7 @@ class CHAOS_API FXPBDSpringConstraints final : public FPBDSpringConstraintsBase
 	using Base::Stiffness;
 
 public:
-	template<int32 Valence>
+	template<int32 Valence, TEMPLATE_REQUIRES(Valence >= 2 && Valence <= 4)>
 	FXPBDSpringConstraints(
 		const FSolverParticles& Particles,
 		int32 ParticleOffset,
@@ -28,15 +28,22 @@ public:
 		const TArray<TVector<int32, Valence>>& InConstraints,
 		const TConstArrayView<FRealSingle>& StiffnessMultipliers,
 		const FSolverVec2& InStiffness,
-		bool bTrimKinematicConstraints = false,
-		typename TEnableIf<Valence >= 2 && Valence <= 4>::Type* = nullptr)
-		: Base(Particles, ParticleOffset, ParticleCount, InConstraints, StiffnessMultipliers, InStiffness, bTrimKinematicConstraints)
+		bool bTrimKinematicConstraints = false)
+		: Base(
+			Particles,
+			ParticleOffset,
+			ParticleCount,
+			InConstraints,
+			StiffnessMultipliers,
+			InStiffness,
+			bTrimKinematicConstraints,
+			XPBDSpringMaxStiffness)
 		, DampingRatio(FSolverVec2::ZeroVector)
 	{
 		Lambdas.Init((FSolverReal)0., Constraints.Num());
 		InitColor(Particles);
 	}
-	template<int32 Valence>
+	template<int32 Valence, TEMPLATE_REQUIRES(Valence >= 2 && Valence <= 4)>
 	FXPBDSpringConstraints(
 		const FSolverParticles& Particles,
 		int32 ParticleOffset,
@@ -46,10 +53,22 @@ public:
 		const TConstArrayView<FRealSingle>& DampingMultipliers,
 		const FSolverVec2& InStiffness,
 		const FSolverVec2& InDampingRatio,
-		bool bTrimKinematicConstraints = false,
-		typename TEnableIf<Valence >= 2 && Valence <= 4>::Type* = nullptr)
-		: Base(Particles, ParticleOffset, ParticleCount, InConstraints, StiffnessMultipliers, InStiffness, bTrimKinematicConstraints)
-		, DampingRatio(InDampingRatio, DampingMultipliers, TConstArrayView<TVec2<int32>>(Constraints), ParticleOffset, ParticleCount)
+		bool bTrimKinematicConstraints = false)
+		: Base(
+			Particles,
+			ParticleOffset,
+			ParticleCount,
+			InConstraints,
+			StiffnessMultipliers,
+			InStiffness,
+			bTrimKinematicConstraints,
+			XPBDSpringMaxStiffness)
+		, DampingRatio(
+			InDampingRatio,
+			DampingMultipliers,
+			TConstArrayView<TVec2<int32>>(Constraints),
+			ParticleOffset,
+			ParticleCount)
 	{
 		Lambdas.Init((FSolverReal)0., Constraints.Num());
 		InitColor(Particles);
@@ -60,10 +79,10 @@ public:
 	void Init() const { for (FSolverReal& Lambda : Lambdas) { Lambda = (FSolverReal)0.; } }
 
 	// Update stiffness values
-	void SetProperties(const FSolverVec2& InStiffness, const FSolverVec2& InDampingRatio = (FSolverReal)0.f)
+	void SetProperties(const FSolverVec2& InStiffness, const FSolverVec2& InDampingRatio = FSolverVec2::ZeroVector)
 	{ 
-		Stiffness.SetWeightedValueUnclamped(InStiffness); 
-		DampingRatio.SetWeightedValueUnclamped(InDampingRatio);
+		Stiffness.SetWeightedValue(InStiffness, XPBDSpringMaxStiffness);
+		DampingRatio.SetWeightedValue(InDampingRatio);
 	}
 	
 	// Update stiffness table, as well as the simulation stiffness exponent
