Chaos Cloth - Fixed missing friction coefficient from CollisionSpringConstraint to Base constructor.

#rnx

diff --git a/Engine/Source/Runtime/Experimental/Chaos/Public/Chaos/PBDCollisionSpringConstraints.h b/Engine/Source/Runtime/Experimental/Chaos/Public/Chaos/PBDCollisionSpringConstraints.h
--- a/Engine/Source/Runtime/Experimental/Chaos/Public/Chaos/PBDCollisionSpringConstraints.h
+++ b/Engine/Source/Runtime/Experimental/Chaos/Public/Chaos/PBDCollisionSpringConstraints.h
@@ -23,7 +23,15 @@ public:
 		const FSolverReal InThickness = Base::BackCompatThickness,
 		const FSolverReal InStiffness = Base::BackCompatStiffness,
 		const FSolverReal InFrictionCoefficient = Base::BackCompatFrictionCoefficient)
-		: Base(InOffset, InNumParticles, InTriangleMesh, InRestPositions, MoveTemp(InDisabledCollisionElements), InThickness, InStiffness)
+		: Base(
+			InOffset,
+			InNumParticles,
+			InTriangleMesh,
+			InRestPositions,
+			MoveTemp(InDisabledCollisionElements),
+			InThickness,
+			InStiffness,
+			InFrictionCoefficient)
 	{}
 
 	virtual ~FPBDCollisionSpringConstraints() override {}
