Sequencer: Scriprting missing file from previous submit
#tivial

diff --git a/Engine/Source/Runtime/Experimental/Animation/Constraints/Public/TransformableHandle.h b/Engine/Source/Runtime/Experimental/Animation/Constraints/Public/TransformableHandle.h
--- a/Engine/Source/Runtime/Experimental/Animation/Constraints/Public/TransformableHandle.h
+++ b/Engine/Source/Runtime/Experimental/Animation/Constraints/Public/TransformableHandle.h
@@ -118,7 +118,7 @@ public:
 #endif
 
 	//possible bindingID
-	UPROPERTY(EditAnywhere, Category = "Binding")
+	UPROPERTY(BlueprintReadOnly, Category = "Binding")
 	FMovieSceneObjectBindingID ConstraintBindingID;
 protected:
 	FHandleModifiedEvent OnHandleModified;
