Iris - Use constant instead of hardcoded value for marking array dirty.

#rnx

diff --git a/Engine/Source/Runtime/Experimental/Iris/Core/Private/Iris/ReplicationState/IrisFastArraySerializer.cpp b/Engine/Source/Runtime/Experimental/Iris/Core/Private/Iris/ReplicationState/IrisFastArraySerializer.cpp
--- a/Engine/Source/Runtime/Experimental/Iris/Core/Private/Iris/ReplicationState/IrisFastArraySerializer.cpp
+++ b/Engine/Source/Runtime/Experimental/Iris/Core/Private/Iris/ReplicationState/IrisFastArraySerializer.cpp
@@ -94,7 +94,7 @@ void FIrisFastArraySerializerPrivateAccessor::MarkAllArrayItemsDirty(FIrisFastAr
 	checkSlow(Array.ReplicationStateHeader.IsBound());
 
 	FNetBitArrayView MemberChangeMask = UE::Net::Private::FIrisFastArraySerializerPrivateAccessor::GetChangeMask(Array);
-	if (!MemberChangeMask.GetBit(0))
+	if (!MemberChangeMask.GetBit(FIrisFastArraySerializer::IrisFastArrayPropertyBitIndex))
 	{
 		MarkNetObjectStateDirty(Array.ReplicationStateHeader);
 	}
@@ -118,8 +118,8 @@ void FIrisFastArraySerializerPrivateAccessor::MarkArrayDirty(FIrisFastArraySeria
 	// Dirty object unless already dirty, we only use the array bit for this purpose
 	if (!MemberChangeMask.GetBit(FIrisFastArraySerializer::IrisFastArrayPropertyBitIndex))
 	{
-		MarkNetObjectStateDirty(Array.ReplicationStateHeader);
 		MemberChangeMask.SetBit(FIrisFastArraySerializer::IrisFastArrayPropertyBitIndex);
+		MarkNetObjectStateDirty(Array.ReplicationStateHeader);
 	}
 }
 
@@ -133,10 +133,10 @@ void FIrisFastArraySerializerPrivateAccessor::MarkArrayItemDirty(FIrisFastArrayS
 	MemberChangeMask.SetBit((ItemIdx % FIrisFastArraySerializer::IrisFastArrayChangeMaskBits) + FIrisFastArraySerializer::IrisFastArrayChangeMaskBitOffset);
 
 	// Dirty object unless already dirty, we only use the array bit for this purpose
-	if (!MemberChangeMask.GetBit(0))
+	if (!MemberChangeMask.GetBit(FIrisFastArraySerializer::IrisFastArrayPropertyBitIndex))
 	{
+		MemberChangeMask.SetBit(FIrisFastArraySerializer::IrisFastArrayPropertyBitIndex);
 		MarkNetObjectStateDirty(Array.ReplicationStateHeader);
-		MemberChangeMask.SetBits(0, 1);
 	}
 }
 
