Iris
- Added CVar to control if we use State received from server or local state when doing onreps with previous state
- Improved backwards compatibility for OnReps by comparing received state with local state before calling OnRep (disabled by default)

diff --git a/Engine/Source/Runtime/Experimental/Iris/Core/Private/Iris/Core/IrisDebugging.cpp b/Engine/Source/Runtime/Experimental/Iris/Core/Private/Iris/Core/IrisDebugging.cpp
--- a/Engine/Source/Runtime/Experimental/Iris/Core/Private/Iris/Core/IrisDebugging.cpp
+++ b/Engine/Source/Runtime/Experimental/Iris/Core/Private/Iris/Core/IrisDebugging.cpp
@@ -69,6 +69,16 @@ bool BreakOnObjectName(UObject* Object)
 	return false;
 }
 
+bool FilterDebuggedObject(UObject* Object)
+{
+	if (IrisDebugHelperInternal::GIrisDebugName.IsEmpty() || GetNameSafe(Object).Contains(IrisDebugHelperInternal::GIrisDebugName))
+	{
+		return true;
+	}
+
+	return false;
+}
+
 bool BreakOnNetRefHandle(FNetRefHandle NetRefHandle)
 {
 	if (IrisDebugHelperInternal::GIrisDebugNetRefHandle.IsValid() && IrisDebugHelperInternal::GIrisDebugNetRefHandle == NetRefHandle)
diff --git a/Engine/Source/Runtime/Experimental/Iris/Core/Private/Iris/ReplicationState/PropertyReplicationState.cpp b/Engine/Source/Runtime/Experimental/Iris/Core/Private/Iris/ReplicationState/PropertyReplicationState.cpp
--- a/Engine/Source/Runtime/Experimental/Iris/Core/Private/Iris/ReplicationState/PropertyReplicationState.cpp
+++ b/Engine/Source/Runtime/Experimental/Iris/Core/Private/Iris/ReplicationState/PropertyReplicationState.cpp
@@ -11,6 +11,9 @@
 #include "Iris/Core/IrisLog.h"
 #include "UObject/PropertyPortFlags.h"
 #include "Containers/StringFwd.h"
+#include "Iris/Core/IrisDebugging.h"
+
+DEFINE_LOG_CATEGORY_STATIC(LogIrisRepNotify, Warning, All);
 
 namespace UE::Net
 {
@@ -300,7 +303,7 @@ bool FPropertyReplicationState::PollObjectReferences(const void* RESTRICT SrcSta
 	return IsDirty();
 }
 
-void FPropertyReplicationState::CallRepNotifies(void* RESTRICT DstData, const FPropertyReplicationState* PreviousState, bool bIsInit) const
+void FPropertyReplicationState::CallRepNotifies(void* RESTRICT DstData, const FCallRepNotifiesParameters& Params) const
 {
 	// $IRIS TODO: Rewrite this to iterate over change mask instead of iterating over all members and querying the mask
 	// Note, we need to use a NetBitStreamReader and the changemask descriptor since each member might have different number of bits.
@@ -334,16 +337,42 @@ void FPropertyReplicationState::CallRepNotifies(void* RESTRICT DstData, const FP
 				const uint32 ExternalMemberOffset = MemberDescriptors[MemberIt].ExternalMemberOffset;
 				
 				UObject* Object = reinterpret_cast<UObject*>(DstData);
-				const uint8* PrevValuePtr = PreviousState ? PreviousState->StateBuffer + ExternalMemberOffset : nullptr;
-				const uint8* ValuePtr = StateBuffer + ExternalMemberOffset;
+				const uint8* PrevValuePtr = Params.PreviousState ? Params.PreviousState->StateBuffer + ExternalMemberOffset : nullptr;
+				const uint8* ValuePtr = StateBuffer + ExternalMemberOffset;				
 
-				// If it is the first time we apply state data, we only call repnotify if the value differs from the default
-				if (!bIsInit || !Private::InternalCompareMember(Descriptor, MemberIt, ValuePtr, PrevValuePtr))
+				bool bShouldCallRepNotify = false;
+				if (Params.bOnlyCallIfDiffersFromLocal)
+				{
+					// We try to be backwards compatible and respect RepNotify_Always/RepNotify_Changed unless it is the intial state where we only will call the repnotify of the received value differs from the local one.
+					bShouldCallRepNotify = Params.bIsInit ? !Private::InternalCompareMember(Descriptor, MemberIt, ValuePtr, PrevValuePtr) : EnumHasAnyFlags(Descriptor->MemberTraitsDescriptors[MemberIt].Traits, EReplicationStateMemberTraits::HasRepNotifyAlways) || !Private::InternalCompareMember(Descriptor, MemberIt, ValuePtr, PrevValuePtr);
+				}
+				else
 				{
+					// Trust data from server and call RepNotify without doing additonal compare unless it is the inital state.
+					bShouldCallRepNotify = !Params.bIsInit || !Private::InternalCompareMember(Descriptor, MemberIt, ValuePtr, PrevValuePtr);
+				}
+
+				if (bShouldCallRepNotify)
+				{
+					// We only want to call RepNotify once for c-arrays
 					LastPropertyWithRepNotify = Property;
-					UE_LOG(LogIris, VeryVerbose, TEXT("Calling RepNotify. Object: %s, Function: %s"), *Object->GetFullName(), ToCStr(RepNotifyFunction->GetName()));
 					Object->ProcessEvent(const_cast<UFunction*>(RepNotifyFunction), const_cast<uint8*>(PrevValuePtr));
 				}
+
+#if !UE_BUILD_SHIPPING
+				if (UE_LOG_ACTIVE(LogIrisRepNotify, Verbose) && IrisDebugHelper::FilterDebuggedObject(Object))
+				{
+					const bool bIsRepNotifyAlways = EnumHasAnyFlags(Descriptor->MemberTraitsDescriptors[MemberIt].Traits, EReplicationStateMemberTraits::HasRepNotifyAlways);
+					if (bShouldCallRepNotify)
+					{
+						UE_LOG(LogIrisRepNotify, Verbose, TEXT("Calling RepNotify. Object: %s, Function: %s IsInit: %d IsRepAlways: %d"), *Object->GetFullName(), ToCStr(RepNotifyFunction->GetName()), Params.bIsInit, bIsRepNotifyAlways);
+					}
+					else
+					{
+						UE_LOG(LogIrisRepNotify, VeryVerbose, TEXT("Skipping RepNotify. Object: %s, Function: %s IsInit: %d IsRepAlways: %d"), *Object->GetFullName(), ToCStr(RepNotifyFunction->GetName()), Params.bIsInit, bIsRepNotifyAlways);
+					}					
+				}
+#endif
 			}
 		}
 	}
diff --git a/Engine/Source/Runtime/Experimental/Iris/Core/Private/Iris/ReplicationSystem/PropertyReplicationFragment.cpp b/Engine/Source/Runtime/Experimental/Iris/Core/Private/Iris/ReplicationSystem/PropertyReplicationFragment.cpp
--- a/Engine/Source/Runtime/Experimental/Iris/Core/Private/Iris/ReplicationSystem/PropertyReplicationFragment.cpp
+++ b/Engine/Source/Runtime/Experimental/Iris/Core/Private/Iris/ReplicationSystem/PropertyReplicationFragment.cpp
@@ -6,10 +6,19 @@
 #include "Iris/ReplicationState/ReplicationStateDescriptorBuilder.h"
 #include "Iris/Core/IrisProfiler.h"
 #include "Iris/Core/IrisLog.h"
+#include "HAL/IConsoleManager.h"
 
 namespace UE::Net
 {
 
+static bool bUsePrevReceivedStateForOnReps = true;
+static FAutoConsoleVariableRef CVarUsePrevReceivedStateForOnReps(
+		TEXT("net.Iris.UsePrevReceivedStateForOnReps"),
+		bUsePrevReceivedStateForOnReps,
+		TEXT("If true OnReps will use the previous received state when doing onreps and not do any compares, if set to false we will copy the local state and do a compare before issuing onreps"
+		));
+
+
 FPropertyReplicationFragment::FPropertyReplicationFragment(EReplicationFragmentTraits InTraits, UObject* InOwner, const FReplicationStateDescriptor* InDescriptor)
 : FReplicationFragment(InTraits)
 , ReplicationStateDescriptor(InDescriptor)
@@ -24,7 +33,8 @@ FPropertyReplicationFragment::FPropertyReplicationFragment(EReplicationFragmentT
 	{
 		if (EnumHasAnyFlags(InDescriptor->Traits, EReplicationStateTraits::HasRepNotifies))
 		{
-			if (EnumHasAnyFlags(InDescriptor->Traits, EReplicationStateTraits::KeepPreviousState))
+			// We need to store the previous state if the onreps require this information or we are not using the previous received state received for onreps as we need to store the local state before overwriting the current state
+			if (!bUsePrevReceivedStateForOnReps || EnumHasAnyFlags(InDescriptor->Traits, EReplicationStateTraits::KeepPreviousState))
 			{
 				PrevReplicationState = MakeUnique<FPropertyReplicationState>(InDescriptor);
 
@@ -82,10 +92,16 @@ void FPropertyReplicationFragment::CallRepNotifies(FReplicationStateApplyContext
 
 	if (PrevReplicationState || !Context.bIsInit)
 	{
-		ReceivedState.CallRepNotifies(Owner, PrevReplicationState.Get(), Context.bIsInit);
+		FPropertyReplicationState::FCallRepNotifiesParameters Params;
+		Params.PreviousState = PrevReplicationState.Get();
+		Params.bIsInit = Context.bIsInit;
+		Params.bOnlyCallIfDiffersFromLocal = !bUsePrevReceivedStateForOnReps;
+
+		ReceivedState.CallRepNotifies(Owner, Params);
 
 		// We keep a copy of the previous state for RepNotifies that need the value
-		if (PrevReplicationState)
+		// If we rely on received data for the onreps, we just copy the received state, otherwise we must store the local state before applying received data.
+		if (bUsePrevReceivedStateForOnReps && PrevReplicationState)
 		{
 			*PrevReplicationState = ReceivedState;
 		}
@@ -103,6 +119,13 @@ void FPropertyReplicationFragment::ApplyReplicatedState(FReplicationStateApplyCo
 {
 	IRIS_PROFILER_SCOPE(PropertyReplicationFragment_ApplyReplicatedState);
 
+	// If we do not rely on received data to issue rep notifies we need to store a copy of the local state before we apply the new received state.
+	if (!bUsePrevReceivedStateForOnReps && PrevReplicationState)
+	{
+		// This could use a partial poll instead only polling RepNotify properties.
+		PrevReplicationState->PollPropertyReplicationState(Owner);
+	}
+
 	// Create a wrapping property replication state, cheap as we are simply injecting the already constructed state
 	const FPropertyReplicationState ReceivedState(Context.Descriptor, Context.StateBufferData.ExternalStateBuffer);
 
diff --git a/Engine/Source/Runtime/Experimental/Iris/Core/Public/Iris/Core/IrisDebugging.h b/Engine/Source/Runtime/Experimental/Iris/Core/Public/Iris/Core/IrisDebugging.h
--- a/Engine/Source/Runtime/Experimental/Iris/Core/Public/Iris/Core/IrisDebugging.h
+++ b/Engine/Source/Runtime/Experimental/Iris/Core/Public/Iris/Core/IrisDebugging.h
@@ -22,13 +22,16 @@ namespace UE::Net::IrisDebugHelper
 IRISCORE_API uint64 Init();
 
 /** Trigger a breakpoint and return true if the object contains the current debug name */
-bool BreakOnObjectName(UObject* Object);
+IRISCORE_API bool BreakOnObjectName(UObject* Object);
 
 /** Trigger a breakpoint and return true if the NetRefHandle is the current debug NetRefHandle */
-bool BreakOnNetRefHandle(FNetRefHandle NetRefHandle);
+IRISCORE_API bool BreakOnNetRefHandle(FNetRefHandle NetRefHandle);
 
 /** Trigger a breakpoint and return true if the name contains the debug RPC string */
-bool BreakOnRPCName(FName RPCName);
+IRISCORE_API bool BreakOnRPCName(FName RPCName);
+
+/** Returns true if the object name contains the current debug name, will return true if no debug name is set */
+IRISCORE_API bool FilterDebuggedObject(UObject* Object);
 
 /** Output state data to StringBuilder for the specified Handle */
 void NetObjectStateToString(FStringBuilderBase& StringBuilder, FNetRefHandle RefHandle);
diff --git a/Engine/Source/Runtime/Experimental/Iris/Core/Public/Iris/ReplicationState/PropertyReplicationState.h b/Engine/Source/Runtime/Experimental/Iris/Core/Public/Iris/ReplicationState/PropertyReplicationState.h
--- a/Engine/Source/Runtime/Experimental/Iris/Core/Public/Iris/ReplicationState/PropertyReplicationState.h
+++ b/Engine/Source/Runtime/Experimental/Iris/Core/Public/Iris/ReplicationState/PropertyReplicationState.h
@@ -105,8 +105,22 @@ public:
 	 */
 	IRISCORE_API bool PollObjectReferences(const void* RESTRICT SrcData);
 
+	struct FCallRepNotifiesParameters
+	{
+		// Previous state if requested
+		const FPropertyReplicationState* PreviousState = nullptr;
+
+		// This is an init state
+		bool bIsInit = false;
+
+		// Only call repnotify if value differs from local value
+		bool bOnlyCallIfDiffersFromLocal = false;
+	};
+
+
 	/** Invoke repnotifies for all dirty members */
-	IRISCORE_API void CallRepNotifies(void* RESTRICT DstData, const FPropertyReplicationState* PreviousState, bool bIsInit) const;
+	IRISCORE_API void CallRepNotifies(void* RESTRICT DstData, const FCallRepNotifiesParameters& Params) const;
+	inline void CallRepNotifies(void* RESTRICT DstData, const FPropertyReplicationState* PreviousState, bool bIsInit) const;
 
 	/** Debug output state to FString */
 	IRISCORE_API FString ToString(bool bIncludeAll = true) const;
@@ -141,4 +155,13 @@ private:
 	uint32 bOwnState : 1;
 };
 
+void FPropertyReplicationState::CallRepNotifies(void* RESTRICT DstData, const FPropertyReplicationState* PreviousState, bool bIsInit) const
+{
+	FCallRepNotifiesParameters Params; 
+
+	Params.PreviousState = PreviousState;
+	Params.bIsInit = bIsInit;
+	CallRepNotifies(DstData, Params);
+}
+
 }
