Dataflow:
- Implemented FPruneInCollectionDataflowNode, FSetVisibilityInCollectionDataflowNode, FMergeInCollectionDataflowNode, FCollectionSetPivotDataflowNode,
                       FMakeSphereDataflowNode, FCollectionTransformSelectionInBoxDataflowNode, FCollectionTransformSelectionInSphereDataflowNode,
                       FCollectionTransformSelectionByFloatAttrDataflowNode, FCollectionTransformSelectionByIntAttrDataflowNode, FCollectionSelectionConvertDataflowNode,
                       FCollectionVertexSelectionCustomDataflowNode, FCollectionFaceSelectionCustomDataflowNode, FCollectionAddCustomAttributeDataflowNode
- Added CollectionTransformSelectionFacace
- Updated other facades
- Added FDataflowVertexSelection, FDataflowFaceSelection

diff --git a/Engine/Plugins/Experimental/Fracture/Fracture.uplugin b/Engine/Plugins/Experimental/Fracture/Fracture.uplugin
--- a/Engine/Plugins/Experimental/Fracture/Fracture.uplugin
+++ b/Engine/Plugins/Experimental/Fracture/Fracture.uplugin
@@ -22,5 +22,12 @@
 		}
 	],
 	"Plugins": [
+		{
+			"Name": "PlanarCut",
+			"Enabled": true
+		}
 	]
-}
\ No newline at end of file
+}
+
+
+
diff --git a/Engine/Plugins/Experimental/Fracture/Source/FractureEngine/FractureEngine.Build.cs b/Engine/Plugins/Experimental/Fracture/Source/FractureEngine/FractureEngine.Build.cs
--- a/Engine/Plugins/Experimental/Fracture/Source/FractureEngine/FractureEngine.Build.cs
+++ b/Engine/Plugins/Experimental/Fracture/Source/FractureEngine/FractureEngine.Build.cs
@@ -15,7 +15,8 @@ namespace UnrealBuildTool.Rules
 				{
 					"Chaos",
 					"Core",
-					"DataflowCore"
+					"DataflowCore",
+					"PlanarCut"
 				}
 				);
 
diff --git a/Engine/Plugins/Experimental/Fracture/Source/FractureEngine/Private/FractureEngineEdit.cpp b/Engine/Plugins/Experimental/Fracture/Source/FractureEngine/Private/FractureEngineEdit.cpp
new file mode 100644
--- /dev/null
+++ b/Engine/Plugins/Experimental/Fracture/Source/FractureEngine/Private/FractureEngineEdit.cpp
@@ -0,0 +1,91 @@
+// Copyright Epic Games, Inc. All Rights Reserved.
+
+#include "FractureEngineEdit.h"
+#include "GeometryCollection/GeometryCollectionClusteringUtility.h"
+#include "GeometryCollection/GeometryCollectionAlgo.h"
+#include "GeometryCollection/GeometryCollectionProximityUtility.h"
+#include "Algo/RemoveIf.h"
+#include "GeometryCollection/Facades/CollectionTransformSelectionFacade.h"
+
+#include "PlanarCut.h"
+
+
+ 
+void FFractureEngineEdit::DeleteBranch(FGeometryCollection& GeometryCollection, const TArray<int32>& InBoneSelection)
+{
+	const FManagedArrayCollection& InCollection = (const FManagedArrayCollection&)GeometryCollection;
+	GeometryCollection::Facades::FCollectionTransformSelectionFacade TransformSelectionFacade(InCollection);
+
+	TArray<int32> BoneIndicies = InBoneSelection;
+	TransformSelectionFacade.RemoveRootNodes(BoneIndicies);
+	TransformSelectionFacade.Sanitize(BoneIndicies);
+
+	TArray<int32> NodesForDeletion;
+	const TManagedArray<TSet<int32>>& Children = InCollection.GetAttribute<TSet<int32>>("Children", FGeometryCollection::TransformGroup);
+
+	for (int32 BoneIdx : BoneIndicies)
+	{
+		FGeometryCollectionClusteringUtility::RecursiveAddAllChildren(Children, BoneIdx, NodesForDeletion);
+	}
+
+	NodesForDeletion.Sort();
+	GeometryCollection.RemoveElements(FGeometryCollection::TransformGroup, NodesForDeletion);
+
+	FGeometryCollectionClusteringUtility::RemoveDanglingClusters(&GeometryCollection);
+
+	// Proximity is invalidated
+	if (GeometryCollection.HasAttribute("Proximity", FGeometryCollection::GeometryGroup))
+	{
+		GeometryCollection.RemoveAttribute("Proximity", FGeometryCollection::GeometryGroup);
+	}
+}
+
+
+void FFractureEngineEdit::SetVisibilityInCollection(FManagedArrayCollection& InCollection, const TArray<int32>& InBoneSelection, bool bVisible)
+{
+	GeometryCollection::Facades::FCollectionTransformSelectionFacade TransformSelectionFacade(InCollection);
+
+	const TManagedArray<int32>& TransformToGeometryIndex = InCollection.GetAttribute<int32>("TransformToGeometryIndex", FGeometryCollection::GeometryGroup);
+	const TManagedArray<int32>& FaceStart = InCollection.GetAttribute<int32>("FaceStart", FGeometryCollection::GeometryGroup);
+	const TManagedArray<int32>& FaceCount = InCollection.GetAttribute<int32>("FaceCount", FGeometryCollection::GeometryGroup);
+	TManagedArray<bool>* Visible = InCollection.FindAttribute<bool>("Visible", FGeometryCollection::GeometryGroup);
+
+	TArray<int32> BoneIndicies = InBoneSelection;
+	TransformSelectionFacade.ConvertSelectionToRigidNodes(BoneIndicies);
+
+	for (int32 Index : BoneIndicies)
+	{
+		// Iterate the faces in the geometry of this rigid node and set invisible.
+		if (TransformToGeometryIndex[Index] > INDEX_NONE)
+		{
+			int32 CurrFace = FaceStart[TransformToGeometryIndex[Index]];
+			for (int32 FaceOffset = 0; FaceOffset < FaceCount[TransformToGeometryIndex[Index]]; ++FaceOffset)
+			{
+				(*Visible)[CurrFace + FaceOffset] = bVisible;
+			}
+		}
+	}
+}
+
+
+void FFractureEngineEdit::Merge(FGeometryCollection& GeometryCollection, const TArray<int32>& InBoneSelection)
+{
+	const FManagedArrayCollection& InCollection = (const FManagedArrayCollection&)GeometryCollection;
+	GeometryCollection::Facades::FCollectionTransformSelectionFacade TransformSelectionFacade(InCollection);
+
+	TArray<int32> BoneIndicies = InBoneSelection;
+	TransformSelectionFacade.Sanitize(BoneIndicies);
+
+	const TArray<int32>& NodesForMerge = InBoneSelection;
+
+	constexpr bool bBooleanUnion = false;
+	MergeAllSelectedBones(GeometryCollection, NodesForMerge, bBooleanUnion);
+
+	// Proximity is invalidated
+	if (GeometryCollection.HasAttribute("Proximity", FGeometryCollection::GeometryGroup))
+	{
+		GeometryCollection.RemoveAttribute("Proximity", FGeometryCollection::GeometryGroup);
+	}
+}
+
+
diff --git a/Engine/Plugins/Experimental/Fracture/Source/FractureEngine/Public/FractureEngineEdit.h b/Engine/Plugins/Experimental/Fracture/Source/FractureEngine/Public/FractureEngineEdit.h
new file mode 100644
--- /dev/null
+++ b/Engine/Plugins/Experimental/Fracture/Source/FractureEngine/Public/FractureEngineEdit.h
@@ -0,0 +1,18 @@
+// Copyright Epic Games, Inc. All Rights Reserved.
+
+#pragma once
+
+#include "GeometryCollection/GeometryCollection.h"
+#include "Dataflow/DataflowSelection.h"
+
+class FRACTUREENGINE_API FFractureEngineEdit
+{
+public:
+
+	static void DeleteBranch(FGeometryCollection& GeometryCollection, const TArray<int32>& InBoneSelection);
+
+	static void SetVisibilityInCollection(FManagedArrayCollection& InCollection, const TArray<int32>& InBoneSelection, bool bVisible);
+
+	static void Merge(FGeometryCollection& GeometryCollection, const TArray<int32>& InBoneSelection);
+
+};
\ No newline at end of file
diff --git a/Engine/Plugins/Experimental/GeometryCollectionPlugin/Source/GeometryCollectionNodes/Private/Dataflow/GeometryCollectionEditNodes.cpp b/Engine/Plugins/Experimental/GeometryCollectionPlugin/Source/GeometryCollectionNodes/Private/Dataflow/GeometryCollectionEditNodes.cpp
new file mode 100644
--- /dev/null
+++ b/Engine/Plugins/Experimental/GeometryCollectionPlugin/Source/GeometryCollectionNodes/Private/Dataflow/GeometryCollectionEditNodes.cpp
@@ -0,0 +1,152 @@
+// Copyright Epic Games, Inc. All Rights Reserved.
+
+#include "Dataflow/GeometryCollectionEditNodes.h"
+#include "Dataflow/DataflowCore.h"
+
+#include "Engine/StaticMesh.h"
+#include "GeometryCollection/GeometryCollectionObject.h"
+#include "GeometryCollection/ManagedArrayCollection.h"
+#include "GeometryCollection/GeometryCollection.h"
+#include "GeometryCollection/GeometryCollectionEngineUtility.h"
+#include "GeometryCollection/GeometryCollectionEngineConversion.h"
+#include "Logging/LogMacros.h"
+#include "Templates/SharedPointer.h"
+#include "UObject/UnrealTypePrivate.h"
+#include "DynamicMeshToMeshDescription.h"
+#include "MeshDescriptionToDynamicMesh.h"
+#include "StaticMeshAttributes.h"
+#include "DynamicMeshEditor.h"
+#include "Operations/MeshBoolean.h"
+
+#include "EngineGlobals.h"
+#include "GeometryCollection/GeometryCollectionAlgo.h"
+#include "GeometryCollection/GeometryCollectionClusteringUtility.h"
+#include "GeometryCollection/GeometryCollectionConvexUtility.h"
+#include "Voronoi/Voronoi.h"
+#include "PlanarCut.h"
+#include "GeometryCollection/GeometryCollectionProximityUtility.h"
+#include "FractureEngineEdit.h"
+#include "FractureEngineSelection.h"
+
+#include UE_INLINE_GENERATED_CPP_BY_NAME(GeometryCollectionEditNodes)
+
+namespace Dataflow
+{
+
+	void GeometryCollectionEditNodes()
+	{
+		static const FLinearColor CDefaultNodeBodyTintColor = FLinearColor(0.f, 0.f, 0.f, 0.5f);
+
+		DATAFLOW_NODE_REGISTER_CREATION_FACTORY(FPruneInCollectionDataflowNode);
+		DATAFLOW_NODE_REGISTER_CREATION_FACTORY(FSetVisibilityInCollectionDataflowNode);
+		DATAFLOW_NODE_REGISTER_CREATION_FACTORY(FMergeInCollectionDataflowNode);
+
+		// GeometryCollection|Edit
+		DATAFLOW_NODE_REGISTER_CREATION_FACTORY_NODE_COLORS_BY_CATEGORY("GeometryCollection|Edit", FLinearColor(0.f, 1.f, 0.05f), CDefaultNodeBodyTintColor);
+	}
+}
+
+
+void FPruneInCollectionDataflowNode::Evaluate(Dataflow::FContext& Context, const FDataflowOutput* Out) const
+{
+	if (Out->IsA<FManagedArrayCollection>(&Collection))
+	{
+		const FManagedArrayCollection& InCollection = GetValue<FManagedArrayCollection>(Context, &Collection);
+		const FDataflowTransformSelection& InTransformSelection = GetValue<FDataflowTransformSelection>(Context, &TransformSelection);
+
+		const int32 NumTransforms = InCollection.NumElements(FGeometryCollection::TransformGroup);
+		if (InTransformSelection.Num() == NumTransforms)
+		{
+			if (InTransformSelection.AnySelected())
+			{
+				if (TUniquePtr<FGeometryCollection> GeomCollection = TUniquePtr<FGeometryCollection>(InCollection.NewCopy<FGeometryCollection>()))
+				{
+					TArray<int32> BoneIndices;
+					InTransformSelection.AsArray(BoneIndices);
+
+					FFractureEngineEdit::DeleteBranch(*GeomCollection, BoneIndices);
+
+					SetValue<FManagedArrayCollection>(Context, (const FManagedArrayCollection&)(*GeomCollection), &Collection);
+					return;
+				}
+			}
+		}
+		else
+		{
+			// ERROR: InTransformSelection's length doesn't much number of transforms in InCollection
+			FString ErrorStr = "TransformSelection's number of elements doesn't match number of transforms in Collection.";
+			UE_LOG(LogTemp, Error, TEXT("[Dataflow ERROR] %s"), *ErrorStr);
+		}
+
+		SetValue<FManagedArrayCollection>(Context, InCollection, &Collection);
+	}
+}
+
+
+void FSetVisibilityInCollectionDataflowNode::Evaluate(Dataflow::FContext& Context, const FDataflowOutput* Out) const
+{
+	if (Out->IsA<FManagedArrayCollection>(&Collection))
+	{
+		FManagedArrayCollection InCollection = GetValue<FManagedArrayCollection>(Context, &Collection);
+		const FDataflowTransformSelection& InTransformSelection = GetValue<FDataflowTransformSelection>(Context, &TransformSelection);
+
+		const int32 NumTransforms = InCollection.NumElements(FGeometryCollection::TransformGroup);
+		if (InTransformSelection.Num() == NumTransforms)
+		{
+			if (InTransformSelection.AnySelected())
+			{
+				TArray<int32> BoneIndices;
+				InTransformSelection.AsArray(BoneIndices);
+
+				FFractureEngineEdit::SetVisibilityInCollection(InCollection, BoneIndices, Visibility == EVisibiltyOptionsEnum::Dataflow_VisibilityOptions_Visible);
+
+				SetValue<FManagedArrayCollection>(Context, InCollection, &Collection);
+				return;
+			}
+		}
+		else
+		{
+			// ERROR: InTransformSelection's length doesn't much number of transforms in InCollection
+			FString ErrorStr = "TransformSelection's number of elements doesn't match number of transforms in Collection.";
+			UE_LOG(LogTemp, Error, TEXT("[Dataflow ERROR] %s"), *ErrorStr);
+		}
+
+		SetValue<FManagedArrayCollection>(Context, InCollection, &Collection);
+	}
+}
+
+
+void FMergeInCollectionDataflowNode::Evaluate(Dataflow::FContext& Context, const FDataflowOutput* Out) const
+{
+	if (Out->IsA<FManagedArrayCollection>(&Collection))
+	{
+		FManagedArrayCollection InCollection = GetValue<FManagedArrayCollection>(Context, &Collection);
+		const FDataflowTransformSelection& InTransformSelection = GetValue<FDataflowTransformSelection>(Context, &TransformSelection);
+
+		const int32 NumTransforms = InCollection.NumElements(FGeometryCollection::TransformGroup);
+		if (InTransformSelection.Num() == NumTransforms)
+		{
+			if (InTransformSelection.AnySelected())
+			{
+				if (TUniquePtr<FGeometryCollection> GeomCollection = TUniquePtr<FGeometryCollection>(InCollection.NewCopy<FGeometryCollection>()))
+				{
+					TArray<int32> BoneIndices;
+					InTransformSelection.AsArray(BoneIndices);
+
+					FFractureEngineEdit::Merge(*GeomCollection, BoneIndices);
+
+					SetValue<FManagedArrayCollection>(Context, (const FManagedArrayCollection&)(*GeomCollection), &Collection);
+					return;
+				}
+			}
+		}
+		else
+		{
+			// ERROR: InTransformSelection's length doesn't much number of transforms in InCollection
+			FString ErrorStr = "TransformSelection's number of elements doesn't match number of transforms in Collection.";
+			UE_LOG(LogTemp, Error, TEXT("[Dataflow ERROR] %s"), *ErrorStr);
+		}
+
+		SetValue<FManagedArrayCollection>(Context, InCollection, &Collection);
+	}
+}
diff --git a/Engine/Plugins/Experimental/GeometryCollectionPlugin/Source/GeometryCollectionNodes/Private/Dataflow/GeometryCollectionNodes.cpp b/Engine/Plugins/Experimental/GeometryCollectionPlugin/Source/GeometryCollectionNodes/Private/Dataflow/GeometryCollectionNodes.cpp
--- a/Engine/Plugins/Experimental/GeometryCollectionPlugin/Source/GeometryCollectionNodes/Private/Dataflow/GeometryCollectionNodes.cpp
+++ b/Engine/Plugins/Experimental/GeometryCollectionPlugin/Source/GeometryCollectionNodes/Private/Dataflow/GeometryCollectionNodes.cpp
@@ -32,6 +32,7 @@
 #include "GeometryCollection/Facades/CollectionBoundsFacade.h"
 #include "GeometryCollection/Facades/CollectionAnchoringFacade.h"
 #include "GeometryCollection/Facades/CollectionRemoveOnBreakFacade.h"
+#include "GeometryCollection/Facades/CollectionTransformFacade.h"
 
 #include UE_INLINE_GENERATED_CPP_BY_NAME(GeometryCollectionNodes)
 
@@ -56,6 +57,7 @@ namespace Dataflow
 		DATAFLOW_NODE_REGISTER_CREATION_FACTORY(FFloatToStringDataflowNode);
 		DATAFLOW_NODE_REGISTER_CREATION_FACTORY(FMakePointsDataflowNode);
 		DATAFLOW_NODE_REGISTER_CREATION_FACTORY(FMakeBoxDataflowNode);
+		DATAFLOW_NODE_REGISTER_CREATION_FACTORY(FMakeSphereDataflowNode);
 		DATAFLOW_NODE_REGISTER_CREATION_FACTORY(FMakeLiteralFloatDataflowNode);
 		DATAFLOW_NODE_REGISTER_CREATION_FACTORY(FMakeLiteralIntDataflowNode);
 		DATAFLOW_NODE_REGISTER_CREATION_FACTORY(FMakeLiteralBoolDataflowNode);
@@ -79,8 +81,8 @@ namespace Dataflow
 		DATAFLOW_NODE_REGISTER_CREATION_FACTORY(FMathConstantsDataflowNode);
 		DATAFLOW_NODE_REGISTER_CREATION_FACTORY(FGetArrayElementDataflowNode);
 		DATAFLOW_NODE_REGISTER_CREATION_FACTORY(FGetNumArrayElementsDataflowNode);
-		DATAFLOW_NODE_REGISTER_CREATION_FACTORY(FGetBoundingBoxesDataflowNode);
-		DATAFLOW_NODE_REGISTER_CREATION_FACTORY(FGetCentroidsDataflowNode);
+		DATAFLOW_NODE_REGISTER_CREATION_FACTORY(FGetBoundingBoxesFromCollectionDataflowNode);
+		DATAFLOW_NODE_REGISTER_CREATION_FACTORY(FGetCentroidsFromCollectionDataflowNode);
 		DATAFLOW_NODE_REGISTER_CREATION_FACTORY(FTransformDataflowNode);
 		DATAFLOW_NODE_REGISTER_CREATION_FACTORY(FCompareIntDataflowNode);
 		DATAFLOW_NODE_REGISTER_CREATION_FACTORY(FBranchDataflowNode);
@@ -88,6 +90,11 @@ namespace Dataflow
 		DATAFLOW_NODE_REGISTER_CREATION_FACTORY(FRemoveOnBreakDataflowNode);
 		DATAFLOW_NODE_REGISTER_CREATION_FACTORY(FSetAnchorStateDataflowNode);
 		DATAFLOW_NODE_REGISTER_CREATION_FACTORY(FProximityDataflowNode);
+		DATAFLOW_NODE_REGISTER_CREATION_FACTORY(FCollectionSetPivotDataflowNode);
+		DATAFLOW_NODE_REGISTER_CREATION_FACTORY(FAddCustomCollectionAttributeDataflowNode);
+		DATAFLOW_NODE_REGISTER_CREATION_FACTORY(FGetNumElementsInCollectionGroupDataflowNode);
+		DATAFLOW_NODE_REGISTER_CREATION_FACTORY(FGetCollectionAttributeDataTypedDataflowNode);
+		DATAFLOW_NODE_REGISTER_CREATION_FACTORY(FSetCollectionAttributeDataTypedDataflowNode);
 
 		// GeometryCollection
 		DATAFLOW_NODE_REGISTER_CREATION_FACTORY_NODE_COLORS_BY_CATEGORY("GeometryCollection", FLinearColor(0.55f, 0.45f, 1.0f), CDefaultNodeBodyTintColor);
@@ -273,9 +280,9 @@ void FBoundingBoxDataflowNode::Evaluate(Dataflow::FContext& Context, const FData
 		const FManagedArrayCollection& InCollection = GetValue<FManagedArrayCollection>(Context, &Collection);
 
 		GeometryCollection::Facades::FBoundsFacade BoundsFacade(InCollection);
-		const FBox& BBox = BoundsFacade.GetBoundingBox();
+		const FBox& BoundingBoxInCollectionSpace = BoundsFacade.GetBoundingBoxInCollectionSpace();
 
-		SetValue<FBox>(Context, BBox, &BoundingBox);
+		SetValue<FBox>(Context, BoundingBoxInCollectionSpace, &BoundingBox);
 	}
 }
 
@@ -354,6 +361,18 @@ void FMakeBoxDataflowNode::Evaluate(Dataflow::FContext& Context, const FDataflow
 }
 
 
+void FMakeSphereDataflowNode::Evaluate(Dataflow::FContext& Context, const FDataflowOutput* Out) const
+{
+	if (Out->IsA<FSphere>(&Sphere))
+	{
+		FVector CenterVal = GetValue<FVector>(Context, &Center);
+		float RadiusVal = GetValue<float>(Context, &Radius);
+
+		SetValue<FSphere>(Context, FSphere(CenterVal, RadiusVal), &Sphere);
+	}
+}
+
+
 void FMakeLiteralFloatDataflowNode::Evaluate(Dataflow::FContext& Context, const FDataflowOutput* Out) const
 {
 	if (Out->IsA<float>(&Float))
@@ -806,7 +825,7 @@ void FGetArrayElementDataflowNode::Evaluate(Dataflow::FContext& Context, const F
 	if (Out->IsA<FVector>(&Point))
 	{
 		const TArray<FVector>& Array = GetValue<TArray<FVector>>(Context, &Points);
-		if (Index >= 0 && Index < Array.Num())
+		if (Array.Num() > 0 && Index >= 0 && Index < Array.Num())
 		{
 			SetValue<FVector>(Context, Array[Index], &Point);
 		}
@@ -821,7 +840,7 @@ void FGetNumArrayElementsDataflowNode::Evaluate(Dataflow::FContext& Context, con
 	}
 }
 
-void FGetBoundingBoxesDataflowNode::Evaluate(Dataflow::FContext& Context, const FDataflowOutput* Out) const
+void FGetBoundingBoxesFromCollectionDataflowNode::Evaluate(Dataflow::FContext& Context, const FDataflowOutput* Out) const
 {
 	if (Out->IsA<TArray<FBox>>(&BoundingBoxes))
 	{
@@ -831,27 +850,36 @@ void FGetBoundingBoxesDataflowNode::Evaluate(Dataflow::FContext& Context, const
 		GeometryCollection::Facades::FBoundsFacade BoundsFacade(InCollection);
 		const TManagedArray<FBox>& InBoundingBoxes = BoundsFacade.GetBoundingBoxes();
 
+		GeometryCollection::Facades::FCollectionTransformFacade TransformFacade(InCollection);
+
 		TArray<FBox> BoundingBoxesArr;
 		for (int32 Idx = 0; Idx < InBoundingBoxes.Num(); ++Idx)
 		{
+			const FBox BoundingBoxInBoneSpace = InBoundingBoxes[Idx];
+
+			// Transform from BoneSpace to CollectionSpace
+			const FTransform CollectionSpaceTransform = TransformFacade.ComputeCollectionSpaceTransform(Idx);
+			const FBox BoundingBoxInCollectionSpace = BoundingBoxInBoneSpace.TransformBy(CollectionSpaceTransform);
+
 			if (IsConnected<FDataflowTransformSelection>(&TransformSelection))
 			{
 				if (InTransformSelection.IsSelected(Idx))
 				{
-					BoundingBoxesArr.Add(InBoundingBoxes[Idx]);
+					BoundingBoxesArr.Add(BoundingBoxInCollectionSpace);
 				}
 			}
 			else
 			{
-				BoundingBoxesArr.Add(InBoundingBoxes[Idx]);
+				BoundingBoxesArr.Add(BoundingBoxInCollectionSpace);
 			}
 
 		}
+
 		SetValue<TArray<FBox>>(Context, BoundingBoxesArr, &BoundingBoxes);
 	}
 }
 
-void FGetCentroidsDataflowNode::Evaluate(Dataflow::FContext& Context, const FDataflowOutput* Out) const
+void FGetCentroidsFromCollectionDataflowNode::Evaluate(Dataflow::FContext& Context, const FDataflowOutput* Out) const
 {
 	if (Out->IsA<TArray<FVector>>(&Centroids))
 	{
@@ -859,24 +887,32 @@ void FGetCentroidsDataflowNode::Evaluate(Dataflow::FContext& Context, const FDat
 		const FDataflowTransformSelection& InTransformSelection = GetValue<FDataflowTransformSelection>(Context, &TransformSelection);
 
 		GeometryCollection::Facades::FBoundsFacade BoundsFacade(InCollection);
-		const TManagedArray<FBox>& InBoundingBoxes = BoundsFacade.GetBoundingBoxes();
+		const TArray<FVector>& InCentroids = BoundsFacade.GetCentroids();
+
+		GeometryCollection::Facades::FCollectionTransformFacade TransformFacade(InCollection);
 
 		TArray<FVector> CentroidsArr;
-		for (int32 Idx = 0; Idx < InBoundingBoxes.Num(); ++Idx)
+		for (int32 Idx = 0; Idx < InCentroids.Num(); ++Idx)
 		{
-			const FBox& BoundingBox = InBoundingBoxes[Idx];
+			const FVector PositionInBoneSpace(InCentroids[Idx]);
+
+			// Transform from BoneSpace to CollectionSpace
+			const FTransform CollectionSpaceTransform = TransformFacade.ComputeCollectionSpaceTransform(Idx);
+			const FVector PositionInCollectionSpace = CollectionSpaceTransform.TransformPosition(PositionInBoneSpace);
+
 			if (IsConnected<FDataflowTransformSelection>(&TransformSelection))
 			{
 				if (InTransformSelection.IsSelected(Idx))
 				{
-					CentroidsArr.Add(BoundingBox.GetCenter());
+					CentroidsArr.Add(PositionInCollectionSpace);
 				}
 			}
 			else
 			{
-				CentroidsArr.Add(BoundingBox.GetCenter());
+				CentroidsArr.Add(PositionInCollectionSpace);
 			}
 		}
+
 		SetValue<TArray<FVector>>(Context, CentroidsArr, &Centroids);
 	}
 }
@@ -886,6 +922,7 @@ void FTransformDataflowNode::Evaluate(Dataflow::FContext& Context, const FDatafl
 {
 	if (Out->IsA<TObjectPtr<UDynamicMesh>>(&Mesh))
 	{
+		// TODO: Implement this
 		SetValue<TObjectPtr<UDynamicMesh>>(Context, NewObject<UDynamicMesh>(), &Mesh);
 	}
 }
@@ -1088,10 +1125,350 @@ void FProximityDataflowNode::Evaluate(Dataflow::FContext& Context, const FDatafl
 }
 
 
+void FCollectionSetPivotDataflowNode::Evaluate(Dataflow::FContext& Context, const FDataflowOutput* Out) const
+{
+	if (Out->IsA<FManagedArrayCollection>(&Collection))
+	{
+		FManagedArrayCollection InCollection = GetValue<FManagedArrayCollection>(Context, &Collection);
+		const FTransform& InTransform = GetValue<FTransform>(Context, &Transform);
+
+		GeometryCollection::Facades::FCollectionTransformFacade TransformFacade(InCollection);
+		TransformFacade.SetPivot(InTransform);
+
+		SetValue<FManagedArrayCollection>(Context, InCollection, &Collection);
+	}
+}
+
+
+static FName GetGroupName(const EStandardGroupNameEnum& InGroupName)
+{
+	FName GroupNameToUse;
+	if (InGroupName == EStandardGroupNameEnum::Dataflow_EStandardGroupNameEnum_Transform)
+	{
+		GroupNameToUse = FGeometryCollection::TransformGroup;
+	}
+	else if (InGroupName == EStandardGroupNameEnum::Dataflow_EStandardGroupNameEnum_Geometry)
+	{
+		GroupNameToUse = FGeometryCollection::GeometryGroup;
+	}
+	else if (InGroupName == EStandardGroupNameEnum::Dataflow_EStandardGroupNameEnum_Faces)
+	{
+		GroupNameToUse = FGeometryCollection::FacesGroup;
+	}
+	else if (InGroupName == EStandardGroupNameEnum::Dataflow_EStandardGroupNameEnum_Vertices)
+	{
+		GroupNameToUse = FGeometryCollection::VerticesGroup;
+	}
+	else if (InGroupName == EStandardGroupNameEnum::Dataflow_EStandardGroupNameEnum_Material)
+	{
+		GroupNameToUse = FGeometryCollection::MaterialGroup;
+	}
+	else if (InGroupName == EStandardGroupNameEnum::Dataflow_EStandardGroupNameEnum_Breaking)
+	{
+		GroupNameToUse = FGeometryCollection::BreakingGroup;
+	}
+
+	return GroupNameToUse;
+}
+
+void FAddCustomCollectionAttributeDataflowNode::Evaluate(Dataflow::FContext& Context, const FDataflowOutput* Out) const
+{
+	if (Out->IsA<FManagedArrayCollection>(&Collection))
+	{
+		FManagedArrayCollection InCollection = GetValue<FManagedArrayCollection>(Context, &Collection);
+		const int32 InNumElements = GetValue<int32>(Context, &NumElements);
+
+		FName GroupNameToUse;
+		if (GroupName != EStandardGroupNameEnum::Dataflow_EStandardGroupNameEnum_Custom)
+		{
+			GroupNameToUse = GetGroupName(GroupName);
+		}
+		else
+		{
+			GroupNameToUse = FName(*CustomGroupName);
+		}
+
+		if (GroupNameToUse.GetStringLength() > 0 && AttrName.Len() > 0)
+		{
+			if (CustomAttributeType == ECustomAttributeTypeEnum::Dataflow_CustomAttributeType_Bool)
+			{
+				TManagedArrayAccessor<bool> CustomAttribute(InCollection, FName(*AttrName), GroupNameToUse);
+
+				// If the group already exists don't change the number of elements
+				if (!InCollection.HasGroup(GroupNameToUse))
+				{
+					CustomAttribute.AddElements(InNumElements);
+				}
+
+				CustomAttribute.AddAndFill(false);
+			}
+			else if (CustomAttributeType == ECustomAttributeTypeEnum::Dataflow_CustomAttributeType_Float)
+			{
+				TManagedArrayAccessor<float> CustomAttribute(InCollection, FName(*AttrName), GroupNameToUse);
+
+				// If the group already exists don't change the number of elements
+				if (!InCollection.HasGroup(GroupNameToUse))
+				{
+					CustomAttribute.AddElements(InNumElements);
+				}
+
+				CustomAttribute.AddAndFill(0.f);
+			}
+			else if (CustomAttributeType == ECustomAttributeTypeEnum::Dataflow_CustomAttributeType_Int)
+			{
+				TManagedArrayAccessor<int32> CustomAttribute(InCollection, FName(*AttrName), GroupNameToUse);
+
+				// If the group already exists don't change the number of elements
+				if (!InCollection.HasGroup(GroupNameToUse))
+				{
+					CustomAttribute.AddElements(InNumElements);
+				}
+
+				CustomAttribute.AddAndFill(0);
+			}
+			else if (CustomAttributeType == ECustomAttributeTypeEnum::Dataflow_CustomAttributeType_Vector)
+			{
+				TManagedArrayAccessor<FVector> CustomAttribute(InCollection, FName(*AttrName), GroupNameToUse);
+
+				// If the group already exists don't change the number of elements
+				if (!InCollection.HasGroup(GroupNameToUse))
+				{
+					CustomAttribute.AddElements(InNumElements);
+				}
+
+				CustomAttribute.AddAndFill(FVector(EForceInit::ForceInitToZero));
+			}
+		}
+
+		SetValue<FManagedArrayCollection>(Context, InCollection, &Collection);
+	}
+}
+
+
+void FGetNumElementsInCollectionGroupDataflowNode::Evaluate(Dataflow::FContext& Context, const FDataflowOutput* Out) const
+{
+	if (Out->IsA<int32>(&NumElements))
+	{
+		const FManagedArrayCollection& InCollection = GetValue<FManagedArrayCollection>(Context, &Collection);
+		
+		int32 OutNumElements = 0;
 
+		FName GroupNameToUse;
+		if (GroupName != EStandardGroupNameEnum::Dataflow_EStandardGroupNameEnum_Custom)
+		{
+			GroupNameToUse = GetGroupName(GroupName);
+		}
+		else
+		{
+			GroupNameToUse = FName(*CustomGroupName);
+		}
 
+		if (GroupNameToUse.GetStringLength() > 0)
+		{
+			if (InCollection.HasGroup(GroupNameToUse))
+			{
+				OutNumElements = InCollection.NumElements(GroupNameToUse);
+			}
+		}
+
+		SetValue<int32>(Context, OutNumElements, &NumElements);
+	}
+}
 
 
+void FGetCollectionAttributeDataTypedDataflowNode::Evaluate(Dataflow::FContext& Context, const FDataflowOutput* Out) const
+{
+	if (Out->IsA<TArray<bool>>(&BoolAttributeData) ||
+		Out->IsA<TArray<float>>(&FloatAttributeData) ||
+		Out->IsA<TArray<int32>>(&IntAttributeData) ||
+		Out->IsA<TArray<FVector>>(&VectorAttributeData))
+	{
+		const FManagedArrayCollection& InCollection = GetValue<FManagedArrayCollection>(Context, &Collection);
+
+		FName GroupNameToUse;
+		if (GroupName != EStandardGroupNameEnum::Dataflow_EStandardGroupNameEnum_Custom)
+		{
+			GroupNameToUse = GetGroupName(GroupName);
+		}
+		else
+		{
+			GroupNameToUse = FName(*CustomGroupName);
+		}
+
+		if (GroupNameToUse.GetStringLength() > 0 && AttrName.Len() > 0)
+		{
+			if (InCollection.HasGroup(GroupNameToUse))
+			{
+				if (InCollection.HasAttribute(FName(*AttrName), GroupNameToUse))
+				{
+					FString TypeStr = GetArrayTypeString(InCollection.GetAttributeType(FName(*AttrName), GroupNameToUse)).ToString();
+
+					if (Out->IsA<TArray<bool>>(&BoolAttributeData))
+					{
+						if (TypeStr == FString("bool"))
+						{
+							const TManagedArray<bool>& BoolAttributeArr = InCollection.GetAttribute<bool>(FName(*AttrName), GroupNameToUse);
+
+							SetValue<TArray<bool>>(Context, BoolAttributeArr.GetConstArray(), &BoolAttributeData);
+						}
+						else
+						{
+							SetValue<TArray<bool>>(Context, TArray<bool>(), &BoolAttributeData);
+						}
+					}
+					else if (Out->IsA<TArray<float>>(&FloatAttributeData))
+					{
+						if (TypeStr == FString("float"))
+						{
+							const TManagedArray<float>& FloatAttributeArr = InCollection.GetAttribute<float>(FName(*AttrName), GroupNameToUse);
+
+							SetValue<TArray<float>>(Context, FloatAttributeArr.GetConstArray(), &FloatAttributeData);
+						}
+						else
+						{
+							SetValue<TArray<float>>(Context, TArray<float>(), &FloatAttributeData);
+						}
+					}
+					else if (Out->IsA<TArray<int32>>(&IntAttributeData))
+					{
+						if (TypeStr == FString("int32"))
+						{
+							const TManagedArray<int32>& IntAttributeArr = InCollection.GetAttribute<int32>(FName(*AttrName), GroupNameToUse);
+
+							SetValue<TArray<int32>>(Context, IntAttributeArr.GetConstArray(), &IntAttributeData);
+						}
+						else
+						{
+							SetValue<TArray<int32>>(Context, TArray<int32>(), &IntAttributeData);
+						}
+					}
+					else if (Out->IsA<TArray<FVector>>(&VectorAttributeData))
+					{
+						if (TypeStr == FString("Vector3d"))
+						{
+							const TManagedArray<FVector>& VectorAttributeArr = InCollection.GetAttribute<FVector>(FName(*AttrName), GroupNameToUse);
+
+							SetValue<TArray<FVector>>(Context, VectorAttributeArr.GetConstArray(), &VectorAttributeData);
+						}
+						else
+						{
+							SetValue<TArray<FVector>>(Context, TArray<FVector>(), &VectorAttributeData);
+						}
+					}
+				}
+			}
+		}
+	}
+}
+
+
+void FSetCollectionAttributeDataTypedDataflowNode::Evaluate(Dataflow::FContext& Context, const FDataflowOutput* Out) const
+{
+	if (Out->IsA<FManagedArrayCollection>(&Collection))
+	{
+		FManagedArrayCollection InCollection = GetValue<FManagedArrayCollection>(Context, &Collection);
+
+		FName GroupNameToUse;
+		if (GroupName != EStandardGroupNameEnum::Dataflow_EStandardGroupNameEnum_Custom)
+		{
+			GroupNameToUse = GetGroupName(GroupName);
+		}
+		else
+		{
+			GroupNameToUse = FName(*CustomGroupName);
+		}
+
+		if (GroupNameToUse.GetStringLength() > 0 && AttrName.Len() > 0)
+		{
+			if (InCollection.HasGroup(GroupNameToUse))
+			{
+				if (InCollection.HasAttribute(FName(*AttrName), GroupNameToUse))
+				{
+					FString TypeStr = GetArrayTypeString(InCollection.GetAttributeType(FName(*AttrName), GroupNameToUse)).ToString();
+
+					if (TypeStr == FString("bool"))
+					{
+						if (IsConnected<TArray<bool>>(&BoolAttributeData))
+						{
+							TArray<bool> InBoolAttributeData = GetValue<TArray<bool>>(Context, &BoolAttributeData);
+
+							TManagedArray<bool>& BoolAttributeArr = InCollection.ModifyAttribute<bool>(FName(*AttrName), GroupNameToUse);
+
+							if (InBoolAttributeData.Num() == BoolAttributeArr.Num())
+							{
+								for (int32 Idx = 0; Idx < BoolAttributeArr.Num(); ++Idx)
+								{
+									BoolAttributeArr[Idx] = InBoolAttributeData[Idx];
+								}
+							}
+
+							SetValue<FManagedArrayCollection>(Context, InCollection, &Collection);
+						}
+					}
+					else if (TypeStr == FString("float"))
+					{
+						if (IsConnected<TArray<float>>(&FloatAttributeData))
+						{
+							TArray<float> InFloatAttributeData = GetValue<TArray<float>>(Context, &FloatAttributeData);
+
+							TManagedArray<float>& FloatAttributeArr = InCollection.ModifyAttribute<float>(FName(*AttrName), GroupNameToUse);
+
+							if (InFloatAttributeData.Num() == FloatAttributeArr.Num())
+							{
+								for (int32 Idx = 0; Idx < FloatAttributeArr.Num(); ++Idx)
+								{
+									FloatAttributeArr[Idx] = InFloatAttributeData[Idx];
+								}
+							}
+
+							SetValue<FManagedArrayCollection>(Context, InCollection, &Collection);
+						}
+					}
+					else if (TypeStr == FString("int32"))
+					{
+						if (IsConnected<TArray<int32>>(&IntAttributeData))
+						{
+							TArray<int32> InIntAttributeData = GetValue<TArray<int32>>(Context, &IntAttributeData);
+
+							TManagedArray<int32>& IntAttributeArr = InCollection.ModifyAttribute<int32>(FName(*AttrName), GroupNameToUse);
+
+							if (InIntAttributeData.Num() == IntAttributeArr.Num())
+							{
+								for (int32 Idx = 0; Idx < IntAttributeArr.Num(); ++Idx)
+								{
+									IntAttributeArr[Idx] = InIntAttributeData[Idx];
+								}
+							}
+
+							SetValue<FManagedArrayCollection>(Context, InCollection, &Collection);
+						}
+					}
+					else if (TypeStr == FString("Vector3d"))
+					{
+						if (IsConnected<TArray<FVector>>(&VectorAttributeData))
+						{
+							TArray<FVector> InVectorAttributeData = GetValue<TArray<FVector>>(Context, &VectorAttributeData);
+
+							TManagedArray<FVector>& VectorAttributeArr = InCollection.ModifyAttribute<FVector>(FName(*AttrName), GroupNameToUse);
+
+							if (InVectorAttributeData.Num() == VectorAttributeArr.Num())
+							{
+								for (int32 Idx = 0; Idx < VectorAttributeArr.Num(); ++Idx)
+								{
+									VectorAttributeArr[Idx] = InVectorAttributeData[Idx];
+								}
+							}
+
+							SetValue<FManagedArrayCollection>(Context, InCollection, &Collection);
+						}
+					}
+				}
+			}
+		}
+
+		SetValue<FManagedArrayCollection>(Context, InCollection, &Collection);
+	}
+}
 
 
 
diff --git a/Engine/Plugins/Experimental/GeometryCollectionPlugin/Source/GeometryCollectionNodes/Private/Dataflow/GeometryCollectionNodesPlugin.cpp b/Engine/Plugins/Experimental/GeometryCollectionPlugin/Source/GeometryCollectionNodes/Private/Dataflow/GeometryCollectionNodesPlugin.cpp
--- a/Engine/Plugins/Experimental/GeometryCollectionPlugin/Source/GeometryCollectionNodes/Private/Dataflow/GeometryCollectionNodesPlugin.cpp
+++ b/Engine/Plugins/Experimental/GeometryCollectionPlugin/Source/GeometryCollectionNodes/Private/Dataflow/GeometryCollectionNodesPlugin.cpp
@@ -13,6 +13,7 @@
 #include "Dataflow/GeometryCollectionMeshNodes.h"
 #include "Dataflow/GeometryCollectionClusteringNodes.h"
 #include "Dataflow/GeometryCollectionFracturingNodes.h"
+#include "Dataflow/GeometryCollectionEditNodes.h"
 
 #define LOCTEXT_NAMESPACE "DataflowNodes"
 
@@ -27,6 +28,7 @@ void IGeometryCollectionNodesPlugin::StartupModule()
 	Dataflow::GeometryCollectionMeshNodes();
 	Dataflow::GeometryCollectionClusteringNodes();
 	Dataflow::GeometryCollectionFracturingNodes();
+	Dataflow::GeometryCollectionEditNodes();
 }
 
 void IGeometryCollectionNodesPlugin::ShutdownModule()
diff --git a/Engine/Plugins/Experimental/GeometryCollectionPlugin/Source/GeometryCollectionNodes/Private/Dataflow/GeometryCollectionSelectionNodes.cpp b/Engine/Plugins/Experimental/GeometryCollectionPlugin/Source/GeometryCollectionNodes/Private/Dataflow/GeometryCollectionSelectionNodes.cpp
--- a/Engine/Plugins/Experimental/GeometryCollectionPlugin/Source/GeometryCollectionNodes/Private/Dataflow/GeometryCollectionSelectionNodes.cpp
+++ b/Engine/Plugins/Experimental/GeometryCollectionPlugin/Source/GeometryCollectionNodes/Private/Dataflow/GeometryCollectionSelectionNodes.cpp
@@ -27,6 +27,7 @@
 #include "GeometryCollection/GeometryCollectionProximityUtility.h"
 #include "FractureEngineClustering.h"
 #include "FractureEngineSelection.h"
+#include "GeometryCollection/Facades/CollectionTransformSelectionFacade.h"
 
 #include UE_INLINE_GENERATED_CPP_BY_NAME(GeometryCollectionSelectionNodes)
 
@@ -55,6 +56,13 @@ namespace Dataflow
 		DATAFLOW_NODE_REGISTER_CREATION_FACTORY(FCollectionTransformSelectionClusterDataflowNode);
 		DATAFLOW_NODE_REGISTER_CREATION_FACTORY(FCollectionTransformSelectionBySizeDataflowNode);
 		DATAFLOW_NODE_REGISTER_CREATION_FACTORY(FCollectionTransformSelectionByVolumeDataflowNode);
+		DATAFLOW_NODE_REGISTER_CREATION_FACTORY(FCollectionTransformSelectionInBoxDataflowNode);
+		DATAFLOW_NODE_REGISTER_CREATION_FACTORY(FCollectionTransformSelectionInSphereDataflowNode);
+		DATAFLOW_NODE_REGISTER_CREATION_FACTORY(FCollectionTransformSelectionByFloatAttrDataflowNode);
+		DATAFLOW_NODE_REGISTER_CREATION_FACTORY(FCollectionTransformSelectionByIntAttrDataflowNode);
+		DATAFLOW_NODE_REGISTER_CREATION_FACTORY(FCollectionVertexSelectionCustomDataflowNode);
+		DATAFLOW_NODE_REGISTER_CREATION_FACTORY(FCollectionFaceSelectionCustomDataflowNode);
+		DATAFLOW_NODE_REGISTER_CREATION_FACTORY(FCollectionSelectionConvertDataflowNode);
 
 		// GeometryCollection|Selection
 		DATAFLOW_NODE_REGISTER_CREATION_FACTORY_NODE_COLORS_BY_CATEGORY("GeometryCollection|Selection", FLinearColor(1.f, 1.f, 0.05f), CDefaultNodeBodyTintColor);
@@ -68,10 +76,12 @@ void FCollectionTransformSelectionAllDataflowNode::Evaluate(Dataflow::FContext&
 	{
 		const FManagedArrayCollection& InCollection = GetValue<FManagedArrayCollection>(Context, &Collection);
 
-		int32 NumTransforms = InCollection.NumElements(FGeometryCollection::TransformGroup);
+		GeometryCollection::Facades::FCollectionTransformSelectionFacade TransformSelectionFacade(InCollection);
+		const TArray<int32>& SelectionArr = TransformSelectionFacade.SelectAll();
 
 		FDataflowTransformSelection NewTransformSelection;
-		NewTransformSelection.Initialize(NumTransforms, true);
+		NewTransformSelection.Initialize(InCollection.NumElements(FGeometryCollection::TransformGroup), false);
+		NewTransformSelection.SetFromArray(SelectionArr);
 
 		SetValue<FDataflowTransformSelection>(Context, NewTransformSelection, &TransformSelection);
 	}
@@ -263,10 +273,12 @@ void FCollectionTransformSelectionNoneDataflowNode::Evaluate(Dataflow::FContext&
 	{
 		const FManagedArrayCollection& InCollection = GetValue<FManagedArrayCollection>(Context, &Collection);
 
-		int32 NumTransforms = InCollection.NumElements(FGeometryCollection::TransformGroup);
+		GeometryCollection::Facades::FCollectionTransformSelectionFacade TransformSelectionFacade(InCollection);
+		const TArray<int32>& SelectionArr = TransformSelectionFacade.SelectNone();
 
 		FDataflowTransformSelection NewTransformSelection;
-		NewTransformSelection.Initialize(NumTransforms, false);
+		NewTransformSelection.Initialize(InCollection.NumElements(FGeometryCollection::TransformGroup), false);
+		NewTransformSelection.SetFromArray(SelectionArr);
 
 		SetValue<FDataflowTransformSelection>(Context, NewTransformSelection, &TransformSelection);
 	}
@@ -297,37 +309,15 @@ void FCollectionTransformSelectionRandomDataflowNode::Evaluate(Dataflow::FContex
 	{
 		const FManagedArrayCollection& InCollection = GetValue<FManagedArrayCollection>(Context, &Collection);
 
-		int32 NumTransforms = InCollection.NumElements(FGeometryCollection::TransformGroup);
-
-		FDataflowTransformSelection NewTransformSelection;
-		NewTransformSelection.Initialize(NumTransforms, false);
-
 		float RandomSeedVal = GetValue<float>(Context, &RandomSeed);
 		float RandomThresholdVal = GetValue<float>(Context, &RandomThreshold);
 
-		FRandomStream Stream(RandomSeedVal);
+		GeometryCollection::Facades::FCollectionTransformSelectionFacade TransformSelectionFacade(InCollection);
+		const TArray<int32>& SelectionArr = TransformSelectionFacade.SelectRandom(bDeterministic, RandomSeedVal, RandomThresholdVal);
 
-		for (int32 Idx = 0; Idx < NumTransforms; ++Idx)
-		{
-			if (Deterministic)
-			{
-				float RandomVal = Stream.FRandRange(0.f, 1.f);
-
-				if (RandomVal > RandomThresholdVal)
-				{
-					NewTransformSelection.SetSelected(Idx);
-				}
-			}
-			else
-			{
-				float RandomVal = FMath::FRandRange(0.f, 1.f);
-
-				if (RandomVal > RandomThresholdVal)
-				{
-					NewTransformSelection.SetSelected(Idx);
-				}
-			}
-		}
+		FDataflowTransformSelection NewTransformSelection;
+		NewTransformSelection.Initialize(InCollection.NumElements(FGeometryCollection::TransformGroup), false);
+		NewTransformSelection.SetFromArray(SelectionArr);
 
 		SetValue<FDataflowTransformSelection>(Context, NewTransformSelection, &TransformSelection);
 	}
@@ -345,24 +335,14 @@ void FCollectionTransformSelectionRootDataflowNode::Evaluate(Dataflow::FContext&
 	{
 		const FManagedArrayCollection& InCollection = GetValue<FManagedArrayCollection>(Context, &Collection);
 
-		if (InCollection.HasGroup(FGeometryCollection::TransformGroup) &&
-			InCollection.HasAttribute("Parent", FGeometryCollection::TransformGroup))
-		{
-			TArray<int32> RootBones;
-			FFractureEngineSelection::GetRootBones(InCollection, RootBones);
-
-			int32 NumTransforms = InCollection.NumElements(FGeometryCollection::TransformGroup);
+		GeometryCollection::Facades::FCollectionTransformSelectionFacade TransformSelectionFacade(InCollection);
+		const TArray<int32>& SelectionArr = TransformSelectionFacade.SelectRootBones();
 
-			FDataflowTransformSelection NewTransformSelection;
-			NewTransformSelection.Initialize(NumTransforms, false);
-			NewTransformSelection.SetFromArray(RootBones);
+		FDataflowTransformSelection NewTransformSelection;
+		NewTransformSelection.Initialize(InCollection.NumElements(FGeometryCollection::TransformGroup), false);
+		NewTransformSelection.SetFromArray(SelectionArr);
 
-			SetValue<FDataflowTransformSelection>(Context, NewTransformSelection, &TransformSelection);
-		}
-		else
-		{
-			SetValue<FDataflowTransformSelection>(Context, FDataflowTransformSelection(), &TransformSelection);
-		}
+		SetValue<FDataflowTransformSelection>(Context, NewTransformSelection, &TransformSelection);
 	}
 	else if (Out->IsA<FManagedArrayCollection>(&Collection))
 	{
@@ -430,7 +410,11 @@ void FCollectionTransformSelectionParentDataflowNode::Evaluate(Dataflow::FContex
 		const FManagedArrayCollection& InCollection = GetValue<FManagedArrayCollection>(Context, &Collection);
 		FDataflowTransformSelection InTransformSelection = GetValue<FDataflowTransformSelection>(Context, &TransformSelection);
 
-		FFractureEngineSelection::SelectParent(InCollection, InTransformSelection);
+		GeometryCollection::Facades::FCollectionTransformSelectionFacade TransformSelectionFacade(InCollection);
+		TArray<int32> SelectionArr = InTransformSelection.AsArray();
+		TransformSelectionFacade.SelectParent(SelectionArr);
+
+		InTransformSelection.SetFromArray(SelectionArr);
 		
 		SetValue<FDataflowTransformSelection>(Context, InTransformSelection, &TransformSelection);
 	}
@@ -451,8 +435,11 @@ void FCollectionTransformSelectionByPercentageDataflowNode::Evaluate(Dataflow::F
 		int32 InPercentage = GetValue<int32>(Context, &Percentage);
 		float InRandomSeed = GetValue<float>(Context, &RandomSeed);
 
-		FFractureEngineSelection::SelectByPercentage(InTransformSelection, InPercentage, Deterministic, InRandomSeed);
+		TArray<int32> SelectionArr = InTransformSelection.AsArray();
+
+		GeometryCollection::Facades::FCollectionTransformSelectionFacade::SelectByPercentage(SelectionArr, InPercentage, bDeterministic, InRandomSeed);
 
+		InTransformSelection.SetFromArray(SelectionArr);
 		SetValue<FDataflowTransformSelection>(Context, InTransformSelection, &TransformSelection);
 	}
 }
@@ -465,7 +452,11 @@ void FCollectionTransformSelectionChildrenDataflowNode::Evaluate(Dataflow::FCont
 		const FManagedArrayCollection& InCollection = GetValue<FManagedArrayCollection>(Context, &Collection);
 		FDataflowTransformSelection InTransformSelection = GetValue<FDataflowTransformSelection>(Context, &TransformSelection);
 
-		FFractureEngineSelection::SelectChildren(InCollection, InTransformSelection);
+		GeometryCollection::Facades::FCollectionTransformSelectionFacade TransformSelectionFacade(InCollection);
+		TArray<int32> SelectionArr = InTransformSelection.AsArray();
+
+		TransformSelectionFacade.SelectChildren(SelectionArr);
+		InTransformSelection.SetFromArray(SelectionArr);
 
 		SetValue<FDataflowTransformSelection>(Context, InTransformSelection, &TransformSelection);
 	}
@@ -484,7 +475,11 @@ void FCollectionTransformSelectionSiblingsDataflowNode::Evaluate(Dataflow::FCont
 		const FManagedArrayCollection& InCollection = GetValue<FManagedArrayCollection>(Context, &Collection);
 		FDataflowTransformSelection InTransformSelection = GetValue<FDataflowTransformSelection>(Context, &TransformSelection);
 
-		FFractureEngineSelection::SelectSiblings(InCollection, InTransformSelection);
+		GeometryCollection::Facades::FCollectionTransformSelectionFacade TransformSelectionFacade(InCollection);
+		TArray<int32> SelectionArr = InTransformSelection.AsArray();
+
+		TransformSelectionFacade.SelectSiblings(SelectionArr);
+		InTransformSelection.SetFromArray(SelectionArr);
 
 		SetValue<FDataflowTransformSelection>(Context, InTransformSelection, &TransformSelection);
 	}
@@ -503,7 +498,11 @@ void FCollectionTransformSelectionLevelDataflowNode::Evaluate(Dataflow::FContext
 		const FManagedArrayCollection& InCollection = GetValue<FManagedArrayCollection>(Context, &Collection);
 		FDataflowTransformSelection InTransformSelection = GetValue<FDataflowTransformSelection>(Context, &TransformSelection);
 
-		FFractureEngineSelection::SelectLevel(InCollection, InTransformSelection);
+		GeometryCollection::Facades::FCollectionTransformSelectionFacade TransformSelectionFacade(InCollection);
+		TArray<int32> SelectionArr = InTransformSelection.AsArray();
+
+		TransformSelectionFacade.SelectLevel(SelectionArr);
+		InTransformSelection.SetFromArray(SelectionArr);
 
 		SetValue<FDataflowTransformSelection>(Context, InTransformSelection, &TransformSelection);
 	}
@@ -522,13 +521,13 @@ void FCollectionTransformSelectionContactDataflowNode::Evaluate(Dataflow::FConte
 		const FManagedArrayCollection& InCollection = GetValue<FManagedArrayCollection>(Context, &Collection);
 		FDataflowTransformSelection InTransformSelection = GetValue<FDataflowTransformSelection>(Context, &TransformSelection);
 
-		if (TUniquePtr<FGeometryCollection> GeomCollection = TUniquePtr<FGeometryCollection>(InCollection.NewCopy<FGeometryCollection>()))
-		{
-			FFractureEngineSelection::SelectContact(*GeomCollection, InTransformSelection);
-		}
+		GeometryCollection::Facades::FCollectionTransformSelectionFacade TransformSelectionFacade(InCollection);
+		TArray<int32> SelectionArr = InTransformSelection.AsArray();
+
+		TransformSelectionFacade.SelectLevel(SelectionArr);
+		InTransformSelection.SetFromArray(SelectionArr);
 
 		SetValue<FDataflowTransformSelection>(Context, InTransformSelection, &TransformSelection);
-		SetValue<FManagedArrayCollection>(Context, InCollection, &Collection);
 	}
 	else if (Out->IsA<FManagedArrayCollection>(&Collection))
 	{
@@ -544,12 +543,12 @@ void FCollectionTransformSelectionLeafDataflowNode::Evaluate(Dataflow::FContext&
 	{
 		const FManagedArrayCollection& InCollection = GetValue<FManagedArrayCollection>(Context, &Collection);
 
-		FDataflowTransformSelection NewTransformSelection;
+		GeometryCollection::Facades::FCollectionTransformSelectionFacade TransformSelectionFacade(InCollection);
+		const TArray<int32>& SelectionArr = TransformSelectionFacade.SelectLeaf();
 
-		if (TUniquePtr<FGeometryCollection> GeomCollection = TUniquePtr<FGeometryCollection>(InCollection.NewCopy<FGeometryCollection>()))
-		{
-			FFractureEngineSelection::SelectLeaf(*GeomCollection, NewTransformSelection);
-		}
+		FDataflowTransformSelection NewTransformSelection;
+		NewTransformSelection.Initialize(InCollection.NumElements(FGeometryCollection::TransformGroup), false);
+		NewTransformSelection.SetFromArray(SelectionArr);
 
 		SetValue<FDataflowTransformSelection>(Context, NewTransformSelection, &TransformSelection);
 	}
@@ -567,16 +566,12 @@ void FCollectionTransformSelectionClusterDataflowNode::Evaluate(Dataflow::FConte
 	{
 		const FManagedArrayCollection& InCollection = GetValue<FManagedArrayCollection>(Context, &Collection);
 
-		const int32 NumTransforms = InCollection.NumElements(FGeometryCollection::TransformGroup);
+		GeometryCollection::Facades::FCollectionTransformSelectionFacade TransformSelectionFacade(InCollection);
+		const TArray<int32>& SelectionArr = TransformSelectionFacade.SelectCluster();
 
-		// TODO: Convert FractureEngine API to use only FManagedArrayCollection
 		FDataflowTransformSelection NewTransformSelection;
-		NewTransformSelection.Initialize(NumTransforms, false);
-
-		if (TUniquePtr<FGeometryCollection> GeomCollection = TUniquePtr<FGeometryCollection>(InCollection.NewCopy<FGeometryCollection>()))
-		{
-			FFractureEngineSelection::SelectCluster(*GeomCollection, NewTransformSelection);			
-		}
+		NewTransformSelection.Initialize(InCollection.NumElements(FGeometryCollection::TransformGroup), false);
+		NewTransformSelection.SetFromArray(SelectionArr);
 
 		SetValue<FDataflowTransformSelection>(Context, NewTransformSelection, &TransformSelection);
 	}
@@ -595,13 +590,14 @@ void FCollectionTransformSelectionBySizeDataflowNode::Evaluate(Dataflow::FContex
 		const FManagedArrayCollection& InCollection = GetValue<FManagedArrayCollection>(Context, &Collection);
 		float InSizeMin = GetValue<float>(Context, &SizeMin);
 		float InSizeMax = GetValue<float>(Context, &SizeMax);
+		bool bInsideRange = RangeSetting == ERangeSettingEnum::Dataflow_RangeSetting_InsideRange;
 
-		FDataflowTransformSelection NewTransformSelection;
+		GeometryCollection::Facades::FCollectionTransformSelectionFacade TransformSelectionFacade(InCollection);
+		const TArray<int32>& SelectionArr = TransformSelectionFacade.SelectBySize(InSizeMin, InSizeMax, bInclusive, bInsideRange);
 
-		if (TUniquePtr<FGeometryCollection> GeomCollection = TUniquePtr<FGeometryCollection>(InCollection.NewCopy<FGeometryCollection>()))
-		{
-			FFractureEngineSelection::SelectBySize(*GeomCollection, NewTransformSelection, InSizeMin, InSizeMax);
-		}
+		FDataflowTransformSelection NewTransformSelection;
+		NewTransformSelection.Initialize(InCollection.NumElements(FGeometryCollection::TransformGroup), false);
+		NewTransformSelection.SetFromArray(SelectionArr);
 
 		SetValue<FDataflowTransformSelection>(Context, NewTransformSelection, &TransformSelection);
 	}
@@ -620,13 +616,142 @@ void FCollectionTransformSelectionByVolumeDataflowNode::Evaluate(Dataflow::FCont
 		const FManagedArrayCollection& InCollection = GetValue<FManagedArrayCollection>(Context, &Collection);
 		float InVolumeMin = GetValue<float>(Context, &VolumeMin);
 		float InVolumeMax = GetValue<float>(Context, &VolumeMax);
+		bool bInsideRange = RangeSetting == ERangeSettingEnum::Dataflow_RangeSetting_InsideRange;
+
+		GeometryCollection::Facades::FCollectionTransformSelectionFacade TransformSelectionFacade(InCollection);
+		const TArray<int32>& SelectionArr = TransformSelectionFacade.SelectByVolume(InVolumeMin, InVolumeMax, bInclusive, bInsideRange);
 
 		FDataflowTransformSelection NewTransformSelection;
+		NewTransformSelection.Initialize(InCollection.NumElements(FGeometryCollection::TransformGroup), false);
+		NewTransformSelection.SetFromArray(SelectionArr);
+
+		SetValue<FDataflowTransformSelection>(Context, NewTransformSelection, &TransformSelection);
+	}
+	else if (Out->IsA<FManagedArrayCollection>(&Collection))
+	{
+		const FManagedArrayCollection& InCollection = GetValue<FManagedArrayCollection>(Context, &Collection);
+		SetValue<FManagedArrayCollection>(Context, InCollection, &Collection);
+	}
+}
+
+
+void FCollectionTransformSelectionInBoxDataflowNode::Evaluate(Dataflow::FContext& Context, const FDataflowOutput* Out) const
+{
+	if (Out->IsA<FDataflowTransformSelection>(&TransformSelection))
+	{
+		const FManagedArrayCollection& InCollection = GetValue<FManagedArrayCollection>(Context, &Collection);
+		const FBox& InBox = GetValue<FBox>(Context, &Box);
+		const FTransform& InTransform = GetValue<FTransform>(Context, &Transform);
+
+		GeometryCollection::Facades::FCollectionTransformSelectionFacade TransformSelectionFacade(InCollection);
 
-		if (TUniquePtr<FGeometryCollection> GeomCollection = TUniquePtr<FGeometryCollection>(InCollection.NewCopy<FGeometryCollection>()))
+		TArray<int32> SelectionArr;
+		if (Type == ESelectSubjectTypeEnum::Dataflow_SelectSubjectType_Vertices)
 		{
-			FFractureEngineSelection::SelectByVolume(*GeomCollection, NewTransformSelection, InVolumeMin, InVolumeMax);
+			SelectionArr = TransformSelectionFacade.SelectVerticesInBox(InBox, InTransform, bAllVerticesMustContainedInBox);
 		}
+		else if (Type == ESelectSubjectTypeEnum::Dataflow_SelectSubjectType_BoundingBox)
+		{
+			SelectionArr = TransformSelectionFacade.SelectBoundingBoxInBox(InBox, InTransform);
+		}
+		else if (Type == ESelectSubjectTypeEnum::Dataflow_SelectSubjectType_Centroid)
+		{
+			SelectionArr = TransformSelectionFacade.SelectCentroidInBox(InBox, InTransform);
+		}
+
+		FDataflowTransformSelection NewTransformSelection;
+		NewTransformSelection.Initialize(InCollection.NumElements(FGeometryCollection::TransformGroup), false);
+		NewTransformSelection.SetFromArray(SelectionArr);
+
+		SetValue<FDataflowTransformSelection>(Context, NewTransformSelection, &TransformSelection);
+	}
+	else if (Out->IsA<FManagedArrayCollection>(&Collection))
+	{
+		const FManagedArrayCollection& InCollection = GetValue<FManagedArrayCollection>(Context, &Collection);
+		SetValue<FManagedArrayCollection>(Context, InCollection, &Collection);
+	}
+}
+
+
+void FCollectionTransformSelectionInSphereDataflowNode::Evaluate(Dataflow::FContext& Context, const FDataflowOutput* Out) const
+{
+	if (Out->IsA<FDataflowTransformSelection>(&TransformSelection))
+	{
+		const FManagedArrayCollection& InCollection = GetValue<FManagedArrayCollection>(Context, &Collection);
+		const FSphere& InSphere = GetValue<FSphere>(Context, &Sphere);
+		const FTransform& InTransform = GetValue<FTransform>(Context, &Transform);
+
+		GeometryCollection::Facades::FCollectionTransformSelectionFacade TransformSelectionFacade(InCollection);
+
+		TArray<int32> SelectionArr;
+		if (Type == ESelectSubjectTypeEnum::Dataflow_SelectSubjectType_Vertices)
+		{
+			SelectionArr = TransformSelectionFacade.SelectVerticesInSphere(InSphere, InTransform, bAllVerticesMustContainedInSphere);
+		}
+		else if (Type == ESelectSubjectTypeEnum::Dataflow_SelectSubjectType_BoundingBox)
+		{
+			SelectionArr = TransformSelectionFacade.SelectBoundingBoxInSphere(InSphere, InTransform);
+		}
+		else if (Type == ESelectSubjectTypeEnum::Dataflow_SelectSubjectType_Centroid)
+		{
+			SelectionArr = TransformSelectionFacade.SelectCentroidInSphere(InSphere, InTransform);
+		}
+
+		FDataflowTransformSelection NewTransformSelection;
+		NewTransformSelection.Initialize(InCollection.NumElements(FGeometryCollection::TransformGroup), false);
+		NewTransformSelection.SetFromArray(SelectionArr);
+
+		SetValue<FDataflowTransformSelection>(Context, NewTransformSelection, &TransformSelection);
+	}
+	else if (Out->IsA<FManagedArrayCollection>(&Collection))
+	{
+		const FManagedArrayCollection& InCollection = GetValue<FManagedArrayCollection>(Context, &Collection);
+		SetValue<FManagedArrayCollection>(Context, InCollection, &Collection);
+	}
+}
+
+
+void FCollectionTransformSelectionByFloatAttrDataflowNode::Evaluate(Dataflow::FContext& Context, const FDataflowOutput* Out) const
+{
+	if (Out->IsA<FDataflowTransformSelection>(&TransformSelection))
+	{
+		const FManagedArrayCollection& InCollection = GetValue<FManagedArrayCollection>(Context, &Collection);
+		float InMin = GetValue<float>(Context, &Min);
+		float InMax = GetValue<float>(Context, &Max);
+		bool bInsideRange = RangeSetting == ERangeSettingEnum::Dataflow_RangeSetting_InsideRange;
+
+		GeometryCollection::Facades::FCollectionTransformSelectionFacade TransformSelectionFacade(InCollection);
+		const TArray<int32>& SelectionArr = TransformSelectionFacade.SelectByFloatAttribute(GroupName, AttrName, InMin, InMax, bInclusive, bInsideRange);
+
+		FDataflowTransformSelection NewTransformSelection;
+		NewTransformSelection.Initialize(InCollection.NumElements(FGeometryCollection::TransformGroup), false);
+		NewTransformSelection.SetFromArray(SelectionArr);
+
+		SetValue<FDataflowTransformSelection>(Context, NewTransformSelection, &TransformSelection);
+	}
+	else if (Out->IsA<FManagedArrayCollection>(&Collection))
+	{
+		const FManagedArrayCollection& InCollection = GetValue<FManagedArrayCollection>(Context, &Collection);
+		SetValue<FManagedArrayCollection>(Context, InCollection, &Collection);
+	}
+}
+
+
+void FCollectionTransformSelectionByIntAttrDataflowNode::Evaluate(Dataflow::FContext& Context, const FDataflowOutput* Out) const
+{
+	if (Out->IsA<FDataflowTransformSelection>(&TransformSelection))
+	{
+		const FManagedArrayCollection& InCollection = GetValue<FManagedArrayCollection>(Context, &Collection);
+		int32 InMin = GetValue<int32>(Context, &Min);
+		int32 InMax = GetValue<int32>(Context, &Max);
+		bool bInsideRange = RangeSetting == ERangeSettingEnum::Dataflow_RangeSetting_InsideRange;
+
+		GeometryCollection::Facades::FCollectionTransformSelectionFacade TransformSelectionFacade(InCollection);
+		const TArray<int32>& SelectionArr = TransformSelectionFacade.SelectByIntAttribute(GroupName, AttrName, InMin, InMax, bInclusive, bInsideRange);
+
+		FDataflowTransformSelection NewTransformSelection;
+		NewTransformSelection.Initialize(InCollection.NumElements(FGeometryCollection::TransformGroup), false);
+		NewTransformSelection.SetFromArray(SelectionArr);
 
 		SetValue<FDataflowTransformSelection>(Context, NewTransformSelection, &TransformSelection);
 	}
@@ -638,3 +763,224 @@ void FCollectionTransformSelectionByVolumeDataflowNode::Evaluate(Dataflow::FCont
 }
 
 
+void FCollectionVertexSelectionCustomDataflowNode::Evaluate(Dataflow::FContext& Context, const FDataflowOutput* Out) const
+{
+	if (Out->IsA<FDataflowVertexSelection>(&VertexSelection))
+	{
+		const FManagedArrayCollection& InCollection = GetValue<FManagedArrayCollection>(Context, &Collection);
+
+		if (InCollection.HasGroup(FGeometryCollection::VerticesGroup))
+		{
+			const int32 NumVertices = InCollection.NumElements(FGeometryCollection::VerticesGroup);
+
+			FDataflowVertexSelection NewVertexSelection;
+			NewVertexSelection.Initialize(NumVertices, false);
+
+			const FString InVertexIndicies = GetValue<FString>(Context, &VertexIndicies);
+
+			TArray<FString> Indicies;
+			InVertexIndicies.ParseIntoArray(Indicies, TEXT(" "), true);
+
+			for (FString IndexStr : Indicies)
+			{
+				if (IndexStr.IsNumeric())
+				{
+					int32 Index = FCString::Atoi(*IndexStr);
+					if (Index >= 0 && Index < NumVertices)
+					{
+						NewVertexSelection.SetSelected(Index);
+					}
+					else
+					{
+						// ERROR: INVALID INDEX
+						FString ErrorStr = "Invalid specified index found.";
+						UE_LOG(LogTemp, Error, TEXT("[Dataflow ERROR] %s"), *ErrorStr);
+					}
+				}
+			}
+
+			SetValue<FDataflowVertexSelection>(Context, NewVertexSelection, &VertexSelection);
+		}
+		else
+		{
+			SetValue<FDataflowVertexSelection>(Context, FDataflowVertexSelection(), &VertexSelection);
+		}
+	}
+	else if (Out->IsA<FManagedArrayCollection>(&Collection))
+	{
+		const FManagedArrayCollection& InCollection = GetValue<FManagedArrayCollection>(Context, &Collection);
+		SetValue<FManagedArrayCollection>(Context, InCollection, &Collection);
+	}
+}
+
+
+void FCollectionFaceSelectionCustomDataflowNode::Evaluate(Dataflow::FContext& Context, const FDataflowOutput* Out) const
+{
+	if (Out->IsA<FDataflowFaceSelection>(&FaceSelection))
+	{
+		const FManagedArrayCollection& InCollection = GetValue<FManagedArrayCollection>(Context, &Collection);
+
+		if (InCollection.HasGroup(FGeometryCollection::GeometryGroup))
+		{
+			const int32 NumFaces = InCollection.NumElements(FGeometryCollection::FacesGroup);
+
+			FDataflowFaceSelection NewFaceSelection;
+			NewFaceSelection.Initialize(NumFaces, false);
+
+			const FString InFaceIndicies = GetValue<FString>(Context, &FaceIndicies);
+
+			TArray<FString> Indicies;
+			InFaceIndicies.ParseIntoArray(Indicies, TEXT(" "), true);
+
+			for (FString IndexStr : Indicies)
+			{
+				if (IndexStr.IsNumeric())
+				{
+					int32 Index = FCString::Atoi(*IndexStr);
+					if (Index >= 0 && Index < NumFaces)
+					{
+						NewFaceSelection.SetSelected(Index);
+					}
+					else
+					{
+						// ERROR: INVALID INDEX
+						FString ErrorStr = "Invalid specified index found.";
+						UE_LOG(LogTemp, Error, TEXT("[Dataflow ERROR] %s"), *ErrorStr);
+					}
+				}
+			}
+
+			SetValue<FDataflowFaceSelection>(Context, NewFaceSelection, &FaceSelection);
+		}
+		else
+		{
+			SetValue<FDataflowFaceSelection>(Context, FDataflowFaceSelection(), &FaceSelection);
+		}
+	}
+	else if (Out->IsA<FManagedArrayCollection>(&Collection))
+	{
+		const FManagedArrayCollection& InCollection = GetValue<FManagedArrayCollection>(Context, &Collection);
+		SetValue<FManagedArrayCollection>(Context, InCollection, &Collection);
+	}
+}
+
+
+void FCollectionSelectionConvertDataflowNode::Evaluate(Dataflow::FContext& Context, const FDataflowOutput* Out) const
+{
+	if (Out->IsA<FDataflowTransformSelection>(&TransformSelection))
+	{
+		if (IsConnected<FDataflowVertexSelection>(&VertexSelection))
+		{
+			const FManagedArrayCollection& InCollection = GetValue<FManagedArrayCollection>(Context, &Collection);
+			const FDataflowVertexSelection& InVertexSelection = GetValue<FDataflowVertexSelection>(Context, &VertexSelection);
+
+			GeometryCollection::Facades::FCollectionTransformSelectionFacade TransformSelectionFacade(InCollection);
+			const TArray<int32>& SelectionArr = TransformSelectionFacade.ConvertVertexSelectionToTransformSelection(InVertexSelection.AsArray(), bAllElementsMustBeSelected);
+
+			FDataflowTransformSelection NewTransformSelection;
+			NewTransformSelection.Initialize(InCollection.NumElements(FGeometryCollection::TransformGroup), false);
+			NewTransformSelection.SetFromArray(SelectionArr);
+
+			SetValue<FDataflowTransformSelection>(Context, NewTransformSelection, &TransformSelection);
+		}
+		else if (IsConnected<FDataflowFaceSelection>(&FaceSelection))
+		{
+			const FManagedArrayCollection& InCollection = GetValue<FManagedArrayCollection>(Context, &Collection);
+			const FDataflowFaceSelection& InFaceSelection = GetValue<FDataflowFaceSelection>(Context, &FaceSelection);
+
+			GeometryCollection::Facades::FCollectionTransformSelectionFacade TransformSelectionFacade(InCollection);
+			const TArray<int32>& SelectionArr = TransformSelectionFacade.ConvertFaceSelectionToTransformSelection(InFaceSelection.AsArray(), bAllElementsMustBeSelected);
+
+			FDataflowTransformSelection NewTransformSelection;
+			NewTransformSelection.Initialize(InCollection.NumElements(FGeometryCollection::TransformGroup), false);
+			NewTransformSelection.SetFromArray(SelectionArr);
+
+			SetValue<FDataflowTransformSelection>(Context, NewTransformSelection, &TransformSelection);
+		}
+		else if (IsConnected<FDataflowTransformSelection>(&TransformSelection))
+		{
+			// Passthrough
+			const FDataflowTransformSelection& InTransformSelection = GetValue<FDataflowTransformSelection>(Context, &TransformSelection);
+			SetValue<FDataflowTransformSelection>(Context, InTransformSelection, &TransformSelection);
+		}
+	}
+	else if (Out->IsA<FDataflowFaceSelection>(&FaceSelection))
+	{
+		if (IsConnected<FDataflowVertexSelection>(&VertexSelection))
+		{
+			const FManagedArrayCollection& InCollection = GetValue<FManagedArrayCollection>(Context, &Collection);
+			const FDataflowVertexSelection& InVertexSelection = GetValue<FDataflowVertexSelection>(Context, &VertexSelection);
+
+			GeometryCollection::Facades::FCollectionTransformSelectionFacade TransformSelectionFacade(InCollection);
+			const TArray<int32>& SelectionArr = TransformSelectionFacade.ConvertVertexSelectionToFaceSelection(InVertexSelection.AsArray(), bAllElementsMustBeSelected);
+
+			FDataflowFaceSelection NewFaceSelection;
+			NewFaceSelection.Initialize(InCollection.NumElements(FGeometryCollection::FacesGroup), false);
+			NewFaceSelection.SetFromArray(SelectionArr);
+
+			SetValue<FDataflowFaceSelection>(Context, NewFaceSelection, &FaceSelection);
+		}
+		else if (IsConnected<FDataflowFaceSelection>(&FaceSelection))
+		{
+			// Passthrough
+			const FDataflowFaceSelection& InFaceSelection = GetValue<FDataflowFaceSelection>(Context, &FaceSelection);
+			SetValue<FDataflowFaceSelection>(Context, InFaceSelection, &InFaceSelection);
+		}
+		else if (IsConnected<FDataflowTransformSelection>(&TransformSelection))
+		{
+			const FManagedArrayCollection& InCollection = GetValue<FManagedArrayCollection>(Context, &Collection);
+			const FDataflowTransformSelection& InTransformSelection = GetValue<FDataflowTransformSelection>(Context, &TransformSelection);
+
+			GeometryCollection::Facades::FCollectionTransformSelectionFacade TransformSelectionFacade(InCollection);
+			const TArray<int32>& SelectionArr = TransformSelectionFacade.ConvertTransformSelectionToFaceSelection(InTransformSelection.AsArray());
+
+			FDataflowFaceSelection NewFaceSelection;
+			NewFaceSelection.Initialize(InCollection.NumElements(FGeometryCollection::FacesGroup), false);
+			NewFaceSelection.SetFromArray(SelectionArr);
+
+			SetValue<FDataflowFaceSelection>(Context, NewFaceSelection, &FaceSelection);
+		}
+	}
+	else if (Out->IsA<FDataflowVertexSelection>(&VertexSelection))
+	{
+		if (IsConnected<FDataflowVertexSelection>(&VertexSelection))
+		{
+			// Passthrough
+			const FDataflowVertexSelection& InVertexSelection = GetValue<FDataflowVertexSelection>(Context, &VertexSelection);
+			SetValue<FDataflowVertexSelection>(Context, InVertexSelection, &VertexSelection);
+		}
+		else if (IsConnected<FDataflowFaceSelection>(&FaceSelection))
+		{
+			const FManagedArrayCollection& InCollection = GetValue<FManagedArrayCollection>(Context, &Collection);
+			const FDataflowFaceSelection& InFaceSelection = GetValue<FDataflowFaceSelection>(Context, &FaceSelection);
+
+			GeometryCollection::Facades::FCollectionTransformSelectionFacade TransformSelectionFacade(InCollection);
+			const TArray<int32>& SelectionArr = TransformSelectionFacade.ConvertFaceSelectionToVertexSelection(InFaceSelection.AsArray());
+
+			FDataflowVertexSelection NewVertexSelection;
+			NewVertexSelection.Initialize(InCollection.NumElements(FGeometryCollection::VerticesGroup), false);
+			NewVertexSelection.SetFromArray(SelectionArr);
+
+			SetValue<FDataflowVertexSelection>(Context, NewVertexSelection, &VertexSelection);
+		}
+		else if (IsConnected<FDataflowTransformSelection>(&TransformSelection))
+		{
+			const FManagedArrayCollection& InCollection = GetValue<FManagedArrayCollection>(Context, &Collection);
+			const FDataflowTransformSelection& InTransformSelection = GetValue<FDataflowTransformSelection>(Context, &TransformSelection);
+
+			GeometryCollection::Facades::FCollectionTransformSelectionFacade TransformSelectionFacade(InCollection);
+			const TArray<int32>& SelectionArr = TransformSelectionFacade.ConvertTransformSelectionToVertexSelection(InTransformSelection.AsArray());
+
+			FDataflowVertexSelection NewVertexSelection;
+			NewVertexSelection.Initialize(InCollection.NumElements(FGeometryCollection::VerticesGroup), false);
+			NewVertexSelection.SetFromArray(SelectionArr);
+
+			SetValue<FDataflowVertexSelection>(Context, NewVertexSelection, &VertexSelection);
+		}
+	}
+	else if (Out->IsA<FManagedArrayCollection>(&Collection))
+	{
+		const FManagedArrayCollection& InCollection = GetValue<FManagedArrayCollection>(Context, &Collection);
+		SetValue<FManagedArrayCollection>(Context, InCollection, &Collection);
+	}
+}
diff --git a/Engine/Plugins/Experimental/GeometryCollectionPlugin/Source/GeometryCollectionNodes/Public/Dataflow/GeometryCollectionEditNodes.h b/Engine/Plugins/Experimental/GeometryCollectionPlugin/Source/GeometryCollectionNodes/Public/Dataflow/GeometryCollectionEditNodes.h
new file mode 100644
--- /dev/null
+++ b/Engine/Plugins/Experimental/GeometryCollectionPlugin/Source/GeometryCollectionNodes/Public/Dataflow/GeometryCollectionEditNodes.h
@@ -0,0 +1,132 @@
+// Copyright Epic Games, Inc. All Rights Reserved.
+
+#pragma once 
+
+#include "CoreMinimal.h"
+#include "Dataflow/DataflowEngine.h"
+#include "GeometryCollection/ManagedArrayCollection.h"
+#include "Dataflow/DataflowSelection.h"
+
+#include "GeometryCollectionEditNodes.generated.h"
+
+
+class FGeometryCollection;
+
+/**
+ *
+ * Deletes selected bones from Collection. Empty clusters will be eliminated
+ *
+ */
+USTRUCT(meta = (DataflowGeometryCollection))
+struct FPruneInCollectionDataflowNode : public FDataflowNode
+{
+	GENERATED_USTRUCT_BODY()
+	DATAFLOW_NODE_DEFINE_INTERNAL(FPruneInCollectionDataflowNode, "PruneInCollection", "GeometryCollection|Edit", "")
+
+public:
+	/** Fractured GeometryCollection to prune */
+	UPROPERTY(meta = (DataflowInput, DataflowOutput, DataflowPassthrough = "Collection", DataflowIntrinsic))
+	FManagedArrayCollection Collection;
+
+	/** Bone selection for pruning */
+	UPROPERTY(meta = (DataflowInput, DisplayName = "TransformSelection", DataflowIntrinsic))
+	FDataflowTransformSelection TransformSelection;
+
+	FPruneInCollectionDataflowNode(const Dataflow::FNodeParameters& InParam, FGuid InGuid = FGuid::NewGuid())
+		: FDataflowNode(InParam, InGuid)
+	{
+		RegisterInputConnection(&Collection);
+		RegisterInputConnection(&TransformSelection);
+		RegisterOutputConnection(&Collection, &Collection);
+	}
+
+	virtual void Evaluate(Dataflow::FContext& Context, const FDataflowOutput* Out) const override;
+
+};
+
+
+UENUM(BlueprintType)
+enum class EVisibiltyOptionsEnum : uint8
+{
+	Dataflow_VisibilityOptions_Visible UMETA(DisplayName = "Visible"),
+	Dataflow_VisibilityOptions_Invisible UMETA(DisplayName = "Hidden"),
+	//~~~
+	//256th entry
+	Dataflow_Max                UMETA(Hidden)
+};
+
+/**
+ *
+ * Sets all selected bone's visibilty to Visible/Hidden
+ *
+ */
+USTRUCT(meta = (DataflowGeometryCollection))
+struct FSetVisibilityInCollectionDataflowNode : public FDataflowNode
+{
+	GENERATED_USTRUCT_BODY()
+	DATAFLOW_NODE_DEFINE_INTERNAL(FSetVisibilityInCollectionDataflowNode, "SetVisibilityInCollection", "GeometryCollection|Edit", "")
+
+public:
+	/** What to set the visibility of the selected bones */
+	UPROPERTY(EditAnywhere, Category = "Visibility");
+	EVisibiltyOptionsEnum Visibility = EVisibiltyOptionsEnum::Dataflow_VisibilityOptions_Invisible;
+
+	/** Fractured GeometryCollection to prune */
+	UPROPERTY(meta = (DataflowInput, DataflowOutput, DataflowPassthrough = "Collection", DataflowIntrinsic))
+	FManagedArrayCollection Collection;
+
+	/** Bone selection for pruning */
+	UPROPERTY(meta = (DataflowInput, DisplayName = "TransformSelection", DataflowIntrinsic))
+	FDataflowTransformSelection TransformSelection;
+
+	FSetVisibilityInCollectionDataflowNode(const Dataflow::FNodeParameters& InParam, FGuid InGuid = FGuid::NewGuid())
+		: FDataflowNode(InParam, InGuid)
+	{
+		RegisterInputConnection(&Collection);
+		RegisterInputConnection(&TransformSelection);
+		RegisterOutputConnection(&Collection, &Collection);
+	}
+
+	virtual void Evaluate(Dataflow::FContext& Context, const FDataflowOutput* Out) const override;
+
+};
+
+
+/**
+ *
+ * Merges selected bones into a single bone
+ *
+ */
+USTRUCT(meta = (DataflowGeometryCollection))
+struct FMergeInCollectionDataflowNode : public FDataflowNode
+{
+	GENERATED_USTRUCT_BODY()
+	DATAFLOW_NODE_DEFINE_INTERNAL(FMergeInCollectionDataflowNode, "MergeInCollection", "GeometryCollection|Edit", "")
+
+public:
+	/** Fractured GeometryCollection to merge */
+	UPROPERTY(meta = (DataflowInput, DataflowOutput, DataflowPassthrough = "Collection", DataflowIntrinsic))
+	FManagedArrayCollection Collection;
+
+	/** Bone selection for merging */
+	UPROPERTY(meta = (DataflowInput, DisplayName = "TransformSelection", DataflowIntrinsic))
+	FDataflowTransformSelection TransformSelection;
+
+	FMergeInCollectionDataflowNode(const Dataflow::FNodeParameters& InParam, FGuid InGuid = FGuid::NewGuid())
+		: FDataflowNode(InParam, InGuid)
+	{
+		RegisterInputConnection(&Collection);
+		RegisterInputConnection(&TransformSelection);
+		RegisterOutputConnection(&Collection, &Collection);
+	}
+
+	virtual void Evaluate(Dataflow::FContext& Context, const FDataflowOutput* Out) const override;
+
+};
+
+
+namespace Dataflow
+{
+	void GeometryCollectionEditNodes();
+}
+
diff --git a/Engine/Plugins/Experimental/GeometryCollectionPlugin/Source/GeometryCollectionNodes/Public/Dataflow/GeometryCollectionNodes.h b/Engine/Plugins/Experimental/GeometryCollectionPlugin/Source/GeometryCollectionNodes/Public/Dataflow/GeometryCollectionNodes.h
--- a/Engine/Plugins/Experimental/GeometryCollectionPlugin/Source/GeometryCollectionNodes/Public/Dataflow/GeometryCollectionNodes.h
+++ b/Engine/Plugins/Experimental/GeometryCollectionPlugin/Source/GeometryCollectionNodes/Public/Dataflow/GeometryCollectionNodes.h
@@ -8,6 +8,7 @@
 #include "DynamicMesh/DynamicMesh3.h"
 #include "UDynamicMesh.h"
 #include "Dataflow/DataflowSelection.h"
+#include "Math/MathFwd.h"
 
 #include "GeometryCollectionNodes.generated.h"
 
@@ -24,7 +25,7 @@ USTRUCT(meta = (DataflowGeometryCollection))
 struct FGetCollectionAssetDataflowNode : public FDataflowNode
 {
 	GENERATED_USTRUCT_BODY()
-	DATAFLOW_NODE_DEFINE_INTERNAL(FGetCollectionAssetDataflowNode, "GetCollectionAsset", "GeometryCollection", "")
+	DATAFLOW_NODE_DEFINE_INTERNAL(FGetCollectionAssetDataflowNode, "GetCollectionAsset", "GeometryCollection|Asset", "")
 
 public:
 	typedef FManagedArrayCollection DataType;
@@ -481,6 +482,40 @@ public:
 };
 
 
+/**
+ *
+ * Description for this node
+ *
+ */
+USTRUCT()
+struct FMakeSphereDataflowNode : public FDataflowNode
+{
+	GENERATED_USTRUCT_BODY()
+	DATAFLOW_NODE_DEFINE_INTERNAL(FMakeSphereDataflowNode, "MakeSphere", "Generators|Sphere", "")
+
+public:
+	UPROPERTY(EditAnywhere, Category = "Sphere", meta = (DataflowInput));
+	FVector Center = FVector(0.f);
+
+	UPROPERTY(EditAnywhere, Category = "Sphere", meta = (DataflowInput));
+	float Radius = 0.f;
+
+	UPROPERTY(meta = (DataflowOutput));
+	FSphere Sphere;
+
+	FMakeSphereDataflowNode(const Dataflow::FNodeParameters& InParam, FGuid InGuid = FGuid::NewGuid())
+		: FDataflowNode(InParam, InGuid)
+	{
+		RegisterInputConnection(&Center);
+		RegisterInputConnection(&Radius);
+		RegisterOutputConnection(&Sphere);
+	}
+
+	virtual void Evaluate(Dataflow::FContext& Context, const FDataflowOutput* Out) const override;
+
+};
+
+
 /**
  *
  * Description for this node
@@ -490,7 +525,7 @@ USTRUCT()
 struct FMakeLiteralFloatDataflowNode : public FDataflowNode
 {
 	GENERATED_USTRUCT_BODY()
-		DATAFLOW_NODE_DEFINE_INTERNAL(FMakeLiteralFloatDataflowNode, "MakeLiteralFloat", "Math|Float", "")
+	DATAFLOW_NODE_DEFINE_INTERNAL(FMakeLiteralFloatDataflowNode, "MakeLiteralFloat", "Math|Float", "")
 
 public:
 	UPROPERTY(EditAnywhere, Category = "Float");
@@ -981,7 +1016,8 @@ USTRUCT(meta = (DataflowGeometryCollection))
 struct FExplodedViewDataflowNode : public FDataflowNode
 {
 	GENERATED_USTRUCT_BODY()
-	DATAFLOW_NODE_DEFINE_INTERNAL(FExplodedViewDataflowNode, "ExplodedView", "Fracture|Utilities", "")
+	DATAFLOW_NODE_DEFINE_INTERNAL(FExplodedViewDataflowNode, "ExplodedView", "GeometryCollection|Fracture|Utilities", "")
+	DATAFLOW_NODE_RENDER_TYPE(FGeometryCollection::StaticType(), "Collection")
 
 public:
 	UPROPERTY(meta = (DataflowInput, DataflowOutput))
@@ -1291,10 +1327,10 @@ public:
  *
  */
 USTRUCT(meta = (DataflowGeometryCollection))
-struct FGetBoundingBoxesDataflowNode : public FDataflowNode
+struct FGetBoundingBoxesFromCollectionDataflowNode : public FDataflowNode
 {
 	GENERATED_USTRUCT_BODY()
-	DATAFLOW_NODE_DEFINE_INTERNAL(FGetBoundingBoxesDataflowNode, "GetBoundingBoxes", "GeometryCollection", "")
+	DATAFLOW_NODE_DEFINE_INTERNAL(FGetBoundingBoxesFromCollectionDataflowNode, "GetBoundingBoxesFromCollection", "GeometryCollection|Utilities", "")
 
 public:
 	/** Input Collection */
@@ -1309,7 +1345,7 @@ public:
 	UPROPERTY(meta = (DataflowOutput))
 	TArray<FBox> BoundingBoxes;
 
-	FGetBoundingBoxesDataflowNode(const Dataflow::FNodeParameters& InParam, FGuid InGuid = FGuid::NewGuid())
+	FGetBoundingBoxesFromCollectionDataflowNode(const Dataflow::FNodeParameters& InParam, FGuid InGuid = FGuid::NewGuid())
 		: FDataflowNode(InParam, InGuid)
 	{
 		RegisterInputConnection(&Collection);
@@ -1328,10 +1364,10 @@ public:
  *
  */
 USTRUCT(meta = (DataflowGeometryCollection))
-struct FGetCentroidsDataflowNode : public FDataflowNode
+struct FGetCentroidsFromCollectionDataflowNode : public FDataflowNode
 {
 	GENERATED_USTRUCT_BODY()
-	DATAFLOW_NODE_DEFINE_INTERNAL(FGetCentroidsDataflowNode, "GetCentroids", "GeometryCollection", "")
+	DATAFLOW_NODE_DEFINE_INTERNAL(FGetCentroidsFromCollectionDataflowNode, "GetCentroidsFromCollection", "GeometryCollection|Utilities", "")
 
 public:
 	/** Input Collection */
@@ -1346,7 +1382,7 @@ public:
 	UPROPERTY(meta = (DataflowOutput))
 	TArray<FVector> Centroids = TArray<FVector>();
 
-	FGetCentroidsDataflowNode(const Dataflow::FNodeParameters& InParam, FGuid InGuid = FGuid::NewGuid())
+	FGetCentroidsFromCollectionDataflowNode(const Dataflow::FNodeParameters& InParam, FGuid InGuid = FGuid::NewGuid())
 		: FDataflowNode(InParam, InGuid)
 	{
 		RegisterInputConnection(&Collection);
@@ -1551,7 +1587,7 @@ USTRUCT(meta = (DataflowGeometryCollection))
 struct FRemoveOnBreakDataflowNode : public FDataflowNode
 {
 	GENERATED_USTRUCT_BODY()
-	DATAFLOW_NODE_DEFINE_INTERNAL(FRemoveOnBreakDataflowNode, "RemoveOnBreak", "GeometryCollection", "")
+	DATAFLOW_NODE_DEFINE_INTERNAL(FRemoveOnBreakDataflowNode, "RemoveOnBreak", "GeometryCollection|Utilities", "")
 
 public:
 	/** Collection to set theremoval data on */
@@ -1614,7 +1650,7 @@ USTRUCT(meta = (DataflowGeometryCollection))
 struct FSetAnchorStateDataflowNode : public FDataflowNode
 {
 	GENERATED_USTRUCT_BODY()
-	DATAFLOW_NODE_DEFINE_INTERNAL(FSetAnchorStateDataflowNode, "SetAnchorState", "GeometryCollection", "")
+	DATAFLOW_NODE_DEFINE_INTERNAL(FSetAnchorStateDataflowNode, "SetAnchorState", "GeometryCollection|Utilities", "")
 
 public:
 	/** What anchor state to set on selected bones */
@@ -1667,7 +1703,7 @@ USTRUCT(meta = (DataflowGeometryCollection))
 struct FProximityDataflowNode : public FDataflowNode
 {
 	GENERATED_USTRUCT_BODY()
-	DATAFLOW_NODE_DEFINE_INTERNAL(FProximityDataflowNode, "Proximity", "GeometryCollection", "")
+	DATAFLOW_NODE_DEFINE_INTERNAL(FProximityDataflowNode, "Proximity", "GeometryCollection|Utilities", "")
 
 public:
 	/** Which method to use to decide whether a given piece of geometry is in proximity with another */
@@ -1698,6 +1734,289 @@ public:
 };
 
 
+/**
+ *
+ * Sets pivot for Collection
+ *
+ */
+USTRUCT(meta = (DataflowGeometryCollection))
+struct FCollectionSetPivotDataflowNode : public FDataflowNode
+{
+	GENERATED_USTRUCT_BODY()
+	DATAFLOW_NODE_DEFINE_INTERNAL(FCollectionSetPivotDataflowNode, "SetPivot", "GeometryCollection|Utilities", "")
+
+public:
+	/** Collection for the pivot change */
+	UPROPERTY(meta = (DataflowInput, DataflowOutput, DataflowPassthrough = "Collection", DataflowIntrinsic))
+	FManagedArrayCollection Collection;
+
+	/** Pivot transform */
+	UPROPERTY(EditAnywhere, Category = "Pivot", meta = (DataflowInput))
+	FTransform Transform;
+
+	FCollectionSetPivotDataflowNode(const Dataflow::FNodeParameters& InParam, FGuid InGuid = FGuid::NewGuid())
+		: FDataflowNode(InParam, InGuid)
+	{
+		RegisterInputConnection(&Collection);
+		RegisterInputConnection(&Transform);
+		RegisterOutputConnection(&Collection, &Collection);
+	}
+
+	virtual void Evaluate(Dataflow::FContext& Context, const FDataflowOutput* Out) const override;
+
+};
+
+
+UENUM(BlueprintType)
+enum class EStandardGroupNameEnum : uint8
+{
+	/** Transform group */
+	Dataflow_EStandardGroupNameEnum_Transform UMETA(DisplayName = "Transform"),
+	/** Geometry group */
+	Dataflow_EStandardGroupNameEnum_Geometry UMETA(DisplayName = "Geometry"),
+	/** Faces group */
+	Dataflow_EStandardGroupNameEnum_Faces UMETA(DisplayName = "Faces"),
+	/** Vertices group */
+	Dataflow_EStandardGroupNameEnum_Vertices UMETA(DisplayName = "Vertices"),
+	/** Material group */
+	Dataflow_EStandardGroupNameEnum_Material UMETA(DisplayName = "Material"),
+	/** Breaking group */
+	Dataflow_EStandardGroupNameEnum_Breaking UMETA(DisplayName = "Breaking"),
+	/** User specified group */
+	Dataflow_EStandardGroupNameEnum_Custom UMETA(DisplayName = "Custom"),
+	//~~~
+	//256th entry
+	Dataflow_Max                UMETA(Hidden)
+};
+
+UENUM(BlueprintType)
+enum class ECustomAttributeTypeEnum : uint8
+{
+	/** Int custom attribute type */
+	Dataflow_CustomAttributeType_Int UMETA(DisplayName = "Int"),
+	/** Float custom attribute type */
+	Dataflow_CustomAttributeType_Float UMETA(DisplayName = "Float"),
+	/** Bool custom attribute type */
+	Dataflow_CustomAttributeType_Bool UMETA(DisplayName = "Bool"),
+	/** Vector custom attribute type */
+	Dataflow_CustomAttributeType_Vector UMETA(DisplayName = "Vector"),
+	//~~~
+	//256th entry
+	Dataflow_Max                UMETA(Hidden)
+};
+
+/**
+ *
+ * Adds custom attribute to Collection
+ *
+ */
+USTRUCT(meta = (DataflowGeometryCollection))
+struct FAddCustomCollectionAttributeDataflowNode : public FDataflowNode
+{
+	GENERATED_USTRUCT_BODY()
+	DATAFLOW_NODE_DEFINE_INTERNAL(FAddCustomCollectionAttributeDataflowNode, "AddCustomCollectionAttribute", "GeometryCollection|Utilities", "")
+
+public:
+	/** Collection for the custom attribute */
+	UPROPERTY(meta = (DataflowInput, DataflowOutput, DataflowPassthrough = "Collection", DataflowIntrinsic))
+	FManagedArrayCollection Collection;
+
+	/** Standard group names */
+	UPROPERTY(EditAnywhere, Category = "Attribute", meta = (DisplayName = "Group"))
+	EStandardGroupNameEnum GroupName = EStandardGroupNameEnum::Dataflow_EStandardGroupNameEnum_Transform;
+
+	/** User specified group name */
+	UPROPERTY(EditAnywhere, Category = "Attribute", meta = (DisplayName = "Custom Group", EditCondition = "GroupName == EStandardGroupNameEnum::Dataflow_EStandardGroupNameEnum_Custom"))
+	FString CustomGroupName = FString("");
+
+	/** Attribute name */
+	UPROPERTY(EditAnywhere, Category = "Attribute", meta = (DisplayName = "Attribute"))
+	FString AttrName = FString("");
+
+	/** Attribute type */
+	UPROPERTY(EditAnywhere, Category = "Attribute", meta = (DisplayName = "Attribute Type"));
+	ECustomAttributeTypeEnum CustomAttributeType = ECustomAttributeTypeEnum::Dataflow_CustomAttributeType_Float;
+
+	/** Number of elements for the attribute */
+	UPROPERTY(EditAnywhere, Category = "Attribute", meta = (DisplayName = "Number of Elements"));
+	int32 NumElements = 0;
+
+	FAddCustomCollectionAttributeDataflowNode(const Dataflow::FNodeParameters& InParam, FGuid InGuid = FGuid::NewGuid())
+		: FDataflowNode(InParam, InGuid)
+	{
+		RegisterInputConnection(&Collection);
+		RegisterInputConnection(&NumElements);
+		RegisterOutputConnection(&Collection, &Collection);
+	}
+
+	virtual void Evaluate(Dataflow::FContext& Context, const FDataflowOutput* Out) const override;
+
+};
+
+
+/**
+ *
+ * Returns number of elements in a group in a Collection
+ *
+ */
+USTRUCT(meta = (DataflowGeometryCollection))
+struct FGetNumElementsInCollectionGroupDataflowNode : public FDataflowNode
+{
+	GENERATED_USTRUCT_BODY()
+	DATAFLOW_NODE_DEFINE_INTERNAL(FGetNumElementsInCollectionGroupDataflowNode, "GetNumElementsInCollectionGroup", "GeometryCollection|Utilities", "")
+
+public:
+	/** Collection for the custom attribute */
+	UPROPERTY(meta = (DataflowInput, DataflowIntrinsic))
+	FManagedArrayCollection Collection;
+
+	/** Standard group names */
+	UPROPERTY(EditAnywhere, Category = "Attribute", meta = (DisplayName = "Group"))
+	EStandardGroupNameEnum GroupName = EStandardGroupNameEnum::Dataflow_EStandardGroupNameEnum_Transform;
+
+	/** User specified group name */
+	UPROPERTY(EditAnywhere, Category = "Attribute", meta = (DisplayName = "Custom Group", EditCondition = "GroupName == EStandardGroupNameEnum::Dataflow_EStandardGroupNameEnum_Custom"))
+	FString CustomGroupName = FString("");
+
+	/** Number of elements for the attribute */
+	UPROPERTY(meta = (DataflowOutput));
+	int32 NumElements = 0;
+
+	FGetNumElementsInCollectionGroupDataflowNode(const Dataflow::FNodeParameters& InParam, FGuid InGuid = FGuid::NewGuid())
+		: FDataflowNode(InParam, InGuid)
+	{
+		RegisterInputConnection(&Collection);
+		RegisterOutputConnection(&NumElements);
+	}
+
+	virtual void Evaluate(Dataflow::FContext& Context, const FDataflowOutput* Out) const override;
+
+};
+
+
+/**
+ *
+ * Get attribute data from a Collection
+ *
+ */
+USTRUCT(meta = (DataflowGeometryCollection))
+struct FGetCollectionAttributeDataTypedDataflowNode : public FDataflowNode
+{
+	GENERATED_USTRUCT_BODY()
+	DATAFLOW_NODE_DEFINE_INTERNAL(FGetCollectionAttributeDataTypedDataflowNode, "GetCollectionAttributeDataTyped", "GeometryCollection|Utilities", "")
+
+public:
+	/** Collection for the custom attribute */
+	UPROPERTY(meta = (DataflowInput, DataflowIntrinsic))
+	FManagedArrayCollection Collection;
+
+	/** Standard group names */
+	UPROPERTY(EditAnywhere, Category = "Attribute", meta = (DisplayName = "Group"))
+	EStandardGroupNameEnum GroupName = EStandardGroupNameEnum::Dataflow_EStandardGroupNameEnum_Transform;
+
+	/** User specified group name */
+	UPROPERTY(EditAnywhere, Category = "Attribute", meta = (DisplayName = "Custom Group", EditCondition = "GroupName == EStandardGroupNameEnum::Dataflow_EStandardGroupNameEnum_Custom"))
+	FString CustomGroupName = FString("");
+
+	/** Attribute name */
+	UPROPERTY(EditAnywhere, Category = "Attribute", meta = (DisplayName = "Attribute"))
+	FString AttrName = FString("");
+
+	/** Bool type attribute data */
+	UPROPERTY(meta = (DataflowOutput));
+	TArray<bool> BoolAttributeData;
+
+	/** Float type attribute data */
+	UPROPERTY(meta = (DataflowOutput));
+	TArray<float> FloatAttributeData;
+
+	/** Int type attribute data */
+	UPROPERTY(meta = (DataflowOutput));
+	TArray<int32> IntAttributeData;
+
+	/** Vector type attribute data */
+	UPROPERTY(meta = (DataflowOutput));
+	TArray<FVector> VectorAttributeData;
+
+	FGetCollectionAttributeDataTypedDataflowNode(const Dataflow::FNodeParameters& InParam, FGuid InGuid = FGuid::NewGuid())
+		: FDataflowNode(InParam, InGuid)
+	{
+		RegisterInputConnection(&Collection);
+		RegisterOutputConnection(&BoolAttributeData);
+		RegisterOutputConnection(&FloatAttributeData);
+		RegisterOutputConnection(&IntAttributeData);
+		RegisterOutputConnection(&VectorAttributeData);
+	}
+
+	virtual void Evaluate(Dataflow::FContext& Context, const FDataflowOutput* Out) const override;
+
+};
+
+
+/**
+ *
+ * Set attribute data in a Collection
+ *
+ */
+USTRUCT(meta = (DataflowGeometryCollection))
+struct FSetCollectionAttributeDataTypedDataflowNode : public FDataflowNode
+{
+	GENERATED_USTRUCT_BODY()
+	DATAFLOW_NODE_DEFINE_INTERNAL(FSetCollectionAttributeDataTypedDataflowNode, "SetCollectionAttributeDataTyped", "GeometryCollection|Utilities", "")
+
+public:
+	/** Collection for the custom attribute */
+	UPROPERTY(meta = (DataflowInput, DataflowOutput, DataflowPassthrough = "Collection", DataflowIntrinsic))
+	FManagedArrayCollection Collection;
+
+	/** Standard group names */
+	UPROPERTY(EditAnywhere, Category = "Attribute", meta = (DisplayName = "Group"))
+	EStandardGroupNameEnum GroupName = EStandardGroupNameEnum::Dataflow_EStandardGroupNameEnum_Transform;
+
+	/** User specified group name */
+	UPROPERTY(EditAnywhere, Category = "Attribute", meta = (DisplayName = "Custom Group", EditCondition = "GroupName == EStandardGroupNameEnum::Dataflow_EStandardGroupNameEnum_Custom"))
+	FString CustomGroupName = FString("");
+
+	/** Attribute name */
+	UPROPERTY(EditAnywhere, Category = "Attribute", meta = (DisplayName = "Attribute"))
+	FString AttrName = FString("");
+
+	/** Bool type attribute data */
+	UPROPERTY(meta = (DataflowInput));
+	TArray<bool> BoolAttributeData;
+
+	/** Float type attribute data */
+	UPROPERTY(meta = (DataflowInput));
+	TArray<float> FloatAttributeData;
+
+	/** Int type attribute data */
+	UPROPERTY(meta = (DataflowInput));
+	TArray<int32> IntAttributeData;
+
+	/** Vector type attribute data */
+	UPROPERTY(meta = (DataflowInput));
+	TArray<FVector> VectorAttributeData;
+
+	/** If yes then create attribute if it doesn't exist */
+//	UPROPERTY(EditAnywhere, Category = "Attribute", meta = (DisplayName = "Attribute"))
+//	FString AttrName = FString("");
+
+	FSetCollectionAttributeDataTypedDataflowNode(const Dataflow::FNodeParameters& InParam, FGuid InGuid = FGuid::NewGuid())
+		: FDataflowNode(InParam, InGuid)
+	{
+		RegisterInputConnection(&Collection);
+		RegisterInputConnection(&BoolAttributeData);
+		RegisterInputConnection(&FloatAttributeData);
+		RegisterInputConnection(&IntAttributeData);
+		RegisterInputConnection(&VectorAttributeData);
+		RegisterOutputConnection(&Collection, &Collection);
+	}
+
+	virtual void Evaluate(Dataflow::FContext& Context, const FDataflowOutput* Out) const override;
+
+};
+
+
 namespace Dataflow
 {
 	void GeometryCollectionEngineNodes();
diff --git a/Engine/Plugins/Experimental/GeometryCollectionPlugin/Source/GeometryCollectionNodes/Public/Dataflow/GeometryCollectionSelectionNodes.h b/Engine/Plugins/Experimental/GeometryCollectionPlugin/Source/GeometryCollectionNodes/Public/Dataflow/GeometryCollectionSelectionNodes.h
--- a/Engine/Plugins/Experimental/GeometryCollectionPlugin/Source/GeometryCollectionNodes/Public/Dataflow/GeometryCollectionSelectionNodes.h
+++ b/Engine/Plugins/Experimental/GeometryCollectionPlugin/Source/GeometryCollectionNodes/Public/Dataflow/GeometryCollectionSelectionNodes.h
@@ -6,7 +6,8 @@
 #include "Dataflow/DataflowEngine.h"
 #include "GeometryCollection/ManagedArrayCollection.h"
 #include "Dataflow/DataflowSelection.h"
-
+#include "Math/MathFwd.h"
+#include "Math/Sphere.h"
 #include "GeometryCollectionSelectionNodes.generated.h"
 
 
@@ -213,10 +214,10 @@ struct FCollectionTransformSelectionRandomDataflowNode : public FDataflowNode
 public:
 	/** If true, it always generates the same result for the same RandomSeed */
 	UPROPERTY(EditAnywhere, Category = "Random")
-	bool Deterministic = false;
+	bool bDeterministic = false;
 
 	/** Seed for the random generation, only used if Deterministic is on */
-	UPROPERTY(EditAnywhere, Category = "Random", meta = (DataflowInput, EditCondition = "Deterministic"))
+	UPROPERTY(EditAnywhere, Category = "Random", meta = (DataflowInput, EditCondition = "bDeterministic"))
 	float RandomSeed = 0.f;
 
 	/** Bones get selected if RandomValue > RandomThreshold */
@@ -375,10 +376,10 @@ public:
 
 	/** Sets the random generation to deterministic */
 	UPROPERTY(EditAnywhere, Category = "Random")
-	bool Deterministic = false;
+	bool bDeterministic = false;
 
 	/** Seed value for the random generation */
-	UPROPERTY(EditAnywhere, Category = "Random", meta = (DataflowInput, EditCondition = "Deterministic"))
+	UPROPERTY(EditAnywhere, Category = "Random", meta = (DataflowInput, EditCondition = "bDeterministic"))
 	float RandomSeed = 0.f;
 
 	FCollectionTransformSelectionByPercentageDataflowNode(const Dataflow::FNodeParameters& InParam, FGuid InGuid = FGuid::NewGuid())
@@ -598,6 +599,18 @@ public:
 };
 
 
+UENUM(BlueprintType)
+enum class ERangeSettingEnum : uint8
+{
+	/** Values for selection must be inside of the specified range */
+	Dataflow_RangeSetting_InsideRange UMETA(DisplayName = "Inside Range"),
+	/** Values for selection must be outside of the specified range */
+	Dataflow_RangeSetting_OutsideRange UMETA(DisplayName = "Outside Range"),
+	//~~~
+	//256th entry
+	Dataflow_Max                UMETA(Hidden)
+};
+
 /**
  *
  * Selects pieces based on their size
@@ -622,6 +635,14 @@ public:
 	UPROPERTY(EditAnywhere, Category = "Size", meta = (UIMin = 0.f, UIMax = 1000000000.f))
 	float SizeMax = 1000.f;
 
+	/** Values for the selection has to be inside or outside [Min, Max] range */
+	UPROPERTY(EditAnywhere, Category = "Size")
+	ERangeSettingEnum RangeSetting = ERangeSettingEnum::Dataflow_RangeSetting_InsideRange;
+
+	/** If true then range includes Min and Max values */
+	UPROPERTY(EditAnywhere, Category = "Size")
+	bool bInclusive = true;
+
 	/** Array of the selected bone indicies */
 	UPROPERTY(meta = (DataflowOutput, DisplayName = "TransformSelection"))
 	FDataflowTransformSelection TransformSelection;
@@ -665,6 +686,14 @@ public:
 	UPROPERTY(EditAnywhere, Category = "Volume", meta = (UIMin = 0.f, UIMax = 1000000000.f))
 	float VolumeMax = 1000.f;
 
+	/** Values for the selection has to be inside or outside [Min, Max] range */
+	UPROPERTY(EditAnywhere, Category = "Volume")
+	ERangeSettingEnum RangeSetting = ERangeSettingEnum::Dataflow_RangeSetting_InsideRange;
+
+	/** If true then range includes Min and Max values */
+	UPROPERTY(EditAnywhere, Category = "Volume")
+	bool bInclusive = true;
+
 	/** Array of the selected bone indicies */
 	UPROPERTY(meta = (DataflowOutput, DisplayName = "TransformSelection"))
 	FDataflowTransformSelection TransformSelection;
@@ -684,6 +713,368 @@ public:
 };
 
 
+UENUM(BlueprintType)
+enum class ESelectSubjectTypeEnum : uint8
+{
+	/** InBox must contain the vertices of the bone */
+	Dataflow_SelectSubjectType_Vertices UMETA(DisplayName = "Vertices"),
+	/** InBox must contain the BoundingBox of the bone */
+	Dataflow_SelectSubjectType_BoundingBox UMETA(DisplayName = "BoundingBox"),
+	/** InBox must contain the centroid of the bone */
+	Dataflow_SelectSubjectType_Centroid UMETA(DisplayName = "Centroid"),
+	//~~~
+	//256th entry
+	Dataflow_Max                UMETA(Hidden)
+};
+
+/**
+ *
+ * Selects bones if their Vertices/BoundingBox/Centroid in a box
+ *
+ */
+USTRUCT(meta = (DataflowGeometryCollection))
+struct FCollectionTransformSelectionInBoxDataflowNode : public FDataflowNode
+{
+	GENERATED_USTRUCT_BODY()
+	DATAFLOW_NODE_DEFINE_INTERNAL(FCollectionTransformSelectionInBoxDataflowNode, "CollectionTransformSelectInBox", "GeometryCollection|Selection", "")
+
+public:
+	/** GeometryCollection for the selection */
+	UPROPERTY(meta = (DataflowInput, DataflowOutput, DataflowPassthrough = "Collection", DataflowIntrinsic))
+	FManagedArrayCollection Collection;
+
+	/** Box to contain Vertices/BoundingBox/Centroid */
+	UPROPERTY(meta = (DataflowInput, DataflowIntrinsic))
+	FBox Box;
+
+	/** Transform for the box */
+	UPROPERTY(EditAnywhere, Category = "Select")
+	FTransform Transform;
+
+	/** Subject (Vertices/BoundingBox/Centroid) to check against box */
+	UPROPERTY(EditAnywhere, Category = "Select", DisplayName = "Type to Check in Box")
+	ESelectSubjectTypeEnum Type = ESelectSubjectTypeEnum::Dataflow_SelectSubjectType_Centroid;
+
+	/** If true all the vertices of the piece must be inside of box */
+	UPROPERTY(EditAnywhere, Category = "Select", meta = (EditCondition = "Type == ESelectSubjectTypeEnum::Dataflow_SelectSubjectType_Vertices"))
+	bool bAllVerticesMustContainedInBox = true;
+
+	/** Array of the selected bone indicies */
+	UPROPERTY(meta = (DataflowOutput, DisplayName = "TransformSelection"))
+	FDataflowTransformSelection TransformSelection;
+
+	FCollectionTransformSelectionInBoxDataflowNode(const Dataflow::FNodeParameters& InParam, FGuid InGuid = FGuid::NewGuid())
+		: FDataflowNode(InParam, InGuid)
+	{
+		RegisterInputConnection(&Collection);
+		RegisterInputConnection(&Box);
+		RegisterInputConnection(&Transform);
+		RegisterOutputConnection(&TransformSelection);
+		RegisterOutputConnection(&Collection, &Collection);
+	}
+
+	virtual void Evaluate(Dataflow::FContext& Context, const FDataflowOutput* Out) const override;
+
+};
+
+
+/**
+ *
+ * Selects bones if their Vertices/BoundingBox/Centroid in a sphere
+ *
+ */
+USTRUCT(meta = (DataflowGeometryCollection))
+struct FCollectionTransformSelectionInSphereDataflowNode : public FDataflowNode
+{
+	GENERATED_USTRUCT_BODY()
+	DATAFLOW_NODE_DEFINE_INTERNAL(FCollectionTransformSelectionInSphereDataflowNode, "CollectionTransformSelectInSphere", "GeometryCollection|Selection", "")
+
+public:
+	/** GeometryCollection for the selection */
+	UPROPERTY(meta = (DataflowInput, DataflowOutput, DataflowPassthrough = "Collection", DataflowIntrinsic))
+	FManagedArrayCollection Collection;
+
+	/** Sphere to contain Vertices/BoundingBox/Centroid */
+	UPROPERTY(meta = (DataflowInput, DataflowIntrinsic))
+	FSphere Sphere;
+
+	/** Transform for the sphere */
+	UPROPERTY(EditAnywhere, Category = "Select")
+	FTransform Transform;
+
+	/** Subject (Vertices/BoundingBox/Centroid) to check against box */
+	UPROPERTY(EditAnywhere, Category = "Select", DisplayName = "Type to Check in Sphere")
+	ESelectSubjectTypeEnum Type = ESelectSubjectTypeEnum::Dataflow_SelectSubjectType_Centroid;
+
+	/** If true all the vertices of the piece must be inside of box */
+	UPROPERTY(EditAnywhere, Category = "Select", meta = (EditCondition = "Type == ESelectSubjectTypeEnum::Dataflow_SelectSubjectType_Vertices"))
+	bool bAllVerticesMustContainedInSphere = true;
+
+	/** Array of the selected bone indicies */
+	UPROPERTY(meta = (DataflowOutput, DisplayName = "TransformSelection"))
+	FDataflowTransformSelection TransformSelection;
+
+	FCollectionTransformSelectionInSphereDataflowNode(const Dataflow::FNodeParameters& InParam, FGuid InGuid = FGuid::NewGuid())
+		: FDataflowNode(InParam, InGuid)
+	{
+		RegisterInputConnection(&Collection);
+		RegisterInputConnection(&Sphere);
+		RegisterInputConnection(&Transform);
+		RegisterOutputConnection(&TransformSelection);
+		RegisterOutputConnection(&Collection, &Collection);
+	}
+
+	virtual void Evaluate(Dataflow::FContext& Context, const FDataflowOutput* Out) const override;
+
+};
+
+
+/**
+ *
+ * Selects bones by a float attribute
+ *
+ */
+USTRUCT(meta = (DataflowGeometryCollection))
+struct FCollectionTransformSelectionByFloatAttrDataflowNode : public FDataflowNode
+{
+	GENERATED_USTRUCT_BODY()
+	DATAFLOW_NODE_DEFINE_INTERNAL(FCollectionTransformSelectionByFloatAttrDataflowNode, "CollectionTransformSelectByFloatAttribute", "GeometryCollection|Selection", "")
+
+public:
+	/** GeometryCollection for the selection */
+	UPROPERTY(meta = (DataflowInput, DataflowOutput, DataflowPassthrough = "Collection", DataflowIntrinsic))
+	FManagedArrayCollection Collection;
+
+	/** Group name for the attr */
+	UPROPERTY(VisibleAnywhere, Category = "Attribute", meta = (DisplayName = "Group"))
+	FString GroupName = FString("Transform");
+
+	/** Attribute name */
+	UPROPERTY(EditAnywhere, Category = "Attribute", meta = (DisplayName = "Attribute"))
+	FString AttrName = FString("");
+
+	/** Minimum value for the selection */
+	UPROPERTY(EditAnywhere, Category = "Attribute", meta = (UIMin = 0.f, UIMax = 1000000000.f))
+	float Min = 0.f;
+
+	/** Maximum value for the selection */
+	UPROPERTY(EditAnywhere, Category = "Attribute", meta = (UIMin = 0.f, UIMax = 1000000000.f))
+	float Max = 1000.f;
+
+	/** Values for the selection has to be inside or outside [Min, Max] range */
+	UPROPERTY(EditAnywhere, Category = "Attribute")
+	ERangeSettingEnum RangeSetting = ERangeSettingEnum::Dataflow_RangeSetting_InsideRange;
+
+	/** If true then range includes Min and Max values */
+	UPROPERTY(EditAnywhere, Category = "Attribute")
+	bool bInclusive = true;
+
+	/** Array of the selected bone indicies */
+	UPROPERTY(meta = (DataflowOutput, DisplayName = "TransformSelection"))
+	FDataflowTransformSelection TransformSelection;
+
+	FCollectionTransformSelectionByFloatAttrDataflowNode(const Dataflow::FNodeParameters& InParam, FGuid InGuid = FGuid::NewGuid())
+		: FDataflowNode(InParam, InGuid)
+	{
+		RegisterInputConnection(&Collection);
+		RegisterInputConnection(&Min);
+		RegisterInputConnection(&Max);
+		RegisterOutputConnection(&TransformSelection);
+		RegisterOutputConnection(&Collection, &Collection);
+	}
+
+	virtual void Evaluate(Dataflow::FContext& Context, const FDataflowOutput* Out) const override;
+
+};
+
+
+/**
+ *
+ * Selects bones by an int attribute
+ *
+ */
+USTRUCT(meta = (DataflowGeometryCollection))
+struct FCollectionTransformSelectionByIntAttrDataflowNode : public FDataflowNode
+{
+	GENERATED_USTRUCT_BODY()
+	DATAFLOW_NODE_DEFINE_INTERNAL(FCollectionTransformSelectionByIntAttrDataflowNode, "CollectionTransformSelectByIntAttribute", "GeometryCollection|Selection", "")
+
+public:
+	/** GeometryCollection for the selection */
+	UPROPERTY(meta = (DataflowInput, DataflowOutput, DataflowPassthrough = "Collection", DataflowIntrinsic))
+	FManagedArrayCollection Collection;
+
+	/** Group name for the attr */
+	UPROPERTY(VisibleAnywhere, Category = "Attribute", meta = (DisplayName = "Group"))
+	FString GroupName = FString("Transform");
+
+	/** Attribute name */
+	UPROPERTY(EditAnywhere, Category = "Attribute", meta = (DisplayName = "Attribute"))
+	FString AttrName = FString("");
+
+	/** Minimum value for the selection */
+	UPROPERTY(EditAnywhere, Category = "Attribute", meta = (UIMin = 0, UIMax = 1000000000))
+	int32 Min = 0;
+
+	/** Maximum value for the selection */
+	UPROPERTY(EditAnywhere, Category = "Attribute", meta = (UIMin = 0, UIMax = 1000000000))
+	int32 Max = 1000;
+
+	/** Values for the selection has to be inside or outside [Min, Max] range */
+	UPROPERTY(EditAnywhere, Category = "Attribute")
+	ERangeSettingEnum RangeSetting = ERangeSettingEnum::Dataflow_RangeSetting_InsideRange;
+
+	/** If true then range includes Min and Max values */
+	UPROPERTY(EditAnywhere, Category = "Attribute")
+	bool bInclusive = true;
+
+	/** Transform selection including the new indicies */
+	UPROPERTY(meta = (DataflowOutput, DisplayName = "TransformSelection"))
+	FDataflowTransformSelection TransformSelection;
+
+	FCollectionTransformSelectionByIntAttrDataflowNode(const Dataflow::FNodeParameters& InParam, FGuid InGuid = FGuid::NewGuid())
+		: FDataflowNode(InParam, InGuid)
+	{
+		RegisterInputConnection(&Collection);
+		RegisterInputConnection(&Min);
+		RegisterInputConnection(&Max);
+		RegisterOutputConnection(&TransformSelection);
+		RegisterOutputConnection(&Collection, &Collection);
+	}
+
+	virtual void Evaluate(Dataflow::FContext& Context, const FDataflowOutput* Out) const override;
+
+};
+
+
+/**
+ *
+ * Selects specified vertices in the GeometryCollection by using a
+ * space separated list
+ *
+ */
+USTRUCT(meta = (DataflowGeometryCollection))
+struct FCollectionVertexSelectionCustomDataflowNode : public FDataflowNode
+{
+	GENERATED_USTRUCT_BODY()
+	DATAFLOW_NODE_DEFINE_INTERNAL(FCollectionVertexSelectionCustomDataflowNode, "CollectionVertexSelectCustom", "GeometryCollection|Selection", "")
+
+public:
+	/** GeometryCollection for the selection */
+	UPROPERTY(meta = (DataflowInput, DataflowOutput, DataflowPassthrough = "Collection", DataflowIntrinsic))
+	FManagedArrayCollection Collection;
+
+	/** Space separated list of vertex indicies to specify the selection */
+	UPROPERTY(EditAnywhere, Category = "Selection")
+	FString VertexIndicies = FString();
+
+	/** Vertex selection including the new indicies */
+	UPROPERTY(meta = (DataflowOutput, DisplayName = "VertexSelection"))
+	FDataflowVertexSelection VertexSelection;
+
+	FCollectionVertexSelectionCustomDataflowNode(const Dataflow::FNodeParameters& InParam, FGuid InGuid = FGuid::NewGuid())
+		: FDataflowNode(InParam, InGuid)
+	{
+		RegisterInputConnection(&Collection);
+		RegisterInputConnection(&VertexIndicies);
+		RegisterOutputConnection(&VertexSelection);
+		RegisterOutputConnection(&Collection, &Collection);
+	}
+
+	virtual void Evaluate(Dataflow::FContext& Context, const FDataflowOutput* Out) const override;
+
+};
+
+
+/**
+ *
+ * Selects specified faces in the GeometryCollection by using a
+ * space separated list
+ *
+ */
+USTRUCT(meta = (DataflowGeometryCollection))
+struct FCollectionFaceSelectionCustomDataflowNode : public FDataflowNode
+{
+	GENERATED_USTRUCT_BODY()
+	DATAFLOW_NODE_DEFINE_INTERNAL(FCollectionFaceSelectionCustomDataflowNode, "CollectionFaceSelectCustom", "GeometryCollection|Selection", "")
+
+public:
+	/** GeometryCollection for the selection */
+	UPROPERTY(meta = (DataflowInput, DataflowOutput, DataflowPassthrough = "Collection", DataflowIntrinsic))
+	FManagedArrayCollection Collection;
+
+	/** Space separated list of face indicies to specify the selection */
+	UPROPERTY(EditAnywhere, Category = "Selection")
+	FString FaceIndicies = FString();
+
+	/** Face selection including the new indicies */
+	UPROPERTY(meta = (DataflowOutput, DisplayName = "FaceSelection"))
+	FDataflowFaceSelection FaceSelection;
+
+	FCollectionFaceSelectionCustomDataflowNode(const Dataflow::FNodeParameters& InParam, FGuid InGuid = FGuid::NewGuid())
+		: FDataflowNode(InParam, InGuid)
+	{
+		RegisterInputConnection(&Collection);
+		RegisterInputConnection(&FaceIndicies);
+		RegisterOutputConnection(&FaceSelection);
+		RegisterOutputConnection(&Collection, &Collection);
+	}
+
+	virtual void Evaluate(Dataflow::FContext& Context, const FDataflowOutput* Out) const override;
+
+};
+
+
+/**
+ *
+ * Converts Vertex/Face/Transform selection into Vertex/Face/Transform selection
+ *
+ */
+USTRUCT(meta = (DataflowGeometryCollection))
+struct FCollectionSelectionConvertDataflowNode : public FDataflowNode
+{
+	GENERATED_USTRUCT_BODY()
+	DATAFLOW_NODE_DEFINE_INTERNAL(FCollectionSelectionConvertDataflowNode, "CollectionSelectionConvert", "GeometryCollection|Selection", "")
+
+public:
+	/** GeometryCollection for the selection */
+	UPROPERTY(meta = (DataflowInput, DataflowOutput, DataflowPassthrough = "Collection", DataflowIntrinsic))
+	FManagedArrayCollection Collection;
+
+	/** Transform selection including the new indicies */
+	UPROPERTY(meta = (DataflowInput, DataflowOutput, DataflowPassthrough = "TransformSelection", DisplayName = "TransformSelection"))
+	FDataflowTransformSelection TransformSelection;
+
+	/** Face selection including the new indicies */
+	UPROPERTY(meta = (DataflowInput, DataflowOutput, DataflowPassthrough = "FaceSelection", DisplayName = "FaceSelection"))
+	FDataflowFaceSelection FaceSelection;
+
+	/** Vertex selection including the new indicies */
+	UPROPERTY(meta = (DataflowInput, DataflowOutput, DataflowPassthrough = "VertexSelection", DisplayName = "VertexSelection"))
+	FDataflowVertexSelection VertexSelection;
+	
+	/** If true then for converting vertex/face selection to transform selection all vertex/face must be selected for selecting the associated transform */
+	UPROPERTY(EditAnywhere, Category = "Selection")
+	bool bAllElementsMustBeSelected = false;
+
+	FCollectionSelectionConvertDataflowNode(const Dataflow::FNodeParameters& InParam, FGuid InGuid = FGuid::NewGuid())
+		: FDataflowNode(InParam, InGuid)
+	{
+		RegisterInputConnection(&Collection);
+		RegisterInputConnection(&VertexSelection);
+		RegisterInputConnection(&FaceSelection);
+		RegisterInputConnection(&TransformSelection);
+		RegisterOutputConnection(&TransformSelection, &TransformSelection);
+		RegisterOutputConnection(&FaceSelection, &FaceSelection);
+		RegisterOutputConnection(&VertexSelection, &VertexSelection);
+		RegisterOutputConnection(&Collection, &Collection);
+	}
+
+	virtual void Evaluate(Dataflow::FContext& Context, const FDataflowOutput* Out) const override;
+
+};
+
+
 namespace Dataflow
 {
 	void GeometryCollectionSelectionNodes();
diff --git a/Engine/Plugins/Experimental/GeometryCollectionPlugin/Source/GeometryCollectionNodes/Public/Dataflow/Nodes/GeometryCollectionAssetNodes.h b/Engine/Plugins/Experimental/GeometryCollectionPlugin/Source/GeometryCollectionNodes/Public/Dataflow/Nodes/GeometryCollectionAssetNodes.h
--- a/Engine/Plugins/Experimental/GeometryCollectionPlugin/Source/GeometryCollectionNodes/Public/Dataflow/Nodes/GeometryCollectionAssetNodes.h
+++ b/Engine/Plugins/Experimental/GeometryCollectionPlugin/Source/GeometryCollectionNodes/Public/Dataflow/Nodes/GeometryCollectionAssetNodes.h
@@ -56,7 +56,7 @@ USTRUCT(meta = (DataflowGeometryCollection))
 struct FGetGeometryCollectionAssetDataflowNode : public FDataflowNode
 {
 	GENERATED_USTRUCT_BODY()
-	DATAFLOW_NODE_DEFINE_INTERNAL(FGetGeometryCollectionAssetDataflowNode, "GetGeometryCollectionAsset", "GeometryCollection", "")
+	DATAFLOW_NODE_DEFINE_INTERNAL(FGetGeometryCollectionAssetDataflowNode, "GetGeometryCollectionAsset", "GeometryCollection|Asset", "")
 
 public:
 	FGetGeometryCollectionAssetDataflowNode(const Dataflow::FNodeParameters& InParam, FGuid InGuid = FGuid::NewGuid());
@@ -75,7 +75,7 @@ USTRUCT(meta = (DataflowGeometryCollection))
 struct FGetGeometryCollectionSourcesDataflowNode : public FDataflowNode
 {
 	GENERATED_USTRUCT_BODY()
-	DATAFLOW_NODE_DEFINE_INTERNAL(FGetGeometryCollectionSourcesDataflowNode, "GetGeometryCollectionSources", "GeometryCollection", "")
+	DATAFLOW_NODE_DEFINE_INTERNAL(FGetGeometryCollectionSourcesDataflowNode, "GetGeometryCollectionSources", "GeometryCollection|Utilities", "")
 
 public:
 	FGetGeometryCollectionSourcesDataflowNode(const Dataflow::FNodeParameters& InParam, FGuid InGuid = FGuid::NewGuid());
@@ -97,7 +97,7 @@ USTRUCT(meta = (DataflowGeometryCollection))
 struct FCreateGeometryCollectionFromSourcesDataflowNode : public FDataflowNode
 {
 	GENERATED_USTRUCT_BODY()
-	DATAFLOW_NODE_DEFINE_INTERNAL(FCreateGeometryCollectionFromSourcesDataflowNode, "CreateGeometryCollectionFromSources", "GeometryCollection", "")
+	DATAFLOW_NODE_DEFINE_INTERNAL(FCreateGeometryCollectionFromSourcesDataflowNode, "CreateGeometryCollectionFromSources", "GeometryCollection|Utilities", "")
 
 public:
 	FCreateGeometryCollectionFromSourcesDataflowNode(const Dataflow::FNodeParameters& InParam, FGuid InGuid = FGuid::NewGuid());
diff --git a/Engine/Plugins/Experimental/PlanarCutPlugin/Source/PlanarCut/Private/PlanarCut.cpp b/Engine/Plugins/Experimental/PlanarCutPlugin/Source/PlanarCut/Private/PlanarCut.cpp
--- a/Engine/Plugins/Experimental/PlanarCutPlugin/Source/PlanarCut/Private/PlanarCut.cpp
+++ b/Engine/Plugins/Experimental/PlanarCutPlugin/Source/PlanarCut/Private/PlanarCut.cpp
@@ -104,13 +104,13 @@ int32 FInternalSurfaceMaterials::GetDefaultMaterialIDForGeometry(const FGeometry
 	return InternalMaterialID;
 }
 
-void FInternalSurfaceMaterials::SetUVScaleFromCollection(const GeometryCollection::Facades::FMeshFacade& CollectionMesh, int32 GeometryIdx)
+void FInternalSurfaceMaterials::SetUVScaleFromCollection(const GeometryCollection::Facades::FCollectionMeshFacade& CollectionMesh, int32 GeometryIdx)
 {
-	const auto& VertexArray = CollectionMesh.Vertex.Get();
-	const auto& UVsArray = CollectionMesh.UVs.Get();
-	const auto& IndicesArray = CollectionMesh.Indices.Get();
-	const auto& FaceStartArray = CollectionMesh.FaceStart.Get();
-	const auto& FaceCountArray = CollectionMesh.FaceCount.Get();
+	const auto& VertexArray = CollectionMesh.VertexAttribute.Get();
+	const auto& UVsArray = CollectionMesh.UVsAttribute.Get();
+	const auto& IndicesArray = CollectionMesh.IndicesAttribute.Get();
+	const auto& FaceStartArray = CollectionMesh.FaceStartAttribute.Get();
+	const auto& FaceCountArray = CollectionMesh.FaceCountAttribute.Get();
 	
 	int32 FaceStart = 0;
 	int32 FaceEnd = IndicesArray.Num();
@@ -756,7 +756,7 @@ int32 CutMultipleWithMultiplePlanes(
 	int32 OrigNumGeom = Collection.FaceCount.Num();
 	int32 CurNumGeom = OrigNumGeom;
 
-	GeometryCollection::Facades::FMeshFacade CollectionMesh(Collection);
+	GeometryCollection::Facades::FCollectionMeshFacade CollectionMesh(Collection);
 	if (!CollectionMesh.IsValid())
 	{
 		return -1;
diff --git a/Engine/Plugins/Experimental/PlanarCutPlugin/Source/PlanarCut/Public/PlanarCut.h b/Engine/Plugins/Experimental/PlanarCutPlugin/Source/PlanarCut/Public/PlanarCut.h
--- a/Engine/Plugins/Experimental/PlanarCutPlugin/Source/PlanarCut/Public/PlanarCut.h
+++ b/Engine/Plugins/Experimental/PlanarCutPlugin/Source/PlanarCut/Public/PlanarCut.h
@@ -82,7 +82,7 @@ struct PLANARCUT_API FInternalSurfaceMaterials
 	 * @param Collection	Reference collection to use for setting UV scale
 	 * @param GeometryIdx	Reference geometry inside collection; if -1, use all geometries in collection
 	 */
-	void SetUVScaleFromCollection(const GeometryCollection::Facades::FMeshFacade& CollectionMesh, int32 GeometryIdx = -1);
+	void SetUVScaleFromCollection(const GeometryCollection::Facades::FCollectionMeshFacade& CollectionMesh, int32 GeometryIdx = -1);
 
 	
 	int32 GetDefaultMaterialIDForGeometry(const FGeometryCollection& Collection, int32 GeometryIdx = -1) const;
diff --git a/Engine/Source/Runtime/Experimental/Chaos/Private/GeometryCollection/Facades/CollectionBoundsFacade.cpp b/Engine/Source/Runtime/Experimental/Chaos/Private/GeometryCollection/Facades/CollectionBoundsFacade.cpp
--- a/Engine/Source/Runtime/Experimental/Chaos/Private/GeometryCollection/Facades/CollectionBoundsFacade.cpp
+++ b/Engine/Source/Runtime/Experimental/Chaos/Private/GeometryCollection/Facades/CollectionBoundsFacade.cpp
@@ -7,12 +7,15 @@
 #include "GeometryCollection/Facades/CollectionBoundsFacade.h"
 #include "GeometryCollection/GeometryCollection.h"
 #include "GeometryCollection/Facades/CollectionHierarchyFacade.h"
+#include "GeometryCollection/Facades/CollectionTransformFacade.h"
 
 namespace GeometryCollection::Facades
 {
 
 	FBoundsFacade::FBoundsFacade(FManagedArrayCollection& InCollection)
-		: BoundingBoxAttribute(InCollection,"BoundingBox", FGeometryCollection::GeometryGroup)
+		: ConstCollection(InCollection)
+		, Collection(&InCollection)
+		, BoundingBoxAttribute(InCollection,"BoundingBox", FGeometryCollection::GeometryGroup)
 		, VertexAttribute(InCollection,"Vertex", FGeometryCollection::VerticesGroup, FGeometryCollection::VerticesGroup)
 		, BoneMapAttribute(InCollection,"BoneMap", FGeometryCollection::VerticesGroup)
 		, TransformToGeometryIndexAttribute(InCollection,"TransformToGeometryIndex", FTransformCollection::TransformGroup)
@@ -20,7 +23,9 @@ namespace GeometryCollection::Facades
 	{}
 
 	FBoundsFacade::FBoundsFacade(const FManagedArrayCollection& InCollection)
-		: BoundingBoxAttribute(InCollection, "BoundingBox", FGeometryCollection::GeometryGroup)
+		: ConstCollection(InCollection)
+		, Collection(nullptr)
+		, BoundingBoxAttribute(InCollection, "BoundingBox", FGeometryCollection::GeometryGroup)
 		, VertexAttribute(InCollection, "Vertex", FGeometryCollection::VerticesGroup, FGeometryCollection::VerticesGroup)
 		, BoneMapAttribute(InCollection, "BoneMap", FGeometryCollection::VerticesGroup)
 		, TransformToGeometryIndexAttribute(InCollection, "TransformToGeometryIndex", FTransformCollection::TransformGroup)
@@ -80,22 +85,89 @@ namespace GeometryCollection::Facades
 		}
 	}
 
-	FBox FBoundsFacade::GetBoundingBox()
+	TArray<FVector> FBoundsFacade::GetCentroids() const
 	{
-		const TSet<int32>& RootIndices = Chaos::Facades::FCollectionHierarchyFacade::GetRootIndices(ParentAttribute);
+		TArray<FVector> Centroids;
 
-		const TManagedArray<FBox>& BoundingBoxArr = BoundingBoxAttribute.Get();
+		if (IsValid())
+		{
+			const TManagedArray<FBox>& BoundingBoxes = BoundingBoxAttribute.Get();
+
+			for (int32 Idx = 0; Idx < BoundingBoxes.Num(); ++Idx)
+			{
+				Centroids.Add(BoundingBoxes[Idx].GetCenter());
+			}
+		}
+		
+		return Centroids;
+	}
+
+	FBox FBoundsFacade::GetBoundingBox() const
+	{
+		FBox BoundingBox;
+
+		if (IsValid())
+		{
+			const TArray<int32>& RootIndices = Chaos::Facades::FCollectionHierarchyFacade::GetRootIndices(ParentAttribute);
+
+			const TManagedArray<FBox>& BoundingBoxArr = BoundingBoxAttribute.Get();
+
+			BoundingBox.Init();
+
+			for (int32 Idx : RootIndices)
+			{
+				BoundingBox += BoundingBoxArr[Idx];
+			}
+		}
+
+		return BoundingBox;
+	}
 
+	FBox FBoundsFacade::GetBoundingBoxInCollectionSpace() const
+	{
 		FBox BoundingBox;
-		BoundingBox.Init();
 
-		for (auto& Idx : RootIndices)
+		if (IsValid())
 		{
-			BoundingBox += BoundingBoxArr[Idx];
+			const TArray<int32>& RootIndices = Chaos::Facades::FCollectionHierarchyFacade::GetRootIndices(ParentAttribute);
+
+			const TManagedArray<FBox>& BoundingBoxArr = BoundingBoxAttribute.Get();
+
+			GeometryCollection::Facades::FCollectionTransformFacade TransformFacade(ConstCollection);
+
+			BoundingBox.Init();
+
+			for (int32 Idx : RootIndices)
+			{
+				const FTransform CollectionSpaceTransform = TransformFacade.ComputeCollectionSpaceTransform(Idx);
+				const FBox BoundingBoxInCollectionSpace = BoundingBoxArr[Idx].TransformBy(CollectionSpaceTransform);
+
+				BoundingBox += BoundingBoxInCollectionSpace;
+			}
 		}
 
 		return BoundingBox;
 	}
+
+	TArray<FVector> FBoundsFacade::GetBoundingBoxVertexPositions(const FBox& InBox)
+	{
+		FVector Min = InBox.Min;
+		FVector Max = InBox.Max;
+		FVector Extent(Max.X - Min.X, Max.Y - Min.Y, Max.Z - Min.Z);
+
+		TArray<FVector> Vertices;
+		Vertices.Add(Min);
+		Vertices.Add({ Min.X + Extent.X, Min.Y, Min.Z });
+		Vertices.Add({ Min.X + Extent.X, Min.Y + Extent.Y, Min.Z });
+		Vertices.Add({ Min.X, Min.Y + Extent.Y, Min.Z });
+		Vertices.Add({ Min.X, Min.Y, Min.Z + Extent.Z });
+		Vertices.Add({ Min.X + Extent.X, Min.Y, Min.Z + Extent.Z });
+		Vertices.Add(Max);
+		Vertices.Add({ Min.X, Min.Y + Extent.Y, Min.Z + Extent.Z });
+
+		return Vertices;
+	}
+
 }; // GeometryCollection::Facades
 
 
diff --git a/Engine/Source/Runtime/Experimental/Chaos/Private/GeometryCollection/Facades/CollectionHierarchyFacade.cpp b/Engine/Source/Runtime/Experimental/Chaos/Private/GeometryCollection/Facades/CollectionHierarchyFacade.cpp
--- a/Engine/Source/Runtime/Experimental/Chaos/Private/GeometryCollection/Facades/CollectionHierarchyFacade.cpp
+++ b/Engine/Source/Runtime/Experimental/Chaos/Private/GeometryCollection/Facades/CollectionHierarchyFacade.cpp
@@ -50,6 +50,11 @@ namespace Chaos::Facades
 		return INDEX_NONE;
 	}
 
+	TArray<int32> FCollectionHierarchyFacade::GetRootIndices() const
+	{
+		return GetRootIndices(ParentAttribute);
+	}
+
 	void FCollectionHierarchyFacade::GenerateLevelAttribute()
 	{
 		check(!IsConst());
@@ -103,9 +108,9 @@ namespace Chaos::Facades
 		}
 	}
 
-	TSet<int32> FCollectionHierarchyFacade::GetRootIndices(const TManagedArrayAccessor<int32>& ParentAttribute)
+	TArray<int32> FCollectionHierarchyFacade::GetRootIndices(const TManagedArrayAccessor<int32>& ParentAttribute)
 	{
-		TSet<int32> Roots;
+		TArray<int32> Roots;
 		if (ParentAttribute.IsValid())
 		{
 			const TManagedArray<int32>& Parents = ParentAttribute.Get();
@@ -117,6 +122,7 @@ namespace Chaos::Facades
 				}
 			}
 		}
+
 		return Roots;
 	}
 }
\ No newline at end of file
diff --git a/Engine/Source/Runtime/Experimental/Chaos/Private/GeometryCollection/Facades/CollectionMeshFacade.cpp b/Engine/Source/Runtime/Experimental/Chaos/Private/GeometryCollection/Facades/CollectionMeshFacade.cpp
--- a/Engine/Source/Runtime/Experimental/Chaos/Private/GeometryCollection/Facades/CollectionMeshFacade.cpp
+++ b/Engine/Source/Runtime/Experimental/Chaos/Private/GeometryCollection/Facades/CollectionMeshFacade.cpp
@@ -10,63 +10,131 @@
 
 namespace GeometryCollection::Facades
 {
-	FMeshFacade::FMeshFacade(FManagedArrayCollection& InCollection)
-		: Vertex(InCollection, "Vertex", FGeometryCollection::VerticesGroup)
-		, TangentU(InCollection, "TangentU", FGeometryCollection::VerticesGroup)
-		, TangentV(InCollection, "TangentV", FGeometryCollection::VerticesGroup)
-		, Normal(InCollection, "Normal", FGeometryCollection::VerticesGroup)
-		, UVs(InCollection, "UVs", FGeometryCollection::VerticesGroup)
-		, Color(InCollection, "Color", FGeometryCollection::VerticesGroup)
-		, BoneMap(InCollection, "BoneMap", FGeometryCollection::VerticesGroup)
-		, VertexStart(InCollection, "VertexStart", FGeometryCollection::GeometryGroup)
-		, VertexCount(InCollection, "VertexCount", FGeometryCollection::GeometryGroup)
-		, Indices(InCollection, "Indices", FGeometryCollection::FacesGroup)
-		, Visible(InCollection, "Visible", FGeometryCollection::FacesGroup)
-		, MaterialIndex(InCollection, "MaterialIndex", FGeometryCollection::FacesGroup)
-		, MaterialID(InCollection, "MaterialID", FGeometryCollection::FacesGroup)
-		, FaceStart(InCollection, "FaceStart", FGeometryCollection::GeometryGroup)
-		, FaceCount(InCollection, "FaceCount", FGeometryCollection::GeometryGroup)
+	FCollectionMeshFacade::FCollectionMeshFacade(FManagedArrayCollection& InCollection)
+		: TransformToGeometryIndexAttribute(InCollection, "TransformToGeometryIndex", FGeometryCollection::GeometryGroup)
+		, VertexAttribute(InCollection, "Vertex", FGeometryCollection::VerticesGroup)
+		, TangentUAttribute(InCollection, "TangentU", FGeometryCollection::VerticesGroup)
+		, TangentVAttribute(InCollection, "TangentV", FGeometryCollection::VerticesGroup)
+		, NormalAttribute(InCollection, "Normal", FGeometryCollection::VerticesGroup)
+		, UVsAttribute(InCollection, "UVs", FGeometryCollection::VerticesGroup)
+		, ColorAttribute(InCollection, "Color", FGeometryCollection::VerticesGroup)
+		, BoneMapAttribute(InCollection, "BoneMap", FGeometryCollection::VerticesGroup)
+		, VertexStartAttribute(InCollection, "VertexStart", FGeometryCollection::GeometryGroup)
+		, VertexCountAttribute(InCollection, "VertexCount", FGeometryCollection::GeometryGroup)
+		, IndicesAttribute(InCollection, "Indices", FGeometryCollection::FacesGroup)
+		, VisibleAttribute(InCollection, "Visible", FGeometryCollection::FacesGroup)
+		, MaterialIndexAttribute(InCollection, "MaterialIndex", FGeometryCollection::FacesGroup)
+		, MaterialIDAttribute(InCollection, "MaterialID", FGeometryCollection::FacesGroup)
+		, FaceStartAttribute(InCollection, "FaceStart", FGeometryCollection::GeometryGroup)
+		, FaceCountAttribute(InCollection, "FaceCount", FGeometryCollection::GeometryGroup)
 	{
 	}
 
-	bool FMeshFacade::IsValid() const
+	FCollectionMeshFacade::FCollectionMeshFacade(const FManagedArrayCollection& InCollection)
+		: TransformToGeometryIndexAttribute(InCollection, "TransformToGeometryIndex", FGeometryCollection::GeometryGroup)
+		, VertexAttribute(InCollection, "Vertex", FGeometryCollection::VerticesGroup)
+		, TangentUAttribute(InCollection, "TangentU", FGeometryCollection::VerticesGroup)
+		, TangentVAttribute(InCollection, "TangentV", FGeometryCollection::VerticesGroup)
+		, NormalAttribute(InCollection, "Normal", FGeometryCollection::VerticesGroup)
+		, UVsAttribute(InCollection, "UVs", FGeometryCollection::VerticesGroup)
+		, ColorAttribute(InCollection, "Color", FGeometryCollection::VerticesGroup)
+		, BoneMapAttribute(InCollection, "BoneMap", FGeometryCollection::VerticesGroup)
+		, VertexStartAttribute(InCollection, "VertexStart", FGeometryCollection::GeometryGroup)
+		, VertexCountAttribute(InCollection, "VertexCount", FGeometryCollection::GeometryGroup)
+		, IndicesAttribute(InCollection, "Indices", FGeometryCollection::FacesGroup)
+		, VisibleAttribute(InCollection, "Visible", FGeometryCollection::FacesGroup)
+		, MaterialIndexAttribute(InCollection, "MaterialIndex", FGeometryCollection::FacesGroup)
+		, MaterialIDAttribute(InCollection, "MaterialID", FGeometryCollection::FacesGroup)
+		, FaceStartAttribute(InCollection, "FaceStart", FGeometryCollection::GeometryGroup)
+		, FaceCountAttribute(InCollection, "FaceCount", FGeometryCollection::GeometryGroup)
 	{
-		return Vertex.IsValid()
-			&& TangentU.IsValid()
-			&& TangentV.IsValid()
-			&& Normal.IsValid()
-			&& UVs.IsValid()
-			&& Color.IsValid()
-			&& BoneMap.IsValid()
-			&& VertexStart.IsValid()
-			&& VertexCount.IsValid()
-			&& Indices.IsValid()
-			&& Visible.IsValid()
-			&& MaterialIndex.IsValid()
-			&& MaterialID.IsValid()
-			&& FaceStart.IsValid()
-			&& FaceCount.IsValid()
+	}
+
+	bool FCollectionMeshFacade::IsValid() const
+	{
+		return TransformToGeometryIndexAttribute.IsValid()
+			&& VertexAttribute.IsValid()
+			&& TangentUAttribute.IsValid()
+			&& TangentVAttribute.IsValid()
+			&& NormalAttribute.IsValid()
+			&& UVsAttribute.IsValid()
+			&& ColorAttribute.IsValid()
+			&& BoneMapAttribute.IsValid()
+			&& VertexStartAttribute.IsValid()
+			&& VertexCountAttribute.IsValid()
+			&& IndicesAttribute.IsValid()
+			&& VisibleAttribute.IsValid()
+			&& MaterialIndexAttribute.IsValid()
+			&& MaterialIDAttribute.IsValid()
+			&& FaceStartAttribute.IsValid()
+			&& FaceCountAttribute.IsValid()
 			;
 	}
 
-	void FMeshFacade::DefineSchema()
+	void FCollectionMeshFacade::DefineSchema()
+	{
+		TransformToGeometryIndexAttribute.Add();
+		VertexAttribute.Add();
+		TangentUAttribute.Add();
+		TangentVAttribute.Add();
+		NormalAttribute.Add();
+		UVsAttribute.Add();
+		ColorAttribute.Add();
+		BoneMapAttribute.Add();
+		VertexStartAttribute.Add();
+		VertexCountAttribute.Add();
+		IndicesAttribute.Add();
+		VisibleAttribute.Add();
+		MaterialIndexAttribute.Add();
+		MaterialIDAttribute.Add();
+		FaceStartAttribute.Add();
+		FaceCountAttribute.Add();
+	}
+
+	const TArray<int32> FCollectionMeshFacade::GetVertexIndices(int32 BoneIdx) const
+	{
+		const TManagedArray<int32>& TransformToGeometryIndicies = TransformToGeometryIndexAttribute.Get();
+		const TManagedArray<int32>& VertexStarts = VertexStartAttribute.Get();
+		const TManagedArray<int32>& VertexCounts = VertexCountAttribute.Get();
+
+		TArray<int32> VertexIndices;
+
+		const int32 VertexIndexStart = VertexStarts[TransformToGeometryIndicies[BoneIdx]];
+		for (int32 Offset = 0; Offset < VertexCounts[TransformToGeometryIndicies[BoneIdx]]; ++Offset)
+		{
+			VertexIndices.Add(VertexIndexStart + Offset);
+		}
+
+		return VertexIndices;
+	}
+
+	const TArrayView<const FVector3f> FCollectionMeshFacade::GetVertexPositions(int32 BoneIdx) const
+	{
+		const TManagedArray<int32>& TransformToGeometryIndicies = TransformToGeometryIndexAttribute.Get();
+		const TManagedArray<int32>& VertexStarts = VertexStartAttribute.Get();
+		const TManagedArray<int32>& VertexCounts = VertexCountAttribute.Get();
+		const TManagedArray<FVector3f>& Vertices = VertexAttribute.Get();
+
+		return TArrayView<const FVector3f>(Vertices.GetData() + VertexStarts[TransformToGeometryIndicies[BoneIdx]], VertexCounts[TransformToGeometryIndicies[BoneIdx]] + 1);
+	}
+
+	const TArray<int32> FCollectionMeshFacade::GetFaceIndices(int32 BoneIdx) const
 	{
-		Vertex.Add();
-		TangentU.Add();
-		TangentV.Add();
-		Normal.Add();
-		UVs.Add();
-		Color.Add();
-		BoneMap.Add();
-		VertexStart.Add();
-		VertexCount.Add();
-		Indices.Add();
-		Visible.Add();
-		MaterialIndex.Add();
-		MaterialID.Add();
-		FaceStart.Add();
-		FaceCount.Add();
+		const TManagedArray<int32>& TransformToGeometryIndicies = TransformToGeometryIndexAttribute.Get();
+		const TManagedArray<int32>& FaceStarts = FaceStartAttribute.Get();
+		const TManagedArray<int32>& FaceCounts = FaceCountAttribute.Get();
+
+		TArray<int32> FaceIndicies;
+
+		const int32 FaceIndexStart = FaceStarts[TransformToGeometryIndicies[BoneIdx]];
+		for (int32 Offset = 0; Offset < FaceCounts[TransformToGeometryIndicies[BoneIdx]]; ++Offset)
+		{
+			FaceIndicies.Add(FaceIndexStart + Offset);
+		}
+
+		return FaceIndicies;
 	}
+
 };
 
 
diff --git a/Engine/Source/Runtime/Experimental/Chaos/Private/GeometryCollection/Facades/CollectionTransformFacade.cpp b/Engine/Source/Runtime/Experimental/Chaos/Private/GeometryCollection/Facades/CollectionTransformFacade.cpp
--- a/Engine/Source/Runtime/Experimental/Chaos/Private/GeometryCollection/Facades/CollectionTransformFacade.cpp
+++ b/Engine/Source/Runtime/Experimental/Chaos/Private/GeometryCollection/Facades/CollectionTransformFacade.cpp
@@ -3,33 +3,73 @@
 #include "GeometryCollection/Facades/CollectionTransformFacade.h"
 #include "GeometryCollection/TransformCollection.h"
 #include "GeometryCollection/Facades/CollectionHierarchyFacade.h"
+#include "GeometryCollection/GeometryCollectionAlgo.h"
 
 namespace GeometryCollection::Facades
 {
-	// Group
-	const FName FTransformFacade::TransformGroup = FTransformCollection::TransformGroup;
+	FCollectionTransformFacade::FCollectionTransformFacade(FManagedArrayCollection& InCollection)
+		: ParentAttribute(InCollection, FTransformCollection::ParentAttribute, FTransformCollection::TransformGroup)
+		, ChildrenAttribute(InCollection, FTransformCollection::ChildrenAttribute, FTransformCollection::TransformGroup)
+		, TransformAttribute(InCollection, FTransformCollection::TransformAttribute, FTransformCollection::TransformGroup)
+	{}
 
-	// Attributes
-	const FName FTransformFacade::ParentAttribute = FTransformCollection::ParentAttribute;
-	const FName FTransformFacade::ChildrenAttribute = FTransformCollection::ChildrenAttribute;
-	const FName FTransformFacade::TransformAttribute = FTransformCollection::TransformAttribute;
+	FCollectionTransformFacade::FCollectionTransformFacade(const FManagedArrayCollection& InCollection)
+		: ParentAttribute(InCollection, FTransformCollection::ParentAttribute, FTransformCollection::TransformGroup)
+		, ChildrenAttribute(InCollection, FTransformCollection::ChildrenAttribute, FTransformCollection::TransformGroup)
+		, TransformAttribute(InCollection, FTransformCollection::TransformAttribute, FTransformCollection::TransformGroup)
+	{}
 
+	bool FCollectionTransformFacade::IsValid() const
+	{
+		return ParentAttribute.IsValid() && ChildrenAttribute.IsValid() && TransformAttribute.IsValid();
+	}
 
-	FTransformFacade::FTransformFacade(FManagedArrayCollection& InCollection)
-		: Self(InCollection)
-		, Parent(InCollection, ChildrenAttribute, TransformGroup)
-		, Children(InCollection, ParentAttribute, TransformGroup)
-		, Transform(InCollection, TransformAttribute, TransformGroup)
-	{}
+	TArray<int32> FCollectionTransformFacade::GetRootIndices() const
+	{
+		return Chaos::Facades::FCollectionHierarchyFacade::GetRootIndices(ParentAttribute);
+	}
+
+	TArray<FTransform> FCollectionTransformFacade::ComputeCollectionSpaceTransforms() const
+	{
+		TArray<FTransform> OutTransforms;
+
+		const TManagedArray<FTransform>& BoneTransforms = TransformAttribute.Get();
+		const TManagedArray<int32>& Parents = ParentAttribute.Get();
+
+		GeometryCollectionAlgo::GlobalMatrices(BoneTransforms, Parents, OutTransforms);
 
-	bool FTransformFacade::IsValid() const
+		return OutTransforms;
+	}
+
+	FTransform FCollectionTransformFacade::ComputeCollectionSpaceTransform(int32 BoneIdx) const
+	{
+		const TManagedArray<FTransform>& BoneTransforms = TransformAttribute.Get();
+		const TManagedArray<int32>& Parents = ParentAttribute.Get();
+
+		return GeometryCollectionAlgo::GlobalMatrix(BoneTransforms, Parents, BoneIdx);
+	}
+
+	void FCollectionTransformFacade::SetPivot(const FTransform& InTransform)
 	{
-		return Parent.IsValid() && Children.IsValid() && Transform.IsValid();
+		Transform(InTransform.Inverse());
+	}
+
+	void FCollectionTransformFacade::Transform(const FTransform& InTransform)
+	{
+		// Update only root transforms
+		const TArray<int32>& RootIndices = GetRootIndices();
+
+		TManagedArray<FTransform>& Transforms = TransformAttribute.Modify();
+
+		for (int32 Idx : RootIndices)
+		{
+			Transforms[Idx] = Transforms[Idx] * InTransform;
+		}
 	}
 
-	TSet<int32> FTransformFacade::GetRootIndices() const
+	void FCollectionTransformFacade::Transform(const FTransform& InTransform, const TArray<int32>& InSelection)
 	{
-		return Chaos::Facades::FCollectionHierarchyFacade::GetRootIndices(Parent);
+		ensureAlwaysMsgf(false, TEXT("NOT YET IMPLEMENTED"));
 	}
 
 }
\ No newline at end of file
diff --git a/Engine/Source/Runtime/Experimental/Chaos/Private/GeometryCollection/Facades/CollectionTransformSelectionFacade.cpp b/Engine/Source/Runtime/Experimental/Chaos/Private/GeometryCollection/Facades/CollectionTransformSelectionFacade.cpp
new file mode 100644
--- /dev/null
+++ b/Engine/Source/Runtime/Experimental/Chaos/Private/GeometryCollection/Facades/CollectionTransformSelectionFacade.cpp
@@ -0,0 +1,1167 @@
+// Copyright Epic Games, Inc. All Rights Reserved.
+
+#include "GeometryCollection/Facades/CollectionTransformSelectionFacade.h"
+#include "GeometryCollection/GeometryCollection.h"
+#include "GeometryCollection/Facades/CollectionHierarchyFacade.h"
+#include "Algo/RemoveIf.h"
+#include "GeometryCollection/Facades/CollectionMeshFacade.h"
+#include "GeometryCollection/Facades/CollectionTransformFacade.h"
+#include "GeometryCollection/Facades/CollectionBoundsFacade.h"
+
+namespace GeometryCollection::Facades
+{
+	FCollectionTransformSelectionFacade::FCollectionTransformSelectionFacade(const FManagedArrayCollection& InCollection)
+		: ConstCollection(InCollection)
+		, ParentAttribute(InCollection, "Parent", FTransformCollection::TransformGroup)
+		, ChildrenAttribute(InCollection, "Children", FTransformCollection::TransformGroup)
+		, LevelAttribute(InCollection, "Level", FTransformCollection::TransformGroup)
+		, SimulationTypeAttribute(InCollection, "SimulationType", FTransformCollection::TransformGroup)
+		, TransformIndexAttribute(InCollection, "TransformIndex", FTransformCollection::TransformGroup)
+		, TransformToGeometryIndexAttribute(InCollection, "TransformToGeometryIndex", FTransformCollection::TransformGroup)
+		, VertexStartAttribute(InCollection, "VertexStart", FGeometryCollection::GeometryGroup)
+		, VertexCountAttribute(InCollection, "VertexCount", FGeometryCollection::GeometryGroup)
+		, FaceStartAttribute(InCollection, "FaceStart", FGeometryCollection::GeometryGroup)
+		, FaceCountAttribute(InCollection, "FaceCount", FGeometryCollection::GeometryGroup)
+		, IndicesAttribute(InCollection, "Indices", FGeometryCollection::FacesGroup)
+	{}
+
+	//
+	//  Initialization
+	//
+
+	// There is no need for schema for this facade
+	void FCollectionTransformSelectionFacade::DefineSchema()
+	{
+		check(!IsConst());
+		check(false);
+	}
+
+	bool FCollectionTransformSelectionFacade::IsValid() const
+	{
+		return ParentAttribute.IsValid() &&
+			ChildrenAttribute.IsValid() &&
+			LevelAttribute.IsValid() &&
+			SimulationTypeAttribute.IsValid() &&
+			TransformIndexAttribute.IsValid() &&
+			TransformToGeometryIndexAttribute.IsValid();
+	}
+
+	bool FCollectionTransformSelectionFacade::IsARootBone(const int32 Index) const
+	{
+		const TManagedArray<int32>& Parents = ParentAttribute.Get();
+		return (Parents[Index] == INDEX_NONE);
+	}
+
+	bool FCollectionTransformSelectionFacade::HasSelectedAncestor(const TArray<int32>& InSelection, const int32 Index) const
+	{
+		if (ParentAttribute.IsValid())
+		{
+			const TManagedArray<int32>& Parents = ParentAttribute.Get();
+
+			if (Index < 0 || Index >= ParentAttribute.Num())
+			{
+				return false;
+			}
+
+			int32 CurrIndex = Index;
+			while (Parents[CurrIndex] != INDEX_NONE)
+			{
+				CurrIndex = Parents[CurrIndex];
+				if (InSelection.Contains(CurrIndex))
+				{
+					return true;
+				}
+			}
+		}
+
+		// We've arrived at the top of the hierarchy with no selected ancestors
+		return false;
+	}
+
+	void FCollectionTransformSelectionFacade::RemoveRootNodes(TArray<int32>& InOutSelection) const
+	{
+		if (ParentAttribute.IsValid())
+		{
+			// Ensure that selected indices are valid
+			const int32 NumTransforms = ParentAttribute.Num();
+
+			InOutSelection.RemoveAll([this, NumTransforms](int32 Index) {
+				return Index == INDEX_NONE || !ensure(Index < NumTransforms);
+				});
+
+			InOutSelection.RemoveAll([this](int32 Index) {
+				return IsARootBone(Index);
+				});
+		}
+	}
+
+	// Originally was called GetBonesToLevel() in Geometry library
+	TArray<int32> FCollectionTransformSelectionFacade::GetBonesByLevel(const int32 Level, bool bOnlyClusteredOrRigid, bool bSkipFiltered) const
+	{
+		TArray<int32> OutBones;
+
+		if (LevelAttribute.IsValid() && SimulationTypeAttribute.IsValid() && ChildrenAttribute.IsValid())
+		{
+			const TManagedArray<int32>& Levels = LevelAttribute.Get();
+			const TManagedArray<int32>& SimType = SimulationTypeAttribute.Get();
+			const TManagedArray<TSet<int32>>& Children = ChildrenAttribute.Get();
+
+			bool bAllLevels = Level == INDEX_NONE;
+
+			const int32 NumTransforms = ParentAttribute.Num();
+			for (int32 BoneIdx = 0; BoneIdx < NumTransforms; BoneIdx++)
+			{
+				bool bIsRigid = SimType[BoneIdx] == FGeometryCollection::ESimulationTypes::FST_Rigid;
+				bool bIsClustered = SimType[BoneIdx] == FGeometryCollection::ESimulationTypes::FST_Clustered;
+				if (
+					// (if skipping embedded) sim type is clustered or rigid 
+					(!bOnlyClusteredOrRigid || bIsClustered || bIsRigid)
+					&&
+					// (if skipping nodes the outliner has filtered) sim type is clustered or level is an exact match or level has an exact-match child
+					(bAllLevels || !bSkipFiltered || bIsClustered || Levels[BoneIdx] == Level || (Children[BoneIdx].Num() > 0 && Levels[BoneIdx] + 1 == Level))
+					&&
+					// level is at or before the target
+					(bAllLevels || Levels[BoneIdx] <= Level)
+					)
+				{
+					OutBones.Add(BoneIdx);
+				}
+			}
+		}
+
+		return OutBones;
+	}
+
+	void FCollectionTransformSelectionFacade::Sanitize(TArray<int32>& InOutSelection, bool bFavorParents) const
+	{
+		if (ParentAttribute.IsValid())
+		{
+			const int32 NumTransforms = ParentAttribute.Num();
+
+			InOutSelection.RemoveAll([this, NumTransforms](int32 Index) {
+				return Index == INDEX_NONE || !ensure(Index < NumTransforms);
+				});
+
+			// Ensure that children of a selected node are not also selected.
+			if (bFavorParents)
+			{
+				InOutSelection.RemoveAll([this, InOutSelection](int32 Index) {
+					return !IsValidBone(Index) || HasSelectedAncestor(InOutSelection, Index);
+					});
+			}
+
+			InOutSelection.Sort();
+		}
+	}
+
+	void FCollectionTransformSelectionFacade::ConvertSelectionToRigidNodes(const int32 Index, TArray<int32>& InOutSelection) const
+	{
+		if (ChildrenAttribute.IsValid() && SimulationTypeAttribute.IsValid())
+		{
+			const TManagedArray<TSet<int32>>& Children = ChildrenAttribute.Get();
+			const TManagedArray<int32>& SimulationType = SimulationTypeAttribute.Get();
+
+			if (SimulationType[Index] == FGeometryCollection::ESimulationTypes::FST_Rigid)
+			{
+				InOutSelection.Add(Index);
+			}
+			else
+			{
+				for (int32 Child : Children[Index])
+				{
+					ConvertSelectionToRigidNodes(Child, InOutSelection);
+				}
+			}
+		}
+	}
+
+	void FCollectionTransformSelectionFacade::ConvertSelectionToRigidNodes(TArray<int32>& InOutSelection) const
+	{
+		Sanitize(InOutSelection);
+
+		TArray<int32> RigidSelection;
+		for (int32 Index : InOutSelection)
+		{
+			ConvertSelectionToRigidNodes(Index, RigidSelection);
+		}
+
+		InOutSelection = RigidSelection;
+	}
+
+	TArray<int32> FCollectionTransformSelectionFacade::SelectRootBones() const
+	{
+		TArray<int32> OutSelection;
+
+		if (CanSelectRootBones())
+		{
+			const TManagedArray<int32>& Parents = ParentAttribute.Get();
+			const int32 NumTransforms = ParentAttribute.Num();
+
+			for (int32 Idx = 0; Idx < Parents.Num(); ++Idx)
+			{
+				if (Parents[Idx] == INDEX_NONE)
+				{
+					OutSelection.Add(Idx);
+				}
+			}
+		}
+		
+		return OutSelection;
+	}
+
+	TArray<int32> FCollectionTransformSelectionFacade::SelectNone() const
+	{
+		TArray<int32> OutSelection;
+		return OutSelection;
+
+	}
+
+	TArray<int32> FCollectionTransformSelectionFacade::SelectAll() const
+	{
+		TArray<int32> OutSelection;
+
+		if (CanSelectAll())
+		{
+			const TManagedArray<int32>& Parents = ParentAttribute.Get();
+			const int32 NumTransforms = ParentAttribute.Num();
+
+			for (int32 Idx = 0; Idx < Parents.Num(); ++Idx)
+			{
+				OutSelection.Add(Idx);
+			}
+		}
+		
+		return OutSelection;
+	}
+
+	void FCollectionTransformSelectionFacade::SelectInverse(TArray<int32>& InOutSelection) const
+	{
+		TArray<int32> NewSelection;
+
+		if (CanSelectInverse())
+		{
+			const TManagedArray<int32>& Parents = ParentAttribute.Get();
+			const int32 NumTransforms = ParentAttribute.Num();
+
+			for (int32 Idx = 0; Idx < NumTransforms; ++Idx)
+			{
+				if (!InOutSelection.Contains(Idx))
+				{
+					NewSelection.Add(Idx);
+				}
+			}
+		}
+		
+		InOutSelection = NewSelection;
+	}
+
+	TArray<int32> FCollectionTransformSelectionFacade::SelectRandom(bool bDeterministic, float RandomSeed, float RandomThresholdVal) const
+	{
+		TArray<int32> OutSelection;
+
+		if (CanSelectRandom())
+		{
+			const TManagedArray<int32>& Parents = ParentAttribute.Get();
+			const int32 NumTransforms = ParentAttribute.Num();
+
+			FRandomStream Stream((int32)RandomSeed);
+
+			for (int32 Idx = 0; Idx < NumTransforms; ++Idx)
+			{
+				if (bDeterministic)
+				{
+					float RandomVal = (float)Stream.FRandRange(0.f, 1.f);
+
+					if (RandomVal > RandomThresholdVal)
+					{
+						OutSelection.Add(Idx);
+					}
+				}
+				else
+				{
+					float RandomVal = FMath::FRandRange(0.f, 1.f);
+
+					if (RandomVal > RandomThresholdVal)
+					{
+						OutSelection.Add(Idx);
+					}
+				}
+			}
+		}
+
+		return OutSelection;
+	}
+
+	TArray<int32> FCollectionTransformSelectionFacade::SelectLeaf() const
+	{
+		TArray<int32> OutSelection;
+
+		if (CanSelectLeaf())
+		{
+			const TManagedArray<int32>& Levels = LevelAttribute.Get();
+			const TManagedArray<int32>& SimType = SimulationTypeAttribute.Get();
+
+			const int32 NumTransforms = ParentAttribute.Num();
+
+			const int32 ViewLevel = INDEX_NONE;
+
+			OutSelection = GetBonesByLevel(ViewLevel, true, true);
+
+			OutSelection.SetNum(Algo::RemoveIf(OutSelection, [&](int32 BoneIdx)
+				{
+					return SimType[BoneIdx] != FGeometryCollection::ESimulationTypes::FST_Rigid
+						|| (ViewLevel != INDEX_NONE && Levels[BoneIdx] != ViewLevel);
+				}));
+		}
+		
+		return OutSelection;
+	}
+
+	TArray<int32> FCollectionTransformSelectionFacade::SelectCluster() const
+	{
+		TArray<int32> OutSelection;
+
+		if (CanSelectCluster())
+		{
+			const TManagedArray<int32>& Levels = LevelAttribute.Get();
+			const TManagedArray<int32>& SimType = SimulationTypeAttribute.Get();
+
+			const int32 NumTransforms = ParentAttribute.Num();
+
+			const int32 ViewLevel = INDEX_NONE;
+
+			OutSelection = GetBonesByLevel(ViewLevel, true, true);
+
+			OutSelection.SetNum(Algo::RemoveIf(OutSelection, [&](int32 BoneIdx)
+				{
+					return SimType[BoneIdx] != FGeometryCollection::ESimulationTypes::FST_Rigid
+						|| (ViewLevel != INDEX_NONE && Levels[BoneIdx] != ViewLevel);
+				}));
+		}
+		
+		return OutSelection;
+	}
+
+	void FCollectionTransformSelectionFacade::SelectContact(TArray<int32>& InOutSelection) const
+	{
+		TSet<int32> NewSelection;
+
+		if (CanSelectContact())
+		{
+			// TODO: This needs to be looked at
+			if (TUniquePtr<FGeometryCollection> TempGeomCollection = TUniquePtr<FGeometryCollection>(ConstCollection.NewCopy<FGeometryCollection>()))
+			{
+				FGeometryCollectionProximityUtility ProximityUtility(TempGeomCollection.Get());
+				ProximityUtility.RequireProximity();
+
+				const TManagedArray<int32>& TransformIndex = TransformIndexAttribute.Get();
+				const TManagedArray<int32>& TransformToGeometryIndex = TransformToGeometryIndexAttribute.Get();
+				const TManagedArray<TSet<int32>>& Proximity = TempGeomCollection->GetAttribute<TSet<int32>>("Proximity", FGeometryCollection::TransformGroup);
+
+				for (int32 Bone : InOutSelection)
+				{
+					NewSelection.Add(Bone);
+					int32 GeometryIdx = TransformToGeometryIndex[Bone];
+					if (GeometryIdx != INDEX_NONE)
+					{
+						const TSet<int32>& Neighbors = Proximity[GeometryIdx];
+						for (int32 NeighborGeometryIndex : Neighbors)
+						{
+							NewSelection.Add(TransformIndex[NeighborGeometryIndex]);
+						}
+					}
+				}
+			}
+		}
+
+		InOutSelection = NewSelection.Array();
+	}
+
+	void FCollectionTransformSelectionFacade::SelectParent(TArray<int32>& InOutSelection) const
+	{
+		TSet<int32> NewSelection;
+
+		if (CanSelectParent())
+		{
+			const TManagedArray<int32>& Parents = ParentAttribute.Get();
+
+			for (int32 Bone : InOutSelection)
+			{
+				int32 ParentBone = Parents[Bone];
+				if (ParentBone != FGeometryCollection::Invalid)
+				{
+					NewSelection.Add(ParentBone);
+				}
+			}
+		}
+		
+		InOutSelection = NewSelection.Array();
+	}
+
+	void FCollectionTransformSelectionFacade::SelectChildren(TArray<int32>& InOutSelection) const
+	{
+		TSet<int32> NewSelection;
+
+		if (CanSelectChildren())
+		{
+			const TManagedArray<TSet<int32>>& Children = ChildrenAttribute.Get();
+
+			for (int32 Bone : InOutSelection)
+			{
+				if (Children[Bone].IsEmpty())
+				{
+					NewSelection.Add(Bone);
+					continue;
+				}
+				for (int32 Child : Children[Bone])
+				{
+					NewSelection.Add(Child);
+				}
+			}
+		}
+
+		InOutSelection = NewSelection.Array();
+	}
+
+	void FCollectionTransformSelectionFacade::SelectSiblings(TArray<int32>& InOutSelection) const
+	{
+		TSet<int32> NewSelection;
+
+		if (CanSelectSiblings())
+		{
+			const TManagedArray<int32>& Parents = ParentAttribute.Get();
+			const TManagedArray<TSet<int32>>& Children = ChildrenAttribute.Get();
+
+			for (int32 Bone : InOutSelection)
+			{
+				int32 ParentBone = Parents[Bone];
+				if (ParentBone != INDEX_NONE)
+				{
+					for (int32 Child : Children[ParentBone])
+					{
+						NewSelection.Add(Child);
+					}
+				}
+			}
+		}
+		
+		InOutSelection = NewSelection.Array();
+	}
+
+	void FCollectionTransformSelectionFacade::SelectLevel(TArray<int32>& InOutSelection) const
+	{
+		TSet<int32> NewSelection;
+
+		if (CanSelectLevel())
+		{
+			const TManagedArray<int32>& Levels = LevelAttribute.Get();
+			const int32 NumTransforms = ParentAttribute.Num();
+
+			for (int32 Bone : InOutSelection)
+			{
+				int32 Level = Levels[Bone];
+				for (int32 TransformIdx = 0; TransformIdx < NumTransforms; ++TransformIdx)
+				{
+					if (Levels[TransformIdx] == Level)
+					{
+						NewSelection.Add(TransformIdx);
+					}
+				}
+			}
+		}
+		
+		InOutSelection = NewSelection.Array();
+	}
+
+	static void RandomShuffleArray(TArray<int32>& InArray, bool Deterministic, float RandomSeed)
+	{
+		FRandomStream Stream((int32)RandomSeed);
+
+		int32 LastIndex = InArray.Num() - 1;
+
+		static const int32 NumIterationsMin = 10;
+		static const int32 NumIterationsMax = 50;
+
+		int32 NumIterations = Deterministic ? Stream.RandRange(NumIterationsMin, NumIterationsMax) : FMath::RandRange(NumIterationsMin, NumIterationsMax);
+
+		for (int32 IterIdx = 0; IterIdx < NumIterations; ++IterIdx)
+		{
+			for (int32 Idx = 0; Idx <= LastIndex; ++Idx)
+			{
+				int32 Index;
+				if (Deterministic)
+				{
+					Index = Stream.RandRange(Idx, LastIndex);
+				}
+				else
+				{
+					Index = FMath::RandRange(Idx, LastIndex);
+				}
+
+				if (Idx != Index)
+				{
+					InArray.Swap(Idx, Index);
+				}
+			}
+		}
+	}
+
+	void FCollectionTransformSelectionFacade::SelectByPercentage(TArray<int32>& InOutSelection, int32 InPercentage, bool Deterministic, float RandomSeed)
+	{
+		RandomShuffleArray(InOutSelection, Deterministic, RandomSeed);
+
+		float Percentage = (float)InPercentage * 0.01f;
+		int32 NewNumElements = FMath::RoundToInt32((float)InOutSelection.Num() * Percentage);
+		
+		InOutSelection.SetNum(NewNumElements);
+	}
+
+	TArray<int32> FCollectionTransformSelectionFacade::SelectBySize(float SizeMin, float SizeMax, bool bInclusive, bool bInsideRange) const
+	{
+		TArray<int32> OutSelection;
+
+		if (TUniquePtr<FGeometryCollection> TempGeomCollection = TUniquePtr<FGeometryCollection>(ConstCollection.NewCopy<FGeometryCollection>()))
+		{
+			FGeometryCollectionConvexUtility::SetVolumeAttributes(TempGeomCollection.Get());
+
+			const TManagedArrayAccessor<float> SizeAttribute((FManagedArrayCollection)*TempGeomCollection, "Size", FTransformCollection::TransformGroup);
+			const TManagedArray<float>& Sizes = SizeAttribute.Get();
+
+			for (int32 BoneIdx = 0; BoneIdx < Sizes.Num(); ++BoneIdx)
+			{
+				const float Size = Sizes[BoneIdx];
+
+				if (bInsideRange)
+				{
+					if (bInclusive)
+					{
+						if (Size >= SizeMin && Size <= SizeMax)
+						{
+							OutSelection.Add(BoneIdx);
+						}
+					}
+					else
+					{
+						if (Size > SizeMin && Size < SizeMax)
+						{
+							OutSelection.Add(BoneIdx);
+						}
+					}
+				}
+				else
+				{
+					if (bInclusive)
+					{
+						if (Size <= SizeMin || Size >= SizeMax)
+						{
+							OutSelection.Add(BoneIdx);
+						}
+					}
+					else
+					{
+						if (Size < SizeMin || Size > SizeMax)
+						{
+							OutSelection.Add(BoneIdx);
+						}
+					}
+				}
+			}
+		}
+
+		return OutSelection;
+	}
+
+	TArray<int32> FCollectionTransformSelectionFacade::SelectByVolume(float VolumeMin, float VolumeMax, bool bInclusive, bool bInsideRange) const
+	{
+		TArray<int32> OutSelection;
+
+		if (TUniquePtr<FGeometryCollection> TempGeomCollection = TUniquePtr<FGeometryCollection>(ConstCollection.NewCopy<FGeometryCollection>()))
+		{
+			FGeometryCollectionConvexUtility::SetVolumeAttributes(TempGeomCollection.Get());
+
+			const TManagedArrayAccessor<float> VolumeAttribute((FManagedArrayCollection)*TempGeomCollection, "Volume", FTransformCollection::TransformGroup);
+			const TManagedArray<float>& Volumes = VolumeAttribute.Get();
+
+			for (int32 BoneIdx = 0; BoneIdx < Volumes.Num(); ++BoneIdx)
+			{
+				const float Volume = Volumes[BoneIdx];
+
+				if (bInsideRange)
+				{
+					if (bInclusive)
+					{
+						if (Volume >= VolumeMin && Volume <= VolumeMax)
+						{
+							OutSelection.Add(BoneIdx);
+						}
+					}
+					else
+					{
+						if (Volume > VolumeMin && Volume < VolumeMax)
+						{
+							OutSelection.Add(BoneIdx);
+						}
+					}
+				}
+				else
+				{
+					if (bInclusive)
+					{
+						if (Volume <= VolumeMin || Volume >= VolumeMax)
+						{
+							OutSelection.Add(BoneIdx);
+						}
+					}
+					else
+					{
+						if (!(Volume < VolumeMin || Volume > VolumeMax))
+						{
+							OutSelection.Add(BoneIdx);
+						}
+					}
+				}
+			}
+		}
+
+		return OutSelection;
+	}
+
+	TArray<int32> FCollectionTransformSelectionFacade::SelectVerticesInBox(const FBox& InBox, const FTransform& InBoxTransform, bool bAllVerticesInBox) const
+	{
+		TArray<int32> OutSelection;
+
+		if (CanSelectVerticesInBox())
+		{
+			const int32 NumTransforms = ParentAttribute.Num();
+
+			FCollectionMeshFacade MeshFacade(ConstCollection);
+			FCollectionTransformFacade TransformFacade(ConstCollection);
+
+			for (int32 TransformIdx = 0; TransformIdx < NumTransforms; ++TransformIdx)
+			{
+				int32 NumVerticesInsideBox = 0;
+
+				const TArrayView<const FVector3f> VertexPositions = MeshFacade.GetVertexPositions(TransformIdx);
+				for (const FVector3f& Vertex : VertexPositions)
+				{
+					const FVector VertexInBoneSpace(Vertex);
+
+					// Transform from BoneSpace to CollectionSpace
+					const FTransform CollectionSpaceTransform = TransformFacade.ComputeCollectionSpaceTransform(TransformIdx);
+					const FVector VertexInCollectionSpace = CollectionSpaceTransform.TransformPosition(VertexInBoneSpace);
+
+					// Transform with specified transform
+					const FVector VertexInBoxSpace = InBoxTransform.InverseTransformPosition(VertexInCollectionSpace);
+
+					if (!InBox.IsInside(VertexInBoxSpace))
+					{
+						break;
+					}
+
+					NumVerticesInsideBox++;
+					if (!bAllVerticesInBox)
+					{
+						break;
+					}
+				}
+
+				if ((bAllVerticesInBox && NumVerticesInsideBox == VertexPositions.Num()) ||
+					(!bAllVerticesInBox && NumVerticesInsideBox > 0))
+				{
+					OutSelection.Add(TransformIdx);
+				}
+			}
+		}
+		
+		return OutSelection;
+	}
+
+	TArray<int32> FCollectionTransformSelectionFacade::SelectCentroidInBox(const FBox& InBox, const FTransform& InBoxTransform) const
+	{
+		TArray<int32> OutSelection;
+
+		if (CanSelectCentroidInBox())
+		{
+			const int32 NumTransforms = ParentAttribute.Num();
+
+			FBoundsFacade BoundsFacade(ConstCollection);
+			FCollectionTransformFacade TransformFacade(ConstCollection);
+
+			const TArray<FVector> Centroids = BoundsFacade.GetCentroids();
+
+			for (int32 TransformIdx = 0; TransformIdx < NumTransforms; ++TransformIdx)
+			{
+				const FVector CentroidInBoneSpace(Centroids[TransformIdx]);
+
+				// Transform from BoneSpace to CollectionSpace
+				const FTransform CollectionSpaceTransform = TransformFacade.ComputeCollectionSpaceTransform(TransformIdx);
+				const FVector CentroidInCollectionSpace = CollectionSpaceTransform.TransformPosition(CentroidInBoneSpace);
+
+				// Transform with specified transform
+				const FVector CentroidInBoxSpace = InBoxTransform.InverseTransformPosition(CentroidInCollectionSpace);
+
+				if (InBox.IsInside(CentroidInBoxSpace))
+				{
+					OutSelection.Add(TransformIdx);
+				}
+			}
+		}
+		
+		return OutSelection;
+	}
+
+	TArray<int32> FCollectionTransformSelectionFacade::SelectBoundingBoxInBox(const FBox& InBox, const FTransform& InBoxTransform) const
+	{
+		TArray<int32> OutSelection;
+
+		if (CanSelectBoundingBoxInBox())
+		{
+			const int32 NumTransforms = ParentAttribute.Num();
+
+			FBoundsFacade BoundsFacade(ConstCollection);
+			FCollectionTransformFacade TransformFacade(ConstCollection);
+
+			const TManagedArray<FBox>& BoundingBoxes = BoundsFacade.GetBoundingBoxes();
+
+			for (int32 TransformIdx = 0; TransformIdx < NumTransforms; ++TransformIdx)
+			{
+				const FBox& BoundingBoxInBoneSpace = BoundingBoxes[TransformIdx];
+
+				// Transform from BoneSpace to CollectionSpace
+				const FTransform CollectionSpaceTransform = TransformFacade.ComputeCollectionSpaceTransform(TransformIdx);
+				const FBox BoundingBoxInCollectionSpace = BoundingBoxInBoneSpace.TransformBy(CollectionSpaceTransform);
+
+				// Transform with specified transform
+				const FBox BoundingBoxInBoxSpace = BoundingBoxInCollectionSpace.InverseTransformBy(InBoxTransform);
+
+				if (InBox.IsInside(BoundingBoxInBoxSpace))
+				{
+					OutSelection.Add(TransformIdx);
+				}
+			}
+		}
+		
+		return OutSelection;
+	}
+
+	TArray<int32> FCollectionTransformSelectionFacade::SelectVerticesInSphere(const FSphere& InSphere, const FTransform& InSphereTransform, bool bAllVerticesInSphere) const
+	{
+		TArray<int32> OutSelection;
+
+		if (CanSelectVerticesInSphere())
+		{
+			const int32 NumTransforms = ParentAttribute.Num();
+
+			FCollectionMeshFacade MeshFacade(ConstCollection);
+			FCollectionTransformFacade TransformFacade(ConstCollection);
+
+			for (int32 TransformIdx = 0; TransformIdx < NumTransforms; ++TransformIdx)
+			{
+				int32 NumVerticesInside = 0;
+
+				const TArrayView<const FVector3f> VertexPositions = MeshFacade.GetVertexPositions(TransformIdx);
+				for (const FVector3f& Vertex : VertexPositions)
+				{
+					const FVector VertexInBoneSpace(Vertex);
+
+					// Transform from BoneSpace to CollectionSpace
+					const FTransform CollectionSpaceTransform = TransformFacade.ComputeCollectionSpaceTransform(TransformIdx);
+					const FVector VertexInCollectionSpace = CollectionSpaceTransform.TransformPosition(VertexInBoneSpace);
+
+					// Transform with specified transform
+					const FVector VertexInSphereSpace = InSphereTransform.InverseTransformPosition(VertexInCollectionSpace);
+
+					if (!InSphere.IsInside(VertexInSphereSpace))
+					{
+						break;
+					}
+
+					NumVerticesInside++;
+					if (!bAllVerticesInSphere)
+					{
+						break;
+					}
+				}
+
+				if ((bAllVerticesInSphere && NumVerticesInside == VertexPositions.Num()) ||
+					(!bAllVerticesInSphere && NumVerticesInside > 0))
+				{
+					OutSelection.Add(TransformIdx);
+				}
+			}
+		}
+
+		return OutSelection;
+	}
+
+	TArray<int32> FCollectionTransformSelectionFacade::SelectCentroidInSphere(const FSphere& InSphere, const FTransform& InSphereTransform) const
+	{
+		TArray<int32> OutSelection;
+
+		if (CanSelectCentroidInSphere())
+		{
+			const int32 NumTransforms = ParentAttribute.Num();
+
+			FBoundsFacade BoundsFacade(ConstCollection);
+			FCollectionTransformFacade TransformFacade(ConstCollection);
+
+			const TArray<FVector> Centroids = BoundsFacade.GetCentroids();
+
+			for (int32 TransformIdx = 0; TransformIdx < NumTransforms; ++TransformIdx)
+			{
+				const FVector CentroidInBoneSpace(Centroids[TransformIdx]);
+
+				// Transform from BoneSpace to CollectionSpace
+				const FTransform CollectionSpaceTransform = TransformFacade.ComputeCollectionSpaceTransform(TransformIdx);
+				const FVector CentroidInCollectionSpace = CollectionSpaceTransform.TransformPosition(CentroidInBoneSpace);
+
+				// Transform with specified transform
+				const FVector CentroidInSphereSpace = InSphereTransform.InverseTransformPosition(CentroidInCollectionSpace);
+
+				if (InSphere.IsInside(CentroidInSphereSpace))
+				{
+					OutSelection.Add(TransformIdx);
+				}
+			}
+		}
+		
+		return OutSelection;
+	}
+
+	TArray<int32> FCollectionTransformSelectionFacade::SelectBoundingBoxInSphere(const FSphere& InSphere, const FTransform& InSphereTransform) const
+	{
+		TArray<int32> OutSelection;
+
+		if (CanSelectBoundingBoxInSphere())
+		{
+			const int32 NumTransforms = ParentAttribute.Num();
+
+			FBoundsFacade BoundsFacade(ConstCollection);
+			FCollectionTransformFacade TransformFacade(ConstCollection);
+
+			const TManagedArray<FBox>& BoundingBoxes = BoundsFacade.GetBoundingBoxes();
+
+			for (int32 TransformIdx = 0; TransformIdx < NumTransforms; ++TransformIdx)
+			{
+				const FBox& BoundingBoxInBoneSpace = BoundingBoxes[TransformIdx];
+
+				int32 NumVerticesInside = 0;
+
+				const TArray<FVector>& VertexInBoneSpaceArr = FBoundsFacade::GetBoundingBoxVertexPositions(BoundingBoxInBoneSpace);
+				for (int32 VertexIdx = 0; VertexIdx < VertexInBoneSpaceArr.Num(); ++VertexIdx)
+				{
+					// Transform from BoneSpace to CollectionSpace
+					const FTransform CollectionSpaceTransform = TransformFacade.ComputeCollectionSpaceTransform(TransformIdx);
+					const FVector VertexInCollectionSpace = CollectionSpaceTransform.TransformPosition(VertexInBoneSpaceArr[VertexIdx]);
+
+					// Transform with specified transform
+					const FVector VertexInSphereSpace = InSphereTransform.InverseTransformPosition(VertexInCollectionSpace);
+
+					if (!InSphere.IsInside(VertexInSphereSpace))
+					{
+						break;
+					}
+
+					NumVerticesInside++;
+				}
+
+				if (NumVerticesInside == VertexInBoneSpaceArr.Num())
+				{
+					OutSelection.Add(TransformIdx);
+				}
+			}
+		}
+		
+		return OutSelection;
+	}
+
+	TArray<int32> FCollectionTransformSelectionFacade::SelectByFloatAttribute(FString GroupName, FString AttrName, float Min, float Max, bool bInclusive, bool bInsideRange) const
+	{
+		TArray<int32> OutSelection;
+
+		if (ConstCollection.HasGroup(FName(*GroupName)))
+		{
+			if (ConstCollection.HasAttribute(FName(*AttrName), FName(*GroupName)))
+			{
+				if (ConstCollection.GetAttributeType(FName(*AttrName), FName(*GroupName)) == EManagedArrayType::FFloatType)
+				{
+					const TManagedArray<float>& FloatArray = (const TManagedArray<float>&)ConstCollection.GetAttribute<float>(FName(*AttrName), FName(*GroupName));
+
+					for (int32 Idx = 0; Idx < FloatArray.Num(); ++Idx)
+					{
+						const float FloatValue = FloatArray[Idx];
+
+						if (bInsideRange)
+						{
+							if (bInclusive)
+							{
+								if (FloatValue >= Min && FloatValue <= Max)
+								{
+									OutSelection.Add(Idx);
+								}
+							}
+							else
+							{
+								if (FloatValue > Min && FloatValue < Max)
+								{
+									OutSelection.Add(Idx);
+								}
+							}
+						}
+						else
+						{
+							if (bInclusive)
+							{
+								if (FloatValue <= Min || FloatValue >= Max)
+								{
+									OutSelection.Add(Idx);
+								}
+							}
+							else
+							{
+								if (FloatValue < Min || FloatValue > Max)
+								{
+									OutSelection.Add(Idx);
+								}
+							}
+						}
+					}
+
+				}
+			}
+		}
+
+		return OutSelection;
+	}
+
+	TArray<int32> FCollectionTransformSelectionFacade::SelectByIntAttribute(FString GroupName, FString AttrName, int32 Min, int32 Max, bool bInclusive, bool bInsideRange) const
+	{
+		TArray<int32> OutSelection;
+
+		if (ConstCollection.HasGroup(FName(*GroupName)))
+		{
+			if (ConstCollection.HasAttribute(FName(*AttrName), FName(*GroupName)))
+			{
+				if (ConstCollection.GetAttributeType(FName(*AttrName), FName(*GroupName)) == EManagedArrayType::FFloatType)
+				{
+					const TManagedArray<int32>& IntArray = (const TManagedArray<int32>&)ConstCollection.GetAttribute<int>(FName(*AttrName), FName(*GroupName));
+
+					for (int32 Idx = 0; Idx < IntArray.Num(); ++Idx)
+					{
+						const int32 IntValue = IntArray[Idx];
+
+						if (bInsideRange)
+						{
+							if (bInclusive)
+							{
+								if (IntValue >= Min && IntValue <= Max)
+								{
+									OutSelection.Add(Idx);
+								}
+							}
+							else
+							{
+								if (IntValue > Min && IntValue < Max)
+								{
+									OutSelection.Add(Idx);
+								}
+							}
+						}
+						else
+						{
+							if (bInclusive)
+							{
+								if (IntValue <= Min || IntValue >= Max)
+								{
+									OutSelection.Add(Idx);
+								}
+							}
+							else
+							{
+								if (IntValue < Min || IntValue > Max)
+								{
+									OutSelection.Add(Idx);
+								}
+							}
+						}
+					}
+
+				}
+			}
+		}
+
+		return OutSelection;
+	}
+
+	TArray<int32> FCollectionTransformSelectionFacade::ConvertVertexSelectionToTransformSelection(const TArray<int32>& InVertexSelection, bool bAllElementsMustBeSelected) const
+	{
+		TArray<int32> OutSelection;
+
+		if (TransformIndexAttribute.IsValid() && VertexStartAttribute.IsValid() && VertexCountAttribute.IsValid())
+		{
+			const TManagedArray<int32>& TransformIndicies = TransformIndexAttribute.Get();
+			const TManagedArray<int32>& VertexStarts = VertexStartAttribute.Get();
+			const TManagedArray<int32>& VertexCounts = VertexCountAttribute.Get();
+
+			const int32 NumGeos = ConstCollection.NumElements(FGeometryCollection::GeometryGroup);
+
+			for (int32 VertexIdx : InVertexSelection)
+			{
+				for (int32 GeometryIdx = 0; GeometryIdx < NumGeos; ++GeometryIdx)
+				{
+					const int32 VertexStart = VertexStarts[GeometryIdx];
+					const int32 VertexEnd = VertexStart + VertexCounts[GeometryIdx];
+
+					if (VertexIdx >= VertexStart && VertexIdx < VertexEnd)
+					{
+						OutSelection.AddUnique(TransformIndicies[GeometryIdx]);
+						break;
+					}
+				}
+			}
+		}
+
+		return OutSelection;
+	}
+
+	TArray<int32> FCollectionTransformSelectionFacade::ConvertFaceSelectionToTransformSelection(const TArray<int32>& InFaceSelection, bool bAllElementsMustBeSelected) const
+	{
+		TArray<int32> OutSelection;
+
+		if (TransformIndexAttribute.IsValid() && FaceStartAttribute.IsValid() && FaceCountAttribute.IsValid())
+		{
+			const TManagedArray<int32>& TransformIndicies = TransformIndexAttribute.Get();
+			const TManagedArray<int32>& FaceStarts = FaceStartAttribute.Get();
+			const TManagedArray<int32>& FaceCounts = FaceCountAttribute.Get();
+
+			const int32 NumGeos = ConstCollection.NumElements(FGeometryCollection::GeometryGroup);
+
+			for (int32 FaceIdx : InFaceSelection)
+			{
+				for (int32 GeometryIdx = 0; GeometryIdx < NumGeos; ++GeometryIdx)
+				{
+					const int32 FaceStart = FaceStarts[GeometryIdx];
+					const int32 FaceEnd = FaceStart + FaceCounts[GeometryIdx];
+
+					if (FaceIdx >= FaceStart && FaceIdx < FaceEnd)
+					{
+						OutSelection.AddUnique(TransformIndicies[GeometryIdx]);
+						break;
+					}
+				}
+			}
+		}
+
+		return OutSelection;
+	}
+
+	TArray<int32> FCollectionTransformSelectionFacade::ConvertVertexSelectionToFaceSelection(const TArray<int32>& InVertexSelection, bool bAllElementsMustBeSelected) const
+	{
+		TArray<int32> OutSelection;
+
+		if (FaceStartAttribute.IsValid() && FaceCountAttribute.IsValid() && IndicesAttribute.IsValid())
+		{
+			// Build VertexToFaceIndicies inverse lookup
+			const TManagedArray<int32>& FaceStarts = FaceStartAttribute.Get();
+			const TManagedArray<int32>& FaceCounts = FaceCountAttribute.Get();
+
+			const int32 NumFaces = ConstCollection.NumElements(FGeometryCollection::FacesGroup);
+			const int32 NumVertices = ConstCollection.NumElements(FGeometryCollection::VerticesGroup);
+
+			TArray<TArray<int32>> VertexToFaceIndicies;
+			VertexToFaceIndicies.Reset(NumVertices);
+
+			const TManagedArray<FIntVector>& Indicies = IndicesAttribute.Get();
+
+			for (int32 FaceIdx = 0; FaceIdx < NumFaces; ++FaceIdx)
+			{
+				VertexToFaceIndicies[Indicies[FaceIdx].X].AddUnique(FaceIdx);
+				VertexToFaceIndicies[Indicies[FaceIdx].Y].AddUnique(FaceIdx);
+				VertexToFaceIndicies[Indicies[FaceIdx].Z].AddUnique(FaceIdx);
+			}
+
+			// Convert VertexSelection to FaceSelection
+			for (int32 VertexIdx : InVertexSelection)
+			{
+				for (int32 FaceIdx = 0; FaceIdx < VertexToFaceIndicies[VertexIdx].Num(); ++FaceIdx)
+				{
+					OutSelection.AddUnique(VertexToFaceIndicies[VertexIdx][FaceIdx]);
+				}
+			}
+		}
+
+		return OutSelection;
+	}
+
+	TArray<int32> FCollectionTransformSelectionFacade::ConvertTransformSelectionToFaceSelection(const TArray<int32>& InTransformSelection) const
+	{
+		TArray<int32> OutSelection;
+
+		if (TransformIndexAttribute.IsValid() && FaceStartAttribute.IsValid() && FaceCountAttribute.IsValid())
+		{
+			const TManagedArray<int32>& TransformToGeometryIndicies = TransformToGeometryIndexAttribute.Get();
+			const TManagedArray<int32>& FaceStarts = FaceStartAttribute.Get();
+			const TManagedArray<int32>& FaceCounts = FaceCountAttribute.Get();
+
+			for (int32 TransformIdx : InTransformSelection)
+			{
+				const int32 FaceIdxStart = FaceStarts[TransformToGeometryIndicies[TransformIdx]];
+
+				for (int32 Offset = 0; Offset < FaceCounts[TransformToGeometryIndicies[TransformIdx]]; ++Offset)
+				{
+					const int32 FaceIdx = FaceIdxStart + Offset;
+					OutSelection.Add(FaceIdx);
+				}
+			}
+		}
+
+		return OutSelection;
+	}
+
+	TArray<int32> FCollectionTransformSelectionFacade::ConvertFaceSelectionToVertexSelection(const TArray<int32>& InFaceSelection) const
+	{
+		TArray<int32> OutSelection;
+
+		if (IndicesAttribute.IsValid())
+		{
+			const TManagedArray<FIntVector>& Indicies = IndicesAttribute.Get();
+
+			for (int32 FaceIdx : InFaceSelection)
+			{
+				OutSelection.AddUnique(Indicies[FaceIdx].X);
+				OutSelection.AddUnique(Indicies[FaceIdx].Y);
+				OutSelection.AddUnique(Indicies[FaceIdx].Z);
+			}
+		}
+
+		return OutSelection;
+	}
+
+	TArray<int32> FCollectionTransformSelectionFacade::ConvertTransformSelectionToVertexSelection(const TArray<int32>& InTransformSelection) const
+	{
+		TArray<int32> OutSelection;
+
+		if (TransformToGeometryIndexAttribute.IsValid() && VertexStartAttribute.IsValid() && VertexCountAttribute.IsValid())
+		{
+			const TManagedArray<int32>& TransformToGeometryIndicies = TransformToGeometryIndexAttribute.Get();
+			const TManagedArray<int32>& VertexStarts = VertexStartAttribute.Get();
+			const TManagedArray<int32>& VertexCounts = VertexCountAttribute.Get();
+
+			for (int32 TransformIdx : InTransformSelection)
+			{
+				const int32 VertexIndexStart = VertexStarts[TransformToGeometryIndicies[TransformIdx]];
+
+				for (int32 Offset = 0; Offset < VertexCounts[TransformToGeometryIndicies[TransformIdx]]; ++Offset)
+				{
+					const int32 VertexIdx = VertexIndexStart + Offset;
+					OutSelection.Add(VertexIdx);
+				}
+			}
+		}
+
+		return OutSelection;
+	}
+
+}; // GeometryCollection::Facades
+
+
diff --git a/Engine/Source/Runtime/Experimental/Chaos/Public/GeometryCollection/Facades/CollectionBoundsFacade.h b/Engine/Source/Runtime/Experimental/Chaos/Public/GeometryCollection/Facades/CollectionBoundsFacade.h
--- a/Engine/Source/Runtime/Experimental/Chaos/Public/GeometryCollection/Facades/CollectionBoundsFacade.h
+++ b/Engine/Source/Runtime/Experimental/Chaos/Public/GeometryCollection/Facades/CollectionBoundsFacade.h
@@ -44,13 +44,25 @@ namespace GeometryCollection::Facades
 		/** BoundingBox access */
 		const TManagedArray< FBox >& GetBoundingBoxes() const { return BoundingBoxAttribute.Get(); }
 
+		/** Centroids (Centers of BoundingBoxes) access */
+		TArray<FVector> GetCentroids() const;
+
 		/** BoundingBox for the whole collection */
-		FBox GetBoundingBox();
+		FBox GetBoundingBox() const;
+
+		/** BoundingBox for the whole collection in Collection space */
+		FBox GetBoundingBoxInCollectionSpace() const;
+
+		/** Returns the positions of the vertices of an FBox */
+		static TArray<FVector> GetBoundingBoxVertexPositions(const FBox& InBox);
 
 		/** TransformToGeometryIndex Access */
 		const TManagedArray< int32 >& GetTransformToGeometryIndex() const { return TransformToGeometryIndexAttribute.Get(); }
 
-private:
+	private:
+		const FManagedArrayCollection& ConstCollection;
+		FManagedArrayCollection* Collection = nullptr;
+
 		TManagedArrayAccessor<FBox>			BoundingBoxAttribute;
 		TManagedArrayAccessor<FVector3f>	VertexAttribute;
 		TManagedArrayAccessor<int32>		BoneMapAttribute;
diff --git a/Engine/Source/Runtime/Experimental/Chaos/Public/GeometryCollection/Facades/CollectionHierarchyFacade.h b/Engine/Source/Runtime/Experimental/Chaos/Public/GeometryCollection/Facades/CollectionHierarchyFacade.h
--- a/Engine/Source/Runtime/Experimental/Chaos/Public/GeometryCollection/Facades/CollectionHierarchyFacade.h
+++ b/Engine/Source/Runtime/Experimental/Chaos/Public/GeometryCollection/Facades/CollectionHierarchyFacade.h
@@ -41,6 +41,9 @@ namespace Chaos::Facades
 		/** Get the root index */
 		int32 GetRootIndex() const;
 
+		/** Get the root indicies */
+		TArray<int32> GetRootIndices() const;
+
 		/** 
 		* Update level attribute for all elements (and create it if it is missing ) 
 		* @Param PersistencePolicy whether to make the attribute persistent or keep the existing state
@@ -49,7 +52,7 @@ namespace Chaos::Facades
 
 	public:
 		/** Get the root indicies */
-		static TSet<int32> GetRootIndices(const TManagedArrayAccessor<int32>& ParentAttribute);
+		static TArray<int32> GetRootIndices(const TManagedArrayAccessor<int32>& ParentAttribute);
 
 	private:
 		TManagedArrayAccessor<int32>		ParentAttribute;
diff --git a/Engine/Source/Runtime/Experimental/Chaos/Public/GeometryCollection/Facades/CollectionMeshFacade.h b/Engine/Source/Runtime/Experimental/Chaos/Public/GeometryCollection/Facades/CollectionMeshFacade.h
--- a/Engine/Source/Runtime/Experimental/Chaos/Public/GeometryCollection/Facades/CollectionMeshFacade.h
+++ b/Engine/Source/Runtime/Experimental/Chaos/Public/GeometryCollection/Facades/CollectionMeshFacade.h
@@ -9,10 +9,11 @@
 namespace GeometryCollection::Facades
 {
 
-	class CHAOS_API FMeshFacade
+	class CHAOS_API FCollectionMeshFacade
 	{
 	public:
-		FMeshFacade(FManagedArrayCollection& InCollection);
+		FCollectionMeshFacade(FManagedArrayCollection& InCollection);
+		FCollectionMeshFacade(const FManagedArrayCollection& InCollection);
 
 		/**
 		 * returns true if all the necessary attributes are present
@@ -25,22 +26,42 @@ namespace GeometryCollection::Facades
 		 */
 		void DefineSchema();
 
-		TManagedArrayAccessor<FVector3f> Vertex;
-		TManagedArrayAccessor<FVector3f> TangentU;
-		TManagedArrayAccessor<FVector3f> TangentV;
-		TManagedArrayAccessor<FVector3f> Normal;
-		TManagedArrayAccessor<TArray<FVector2f>> UVs;
-		TManagedArrayAccessor<FLinearColor> Color;
-		TManagedArrayAccessor<int32> BoneMap;
-		TManagedArrayAccessor<int32> VertexStart;
-		TManagedArrayAccessor<int32> VertexCount;
-
-		TManagedArrayAccessor<FIntVector> Indices;
-		TManagedArrayAccessor<bool> Visible;
-		TManagedArrayAccessor<int32> MaterialIndex;
-		TManagedArrayAccessor<int32> MaterialID;
-		TManagedArrayAccessor<int32> FaceStart;
-		TManagedArrayAccessor<int32> FaceCount;
+		/**
+		 * Returns the vertex indicies for the bone
+		 */
+		const TArray<int32> GetVertexIndices(int32 BoneIdx) const;
+
+		/**
+		 * Returns the vertex positions for the bone in bone space
+		 */
+		const TArrayView<const FVector3f> GetVertexPositions(int32 BoneIdx) const;
+
+		/**
+		 * Returns the face indices for the bone
+		 */
+		const TArray<int32> GetFaceIndices(int32 BoneIdx) const;
+
+		/**
+		 * Returns the vertex indices of the face for the bone
+		 */
+//		const TArrayView<FIntVector> GetFaceVertexIndices(int32 BoneIdx) const;
+
+		TManagedArrayAccessor<int32> TransformToGeometryIndexAttribute;
+		TManagedArrayAccessor<FVector3f> VertexAttribute;
+		TManagedArrayAccessor<FVector3f> TangentUAttribute;
+		TManagedArrayAccessor<FVector3f> TangentVAttribute;
+		TManagedArrayAccessor<FVector3f> NormalAttribute;
+		TManagedArrayAccessor<TArray<FVector2f>> UVsAttribute;
+		TManagedArrayAccessor<FLinearColor> ColorAttribute;
+		TManagedArrayAccessor<int32> BoneMapAttribute;
+		TManagedArrayAccessor<int32> VertexStartAttribute;
+		TManagedArrayAccessor<int32> VertexCountAttribute;
+		TManagedArrayAccessor<FIntVector> IndicesAttribute;
+		TManagedArrayAccessor<bool> VisibleAttribute;
+		TManagedArrayAccessor<int32> MaterialIndexAttribute;
+		TManagedArrayAccessor<int32> MaterialIDAttribute;
+		TManagedArrayAccessor<int32> FaceStartAttribute;
+		TManagedArrayAccessor<int32> FaceCountAttribute;
 	};
 
 }
\ No newline at end of file
diff --git a/Engine/Source/Runtime/Experimental/Chaos/Public/GeometryCollection/Facades/CollectionTransformFacade.h b/Engine/Source/Runtime/Experimental/Chaos/Public/GeometryCollection/Facades/CollectionTransformFacade.h
--- a/Engine/Source/Runtime/Experimental/Chaos/Public/GeometryCollection/Facades/CollectionTransformFacade.h
+++ b/Engine/Source/Runtime/Experimental/Chaos/Public/GeometryCollection/Facades/CollectionTransformFacade.h
@@ -11,23 +11,11 @@ namespace GeometryCollection::Facades
 	/**
 	 * Provides an API to read and manipulate hierarchy in a managed array collection
 	 */
-	class CHAOS_API FTransformFacade
+	class CHAOS_API FCollectionTransformFacade
 	{
-		FManagedArrayCollection& Self;
-
 	public:
-
-		// Group
-		static const FName TransformGroup;
-
-		// Attributes
-		static const FName ParentAttribute;
-		static const FName ChildrenAttribute;
-		static const FName TransformAttribute;
-
-
-		FTransformFacade(FManagedArrayCollection& InCollection);
-		FTransformFacade(const FManagedArrayCollection& InCollection);
+		FCollectionTransformFacade(FManagedArrayCollection& InCollection);
+		FCollectionTransformFacade(const FManagedArrayCollection& InCollection);
 
 		/** Create the facade attributes. */
 		void DefineSchema() {}
@@ -36,29 +24,49 @@ namespace GeometryCollection::Facades
 		bool IsValid() const;
 
 		/** Is the facade defined constant. */
-		bool IsConst() const { return Parent.IsConst(); }
+		bool IsConst() const { return ParentAttribute.IsConst(); }
 
 		/** Get the root index */
-		TSet<int32> GetRootIndices() const;
+		TArray<int32> GetRootIndices() const;
 
 		/**
 		* Return the parent indices from the collection. Null if not initialized.
 		*/
-		const TManagedArray< int32 >* GetParents() const { return Parent.Find(); }
+		const TManagedArray< int32 >* GetParents() const { return ParentAttribute.Find(); }
 
 		/**
 		* Return the child indicesfrom the collection. Null if not initialized.
 		*/
-		const TManagedArray< TSet<int32> >* FindChildren() const { return Children.Find(); }
+		const TManagedArray< TSet<int32> >* FindChildren() const { return ChildrenAttribute.Find(); }
 
 		/**
 		* Return the child indicesfrom the collection. Null if not initialized.
 		*/
-		const TManagedArray< FTransform >* FindTransforms() const { return Transform.Find(); }
+		const TManagedArray< FTransform >* FindTransforms() const { return TransformAttribute.Find(); }
+
+		/**
+		* Return array of transforms for transforming from bone space to collection space
+		* Vertex(inCollectionSpace) = TransformComputed.TransformPosition(Vertex(inBoneSpace));
+		*/
+		TArray<FTransform> ComputeCollectionSpaceTransforms() const;
+
+		/**
+		* Return the child indicesfrom the collection. Null if not initialized.
+		*/
+		FTransform ComputeCollectionSpaceTransform(int32 BoneIdx) const;
+
+		/** Transforms the pivot of a collection */
+		void SetPivot(const FTransform& InTransform);
+
+		/** Transforms collection */
+		void Transform(const FTransform& InTransform);
+
+		/** Transforms selected bones in the collection */
+		void Transform(const FTransform& InTransform, const TArray<int32>& InSelection);
 
 	private:
-		TManagedArrayAccessor<int32>		Parent;
-		TManagedArrayAccessor<TSet<int32>>	Children;
-		TManagedArrayAccessor<FTransform>	Transform;
+		TManagedArrayAccessor<int32>		ParentAttribute;
+		TManagedArrayAccessor<TSet<int32>>	ChildrenAttribute;
+		TManagedArrayAccessor<FTransform>	TransformAttribute;
 	};
 }
\ No newline at end of file
diff --git a/Engine/Source/Runtime/Experimental/Chaos/Public/GeometryCollection/Facades/CollectionTransformSelectionFacade.h b/Engine/Source/Runtime/Experimental/Chaos/Public/GeometryCollection/Facades/CollectionTransformSelectionFacade.h
new file mode 100644
--- /dev/null
+++ b/Engine/Source/Runtime/Experimental/Chaos/Public/GeometryCollection/Facades/CollectionTransformSelectionFacade.h
@@ -0,0 +1,219 @@
+// Copyright Epic Games, Inc. All Rights Reserved.
+
+#pragma once
+#include "Containers/UnrealString.h"
+#include "Containers/Array.h"
+#include "Math/MathFwd.h"
+#include "GeometryCollection/ManagedArrayAccessor.h"
+#include "GeometryCollection/ManagedArrayCollection.h"
+#include "GeometryCollection/Facades/CollectionSelectionFacade.h"
+#include "Chaos/Triangle.h"
+
+class FGeometryCollection;
+
+namespace GeometryCollection::Facades
+{
+
+	/**
+	* FCollectionTransformSelectionFacade
+	* 
+	* Defines common API for selecting transforms on a collection
+	* 
+	*/
+	class CHAOS_API FCollectionTransformSelectionFacade
+	{
+	public:
+
+		/**
+		* FCollectionTransformSelectionFacade Constuctor
+		* @param FManagedArrayCollection : Collection input
+		*/
+		FCollectionTransformSelectionFacade(const FManagedArrayCollection& InSelf);
+
+		/** Create the facade attributes. */
+		void DefineSchema();
+
+		/** Is the facade defined constant. */
+		bool IsConst() const { return ParentAttribute.IsConst(); }
+
+		/** Is the Facade defined on the collection? */
+		bool IsValid() const;		
+
+		/**  */
+		bool IsValidBone(const int32 Index) const { return Index >= 0 && Index < ParentAttribute.Num(); }
+
+		/**  */
+		bool IsARootBone(const int32 Index) const;
+
+		/**  */
+		bool HasSelectedAncestor(const TArray<int32>& InSelection, const int32 Index) const;
+
+		/**  */
+		TArray<int32> GetBonesByLevel(const int32 Level, bool bOnlyClusteredOrRigid, bool bSkipFiltered) const;
+
+		/**  */
+		void Sanitize(TArray<int32>& InSelection, bool bFavorParents = true) const;
+
+		/**  */
+		void RemoveRootNodes(TArray<int32>& InOutSelection) const;
+
+		/**  */
+		void ConvertSelectionToRigidNodes(const int32 Index, TArray<int32>& InSelection) const;
+		void ConvertSelectionToRigidNodes(TArray<int32>& InSelection) const;
+
+		/**  */
+		bool CanSelectRootBones() const { return ParentAttribute.IsValid(); }
+		TArray<int32> SelectRootBones() const;
+
+		/**  */
+		TArray<int32> SelectNone() const;
+
+		/**  */
+		bool CanSelectAll() const { return ParentAttribute.IsValid(); }
+		TArray<int32> SelectAll() const;
+
+		/**  */
+		bool CanSelectInverse() const { return ParentAttribute.IsValid(); }
+		void SelectInverse(TArray<int32>& InOutSelection) const;
+
+		/**  */
+		bool CanSelectRandom() const { return ParentAttribute.IsValid(); }
+		TArray<int32> SelectRandom(bool bDeterministic, float RandomSeed, float RandomThresholdVal) const;
+
+		/**  */
+		bool CanSelectLeaf() const { return ParentAttribute.IsValid() && LevelAttribute.IsValid() && SimulationTypeAttribute.IsValid(); }
+		TArray<int32> SelectLeaf() const;
+
+		/**  */
+		bool CanSelectCluster() const { return ParentAttribute.IsValid() && LevelAttribute.IsValid() && SimulationTypeAttribute.IsValid(); }
+		TArray<int32> SelectCluster() const;
+
+		/**  */
+		bool CanSelectContact() const { return TransformIndexAttribute.IsValid() && TransformToGeometryIndexAttribute.IsValid(); }
+		void SelectContact(TArray<int32>& InOutSelection) const;
+
+		/**  */
+		bool CanSelectParent() const { return ParentAttribute.IsValid(); }
+		void SelectParent(TArray<int32>& InOutSelection) const;
+
+		/**  */
+		bool CanSelectChildren() const { return ChildrenAttribute.IsValid(); }
+		void SelectChildren(TArray<int32>& InOutSelection) const;
+
+		/**  */
+		bool CanSelectSiblings() const { return ParentAttribute.IsValid() && ChildrenAttribute.IsValid(); }
+		void SelectSiblings(TArray<int32>& InOutSelection) const;
+
+		/**  */
+		bool CanSelectLevel() const { return ParentAttribute.IsValid() && LevelAttribute.IsValid(); }
+		void SelectLevel(TArray<int32>& InOutSelection) const;
+
+		/**  */
+		static void SelectByPercentage(TArray<int32>& InOutSelection, int32 Percentage, bool Deterministic, float RandomSeed);
+
+		/**  */
+		TArray<int32> SelectBySize(float SizeMin, float SizeMax, bool bInclusive, bool bInsideRange) const;
+
+		/**  */
+		TArray<int32> SelectByVolume(float VolumeMin, float VolumeMax, bool bInclusive, bool bInsideRange) const;
+
+		/** 
+		* Returns indices of bones where all the vertices or at least one vertex of the bone is inside InBox 
+		*/
+		bool CanSelectVerticesInBox() const { return ParentAttribute.IsValid(); }
+		TArray<int32> SelectVerticesInBox(const FBox& InBox, const FTransform& InBoxTransform, bool bAllVerticesInBox) const;
+
+		/**
+		* Returns indices of bones where centroid of the bone is inside InBox
+		*/
+		bool CanSelectCentroidInBox() const { return ParentAttribute.IsValid(); }
+		TArray<int32> SelectCentroidInBox(const FBox& InBox, const FTransform& InBoxTransform) const;
+
+		/**
+		* Returns indices of bones where BoundingBox of the bone is inside InBox
+		*/
+		bool CanSelectBoundingBoxInBox() const { return ParentAttribute.IsValid(); }
+		TArray<int32> SelectBoundingBoxInBox(const FBox& InBox, const FTransform& InBoxTransform) const;
+
+		/**
+		* Returns indices of bones where all the vertices or at least one vertex of the bone is inside a sphere
+		*/
+		bool CanSelectVerticesInSphere() const { return ParentAttribute.IsValid(); }
+		TArray<int32> SelectVerticesInSphere(const FSphere& InSphere, const FTransform& InSphereTransform, bool bAllVerticesInSphere) const;
+
+		/**
+		* Returns indices of bones where BoundingBox of the bone is inside a sphere
+		*/
+		bool CanSelectBoundingBoxInSphere() const { return ParentAttribute.IsValid(); }
+		TArray<int32> SelectBoundingBoxInSphere(const FSphere& InSphere, const FTransform& InSphereTransform) const;
+
+		/**
+		* Returns indices of bones where centroid of the bone is inside a sphere
+		*/
+		bool CanSelectCentroidInSphere() const { return ParentAttribute.IsValid(); }
+		TArray<int32> SelectCentroidInSphere(const FSphere& InSphere, const FTransform& InSphereTransform) const;
+
+		/** 
+		* Returns indices of bones where centroid of the bone is inside a convex hull
+		*/
+//		TArray<int32> SelectCentroidInConvex(const SOMECONVEXTYPE& InConvex, const FTransform& InTransform, const FTransform& InTransform) const;
+
+		/**
+		* Select by a float attribute in the collection, it can be any existing float attribute specified by GroupName/AttrName
+		*/
+		TArray<int32> SelectByFloatAttribute(FString GroupName, FString AttrName, float Min, float Max, bool bInclusive, bool bInsideRange) const;
+
+		/**
+		* Select by an int attribute in the collection, it can be any existing int attribute specified by GroupName/AttrName 
+		*/
+		TArray<int32> SelectByIntAttribute(FString GroupName, FString AttrName, int32 Min, int32 Max, bool bInclusive, bool bInsideRange) const;
+
+		/**
+		* Convert a VertexSelection to a TransformSelection, if bAllElementsMustBeSelected is true then all vertices of a transform must be selected for the transform to be selected 
+		*/
+		TArray<int32> ConvertVertexSelectionToTransformSelection(const TArray<int32>& InVertexSelection, bool bAllElementsMustBeSelected) const;
+
+		/**
+		* Convert a FaceSelection to a TransformSelection, if bAllElementsMustBeSelected is true then all faces of a transform must be selected for the transform to be selected
+		*/
+		TArray<int32> ConvertFaceSelectionToTransformSelection(const TArray<int32>& InFaceSelection, bool bAllElementsMustBeSelected) const;
+
+		/**
+		* Convert a VertexSelection to a FaceSelection, if bAllElementsMustBeSelected is true then all vertices of a face must be selectod for the face to be selected 
+		*/
+		TArray<int32> ConvertVertexSelectionToFaceSelection(const TArray<int32>& InVertexSelection, bool bAllElementsMustBeSelected) const;
+
+		/**
+		* Convert a TransformSelection to a FaceSelection 
+		*/
+		TArray<int32> ConvertTransformSelectionToFaceSelection(const TArray<int32>& InTransformSelection) const;
+
+		/**
+		* Convert a FaceSelection to a VertexSelection 
+		*/
+		TArray<int32> ConvertFaceSelectionToVertexSelection(const TArray<int32>& InFaceSelection) const;
+
+		/**
+		* Convert a TransformSelection to a VertexSelection 
+		*/
+		TArray<int32> ConvertTransformSelectionToVertexSelection(const TArray<int32>& InTransformSelection) const;
+
+
+	private:
+		// This a const facade
+		const FManagedArrayCollection& ConstCollection;
+
+		TManagedArrayAccessor<int32>		ParentAttribute;
+		TManagedArrayAccessor<TSet<int32>>	ChildrenAttribute;
+		TManagedArrayAccessor<int32>		LevelAttribute;
+		TManagedArrayAccessor<int32>		SimulationTypeAttribute;
+		TManagedArrayAccessor<int32>		TransformIndexAttribute;
+		TManagedArrayAccessor<int32>		TransformToGeometryIndexAttribute;
+		TManagedArrayAccessor<int32>		VertexStartAttribute;
+		TManagedArrayAccessor<int32>		VertexCountAttribute;
+		TManagedArrayAccessor<int32>		FaceStartAttribute;
+		TManagedArrayAccessor<int32>		FaceCountAttribute;
+		TManagedArrayAccessor<FIntVector>	IndicesAttribute;
+	};
+
+}
\ No newline at end of file
diff --git a/Engine/Source/Runtime/Experimental/Dataflow/Core/Private/Dataflow/DataflowSelection.cpp b/Engine/Source/Runtime/Experimental/Dataflow/Core/Private/Dataflow/DataflowSelection.cpp
--- a/Engine/Source/Runtime/Experimental/Dataflow/Core/Private/Dataflow/DataflowSelection.cpp
+++ b/Engine/Source/Runtime/Experimental/Dataflow/Core/Private/Dataflow/DataflowSelection.cpp
@@ -33,6 +33,26 @@ void FDataflowSelection::AsArray(TArray<int32>& SelectionArr) const
 	}
 }
 
+TArray<int32> FDataflowSelection::AsArray() const
+{
+	TArray<int32> SelectionArr;
+
+	TBitArray<>::FConstIterator It(SelectionArray);
+	while (It)
+	{
+		int32 Index = It.GetIndex();
+
+		if (SelectionArray[Index])
+		{
+			SelectionArr.Add(Index);
+		}
+
+		++It;
+	}
+
+	return SelectionArr;
+}
+
 void FDataflowSelection::SetFromArray(const TArray<int32>& SelectionArr)
 {
 	SelectionArray.Init(false, SelectionArray.Num());
diff --git a/Engine/Source/Runtime/Experimental/Dataflow/Core/Public/Dataflow/DataflowSelection.h b/Engine/Source/Runtime/Experimental/Dataflow/Core/Public/Dataflow/DataflowSelection.h
--- a/Engine/Source/Runtime/Experimental/Dataflow/Core/Public/Dataflow/DataflowSelection.h
+++ b/Engine/Source/Runtime/Experimental/Dataflow/Core/Public/Dataflow/DataflowSelection.h
@@ -18,13 +18,16 @@ struct DATAFLOWCORE_API FDataflowSelection
 	void SetSelected(int32 Idx) { SelectionArray[Idx] = true; }
 	void SetNotSelected(int32 Idx) { SelectionArray[Idx] = false; }
 	void AsArray(TArray<int32>& SelectionArr) const;
+	TArray<int32> AsArray() const;
 	void SetFromArray(const TArray<int32>& SelectionArr);
 	void AND(const FDataflowSelection& Other, FDataflowSelection& Result) const;
 	void OR(const FDataflowSelection& Other, FDataflowSelection& Result) const;
 	void XOR(const FDataflowSelection& Other, FDataflowSelection& Result) const;
 	void Invert() { SelectionArray.BitwiseNOT(); }
 	void SetWithMask(const bool Value, const FDataflowSelection& Mask);
+
 private:
+
 	TBitArray<> SelectionArray;
 };
 
@@ -33,3 +36,18 @@ struct DATAFLOWCORE_API FDataflowTransformSelection : public FDataflowSelection
 {
 	GENERATED_USTRUCT_BODY()
 };
+
+
+USTRUCT()
+struct DATAFLOWCORE_API FDataflowVertexSelection : public FDataflowSelection
+{
+	GENERATED_USTRUCT_BODY()
+};
+
+
+USTRUCT()
+struct DATAFLOWCORE_API FDataflowFaceSelection : public FDataflowSelection
+{
+	GENERATED_USTRUCT_BODY()
+};
+
