Chaos - fix joint state when a kinematic has a center of mass offset

diff --git a/Engine/Source/Programs/HeadlessChaos/Private/HeadlessChaosTestJoints.cpp b/Engine/Source/Programs/HeadlessChaos/Private/HeadlessChaosTestJoints.cpp
--- a/Engine/Source/Programs/HeadlessChaos/Private/HeadlessChaosTestJoints.cpp
+++ b/Engine/Source/Programs/HeadlessChaos/Private/HeadlessChaosTestJoints.cpp
@@ -689,6 +689,69 @@ namespace ChaosTest {
 		JointConstraint_DisableOneConstrainedParticle<FPBDRigidsEvolutionGBF>();
 	}
 
+	// Create a kinematic-dynamic particle chain with a center of mass offset on the kinematic
+	// and verify that the joint offsets are used correctly. The two particles are arranged vertically
+	// with the joint between them, the center of mass offset should not affect behaviour.
+	//
+	// NOTE: We create both particles as dynamics and then change one to be kinematic so that we
+	// can alter its mass properties.
+	//
+	GTEST_TEST(JointTests, TestJointCoMOffset)
+	{
+		const int32 NumIterations = 1;
+		const FReal Gravity = 980;
+		FJointConstraintsTest<FPBDRigidsEvolutionGBF> Test(NumIterations, Gravity);
+
+		Test.ParticlePositions =
+		{
+			{ (FReal)0, (FReal)0, (FReal)0 },
+			{ (FReal)0, (FReal)0, (FReal)1000 },
+		};
+		Test.ParticleSizes =
+		{
+			{ (FReal)100, (FReal)100, (FReal)100 },
+			{ (FReal)100, (FReal)100, (FReal)100 },
+		};
+		Test.ParticleMasses =
+		{
+			(FReal)1000,
+			(FReal)1000,
+		};
+
+		Test.JointPositions =
+		{
+			{ (FReal)0, (FReal)0, (FReal)500 },
+		};
+		Test.JointParticleIndices =
+		{
+			{ 0, 1 },
+		};
+
+		Test.Create();
+
+		// Move the center of mass of the soon-to-be kinematic
+		Test.GetParticle(0)->CastToRigidParticle()->SetCenterOfMass(FVec3(0, 0, -100));
+		EXPECT_NEAR((Test.GetParticle(0)->CastToRigidParticle()->CenterOfMass() - FVec3(0, 0, -100)).Size(), 0, UE_KINDA_SMALL_NUMBER);
+
+		// Make the root body kinematic
+		Test.Evolution.SetParticleObjectState(Test.GetParticle(0)->CastToRigidParticle(), EObjectStateType::Kinematic);
+
+		// The kinematic will now report zero center of mass via the GenericParticle API, but internally it will still be set
+		EXPECT_NEAR(FGenericParticleHandle(Test.GetParticle(0))->CenterOfMass().Size(), 0, UE_KINDA_SMALL_NUMBER);
+		EXPECT_NEAR((Test.GetParticle(0)->CastToRigidParticle()->CenterOfMass() - FVec3(0, 0, -100)).Size(), 0, UE_KINDA_SMALL_NUMBER);
+
+		const FReal Dt = 0.01f;
+		for (int32 i = 0; i < 100; ++i)
+		{
+			Test.Evolution.AdvanceOneTimeStep(Dt);
+			Test.Evolution.EndFrame(Dt);
+
+			// Neither particle should have moved
+			EXPECT_LT((Test.GetParticle(0)->X() - Test.ParticlePositions[0]).Size(), (FReal)0.1);
+			EXPECT_LT((Test.GetParticle(1)->X() - Test.ParticlePositions[1]).Size(), (FReal)0.1);
+		}
+	}
+
 	// Check that constraints end up in the same island when graph is fully connected
 	GTEST_TEST(JointTests, TestJointConstraintGraph_Connected)
 	{
diff --git a/Engine/Source/Runtime/Experimental/Chaos/Public/Chaos/ParticleHandle.h b/Engine/Source/Runtime/Experimental/Chaos/Public/Chaos/ParticleHandle.h
--- a/Engine/Source/Runtime/Experimental/Chaos/Public/Chaos/ParticleHandle.h
+++ b/Engine/Source/Runtime/Experimental/Chaos/Public/Chaos/ParticleHandle.h
@@ -1569,7 +1569,7 @@ public:
 
 	FVec3& P()
 	{
-		if (MHandle->CastToRigidParticle())
+		if (IsDynamic())
 		{
 			return MHandle->CastToRigidParticle()->P();
 		}
@@ -1579,7 +1579,7 @@ public:
 
 	const FVec3& P() const
 	{
-		if (MHandle->CastToRigidParticle())
+		if (IsDynamic())
 		{
 			return MHandle->CastToRigidParticle()->P();
 		}
@@ -1589,7 +1589,7 @@ public:
 
 	FRotation3& Q()
 	{
-		if (MHandle->CastToRigidParticle())
+		if (IsDynamic())
 		{
 			return MHandle->CastToRigidParticle()->Q();
 		}
@@ -1599,7 +1599,7 @@ public:
 
 	const FRotation3& Q() const
 	{
-		if (MHandle->CastToRigidParticle())
+		if (IsDynamic())
 		{
 			return MHandle->CastToRigidParticle()->Q();
 		}
@@ -1610,7 +1610,7 @@ public:
 	// World-space center of mass position
 	const FVec3 XCom() const
 	{
-		if (MHandle->CastToRigidParticle())
+		if (IsDynamic())
 		{
 			return MHandle->CastToRigidParticle()->XCom();
 		}
@@ -1618,7 +1618,7 @@ public:
 	}
 	const FVec3 PCom() const
 	{
-		if (MHandle->CastToRigidParticle())
+		if (IsDynamic())
 		{
 			return MHandle->CastToRigidParticle()->PCom();
 		}
@@ -1628,7 +1628,7 @@ public:
 	// World-space center of mass rotation
 	const FRotation3 RCom() const
 	{ 
-		if (MHandle->CastToRigidParticle())
+		if (IsDynamic())
 		{
 			return MHandle->CastToRigidParticle()->RCom();
 		}
@@ -1636,7 +1636,7 @@ public:
 	}
 	const FRotation3 QCom() const
 	{ 
-		if (MHandle->CastToRigidParticle())
+		if (IsDynamic())
 		{
 			return MHandle->CastToRigidParticle()->QCom();
 		}
@@ -1654,7 +1654,7 @@ public:
 
 	void SetTransformPQCom(const FVec3& InPCom, const FRotation3& InQCom)
 	{ 
-		if (MHandle->CastToRigidParticle())
+		if (IsDynamic())
 		{
 			MHandle->CastToRigidParticle()->SetTransformPQCom(InPCom, InQCom);
 		}
@@ -1671,7 +1671,7 @@ public:
 
 	FRigidTransform3 GetTransformPQ() const
 	{
-		if (MHandle->CastToRigidParticle())
+		if (IsDynamic())
 		{
 			return MHandle->CastToRigidParticle()->GetTransformPQ();
 		}
@@ -1680,7 +1680,7 @@ public:
 
 	FRigidTransform3 GetTransformXRCom() const
 	{
-		if (MHandle->CastToRigidParticle())
+		if (IsDynamic())
 		{
 			return MHandle->CastToRigidParticle()->GetTransformXRCom();
 		}
@@ -1689,13 +1689,51 @@ public:
 
 	FRigidTransform3 GetTransformPQCom() const
 	{
-		if (MHandle->CastToRigidParticle())
+		if (IsDynamic())
 		{
 			return MHandle->CastToRigidParticle()->GetTransformPQCom();
 		}
 		return GetTransformXR();
 	}
 
+	/**
+	 * Convert a particle-relative position into a com-relative position
+	 */
+	FVec3 GetComRelativePosition(const FVec3& P)
+	{
+		if (IsDynamic())
+		{
+			FPBDRigidParticleHandle* Rigid = MHandle->CastToRigidParticle();
+			return Rigid->RotationOfMass().UnrotateVector(P - Rigid->CenterOfMass());
+		}
+		return P;
+	}
+
+	/**
+	 * Convert a particle-relative rotation into a com-relative rotation
+	 */
+	FRotation3 GetComRelativeRotation(const FRotation3& Q)
+	{
+		if (IsDynamic())
+		{
+			FPBDRigidParticleHandle* Rigid = MHandle->CastToRigidParticle();
+			return Rigid->RotationOfMass().Inverse() * Q;
+		}
+		return Q;
+	}
+
+	/**
+	 * Convert a particle-relative transform into a com-relative transform
+	 */
+	FRigidTransform3 GetComRelativeTransform(const FRigidTransform3& T)
+	{
+		if (IsDynamic())
+		{
+			return FRigidTransform3(GetComRelativePosition(T.GetLocation()), GetComRelativeRotation(T.GetRotation()));
+		}
+		return T;
+	}
+
 	const FVec3& VSmooth() const
 	{
 		if (MHandle->CastToRigidParticle())
@@ -1718,7 +1756,7 @@ public:
 
 	const FVec3& Acceleration() const
 	{ 
-		if (IsDynamic())
+		if (MHandle->CastToRigidParticle())
 		{
 			return MHandle->CastToRigidParticle()->Acceleration();
 		}
@@ -1727,7 +1765,7 @@ public:
 	}
 	const FVec3& AngularAcceleration() const
 	{ 
-		if (IsDynamic())
+		if (MHandle->CastToRigidParticle())
 		{
 			return MHandle->CastToRigidParticle()->AngularAcceleration();
 		}
@@ -1784,12 +1822,9 @@ public:
 
 	const TVec3<FRealSingle> I() const
 	{ 
-		if (auto RigidHandle = MHandle->CastToRigidParticle())
+		if (IsDynamic())
 		{
-			if (IsDynamic())
-			{
-				return RigidHandle->I();
-			}
+			return MHandle->CastToRigidParticle()->I();
 		}
 
 		return TVec3<FRealSingle>(0);
@@ -1797,12 +1832,9 @@ public:
 
 	const TVec3<FRealSingle> InvI() const
 	{ 
-		if (auto RigidHandle = MHandle->CastToRigidParticle())
+		if (IsDynamic())
 		{
-			if (IsDynamic())
-			{
-				return MHandle->CastToRigidParticle()->InvI();
-			}
+			return MHandle->CastToRigidParticle()->InvI();
 		}
 
 		return TVec3<FRealSingle>(0);
@@ -1810,12 +1842,9 @@ public:
 
 	FReal M() const
 	{
-		if (auto RigidHandle = MHandle->CastToRigidParticle())
+		if (IsDynamic())
 		{
-			if (IsDynamic())
-			{
-				return MHandle->CastToRigidParticle()->M();
-			}
+			return MHandle->CastToRigidParticle()->M();
 		}
 
 		return (FReal)0;
@@ -1823,12 +1852,9 @@ public:
 
 	FReal InvM() const
 	{
-		if (auto RigidHandle = MHandle->CastToRigidParticle())
+		if (IsDynamic())
 		{
-			if (IsDynamic())
-			{
-				return MHandle->CastToRigidParticle()->InvM();
-			}
+			return MHandle->CastToRigidParticle()->InvM();
 		}
 
 		return (FReal)0;
@@ -1836,29 +1862,45 @@ public:
 
 	FVec3 CenterOfMass() const
 	{
-		if (auto RigidHandle = MHandle->CastToRigidParticle())
+		if (IsDynamic())
 		{
-			return RigidHandle->CenterOfMass();
+			return MHandle->CastToRigidParticle()->CenterOfMass();
 		}
 
 		return FVec3(0);
 	}
 
+	void SetCenterOfMass(const FVec3& InCom)
+	{
+		if (IsDynamic())
+		{
+			MHandle->CastToRigidParticle()->SetCenterOfMass(InCom);
+		}
+	}
+
 	FRotation3 RotationOfMass() const
 	{
-		if (auto RigidHandle = MHandle->CastToRigidParticle())
+		if (IsDynamic())
 		{
-			return RigidHandle->RotationOfMass();
+			return MHandle->CastToRigidParticle()->RotationOfMass();
 		}
 
 		return FRotation3::FromIdentity();
 	}
 
+	void SetRotationOfMass(const FRotation3& InRom)
+	{
+		if (IsDynamic())
+		{
+			MHandle->CastToRigidParticle()->SetRotationOfMass(InRom);
+		}
+	}
+
 	TVec3<FRealSingle> InvIConditioning() const
 	{ 
-		if (auto RigidHandle = MHandle->CastToRigidParticle())
+		if (IsDynamic())
 		{
-			return RigidHandle->InvIConditioning();
+			return MHandle->CastToRigidParticle()->InvIConditioning();
 		}
 
 		return TVec3<FRealSingle>(1);
@@ -1874,9 +1916,9 @@ public:
 
 	TVec3<FRealSingle> ConditionedInvI() const
 	{ 
-		if (auto RigidHandle = MHandle->CastToRigidParticle())
+		if (IsDynamic())
 		{
-			return RigidHandle->ConditionedInvI();
+			return MHandle->CastToRigidParticle()->ConditionedInvI();
 		}
 
 		return TVec3<FRealSingle>(0);
@@ -1884,9 +1926,9 @@ public:
 
 	TVec3<FRealSingle> ConditionedI() const
 	{ 
-		if (auto RigidHandle = MHandle->CastToRigidParticle())
+		if (IsDynamic())
 		{
-			return RigidHandle->ConditionedI();
+			return MHandle->CastToRigidParticle()->ConditionedI();
 		}
 
 		return TVec3<FRealSingle>(0);
