Geometry collection : fix crash caused by a geometry collection made of meshes with different number of UV channels

diff --git a/Engine/Source/Runtime/Experimental/GeometryCollectionEngine/Private/GeometryCollection/GeometryCollectionComponent.cpp b/Engine/Source/Runtime/Experimental/GeometryCollectionEngine/Private/GeometryCollection/GeometryCollectionComponent.cpp
--- a/Engine/Source/Runtime/Experimental/GeometryCollectionEngine/Private/GeometryCollection/GeometryCollectionComponent.cpp
+++ b/Engine/Source/Runtime/Experimental/GeometryCollectionEngine/Private/GeometryCollection/GeometryCollectionComponent.cpp
@@ -1804,6 +1804,7 @@ void SetHierarchyStrain(Chaos::TPBDRigidClusteredParticleHandle<Chaos::FReal, 3>
 void UGeometryCollectionComponent::InitConstantData(FGeometryCollectionConstantData* ConstantData) const
 {
 	// Constant data should all be moved to the DDC as time permits.
+	// todo : this should be computed once per asset not per component or at least the part that does not depend on the component properties
 
 	check(ConstantData);
 	check(RestCollection);
@@ -1834,8 +1835,13 @@ void UGeometryCollectionComponent::InitConstantData(FGeometryCollectionConstantD
 		ConstantData->Normals = TArray<FVector3f>(Normal.GetData(), Normal.Num());
 		ConstantData->Colors = TArray<FLinearColor>(Color.GetData(), Color.Num());
 
-		// use the first point to know the number of channels
-		const int32 NumUVChannels = Collection->UVs.Num() ? Collection->UVs[0].Num() : 1;
+		// find the maximum number of channels 
+		int32 NumUVChannels = 0;
+		for (const TArray<FVector2f>& VertexUVs : Collection->UVs.GetConstArray())
+		{
+			NumUVChannels = FMath::Max(NumUVChannels, VertexUVs.Num());
+		}
+
 		ConstantData->UVChannels.SetNum(NumUVChannels);
 		for (int32 UVChannelIndex = 0; UVChannelIndex < NumUVChannels; UVChannelIndex++)
 		{
@@ -1852,10 +1858,15 @@ void UGeometryCollectionComponent::InitConstantData(FGeometryCollectionConstantD
 					const int32 BoneIndex = ConstantData->BoneMap[InPointIndex];
 					ConstantData->BoneColors[InPointIndex] = BoneColors[BoneIndex];
 
-					for (int32 UVChannelIndex = 0; UVChannelIndex < NumUVChannels; UVChannelIndex++)
+					const TArray<FVector2f>& VertexUVs = Collection->UVs[InPointIndex];
+					for (int32 UVChannelIndex = 0; UVChannelIndex < VertexUVs.Num(); UVChannelIndex++)
 					{
 						ConstantData->UVChannels[UVChannelIndex][InPointIndex] = Collection->UVs[InPointIndex][UVChannelIndex];
 					}
+					for (int32 UVChannelIndex = VertexUVs.Num(); UVChannelIndex < NumUVChannels; UVChannelIndex++)
+					{
+						ConstantData->UVChannels[UVChannelIndex][InPointIndex] = FVector2f::ZeroVector;
+					}
 				});
 		}
 		else
