Updated the underlying types for regular and namespaced enums.

diff --git a/Engine/Source/Runtime/Engine/Classes/Materials/MaterialExpressionViewProperty.h b/Engine/Source/Runtime/Engine/Classes/Materials/MaterialExpressionViewProperty.h
--- a/Engine/Source/Runtime/Engine/Classes/Materials/MaterialExpressionViewProperty.h
+++ b/Engine/Source/Runtime/Engine/Classes/Materials/MaterialExpressionViewProperty.h
@@ -9,7 +9,7 @@
 #include "MaterialExpressionViewProperty.generated.h"
 
 UENUM()
-enum EMaterialExposedViewProperty
+enum EMaterialExposedViewProperty : int
 {	
 	/** Horizontal and vertical size of the view's buffer in pixels */
 	MEVP_BufferSize UMETA(DisplayName="Render Target Size"),
diff --git a/Engine/Source/Runtime/Engine/Classes/Materials/MaterialExpressionWorldPosition.h b/Engine/Source/Runtime/Engine/Classes/Materials/MaterialExpressionWorldPosition.h
--- a/Engine/Source/Runtime/Engine/Classes/Materials/MaterialExpressionWorldPosition.h
+++ b/Engine/Source/Runtime/Engine/Classes/Materials/MaterialExpressionWorldPosition.h
@@ -10,7 +10,7 @@
 
 /** Specifies which shader generated offsets should included in the world position (displacement/WPO etc.) */
 UENUM()
-enum EWorldPositionIncludedOffsets
+enum EWorldPositionIncludedOffsets : int
 {
 	/** Absolute world position with all material shader offsets applied */
 	WPT_Default UMETA(DisplayName="Absolute World Position (Including Material Shader Offsets)"),
diff --git a/Engine/Source/Runtime/Engine/Classes/Materials/MaterialInterface.h b/Engine/Source/Runtime/Engine/Classes/Materials/MaterialInterface.h
--- a/Engine/Source/Runtime/Engine/Classes/Materials/MaterialInterface.h
+++ b/Engine/Source/Runtime/Engine/Classes/Materials/MaterialInterface.h
@@ -66,7 +66,7 @@ enum EShaderPlatform : uint16;
 typedef TArray<FMaterialResource*> FMaterialResourceDeferredDeletionArray;
 
 UENUM(BlueprintType)
-enum EMaterialUsage
+enum EMaterialUsage : int
 {
 	MATUSAGE_SkeletalMesh,
 	MATUSAGE_ParticleSprites,
diff --git a/Engine/Source/Runtime/Engine/Classes/Particles/Attractor/ParticleModuleAttractorParticle.h b/Engine/Source/Runtime/Engine/Classes/Particles/Attractor/ParticleModuleAttractorParticle.h
--- a/Engine/Source/Runtime/Engine/Classes/Particles/Attractor/ParticleModuleAttractorParticle.h
+++ b/Engine/Source/Runtime/Engine/Classes/Particles/Attractor/ParticleModuleAttractorParticle.h
@@ -13,7 +13,7 @@ class UParticleModuleTypeDataBase;
 struct FParticleEmitterInstance;
 
 UENUM()
-enum EAttractorParticleSelectionMethod
+enum EAttractorParticleSelectionMethod : int
 {
 	EAPSM_Random UMETA(DisplayName="Random"),
 	EAPSM_Sequential UMETA(DisplayName="Sequential"),
diff --git a/Engine/Source/Runtime/Engine/Classes/Particles/Beam/ParticleModuleBeamBase.h b/Engine/Source/Runtime/Engine/Classes/Particles/Beam/ParticleModuleBeamBase.h
--- a/Engine/Source/Runtime/Engine/Classes/Particles/Beam/ParticleModuleBeamBase.h
+++ b/Engine/Source/Runtime/Engine/Classes/Particles/Beam/ParticleModuleBeamBase.h
@@ -16,7 +16,7 @@
 
 /** The method to use in determining the source/target. */
 UENUM()
-enum Beam2SourceTargetMethod
+enum Beam2SourceTargetMethod : int
 {
 	/** Default	- use the distribution. 
 	 *	This is the fallback for when other modes can't be resolved.
@@ -42,7 +42,7 @@ enum Beam2SourceTargetMethod
 
 /** The method to use in determining the source/target tangent. */
 UENUM()
-enum Beam2SourceTargetTangentMethod
+enum Beam2SourceTargetTangentMethod : int
 {
 	/** Direct - a direct line between source and target.				 */
 	PEB2STTM_Direct UMETA(DisplayName="Direct"),
diff --git a/Engine/Source/Runtime/Engine/Classes/Particles/Beam/ParticleModuleBeamModifier.h b/Engine/Source/Runtime/Engine/Classes/Particles/Beam/ParticleModuleBeamModifier.h
--- a/Engine/Source/Runtime/Engine/Classes/Particles/Beam/ParticleModuleBeamModifier.h
+++ b/Engine/Source/Runtime/Engine/Classes/Particles/Beam/ParticleModuleBeamModifier.h
@@ -25,7 +25,7 @@ struct FParticleEmitterInstance;
 
 /**	What to modify. */
 UENUM()
-enum BeamModifierType
+enum BeamModifierType : int
 {
 	/** Modify the source of the beam.				*/
 	PEB2MT_Source UMETA(DisplayName="Source"),
diff --git a/Engine/Source/Runtime/Engine/Classes/Particles/Camera/ParticleModuleCameraOffset.h b/Engine/Source/Runtime/Engine/Classes/Particles/Camera/ParticleModuleCameraOffset.h
--- a/Engine/Source/Runtime/Engine/Classes/Particles/Camera/ParticleModuleCameraOffset.h
+++ b/Engine/Source/Runtime/Engine/Classes/Particles/Camera/ParticleModuleCameraOffset.h
@@ -16,7 +16,7 @@ struct FParticleEmitterInstance;
  *	The update method for the offset
  */
 UENUM()
-enum EParticleCameraOffsetUpdateMethod
+enum EParticleCameraOffsetUpdateMethod : int
 {
 	EPCOUM_DirectSet UMETA(DisplayName="Direct Set"),
 	EPCOUM_Additive UMETA(DisplayName="Additive"),
diff --git a/Engine/Source/Runtime/Engine/Classes/Particles/Collision/ParticleModuleCollisionBase.h b/Engine/Source/Runtime/Engine/Classes/Particles/Collision/ParticleModuleCollisionBase.h
--- a/Engine/Source/Runtime/Engine/Classes/Particles/Collision/ParticleModuleCollisionBase.h
+++ b/Engine/Source/Runtime/Engine/Classes/Particles/Collision/ParticleModuleCollisionBase.h
@@ -12,7 +12,7 @@
  *	Flags indicating what to do with the particle when MaxCollisions is reached
  */
 UENUM()
-enum EParticleCollisionComplete
+enum EParticleCollisionComplete : int
 {
 	/**	Kill the particle when MaxCollisions is reached		*/
 	EPCC_Kill UMETA(DisplayName="Kill"),
diff --git a/Engine/Source/Runtime/Engine/Classes/Particles/Location/ParticleModuleLocationBoneSocket.h b/Engine/Source/Runtime/Engine/Classes/Particles/Location/ParticleModuleLocationBoneSocket.h
--- a/Engine/Source/Runtime/Engine/Classes/Particles/Location/ParticleModuleLocationBoneSocket.h
+++ b/Engine/Source/Runtime/Engine/Classes/Particles/Location/ParticleModuleLocationBoneSocket.h
@@ -17,7 +17,7 @@ struct FModuleLocationBoneSocketInstancePayload;
 struct FParticleEmitterInstance;
 
 UENUM()
-enum ELocationBoneSocketSource
+enum ELocationBoneSocketSource : int
 {
 	BONESOCKETSOURCE_Bones,
 	BONESOCKETSOURCE_Sockets,
@@ -25,7 +25,7 @@ enum ELocationBoneSocketSource
 };
 
 UENUM()
-enum ELocationBoneSocketSelectionMethod
+enum ELocationBoneSocketSelectionMethod : int
 {
 	BONESOCKETSEL_Sequential,
 	BONESOCKETSEL_Random,
diff --git a/Engine/Source/Runtime/Engine/Classes/Particles/Location/ParticleModuleLocationEmitter.h b/Engine/Source/Runtime/Engine/Classes/Particles/Location/ParticleModuleLocationEmitter.h
--- a/Engine/Source/Runtime/Engine/Classes/Particles/Location/ParticleModuleLocationEmitter.h
+++ b/Engine/Source/Runtime/Engine/Classes/Particles/Location/ParticleModuleLocationEmitter.h
@@ -18,7 +18,7 @@
 struct FParticleEmitterInstance;
 
 UENUM()
-enum ELocationEmitterSelectionMethod
+enum ELocationEmitterSelectionMethod : int
 {
 	ELESM_Random,
 	ELESM_Sequential,
diff --git a/Engine/Source/Runtime/Engine/Classes/Particles/Location/ParticleModuleLocationPrimitiveCylinder.h b/Engine/Source/Runtime/Engine/Classes/Particles/Location/ParticleModuleLocationPrimitiveCylinder.h
--- a/Engine/Source/Runtime/Engine/Classes/Particles/Location/ParticleModuleLocationPrimitiveCylinder.h
+++ b/Engine/Source/Runtime/Engine/Classes/Particles/Location/ParticleModuleLocationPrimitiveCylinder.h
@@ -15,7 +15,7 @@
 struct FParticleEmitterInstance;
 
 UENUM()
-enum CylinderHeightAxis
+enum CylinderHeightAxis : int
 {
 	PMLPC_HEIGHTAXIS_X UMETA(DisplayName="X"),
 	PMLPC_HEIGHTAXIS_Y UMETA(DisplayName="Y"),
diff --git a/Engine/Source/Runtime/Engine/Classes/Particles/Location/ParticleModuleLocationSkelVertSurface.h b/Engine/Source/Runtime/Engine/Classes/Particles/Location/ParticleModuleLocationSkelVertSurface.h
--- a/Engine/Source/Runtime/Engine/Classes/Particles/Location/ParticleModuleLocationSkelVertSurface.h
+++ b/Engine/Source/Runtime/Engine/Classes/Particles/Location/ParticleModuleLocationSkelVertSurface.h
@@ -17,7 +17,7 @@ class FSkeletalMeshLODRenderData;
 struct FSkelMeshRenderSection;
 
 UENUM()
-enum ELocationSkelVertSurfaceSource
+enum ELocationSkelVertSurfaceSource : int
 {
 	VERTSURFACESOURCE_Vert UMETA(DisplayName="Vertices"),
 	VERTSURFACESOURCE_Surface UMETA(DisplayName="Surfaces"),
diff --git a/Engine/Source/Runtime/Engine/Classes/Particles/Orbit/ParticleModuleOrbit.h b/Engine/Source/Runtime/Engine/Classes/Particles/Orbit/ParticleModuleOrbit.h
--- a/Engine/Source/Runtime/Engine/Classes/Particles/Orbit/ParticleModuleOrbit.h
+++ b/Engine/Source/Runtime/Engine/Classes/Particles/Orbit/ParticleModuleOrbit.h
@@ -14,7 +14,7 @@ class UParticleModuleTypeDataBase;
 struct FParticleEmitterInstance;
 
 UENUM()
-enum EOrbitChainMode
+enum EOrbitChainMode : int
 {
 	/** Add the module values to the previous results						*/
 	EOChainMode_Add UMETA(DisplayName="Add"),
diff --git a/Engine/Source/Runtime/Engine/Classes/Particles/Orientation/ParticleModuleOrientationAxisLock.h b/Engine/Source/Runtime/Engine/Classes/Particles/Orientation/ParticleModuleOrientationAxisLock.h
--- a/Engine/Source/Runtime/Engine/Classes/Particles/Orientation/ParticleModuleOrientationAxisLock.h
+++ b/Engine/Source/Runtime/Engine/Classes/Particles/Orientation/ParticleModuleOrientationAxisLock.h
@@ -12,7 +12,7 @@ struct FParticleEmitterInstance;
 
 // Flags indicating lock
 UENUM()
-enum EParticleAxisLock
+enum EParticleAxisLock : int
 {
 	/** No locking to an axis...							*/
 	EPAL_NONE UMETA(DisplayName="None"),
diff --git a/Engine/Source/Runtime/Engine/Classes/Particles/Parameter/ParticleModuleParameterDynamic.h b/Engine/Source/Runtime/Engine/Classes/Particles/Parameter/ParticleModuleParameterDynamic.h
--- a/Engine/Source/Runtime/Engine/Classes/Particles/Parameter/ParticleModuleParameterDynamic.h
+++ b/Engine/Source/Runtime/Engine/Classes/Particles/Parameter/ParticleModuleParameterDynamic.h
@@ -26,7 +26,7 @@ class UParticleModuleTypeDataBase;
  *	Enumeration indicating the way a dynamic parameter should be set.
  */
 UENUM()
-enum EEmitterDynamicParameterValue
+enum EEmitterDynamicParameterValue : int
 {
 	/** UserSet - use the user set values in the distribution (the default) */
 	EDPV_UserSet,
diff --git a/Engine/Source/Runtime/Engine/Classes/Particles/ParticleEmitter.h b/Engine/Source/Runtime/Engine/Classes/Particles/ParticleEmitter.h
--- a/Engine/Source/Runtime/Engine/Classes/Particles/ParticleEmitter.h
+++ b/Engine/Source/Runtime/Engine/Classes/Particles/ParticleEmitter.h
@@ -22,7 +22,7 @@ class UParticleSystemComponent;
 //	Burst emissions
 //~=============================================================================
 UENUM()
-enum EParticleBurstMethod
+enum EParticleBurstMethod : int
 {
 	EPBM_Instant UMETA(DisplayName="Instant"),
 	EPBM_Interpolated UMETA(DisplayName="Interpolated"),
@@ -33,7 +33,7 @@ enum EParticleBurstMethod
 //	SubUV-related
 //~=============================================================================
 UENUM()
-enum EParticleSubUVInterpMethod
+enum EParticleSubUVInterpMethod : int
 {
 	PSUVIM_None UMETA(DisplayName="None"),
 	PSUVIM_Linear UMETA(DisplayName="Linear"),
@@ -47,7 +47,7 @@ enum EParticleSubUVInterpMethod
 //	Cascade-related
 //~=============================================================================
 UENUM()
-enum EEmitterRenderMode
+enum EEmitterRenderMode : int
 {
 	ERM_Normal UMETA(DisplayName="Normal"),
 	ERM_Point UMETA(DisplayName="Point"),
diff --git a/Engine/Source/Runtime/Engine/Classes/Particles/ParticleModule.h b/Engine/Source/Runtime/Engine/Classes/Particles/ParticleModule.h
--- a/Engine/Source/Runtime/Engine/Classes/Particles/ParticleModule.h
+++ b/Engine/Source/Runtime/Engine/Classes/Particles/ParticleModule.h
@@ -27,7 +27,7 @@ struct FParticleEmitterInstance;
  *	The TypeData field is present to speed up finding the TypeData module.
  */
 UENUM()
-enum EModuleType
+enum EModuleType : int
 {
 	/** General - all emitter types can use it			*/
 	EPMT_General UMETA(DisplayName="General"),
@@ -55,7 +55,7 @@ enum EModuleType
  *	as the source points.
  */
 UENUM()
-enum EParticleSourceSelectionMethod
+enum EParticleSourceSelectionMethod : int
 {
 	/** Random		- select a particle at random		*/
 	EPSSM_Random UMETA(DisplayName="Random"),
diff --git a/Engine/Source/Runtime/Engine/Classes/Particles/ParticleModuleRequired.h b/Engine/Source/Runtime/Engine/Classes/Particles/ParticleModuleRequired.h
--- a/Engine/Source/Runtime/Engine/Classes/Particles/ParticleModuleRequired.h
+++ b/Engine/Source/Runtime/Engine/Classes/Particles/ParticleModuleRequired.h
@@ -82,7 +82,7 @@ FORCEINLINE void AdjustParticleBaseSizeForUVFlipping(FVector& OutSize, EParticle
 }
 
 UENUM()
-enum EParticleSortMode
+enum EParticleSortMode : int
 {
 	PSORTMODE_None,
 	PSORTMODE_ViewProjDepth,
@@ -93,7 +93,7 @@ enum EParticleSortMode
 };
 
 UENUM()
-enum EEmitterNormalsMode
+enum EEmitterNormalsMode : int
 {
 	/** Default mode, normals are based on the camera facing geometry. */
 	ENM_CameraFacing,
diff --git a/Engine/Source/Runtime/Engine/Classes/Particles/ParticleSpriteEmitter.h b/Engine/Source/Runtime/Engine/Classes/Particles/ParticleSpriteEmitter.h
--- a/Engine/Source/Runtime/Engine/Classes/Particles/ParticleSpriteEmitter.h
+++ b/Engine/Source/Runtime/Engine/Classes/Particles/ParticleSpriteEmitter.h
@@ -11,7 +11,7 @@
 class UParticleSystemComponent;
 
 UENUM()
-enum EParticleScreenAlignment
+enum EParticleScreenAlignment : int
 {
 	PSA_FacingCameraPosition,
 	PSA_Square,
diff --git a/Engine/Source/Runtime/Engine/Classes/Particles/ParticleSystem.h b/Engine/Source/Runtime/Engine/Classes/Particles/ParticleSystem.h
--- a/Engine/Source/Runtime/Engine/Classes/Particles/ParticleSystem.h
+++ b/Engine/Source/Runtime/Engine/Classes/Particles/ParticleSystem.h
@@ -23,7 +23,7 @@ class UParticleEmitter;
  *	Enumeration indicating the method by which the system should be updated
  */
 UENUM()
-enum EParticleSystemUpdateMode
+enum EParticleSystemUpdateMode : int
 {
 	/** RealTime	- update via the delta time passed in				*/
 	EPSUM_RealTime UMETA(DisplayName="Real-Time"),
@@ -35,7 +35,7 @@ enum EParticleSystemUpdateMode
  *	ParticleSystemLODMethod
  */
 UENUM()
-enum ParticleSystemLODMethod
+enum ParticleSystemLODMethod : int
 {
 	// Automatically set the LOD level, checking every LODDistanceCheckTime seconds.
 	PARTICLESYSTEMLODMETHOD_Automatic UMETA(DisplayName="Automatic"),
@@ -49,7 +49,7 @@ enum ParticleSystemLODMethod
 
 /** Occlusion method enumeration */
 UENUM()
-enum EParticleSystemOcclusionBoundsMethod
+enum EParticleSystemOcclusionBoundsMethod : int
 {
 	/** Don't determine occlusion on this particle system */
 	EPSOBM_None UMETA(DisplayName="None"),
diff --git a/Engine/Source/Runtime/Engine/Classes/Particles/ParticleSystemComponent.h b/Engine/Source/Runtime/Engine/Classes/Particles/ParticleSystemComponent.h
--- a/Engine/Source/Runtime/Engine/Classes/Particles/ParticleSystemComponent.h
+++ b/Engine/Source/Runtime/Engine/Classes/Particles/ParticleSystemComponent.h
@@ -42,7 +42,7 @@ struct FParticleAnimTrailEmitterInstance;
 
 /** Enum for specifying type of a name instance parameter. */
 UENUM(BlueprintType)
-enum EParticleSysParamType
+enum EParticleSysParamType : int
 {
 	PSPT_None UMETA(DisplayName="None"),
 	PSPT_Scalar UMETA(DisplayName="Scalar"),
@@ -58,7 +58,7 @@ enum EParticleSysParamType
 
 /** Particle system replay state */
 UENUM()
-enum ParticleReplayState
+enum ParticleReplayState : int
 {
 	/** Replay system is disabled.  Particles are simulated and rendered normally. */
 	PRS_Disabled UMETA(DisplayName="Disabled"),
@@ -75,7 +75,7 @@ enum ParticleReplayState
  *	Event type
  */
 UENUM()
-enum EParticleEventType
+enum EParticleEventType : int
 {
 	/** Any - allow any event */
 	EPET_Any UMETA(DisplayName="Any"),
diff --git a/Engine/Source/Runtime/Engine/Classes/Particles/SubUVAnimation.h b/Engine/Source/Runtime/Engine/Classes/Particles/SubUVAnimation.h
--- a/Engine/Source/Runtime/Engine/Classes/Particles/SubUVAnimation.h
+++ b/Engine/Source/Runtime/Engine/Classes/Particles/SubUVAnimation.h
@@ -36,14 +36,14 @@ public:
  * and when the particles using the texture will be few and large.
  */
 UENUM()
-enum ESubUVBoundingVertexCount
+enum ESubUVBoundingVertexCount : int
 {
 	BVC_FourVertices,
 	BVC_EightVertices
 };
 
 UENUM()
-enum EOpacitySourceMode
+enum EOpacitySourceMode : int
 {
 	OSM_Alpha,
 	OSM_ColorBrightness,
diff --git a/Engine/Source/Runtime/Engine/Classes/Particles/Trail/ParticleModuleTrailSource.h b/Engine/Source/Runtime/Engine/Classes/Particles/Trail/ParticleModuleTrailSource.h
--- a/Engine/Source/Runtime/Engine/Classes/Particles/Trail/ParticleModuleTrailSource.h
+++ b/Engine/Source/Runtime/Engine/Classes/Particles/Trail/ParticleModuleTrailSource.h
@@ -19,7 +19,7 @@ class UParticleSystemComponent;
 struct FParticleEmitterInstance;
 
 UENUM()
-enum ETrail2SourceMethod
+enum ETrail2SourceMethod : int
 {
 	/** Default	- use the emitter position. 
 	 *	This is the fallback for when other modes can't be resolved.
diff --git a/Engine/Source/Runtime/Engine/Classes/Particles/TypeData/ParticleModuleTypeDataBeam2.h b/Engine/Source/Runtime/Engine/Classes/Particles/TypeData/ParticleModuleTypeDataBeam2.h
--- a/Engine/Source/Runtime/Engine/Classes/Particles/TypeData/ParticleModuleTypeDataBeam2.h
+++ b/Engine/Source/Runtime/Engine/Classes/Particles/TypeData/ParticleModuleTypeDataBeam2.h
@@ -20,7 +20,7 @@ struct FCurveEdEntry;
 struct FParticleEmitterInstance;
 
 UENUM()
-enum EBeam2Method
+enum EBeam2Method : int
 {
 	PEB2M_Distance UMETA(DisplayName="Distance"),
 	PEB2M_Target UMETA(DisplayName="Target"),
@@ -32,7 +32,7 @@ enum EBeam2Method
 // Beam Tapering Variables.
 //
 UENUM()
-enum EBeamTaperMethod
+enum EBeamTaperMethod : int
 {
 	PEBTM_None UMETA(DisplayName="None"),
 	PEBTM_Full UMETA(DisplayName="Full"),
diff --git a/Engine/Source/Runtime/Engine/Classes/Particles/TypeData/ParticleModuleTypeDataMesh.h b/Engine/Source/Runtime/Engine/Classes/Particles/TypeData/ParticleModuleTypeDataMesh.h
--- a/Engine/Source/Runtime/Engine/Classes/Particles/TypeData/ParticleModuleTypeDataMesh.h
+++ b/Engine/Source/Runtime/Engine/Classes/Particles/TypeData/ParticleModuleTypeDataMesh.h
@@ -17,7 +17,7 @@ struct FParticleEmitterInstance;
 class UStaticMesh;
 
 UENUM()
-enum EMeshScreenAlignment
+enum EMeshScreenAlignment : int
 {
 	PSMA_MeshFaceCameraWithRoll UMETA(DisplayName="Face Camera With Roll"),
 	PSMA_MeshFaceCameraWithSpin UMETA(DisplayName="Face Camera With Spin"),
@@ -26,7 +26,7 @@ enum EMeshScreenAlignment
 };
 
 UENUM()
-enum EMeshCameraFacingUpAxis
+enum EMeshCameraFacingUpAxis : int
 {
 	CameraFacing_NoneUP UMETA(DisplayName="None"),
 	CameraFacing_ZUp UMETA(DisplayName="Z Up"),
@@ -37,7 +37,7 @@ enum EMeshCameraFacingUpAxis
 };
 
 UENUM()
-enum EMeshCameraFacingOptions
+enum EMeshCameraFacingOptions : int
 {
 	XAxisFacing_NoUp UMETA(DisplayName="X Axis Facing : No Up"),
 	XAxisFacing_ZUp UMETA(DisplayName="X Axis Facing : Z Up"),
diff --git a/Engine/Source/Runtime/Engine/Classes/Particles/TypeData/ParticleModuleTypeDataRibbon.h b/Engine/Source/Runtime/Engine/Classes/Particles/TypeData/ParticleModuleTypeDataRibbon.h
--- a/Engine/Source/Runtime/Engine/Classes/Particles/TypeData/ParticleModuleTypeDataRibbon.h
+++ b/Engine/Source/Runtime/Engine/Classes/Particles/TypeData/ParticleModuleTypeDataRibbon.h
@@ -18,7 +18,7 @@ class UParticleSystemComponent;
 struct FParticleEmitterInstance;
 
 UENUM()
-enum ETrailsRenderAxisOption
+enum ETrailsRenderAxisOption : int
 {
 	Trails_CameraUp UMETA(DisplayName="Camera Up"),
 	Trails_SourceUp UMETA(DisplayName="Source Up"),
diff --git a/Engine/Source/Runtime/Engine/Classes/Sound/ReverbSettings.h b/Engine/Source/Runtime/Engine/Classes/Sound/ReverbSettings.h
--- a/Engine/Source/Runtime/Engine/Classes/Sound/ReverbSettings.h
+++ b/Engine/Source/Runtime/Engine/Classes/Sound/ReverbSettings.h
@@ -15,7 +15,7 @@ class USoundEffectSubmixPreset;
  *
  */
 UENUM()
-enum ReverbPreset
+enum ReverbPreset : int
 {
 	REVERB_Default,
 	REVERB_Bathroom,
diff --git a/Engine/Source/Runtime/Engine/Classes/Sound/SoundAttenuation.h b/Engine/Source/Runtime/Engine/Classes/Sound/SoundAttenuation.h
--- a/Engine/Source/Runtime/Engine/Classes/Sound/SoundAttenuation.h
+++ b/Engine/Source/Runtime/Engine/Classes/Sound/SoundAttenuation.h
@@ -19,7 +19,7 @@ class USpatializationPluginSourceSettingsBase;
 
 // This enumeration is deprecated
 UENUM()
-enum ESoundDistanceCalc
+enum ESoundDistanceCalc : int
 {
 	SOUNDDISTANCE_Normal,
 	SOUNDDISTANCE_InfiniteXYPlane,
@@ -29,7 +29,7 @@ enum ESoundDistanceCalc
 };
 
 UENUM()
-enum ESoundSpatializationAlgorithm
+enum ESoundSpatializationAlgorithm : int
 {
 	// Standard panning method for spatialization (linear or equal power method defined in project settings)
 	SPATIALIZATION_Default UMETA(DisplayName = "Panning"),
diff --git a/Engine/Source/Runtime/Engine/Classes/Sound/SoundGroups.h b/Engine/Source/Runtime/Engine/Classes/Sound/SoundGroups.h
--- a/Engine/Source/Runtime/Engine/Classes/Sound/SoundGroups.h
+++ b/Engine/Source/Runtime/Engine/Classes/Sound/SoundGroups.h
@@ -7,7 +7,7 @@
 #include "SoundGroups.generated.h"
 
 UENUM()
-enum ESoundGroup
+enum ESoundGroup : int
 {
 	SOUNDGROUP_Default UMETA(DisplayName="Default"),
 	SOUNDGROUP_Effects UMETA(DisplayName="Effects"),
diff --git a/Engine/Source/Runtime/Engine/Classes/Sound/SoundNodeModulatorContinuous.h b/Engine/Source/Runtime/Engine/Classes/Sound/SoundNodeModulatorContinuous.h
--- a/Engine/Source/Runtime/Engine/Classes/Sound/SoundNodeModulatorContinuous.h
+++ b/Engine/Source/Runtime/Engine/Classes/Sound/SoundNodeModulatorContinuous.h
@@ -14,7 +14,7 @@ struct FSoundParseParameters;
 struct FWaveInstance;
 
 UENUM()
-enum ModulationParamMode
+enum ModulationParamMode : int
 {
 	MPM_Normal UMETA(DisplayName = "Normal" , Tooltip = "Clamps input value to the range (MinInput, MaxInput) then maps to the range (MinOutput, MaxOutput)"),
 	MPM_Abs UMETA(DisplayName = "Absolute" , Tooltip = "Same as Normal except that the input value is treated as an absolute value"),
diff --git a/Engine/Source/Runtime/Engine/Classes/Sound/SoundWave.h b/Engine/Source/Runtime/Engine/Classes/Sound/SoundWave.h
--- a/Engine/Source/Runtime/Engine/Classes/Sound/SoundWave.h
+++ b/Engine/Source/Runtime/Engine/Classes/Sound/SoundWave.h
@@ -39,7 +39,7 @@ struct FSoundParseParameters;
 struct FPlatformAudioCookOverrides;
 
 UENUM()
-enum EDecompressionType
+enum EDecompressionType : int
 {
 	DTYPE_Setup,
 	DTYPE_Invalid,
diff --git a/Engine/Source/Runtime/Engine/Classes/VectorField/VectorFieldAnimated.h b/Engine/Source/Runtime/Engine/Classes/VectorField/VectorFieldAnimated.h
--- a/Engine/Source/Runtime/Engine/Classes/VectorField/VectorFieldAnimated.h
+++ b/Engine/Source/Runtime/Engine/Classes/VectorField/VectorFieldAnimated.h
@@ -17,7 +17,7 @@ struct FPropertyChangedEvent;
  * Operation used to construct the vector field from a 2D texture.
  */
 UENUM()
-enum EVectorFieldConstructionOp
+enum EVectorFieldConstructionOp : int
 {
 	VFCO_Extrude UMETA(DisplayName="Extrude"),
 	VFCO_Revolve UMETA(DisplayName="Revolve"),
diff --git a/Engine/Source/Runtime/Engine/Public/Animation/AnimCompressionTypes.h b/Engine/Source/Runtime/Engine/Public/Animation/AnimCompressionTypes.h
--- a/Engine/Source/Runtime/Engine/Public/Animation/AnimCompressionTypes.h
+++ b/Engine/Source/Runtime/Engine/Public/Animation/AnimCompressionTypes.h
@@ -31,7 +31,7 @@
  * Indicates animation data key format.
  */
 UENUM()
-enum AnimationKeyFormat
+enum AnimationKeyFormat : int
 {
 	AKF_ConstantKeyLerp,
 	AKF_VariableKeyLerp,
diff --git a/Engine/Source/Runtime/Engine/Public/Animation/AnimCurveTypes.h b/Engine/Source/Runtime/Engine/Public/Animation/AnimCurveTypes.h
--- a/Engine/Source/Runtime/Engine/Public/Animation/AnimCurveTypes.h
+++ b/Engine/Source/Runtime/Engine/Public/Animation/AnimCurveTypes.h
@@ -31,7 +31,7 @@ ENUM_RANGE_BY_COUNT(EAnimCurveType, EAnimCurveType::MaxAnimCurveType);
 
 /** This is curve flags that are saved in asset and **/
 UENUM(BlueprintType, meta=(Bitflags))
-enum EAnimAssetCurveFlags
+enum EAnimAssetCurveFlags : int
 {
 	AACF_NONE = 0 UMETA(Hidden),
 	// Used as morph target curve
diff --git a/Engine/Source/Runtime/Engine/Public/Animation/AnimTypes.h b/Engine/Source/Runtime/Engine/Public/Animation/AnimTypes.h
--- a/Engine/Source/Runtime/Engine/Public/Animation/AnimTypes.h
+++ b/Engine/Source/Runtime/Engine/Public/Animation/AnimTypes.h
@@ -46,7 +46,7 @@ ENGINE_API float GetTriggerTimeOffsetForType(EAnimEventTriggerOffsets::Type Offs
 
 /** Enum for specifying a specific axis of a bone */
 UENUM()
-enum EBoneAxis
+enum EBoneAxis : int
 {
 	BA_X UMETA(DisplayName = "X Axis"),
 	BA_Y UMETA(DisplayName = "Y Axis"),
@@ -56,7 +56,7 @@ enum EBoneAxis
 
 /** Enum for controlling which reference frame a controller is applied in. */
 UENUM(BlueprintType)
-enum EBoneControlSpace
+enum EBoneControlSpace : int
 {
 	/** Set absolute position of bone in world space. */
 	BCS_WorldSpace UMETA(DisplayName = "World Space"),
@@ -71,7 +71,7 @@ enum EBoneControlSpace
 
 /** Enum for specifying the source of a bone's rotation. */
 UENUM()
-enum EBoneRotationSource
+enum EBoneRotationSource : int
 {
 	/** Don't change rotation at all. */
 	BRS_KeepComponentSpaceRotation UMETA(DisplayName = "No Change (Preserve Existing Component Space Rotation)"),
@@ -85,7 +85,7 @@ enum EBoneRotationSource
 UENUM()
 namespace EMontageNotifyTickType
 {
-	enum Type
+	enum Type : int
 	{
 		/** Queue notifies, and trigger them at the end of the evaluation phase (faster). Not suitable for changing sections or montage position. */
 		Queued,
@@ -98,7 +98,7 @@ namespace EMontageNotifyTickType
 UENUM()
 namespace ENotifyFilterType
 {
-	enum Type
+	enum Type : int
 	{
 		/** No filtering. */
 		NoFiltering,
@@ -577,7 +577,7 @@ struct FAnimNotifyTrack
  * Indicates whether an animation is additive, and what kind.
  */
 UENUM()
-enum EAdditiveAnimationType
+enum EAdditiveAnimationType : int
 {
 	/** No additive. */
 	AAT_None  UMETA(DisplayName="No additive"),
@@ -591,7 +591,7 @@ enum EAdditiveAnimationType
 UENUM()
 namespace ECurveBlendOption
 {
-	enum Type
+	enum Type : int
 	{
 		/* Last pose that contains valid curve value override it. */
 		Override, // redirect from MaxWeight old legacy behavior
@@ -759,7 +759,7 @@ private:
 UENUM()
 namespace EAxisOption
 {
-	enum Type
+	enum Type : int
 	{
 		X,
 		Y,
@@ -806,7 +806,7 @@ struct FAxisOption
 UENUM()
 namespace EComponentType
 {
-	enum Type
+	enum Type : int
 	{
 		None = 0,
 		TranslationX,
diff --git a/Engine/Source/Runtime/Engine/Public/Animation/AttributeTypes.h b/Engine/Source/Runtime/Engine/Public/Animation/AttributeTypes.h
--- a/Engine/Source/Runtime/Engine/Public/Animation/AttributeTypes.h
+++ b/Engine/Source/Runtime/Engine/Public/Animation/AttributeTypes.h
@@ -13,7 +13,7 @@
 #include "CoreTypes.h"
 #include "Templates/UnrealTemplate.h"
 
-enum EAdditiveAnimationType;
+enum EAdditiveAnimationType : int;
 
 class UAnimationSettings;
 
diff --git a/Engine/Source/Runtime/Engine/Public/Animation/AttributesRuntime.h b/Engine/Source/Runtime/Engine/Public/Animation/AttributesRuntime.h
--- a/Engine/Source/Runtime/Engine/Public/Animation/AttributesRuntime.h
+++ b/Engine/Source/Runtime/Engine/Public/Animation/AttributesRuntime.h
@@ -25,7 +25,7 @@ struct FBlendSampleData;
 struct FBoneContainer;
 
 enum class ECustomAttributeBlendType : uint8;
-enum EAdditiveAnimationType;
+enum EAdditiveAnimationType : int;
 
 struct FAnimatedBoneAttribute;
 struct FCompactPoseBoneIndex;
diff --git a/Engine/Source/Runtime/Engine/Public/HitProxies.h b/Engine/Source/Runtime/Engine/Public/HitProxies.h
--- a/Engine/Source/Runtime/Engine/Public/HitProxies.h
+++ b/Engine/Source/Runtime/Engine/Public/HitProxies.h
@@ -21,7 +21,7 @@ class FReferenceCollector;
  * HPP_UI - the priority of the UI components such as the translation widget
  */
 UENUM()
-enum EHitProxyPriority
+enum EHitProxyPriority : int
 {
 	HPP_World = 0,
 	HPP_Wireframe = 1,
diff --git a/Engine/Source/Runtime/Engine/Public/MaterialSceneTextureId.h b/Engine/Source/Runtime/Engine/Public/MaterialSceneTextureId.h
--- a/Engine/Source/Runtime/Engine/Public/MaterialSceneTextureId.h
+++ b/Engine/Source/Runtime/Engine/Public/MaterialSceneTextureId.h
@@ -13,7 +13,7 @@
 
 /** like EPassInputId but can expose more e.g. GBuffer */
 UENUM()
-enum ESceneTextureId
+enum ESceneTextureId : int
 {
 	/** Scene color, normal post process passes should use PostProcessInput0 */
 	PPI_SceneColor UMETA(DisplayName="SceneColor"),
diff --git a/Engine/Source/Runtime/Engine/Public/MaterialShared.h b/Engine/Source/Runtime/Engine/Public/MaterialShared.h
--- a/Engine/Source/Runtime/Engine/Public/MaterialShared.h
+++ b/Engine/Source/Runtime/Engine/Public/MaterialShared.h
@@ -244,7 +244,7 @@ enum struct EDeferredParamStrictness
 
 /** Defines the domain of a material. */
 UENUM()
-enum EMaterialDomain
+enum EMaterialDomain : int
 {
 	/** The material's attributes describe a 3d surface. */
 	MD_Surface UMETA(DisplayName = "Surface"),
diff --git a/Engine/Source/Runtime/Engine/Public/MaterialTypes.h b/Engine/Source/Runtime/Engine/Public/MaterialTypes.h
--- a/Engine/Source/Runtime/Engine/Public/MaterialTypes.h
+++ b/Engine/Source/Runtime/Engine/Public/MaterialTypes.h
@@ -18,7 +18,7 @@ class URuntimeVirtualTexture;
 class USparseVolumeTexture;
 
 UENUM()
-enum EMaterialParameterAssociation
+enum EMaterialParameterAssociation : int
 {
 	LayerParameter,
 	BlendParameter,
diff --git a/Engine/Source/Runtime/Engine/Public/ParticleHelper.h b/Engine/Source/Runtime/Engine/Public/ParticleHelper.h
--- a/Engine/Source/Runtime/Engine/Public/ParticleHelper.h
+++ b/Engine/Source/Runtime/Engine/Public/ParticleHelper.h
@@ -38,7 +38,7 @@ DECLARE_LOG_CATEGORY_EXTERN(LogParticles, Log, All);
 
 /** Detail mode for scene component rendering. */
 UENUM()
-enum EParticleDetailMode 
+enum EParticleDetailMode : int
 {
 	PDM_Low UMETA(DisplayName = "Low"),
 	PDM_Medium UMETA(DisplayName = "Medium"),
diff --git a/Engine/Source/Runtime/Engine/Public/SceneTypes.h b/Engine/Source/Runtime/Engine/Public/SceneTypes.h
--- a/Engine/Source/Runtime/Engine/Public/SceneTypes.h
+++ b/Engine/Source/Runtime/Engine/Public/SceneTypes.h
@@ -200,7 +200,7 @@ ENGINE_API FString LexToString(EMaterialQualityLevel::Type QualityLevel);
 //	EMaterialProperty
 //
 UENUM(BlueprintType)
-enum EMaterialProperty
+enum EMaterialProperty : int
 {
 	MP_EmissiveColor = 0 UMETA(DisplayName = "Emissive"),
 	MP_Opacity UMETA(DisplayName = "Opacity"),
diff --git a/Engine/Source/Runtime/Engine/Public/SceneUtils.h b/Engine/Source/Runtime/Engine/Public/SceneUtils.h
--- a/Engine/Source/Runtime/Engine/Public/SceneUtils.h
+++ b/Engine/Source/Runtime/Engine/Public/SceneUtils.h
@@ -38,7 +38,7 @@ enum class EMobileHDRMode
 
 /** Used by rendering project settings. */
 UENUM()
-enum EAntiAliasingMethod
+enum EAntiAliasingMethod : int
 {
 	AAM_None UMETA(DisplayName = "None"),
 	AAM_FXAA UMETA(DisplayName = "Fast Approximate Anti-Aliasing (FXAA)"),
diff --git a/Engine/Source/Runtime/Engine/Public/SkeletalMeshReductionSettings.h b/Engine/Source/Runtime/Engine/Public/SkeletalMeshReductionSettings.h
--- a/Engine/Source/Runtime/Engine/Public/SkeletalMeshReductionSettings.h
+++ b/Engine/Source/Runtime/Engine/Public/SkeletalMeshReductionSettings.h
@@ -15,7 +15,7 @@ class FSkeletalMeshLODModel;
 
 /** Enum specifying the reduction type to use when simplifying skeletal meshes with internal tool */
 UENUM()
-enum SkeletalMeshTerminationCriterion
+enum SkeletalMeshTerminationCriterion : int
 {
 	SMTC_NumOfTriangles UMETA(DisplayName = "Triangles", ToolTip = "Triangle count criterion will be used for simplification."),
 	SMTC_NumOfVerts UMETA(DisplayName = "Vertices", ToolTip = "Vertex cont criterion will be used for simplification."),
@@ -28,7 +28,7 @@ enum SkeletalMeshTerminationCriterion
 
 /** Enum specifying the reduction type to use when simplifying skeletal meshes with Simmplygon */
 UENUM()
-enum SkeletalMeshOptimizationType
+enum SkeletalMeshOptimizationType : int
 {
 	SMOT_NumOfTriangles UMETA(DisplayName = "Triangles", ToolTip = "Triangle requirement will be used for simplification."),
 	SMOT_MaxDeviation UMETA(DisplayName = "Accuracy", ToolTip = "Accuracy requirement will be used for simplification."),
@@ -38,7 +38,7 @@ enum SkeletalMeshOptimizationType
 
 /** Enum specifying the importance of properties when simplifying skeletal meshes. */
 UENUM()
-enum SkeletalMeshOptimizationImportance
+enum SkeletalMeshOptimizationImportance : int
 {
 	SMOI_Off UMETA(DisplayName = "Off"),
 	SMOI_Lowest UMETA(DisplayName = "Lowest"),
diff --git a/Engine/Source/Runtime/Experimental/Chaos/Public/Field/FieldSystemTypes.h b/Engine/Source/Runtime/Experimental/Chaos/Public/Field/FieldSystemTypes.h
--- a/Engine/Source/Runtime/Experimental/Chaos/Public/Field/FieldSystemTypes.h
+++ b/Engine/Source/Runtime/Experimental/Chaos/Public/Field/FieldSystemTypes.h
@@ -6,7 +6,7 @@
 *
 */
 UENUM(BlueprintType)
-enum ESetMaskConditionType
+enum ESetMaskConditionType : int
 {
 	Field_Set_Always	        UMETA(DisplayName = "Set Always", ToolTip = "The particle output value will be equal to Interior-value if the particle position is inside a sphere / Exterior-value otherwise."),
 	Field_Set_IFF_NOT_Interior  UMETA(DisplayName = "Merge Interior", ToolTip = "The particle output value will be equal to Interior-value if the particle position is inside the sphere or if the particle input value is already Interior-Value / Exterior-value otherwise."),
@@ -20,7 +20,7 @@ enum ESetMaskConditionType
 *
 */
 UENUM(BlueprintType)
-enum EWaveFunctionType
+enum EWaveFunctionType : int
 {
 	Field_Wave_Cosine	 UMETA(DisplayName = "Cosine", ToolTip = "Cosine wave that will move in time."),
 	Field_Wave_Gaussian  UMETA(DisplayName = "Gaussian", ToolTip = "Gaussian wave that will move in time."),
@@ -35,7 +35,7 @@ enum EWaveFunctionType
 *
 */
 UENUM(BlueprintType)
-enum EFieldOperationType
+enum EFieldOperationType : int
 {
 	Field_Multiply  UMETA(DisplayName = "Multiply", ToolTip = "Multiply the fields output values : Output = Left * Right"),
 	Field_Divide    UMETA(DisplayName = "Divide", ToolTip = "Divide the fields output values : Output = Left / Right"),
@@ -50,7 +50,7 @@ enum EFieldOperationType
 *
 */
 UENUM(BlueprintType)
-enum EFieldCullingOperationType
+enum EFieldCullingOperationType : int
 {
 	Field_Culling_Inside  UMETA(DisplayName = "Inside", ToolTip = "Evaluate the input field if the result of the culling field is equal to 0"),
 	Field_Culling_Outside UMETA(DisplayName = "Outside", ToolTip = "Evaluate the input field if the result of the culling field is different from 0"),
@@ -64,7 +64,7 @@ enum EFieldCullingOperationType
 *
 */
 UENUM(BlueprintType)
-enum EFieldResolutionType
+enum EFieldResolutionType : int
 {
 	Field_Resolution_Minimal  UMETA(DisplayName = "Minimum", ToolTip = "Apply the field to all the active particles"),
 	Field_Resolution_DisabledParents  UMETA(DisplayName = "Parents", ToolTip = "Apply the field to all the parent particles"),
@@ -78,7 +78,7 @@ enum EFieldResolutionType
 *
 */
 UENUM(BlueprintType)
-enum EFieldFilterType
+enum EFieldFilterType : int
 {
 	Field_Filter_Dynamic  UMETA(DisplayName = "Dynamic", ToolTip = "Apply the field to all the dynamic particles"),
 	Field_Filter_Kinematic  UMETA(DisplayName = "Kinematic", ToolTip = "Apply the field to all the kinematic particles"),
@@ -95,7 +95,7 @@ enum EFieldFilterType
 *
 */
 UENUM(BlueprintType)
-enum EFieldObjectType
+enum EFieldObjectType : int
 {
 	Field_Object_Rigid  UMETA(DisplayName = "Rigid", ToolTip = "Apply the field to all the rigid particles"),
 	Field_Object_Cloth  UMETA(DisplayName = "Cloth", ToolTip = "Apply the field to all the cloth particles"),
@@ -111,7 +111,7 @@ enum EFieldObjectType
 *
 */
 UENUM(BlueprintType)
-enum EFieldPositionType
+enum EFieldPositionType : int
 {
 	Field_Position_CenterOfMass  UMETA(DisplayName = "CenterOfMass", ToolTip = "Apply the field to the particles center of mass position"),
 	Field_Position_PivotPoint  UMETA(DisplayName = "PivotPoint", ToolTip = "Apply the field to the particles pivot point position"),
@@ -124,7 +124,7 @@ enum EFieldPositionType
 *
 */
 UENUM(BlueprintType)
-enum EFieldFalloffType
+enum EFieldFalloffType : int
 {
 	Field_FallOff_None			UMETA(DisplayName = "None", ToolTip = "No falloff function is used"),
 	Field_Falloff_Linear		UMETA(DisplayName = "Linear", ToolTip = "The falloff function will be proportional to x"),
@@ -140,7 +140,7 @@ enum EFieldFalloffType
 *
 */
 UENUM(BlueprintType)
-enum EFieldPhysicsType
+enum EFieldPhysicsType : int
 {
 	Field_None						UMETA(Hidden),
 	Field_DynamicState				UMETA(DisplayName = "Dynamic State", ToolTip = "Set the dynamic state of a particle (static, dynamic, kinematic...)"),
@@ -169,7 +169,7 @@ enum EFieldPhysicsType
 
 // TODO : Refactor these 3 enums to be in sync with the GetFieldTargetTypes
 UENUM(BlueprintType)
-enum EFieldVectorType
+enum EFieldVectorType : int
 {
 	Vector_LinearForce				UMETA(DisplayName = "Linear Force", ToolTip = "Add a vector field to the particles linear force."),
 	Vector_LinearVelocity			UMETA(DisplayName = "Linear Velocity", ToolTip = "Add a vector field to the particles linear velocity."),
@@ -186,7 +186,7 @@ enum EFieldVectorType
 };
 
 UENUM(BlueprintType)
-enum EFieldScalarType
+enum EFieldScalarType : int
 {
 	Scalar_ExternalClusterStrain		UMETA(DisplayName = "External Strain", ToolTip = "Apply an external strain over the particles. If this strain is over the internal one, the cluster will break."),
 	Scalar_Kill   						UMETA(DisplayName = "Kill Particle", ToolTip = "Disable the particles for which the field will be higher than 0."),
@@ -200,7 +200,7 @@ enum EFieldScalarType
 };
 
 UENUM(BlueprintType)
-enum EFieldIntegerType
+enum EFieldIntegerType : int
 {
 	Integer_DynamicState				UMETA(DisplayName = "Dynamic State", ToolTip = "Set the dynamic state of a particle (static, dynamic, kinematic...)"),
 	Integer_ActivateDisabled			UMETA(DisplayName = "Activate Disabled", ToolTip = "Activate all the disabled particles for which the field value will be 0"),
@@ -214,7 +214,7 @@ enum EFieldIntegerType
 
 /** Defines the type of the output*/
 UENUM()
-enum EFieldOutputType
+enum EFieldOutputType : int
 {
 	/* Vector Field Type */
 	Field_Output_Vector UMETA(DisplayName = "Vector Field"),
@@ -387,7 +387,7 @@ inline CHAOS_API EFieldPhysicsType GetFieldPhysicsType(const FName& PhysicsName)
 *
 */
 UENUM(BlueprintType)
-enum EFieldPhysicsDefaultFields
+enum EFieldPhysicsDefaultFields : int
 {
 	Field_RadialIntMask				UMETA(DisplayName = "RadialIntMask"),
 	Field_RadialFalloff				UMETA(DisplayName = "RadialFalloff"),
diff --git a/Engine/Source/Runtime/Foliage/Public/FoliageType.h b/Engine/Source/Runtime/Foliage/Public/FoliageType.h
--- a/Engine/Source/Runtime/Foliage/Public/FoliageType.h
+++ b/Engine/Source/Runtime/Foliage/Public/FoliageType.h
@@ -15,7 +15,7 @@
 class UStaticMesh;
 
 UENUM()
-enum FoliageVertexColorMask
+enum FoliageVertexColorMask : int
 {
 	FOLIAGEVERTEXCOLORMASK_Disabled UMETA(DisplayName="Disabled"),
 	FOLIAGEVERTEXCOLORMASK_Red		UMETA(DisplayName="Red"),
diff --git a/Engine/Source/Runtime/HeadMountedDisplay/Public/HeadMountedDisplayTypes.h b/Engine/Source/Runtime/HeadMountedDisplay/Public/HeadMountedDisplayTypes.h
--- a/Engine/Source/Runtime/HeadMountedDisplay/Public/HeadMountedDisplayTypes.h
+++ b/Engine/Source/Runtime/HeadMountedDisplay/Public/HeadMountedDisplayTypes.h
@@ -64,7 +64,7 @@ HEADMOUNTEDDISPLAY_API DECLARE_LOG_CATEGORY_EXTERN(LogLoadingSplash, Log, All);
 UENUM()
 namespace EOrientPositionSelector
 {
-	enum Type
+	enum Type : int
 	{
 		Orientation UMETA(DisplayName = "Orientation"),
 		Position UMETA(DisplayName = "Position"),
@@ -78,7 +78,7 @@ namespace EOrientPositionSelector
 UENUM()
 namespace EHMDTrackingOrigin
 {
-	enum Type
+	enum Type : int
 	{
 		Floor UMETA(DisplayName = "Floor Level"),
 		Eye UMETA(DisplayName = "Eye Level"),
@@ -92,7 +92,7 @@ namespace EHMDTrackingOrigin
 UENUM()
 namespace EHMDWornState
 {
-	enum Type
+	enum Type : int
 	{
 		Unknown UMETA(DisplayName = "Unknown"),
 		Worn UMETA(DisplayName = "Worn"),
@@ -107,7 +107,7 @@ namespace EHMDWornState
 UENUM(BlueprintType)
 namespace EXRDeviceConnectionResult
 {
-	enum Type
+	enum Type : int
 	{
 		NoTrackingSystem,
 		FeatureNotSupported,
@@ -123,7 +123,7 @@ namespace EXRDeviceConnectionResult
 UENUM(BlueprintType)
 namespace EXRSystemFlags
 {
-	enum Type
+	enum Type : int
 	{
 		NoFlags       = 0x00 UMETA(Hidden),
 		IsAR          = 0x01,
diff --git a/Engine/Source/Runtime/Landscape/Classes/Landscape.h b/Engine/Source/Runtime/Landscape/Classes/Landscape.h
--- a/Engine/Source/Runtime/Landscape/Classes/Landscape.h
+++ b/Engine/Source/Runtime/Landscape/Classes/Landscape.h
@@ -44,7 +44,7 @@ extern LANDSCAPE_API TAutoConsoleVariable<int32> CVarLandscapeSplineFalloffModul
 #endif
 
 UENUM()
-enum ELandscapeSetupErrors
+enum ELandscapeSetupErrors : int
 {
 	LSE_None,
 	/** No Landscape Info available. */
@@ -154,7 +154,7 @@ private:
 };
 
 UENUM()
-enum ELandscapeBlendMode
+enum ELandscapeBlendMode : int
 {
 	LSBM_AdditiveBlend,
 	LSBM_AlphaBlend,
diff --git a/Engine/Source/Runtime/Landscape/Classes/LandscapeComponent.h b/Engine/Source/Runtime/Landscape/Classes/LandscapeComponent.h
--- a/Engine/Source/Runtime/Landscape/Classes/LandscapeComponent.h
+++ b/Engine/Source/Runtime/Landscape/Classes/LandscapeComponent.h
@@ -367,7 +367,7 @@ static const uint32 DefaultSplineHash = 0xFFFFFFFF;
 #endif
 
 UENUM()
-enum ELandscapeClearMode
+enum ELandscapeClearMode : int
 {
 	Clear_Weightmap = 1 << 0 UMETA(DisplayName = "Paint"),
 	Clear_Heightmap = 1 << 1 UMETA(DisplayName = "Sculpt"),
diff --git a/Engine/Source/Runtime/Landscape/Classes/LandscapeGizmoActiveActor.h b/Engine/Source/Runtime/Landscape/Classes/LandscapeGizmoActiveActor.h
--- a/Engine/Source/Runtime/Landscape/Classes/LandscapeGizmoActiveActor.h
+++ b/Engine/Source/Runtime/Landscape/Classes/LandscapeGizmoActiveActor.h
@@ -16,7 +16,7 @@ class UMaterialInstance;
 class UTexture2D;
 
 UENUM()
-enum ELandscapeGizmoType
+enum ELandscapeGizmoType : int
 {
 	LGT_None,
 	LGT_Height,
diff --git a/Engine/Source/Runtime/Landscape/Classes/LandscapeSplineSegment.h b/Engine/Source/Runtime/Landscape/Classes/LandscapeSplineSegment.h
--- a/Engine/Source/Runtime/Landscape/Classes/LandscapeSplineSegment.h
+++ b/Engine/Source/Runtime/Landscape/Classes/LandscapeSplineSegment.h
@@ -119,7 +119,7 @@ struct FLandscapeSplineSegmentConnection
 
 // Deprecated
 UENUM()
-enum LandscapeSplineMeshOrientation
+enum LandscapeSplineMeshOrientation : int
 {
 	LSMO_XUp,
 	LSMO_YUp,
diff --git a/Engine/Source/Runtime/Landscape/Classes/Materials/MaterialExpressionLandscapeLayerBlend.h b/Engine/Source/Runtime/Landscape/Classes/Materials/MaterialExpressionLandscapeLayerBlend.h
--- a/Engine/Source/Runtime/Landscape/Classes/Materials/MaterialExpressionLandscapeLayerBlend.h
+++ b/Engine/Source/Runtime/Landscape/Classes/Materials/MaterialExpressionLandscapeLayerBlend.h
@@ -15,7 +15,7 @@ struct FPropertyChangedEvent;
 struct FMaterialParameterInfo;
 
 UENUM()
-enum ELandscapeLayerBlendType
+enum ELandscapeLayerBlendType : int
 {
 	LB_WeightBlend,
 	LB_AlphaBlend,
diff --git a/Engine/Source/Runtime/Landscape/Classes/Materials/MaterialExpressionLandscapeLayerCoords.h b/Engine/Source/Runtime/Landscape/Classes/Materials/MaterialExpressionLandscapeLayerCoords.h
--- a/Engine/Source/Runtime/Landscape/Classes/Materials/MaterialExpressionLandscapeLayerCoords.h
+++ b/Engine/Source/Runtime/Landscape/Classes/Materials/MaterialExpressionLandscapeLayerCoords.h
@@ -9,7 +9,7 @@
 #include "MaterialExpressionLandscapeLayerCoords.generated.h"
 
 UENUM()
-enum ETerrainCoordMappingType
+enum ETerrainCoordMappingType : int
 {
 	TCMT_Auto,
 	TCMT_XY,
@@ -19,7 +19,7 @@ enum ETerrainCoordMappingType
 };
 
 UENUM()
-enum ELandscapeCustomizedCoordType
+enum ELandscapeCustomizedCoordType : int
 {
 	/** Don't use customized UV, just use original UV. */
 	LCCT_None,
diff --git a/Engine/Source/Runtime/MediaAssets/Public/MediaTexture.h b/Engine/Source/Runtime/MediaAssets/Public/MediaTexture.h
--- a/Engine/Source/Runtime/MediaAssets/Public/MediaTexture.h
+++ b/Engine/Source/Runtime/MediaAssets/Public/MediaTexture.h
@@ -23,7 +23,7 @@ class UMediaPlayer;
 
 
 UENUM()
-enum MediaTextureOutputFormat
+enum MediaTextureOutputFormat : int
 {
 	MTOF_Default					UMETA(DisplayName = "Default (sRGB)"),
 	MTOF_SRGB_LINOUT				UMETA(DisplayName = "sRGB (linear output)"),		// sRGB data, using sRGB texture formats; hence read as linear RGB
@@ -31,7 +31,7 @@ enum MediaTextureOutputFormat
 };
 
 UENUM()
-enum MediaTextureOrientation
+enum MediaTextureOrientation : int
 {
 	MTORI_Original					UMETA(DisplayName = "Original (as decoded)"),
 	MTORI_CW90						UMETA(DisplayName = "Clockwise 90deg"),
diff --git a/Engine/Source/Runtime/MoviePlayer/Public/MoviePlayer.h b/Engine/Source/Runtime/MoviePlayer/Public/MoviePlayer.h
--- a/Engine/Source/Runtime/MoviePlayer/Public/MoviePlayer.h
+++ b/Engine/Source/Runtime/MoviePlayer/Public/MoviePlayer.h
@@ -13,7 +13,7 @@
 // starting very early and support rendering as soon as it is started and the module containing the movie streamer for the platform must already be loaded
 
 UENUM()
-enum EMoviePlaybackType
+enum EMoviePlaybackType : int
 {
 	/** Normal playback mode.  Play each movie in the play list a single time */
 	MT_Normal UMETA(DisplayName="Normal Playback"),
diff --git a/Engine/Source/Runtime/MovieSceneCapture/Public/Protocols/CompositionGraphCaptureProtocol.h b/Engine/Source/Runtime/MovieSceneCapture/Public/Protocols/CompositionGraphCaptureProtocol.h
--- a/Engine/Source/Runtime/MovieSceneCapture/Public/Protocols/CompositionGraphCaptureProtocol.h
+++ b/Engine/Source/Runtime/MovieSceneCapture/Public/Protocols/CompositionGraphCaptureProtocol.h
@@ -16,7 +16,7 @@ struct FFrameCaptureViewExtension;
 
 /** Used by UCompositionGraphCaptureSettings. Matches gamut order in EDisplayColorGamut */
 UENUM(BlueprintType)
-enum EHDRCaptureGamut
+enum EHDRCaptureGamut : int
 {
 	HCGM_Rec709 UMETA(DisplayName = "Rec.709 / sRGB"),
 	HCGM_P3DCI UMETA(DisplayName = "P3 D65"),
diff --git a/Engine/Source/Runtime/PhysicsCore/Public/BodySetupEnums.h b/Engine/Source/Runtime/PhysicsCore/Public/BodySetupEnums.h
--- a/Engine/Source/Runtime/PhysicsCore/Public/BodySetupEnums.h
+++ b/Engine/Source/Runtime/PhysicsCore/Public/BodySetupEnums.h
@@ -21,7 +21,7 @@ enum ECollisionTraceFlag : int
 };
 
 UENUM()
-enum EPhysicsType
+enum EPhysicsType : int
 {
 	/** Follow owner. */
 	PhysType_Default UMETA(DisplayName="Default"),	
@@ -34,7 +34,7 @@ enum EPhysicsType
 UENUM()
 namespace EBodyCollisionResponse
 {
-	enum Type
+	enum Type : int
 	{
 		BodyCollision_Enabled UMETA(DisplayName="Enabled"), 
 		BodyCollision_Disabled UMETA(DisplayName="Disabled")//, 
diff --git a/Engine/Source/Runtime/PhysicsCore/Public/Chaos/ChaosEngineInterface.h b/Engine/Source/Runtime/PhysicsCore/Public/Chaos/ChaosEngineInterface.h
--- a/Engine/Source/Runtime/PhysicsCore/Public/Chaos/ChaosEngineInterface.h
+++ b/Engine/Source/Runtime/PhysicsCore/Public/Chaos/ChaosEngineInterface.h
@@ -85,7 +85,7 @@ enum EPhysicalSurface : int
 
 /** Enum for controlling the falloff of strength of a radial impulse as a function of distance from Origin. */
 UENUM()
-enum ERadialImpulseFalloff
+enum ERadialImpulseFalloff : int
 {
 	/** Impulse is a constant strength, up to the limit of its range. */
 	RIF_Constant,
@@ -108,7 +108,7 @@ enum class ESleepFamily: uint8
 
 /** Specifies angular degrees of freedom */
 UENUM()
-enum EAngularConstraintMotion
+enum EAngularConstraintMotion : int
 {
 	/** No constraint against this axis. */
 	ACM_Free		UMETA(DisplayName="Free"),
@@ -124,7 +124,7 @@ enum EAngularConstraintMotion
 UENUM()
 namespace EConstraintFrame
 {
-enum Type
+enum Type : int
 {
 	Frame1,
 	Frame2
@@ -167,7 +167,7 @@ namespace PhysicsInterfaceTypes
 
 // LINEAR CCPT
 UENUM()
-enum EConstraintPlasticityType
+enum EConstraintPlasticityType : int
 {
 	/** */
 	CCPT_Free	UMETA(DisplayName = "Free"),
@@ -181,7 +181,7 @@ enum EConstraintPlasticityType
 
 // LINEAR DOF
 UENUM()
-enum ELinearConstraintMotion
+enum ELinearConstraintMotion : int
 {
 	/** No constraint against this axis. */
 	LCM_Free	UMETA(DisplayName="Free"),
diff --git a/Engine/Source/Runtime/RigVM/Public/RigVMCore/RigVMDrawInstruction.h b/Engine/Source/Runtime/RigVM/Public/RigVMCore/RigVMDrawInstruction.h
--- a/Engine/Source/Runtime/RigVM/Public/RigVMCore/RigVMDrawInstruction.h
+++ b/Engine/Source/Runtime/RigVM/Public/RigVMCore/RigVMDrawInstruction.h
@@ -11,7 +11,7 @@ class FMaterialRenderProxy;
 UENUM()
 namespace ERigVMDrawSettings
 {
-	enum Primitive
+	enum Primitive : int
 	{
 		Points,
 		Lines,
