Fix group assignment of TransformToGeometryIndexAttribute in FCollectionMeshFacade
(fixes fracture unit tests)

diff --git a/Engine/Source/Runtime/Experimental/Chaos/Private/GeometryCollection/Facades/CollectionMeshFacade.cpp b/Engine/Source/Runtime/Experimental/Chaos/Private/GeometryCollection/Facades/CollectionMeshFacade.cpp
--- a/Engine/Source/Runtime/Experimental/Chaos/Private/GeometryCollection/Facades/CollectionMeshFacade.cpp
+++ b/Engine/Source/Runtime/Experimental/Chaos/Private/GeometryCollection/Facades/CollectionMeshFacade.cpp
@@ -11,7 +11,7 @@
 namespace GeometryCollection::Facades
 {
 	FCollectionMeshFacade::FCollectionMeshFacade(FManagedArrayCollection& InCollection)
-		: TransformToGeometryIndexAttribute(InCollection, "TransformToGeometryIndex", FGeometryCollection::GeometryGroup)
+		: TransformToGeometryIndexAttribute(InCollection, "TransformToGeometryIndex", FGeometryCollection::TransformGroup)
 		, VertexAttribute(InCollection, "Vertex", FGeometryCollection::VerticesGroup)
 		, TangentUAttribute(InCollection, "TangentU", FGeometryCollection::VerticesGroup)
 		, TangentVAttribute(InCollection, "TangentV", FGeometryCollection::VerticesGroup)
@@ -31,7 +31,7 @@ namespace GeometryCollection::Facades
 	}
 
 	FCollectionMeshFacade::FCollectionMeshFacade(const FManagedArrayCollection& InCollection)
-		: TransformToGeometryIndexAttribute(InCollection, "TransformToGeometryIndex", FGeometryCollection::GeometryGroup)
+		: TransformToGeometryIndexAttribute(InCollection, "TransformToGeometryIndex", FGeometryCollection::TransformGroup)
 		, VertexAttribute(InCollection, "Vertex", FGeometryCollection::VerticesGroup)
 		, TangentUAttribute(InCollection, "TangentU", FGeometryCollection::VerticesGroup)
 		, TangentVAttribute(InCollection, "TangentV", FGeometryCollection::VerticesGroup)
