Chaos - New style joint container solver

Move the joint solve code out of the joint container. This allows us to make a few optimizations including collecting the joint solver data into a per-thread buffer in the order in which they are solved. It will also make it easier to support different joint solvers on a per-physics-asset basis (WIP)

diff --git a/Engine/Source/Programs/HeadlessChaos/Private/HeadlessChaosTestJoints.cpp b/Engine/Source/Programs/HeadlessChaos/Private/HeadlessChaosTestJoints.cpp
--- a/Engine/Source/Programs/HeadlessChaos/Private/HeadlessChaosTestJoints.cpp
+++ b/Engine/Source/Programs/HeadlessChaos/Private/HeadlessChaosTestJoints.cpp
@@ -758,6 +758,7 @@ namespace ChaosTest {
 		FParticleUniqueIndicesMultithreaded UniqueIndices;
 		FPBDRigidsSOAs ParticleContainer(UniqueIndices);
 		FPBDJointConstraints JointContainer;
+		JointContainer.SetSortEnabled(true);
 
 		// Create 3 particles
 		TArray<FPBDRigidParticleHandle*> Rigids = ParticleContainer.CreateDynamicParticles(3);
@@ -773,14 +774,14 @@ namespace ChaosTest {
 		JointContainer.PrepareTick();
 
 		// Both joints should be in an island
-		EXPECT_GE(Joints[0]->GetConstraintIsland(), 0);
-		EXPECT_GE(Joints[0]->GetConstraintIsland(), 0);
+		EXPECT_GE(JointContainer.GetConstraintIsland(0), 0);
+		EXPECT_GE(JointContainer.GetConstraintIsland(1), 0);
 
 		// Both joints should be in same island
-		EXPECT_EQ(Joints[0]->GetConstraintIsland(), Joints[1]->GetConstraintIsland());
+		EXPECT_EQ(JointContainer.GetConstraintIsland(0), JointContainer.GetConstraintIsland(1));
 
 		// Joints should have different colors
-		EXPECT_NE(Joints[0]->GetConstraintColor(), Joints[1]->GetConstraintColor());
+		EXPECT_NE(JointContainer.GetConstraintColor(0), JointContainer.GetConstraintColor(1));
 	}
 
 	// Check that constraints islands are not merged through shared kinematic particles
@@ -789,6 +790,7 @@ namespace ChaosTest {
 		FParticleUniqueIndicesMultithreaded UniqueIndices;
 		FPBDRigidsSOAs ParticleContainer(UniqueIndices);
 		FPBDJointConstraints JointContainer;
+		JointContainer.SetSortEnabled(true);
 
 		// Create 3 particles
 		TArray<FPBDRigidParticleHandle*> Rigids = ParticleContainer.CreateDynamicParticles(3);
@@ -807,17 +809,16 @@ namespace ChaosTest {
 		JointContainer.PrepareTick();
 
 		// Both joints should be in an island
-		EXPECT_GE(Joints[0]->GetConstraintIsland(), 0);
-		EXPECT_GE(Joints[1]->GetConstraintIsland(), 0);
+		EXPECT_GE(JointContainer.GetConstraintIsland(0), 0);
+		EXPECT_GE(JointContainer.GetConstraintIsland(1), 0);
 
 		// Joints should be in different islands
-		EXPECT_NE(Joints[0]->GetConstraintIsland(), Joints[1]->GetConstraintIsland());
+		EXPECT_NE(JointContainer.GetConstraintIsland(0), JointContainer.GetConstraintIsland(1));
 
 		// Both joints should be at level 0
-		EXPECT_EQ(Joints[0]->GetConstraintLevel(), 0);
-		EXPECT_EQ(Joints[1]->GetConstraintLevel(), 0);
+		EXPECT_EQ(JointContainer.GetConstraintLevel(0), 0);
+		EXPECT_EQ(JointContainer.GetConstraintLevel(1), 0);
 	}
 
-
 }
 
diff --git a/Engine/Source/Runtime/Engine/Private/PhysicsEngine/ImmediatePhysics/ImmediatePhysicsChaos/ImmediatePhysicsSimulation_Chaos.cpp b/Engine/Source/Runtime/Engine/Private/PhysicsEngine/ImmediatePhysics/ImmediatePhysicsChaos/ImmediatePhysicsSimulation_Chaos.cpp
--- a/Engine/Source/Runtime/Engine/Private/PhysicsEngine/ImmediatePhysics/ImmediatePhysicsChaos/ImmediatePhysicsSimulation_Chaos.cpp
+++ b/Engine/Source/Runtime/Engine/Private/PhysicsEngine/ImmediatePhysics/ImmediatePhysicsChaos/ImmediatePhysicsSimulation_Chaos.cpp
@@ -46,8 +46,6 @@ int32 ChaosImmediate_Evolution_DeltaTimeCount = 100;
 int32 ChaosImmediate_Evolution_PositionIterations = -1;
 int32 ChaosImmediate_Evolution_VelocityIterations = -1;
 int32 ChaosImmediate_Evolution_ProjectionIterations = -1;
-int32 ChaosImmediate_Evolution_Iterations = -1;				// Legacy
-int32 ChaosImmediate_Evolution_PushOutIterations = -1;		// Legacy
 int32 ChaosImmediate_DisableInactiveByIndex = 1;
 int32 ChaosImmediate_Evolution_NumCollisionsPerBlock = 50;
 FAutoConsoleVariableRef CVarChaosImmPhysStepTime(TEXT("p.Chaos.ImmPhys.StepTime"), ChaosImmediate_Evolution_StepTime, TEXT("Override step time (if not zero)"));
@@ -57,8 +55,6 @@ FAutoConsoleVariableRef CVarChaosImmPhysDeltaTimeCount(TEXT("p.Chaos.ImmPhys.Del
 FAutoConsoleVariableRef CVarChaosImmPhysPositionIterations(TEXT("p.Chaos.ImmPhys.PositionIterations"), ChaosImmediate_Evolution_PositionIterations, TEXT("Override number of position iteration loops in immediate physics (if >= 0)"));
 FAutoConsoleVariableRef CVarChaosImmPhysVelocityIterations(TEXT("p.Chaos.ImmPhys.VelocityIterations"), ChaosImmediate_Evolution_VelocityIterations, TEXT("Override number of velocity iteration loops in immediate physics (if >= 0)"));
 FAutoConsoleVariableRef CVarChaosImmPhysProjectionIterations(TEXT("p.Chaos.ImmPhys.ProjectionIterations"), ChaosImmediate_Evolution_ProjectionIterations, TEXT("Override number of projection iteration loops in immediate physics (if >= 0)"));
-FAutoConsoleVariableRef CVarChaosImmPhysIterations(TEXT("p.Chaos.ImmPhys.Iterations"), ChaosImmediate_Evolution_Iterations, TEXT("[Legacy Solver] Override number of constraint solver loops in immediate physics (if >= 0)"));
-FAutoConsoleVariableRef CVarChaosImmPhysPushOutIterations(TEXT("p.Chaos.ImmPhys.PushOutIterations"), ChaosImmediate_Evolution_PushOutIterations, TEXT("[Legacy Solver] Override number of solver push-out loops (if >= 0)"));
 FAutoConsoleVariableRef CVarChaosImmPhysDisableInactiveByIndex(TEXT("p.Chaos.ImmPhys.DisableInactiveByIndex"), ChaosImmediate_DisableInactiveByIndex, TEXT("Disable bodies that are no longer active based on the index, rather than just count."));
 FAutoConsoleVariableRef CVarChaosImmPhysNumCollisionsPerBlock(TEXT("p.Chaos.ImmPhys.NumCollisionsPerBlock"), ChaosImmediate_Evolution_NumCollisionsPerBlock, TEXT("The number of collision in a block in the collision pool. Higher values give better cache efficieny but waste memory if you do not need that many"));
 
@@ -69,7 +65,6 @@ FAutoConsoleVariableRef CVarChaosImmPhysSimSpaceCentrifugalAlpha(TEXT("p.Chaos.I
 FAutoConsoleVariableRef CVarChaosImmPhysSimSpaceCoriolisAlpha(TEXT("p.Chaos.ImmPhys.SimSpaceCoriolisAlpha"), ChaosImmediate_Evolution_SimSpaceCoriolisAlpha, TEXT("Settings for simulation space system for rigid body nodes"));
 FAutoConsoleVariableRef CVarChaosImmPhysSimSpaceEulerAlpha(TEXT("p.Chaos.ImmPhys.SimSpaceEulerAlpha"), ChaosImmediate_Evolution_SimSpaceEulerAlpha, TEXT("Settings for simulation space system for rigid body nodes"));
 
-
 Chaos::FRealSingle ChaosImmediate_Evolution_MinStepTime = 0.01f;
 Chaos::FRealSingle ChaosImmediate_Evolution_FixedStepTime = -1.0f;
 Chaos::FRealSingle ChaosImmediate_Evolution_FixedStepTolerance = 0.05f;
@@ -78,7 +73,6 @@ FAutoConsoleVariableRef CVarChaosImmPhysFixedStepTime(TEXT("p.Chaos.ImmPhys.Fixe
 FAutoConsoleVariableRef CVarChaosImmPhysFixedStepTolerance(TEXT("p.Chaos.ImmPhys.FixedStepTolerance"), ChaosImmediate_Evolution_FixedStepTolerance, TEXT("Time remainder required to add a new step (fraction of FixedStepTime)"));
 
 int32 ChaosImmediate_Collision_Enabled = 1;
-int32 ChaosImmediate_Collision_PairIterations = -1;			// Legacy
 int32 ChaosImmediate_Collision_NumPositionFrictionIterations = 0;		// No static friction for RBAN
 int32 ChaosImmediate_Collision_NumVelocityFrictionIterations = 1;		// Dynamic friction for RBAN in velocity solve
 int32 ChaosImmediate_Collision_NumPositionShockPropagationIterations = 0;
@@ -91,7 +85,6 @@ int32 ChaosImmediate_Collision_RestitutionEnabled = true;
 int32 ChaosImmediate_Collision_DeferNarrowPhase = 1;
 int32 ChaosImmediate_Collision_UseManifolds = 0;
 FAutoConsoleVariableRef CVarChaosImmPhysCollisionDisable(TEXT("p.Chaos.ImmPhys.Collision.Enabled"), ChaosImmediate_Collision_Enabled, TEXT("Enable/Disable collisions in Immediate Physics."));
-FAutoConsoleVariableRef CVarChaosImmPhysCollisionPairIterations(TEXT("p.Chaos.ImmPhys.Collision.PairIterations"), ChaosImmediate_Collision_PairIterations, TEXT("[Legacy Solver] Override collision pair iterations (if >= 0)"));
 FAutoConsoleVariableRef CVarChaosImmPhysCollisionPriority(TEXT("p.Chaos.ImmPhys.Collision.Priority"), ChaosImmediate_Collision_Priority, TEXT("Set the Collision constraint sort order (Joints have priority 0)"));
 FAutoConsoleVariableRef CVarChaosImmPhysCollisionCullDistance(TEXT("p.Chaos.ImmPhys.Collision.CullDistance"), ChaosImmediate_Collision_CullDistance, TEXT("Set the collision CullDistance (if >= 0)"));
 FAutoConsoleVariableRef CVarChaosImmPhysCollisionMaxDepenetrationVelocity(TEXT("p.Chaos.ImmPhys.Collision.MaxDepenetrationVelocity"), ChaosImmediate_Collision_MaxDepenetrationVelocity, TEXT("Set the collision Max Depenetration Velocity (if >= 0)"));
@@ -100,8 +93,6 @@ FAutoConsoleVariableRef CVarChaosImmPhysCollisionRestitutionEnabled(TEXT("p.Chao
 FAutoConsoleVariableRef CVarChaosImmPhysCollisionDeferNarrowPhase(TEXT("p.Chaos.ImmPhys.Collision.DeferNarrowPhase"), ChaosImmediate_Collision_DeferNarrowPhase, TEXT("[Legacy Solver] Create contacts for all broadphase pairs, perform NarrowPhase later."));
 FAutoConsoleVariableRef CVarChaosImmPhysCollisionUseManifolds(TEXT("p.Chaos.ImmPhys.Collision.UseManifolds"), ChaosImmediate_Collision_UseManifolds, TEXT("[Legacy Solver] Enable/Disable use of manifoldes in collision."));
 
-int32 ChaosImmediate_Joint_PairIterations = -1;				// Legacy
-int32 ChaosImmediate_Joint_PushOutPairIterations = -1;		// Legacy
 Chaos::FRealSingle ChaosImmediate_Joint_SwingTwistAngleTolerance = 1.0e-6f;
 Chaos::FRealSingle ChaosImmediate_Joint_PositionTolerance = 0.025f;
 Chaos::FRealSingle ChaosImmediate_Joint_AngleTolerance = 0.001f;
@@ -125,8 +116,6 @@ Chaos::FRealSingle ChaosImmediate_Joint_AngularDriveStiffness = -1.0f;
 Chaos::FRealSingle ChaosImmediate_Joint_AngularDriveDamping = -1.0f;
 Chaos::FRealSingle ChaosImmediate_Joint_MinParentMassRatio = 0.2f;
 Chaos::FRealSingle ChaosImmediate_Joint_MaxInertiaRatio = 5.0f;
-FAutoConsoleVariableRef CVarChaosImmPhysJointPairIterations(TEXT("p.Chaos.ImmPhys.Joint.PairIterations"), ChaosImmediate_Joint_PairIterations, TEXT("[Legacy Solver] Override joint pair iterations (if >= 0)"));
-FAutoConsoleVariableRef CVarChaosImmPhysJointPushOutPairIterations(TEXT("p.Chaos.ImmPhys.Joint.PushOutPairIterations"), ChaosImmediate_Joint_PushOutPairIterations, TEXT("[Legacy Solver] Override joint push-out pair iterations (if >= 0)"));
 FAutoConsoleVariableRef CVarChaosImmPhysJointSwingTwistAngleTolerance(TEXT("p.Chaos.ImmPhys.Joint.SwingTwistAngleTolerance"), ChaosImmediate_Joint_SwingTwistAngleTolerance, TEXT("SwingTwistAngleTolerance."));
 FAutoConsoleVariableRef CVarChaosImmPhysJointPositionTolerance(TEXT("p.Chaos.ImmPhys.Joint.PositionTolerance"), ChaosImmediate_Joint_PositionTolerance, TEXT("PositionTolerance."));
 FAutoConsoleVariableRef CVarChaosImmPhysJointAngleTolerance(TEXT("p.Chaos.ImmPhys.Joint.AngleTolerance"), ChaosImmediate_Joint_AngleTolerance, TEXT("AngleTolerance."));
@@ -856,6 +845,7 @@ namespace ImmediatePhysics_Chaos
 			JointsSettings.AngularDriveStiffnessOverride = ChaosImmediate_Joint_AngularDriveStiffness;
 			JointsSettings.AngularDriveDampingOverride = ChaosImmediate_Joint_AngularDriveDamping;
 			JointsSettings.NumShockPropagationIterations = ChaosImmediate_Joint_NumShockPropagationIterations;
+			JointsSettings.bSortEnabled = true;
 			Implementation->Joints.SetSettings(JointsSettings);
 
 			Implementation->Collisions.SetPositionFrictionIterations(ChaosImmediate_Collision_NumPositionFrictionIterations);
diff --git a/Engine/Source/Runtime/Experimental/Chaos/Private/Chaos/ChaosDebugDraw.cpp b/Engine/Source/Runtime/Experimental/Chaos/Private/Chaos/ChaosDebugDraw.cpp
--- a/Engine/Source/Runtime/Experimental/Chaos/Private/Chaos/ChaosDebugDraw.cpp
+++ b/Engine/Source/Runtime/Experimental/Chaos/Private/Chaos/ChaosDebugDraw.cpp
@@ -1167,7 +1167,7 @@ namespace Chaos
 			}
 		}
 
-		void DrawJointConstraintImpl(const FRigidTransform3& SpaceTransform, const FVec3& InPa, const FVec3& InCa, const FVec3& InXa, const FMatrix33& Ra, const FVec3& InPb, const FVec3& InCb, const FVec3& InXb, const FMatrix33& Rb, int32 IslandIndex, int32 LevelIndex, int32 ColorIndex, int32 Index, Chaos::FRealSingle ColorScale, const FChaosDebugDrawJointFeatures& FeatureMask, const FChaosDebugDrawSettings& Settings)
+		void DrawJointConstraintImpl(const FRigidTransform3& SpaceTransform, const FVec3& InPa, const FVec3& InCa, const FVec3& InXa, const FMatrix33& Ra, const FVec3& InPb, const FVec3& InCb, const FVec3& InXb, const FMatrix33& Rb, Chaos::FRealSingle ColorScale, const FChaosDebugDrawJointFeatures& FeatureMask, const FChaosDebugDrawSettings& Settings)
 		{
 			using namespace Chaos::DebugDraw;
 			FColor R = (ColorScale * FColor::Red).ToFColor(false);
@@ -1241,27 +1241,6 @@ namespace Chaos
 				FDebugDrawQueue::GetInstance().DrawDebugDirectionalArrow(Xb, Xb + Settings.DrawScale * Settings.ConstraintAxisLen * SpaceTransform.TransformVector(Rb.GetAxis(1)), Settings.DrawScale * Settings.ArrowSize, M, false, UE_KINDA_SMALL_NUMBER, Settings.DrawPriority, Settings.LineThickness);
 				FDebugDrawQueue::GetInstance().DrawDebugDirectionalArrow(Xb, Xb + Settings.DrawScale * Settings.ConstraintAxisLen * SpaceTransform.TransformVector(Rb.GetAxis(2)), Settings.DrawScale * Settings.ArrowSize, Y, false, UE_KINDA_SMALL_NUMBER, Settings.DrawPriority, Settings.LineThickness);
 			}
-			FVec3 TextPos = Xb;
-			if (FeatureMask.bLevel && (LevelIndex >= 0))
-			{
-				FDebugDrawQueue::GetInstance().DrawDebugString(TextPos, FString::Format(TEXT("{0}"), { LevelIndex }), nullptr, FColor::Red, UE_KINDA_SMALL_NUMBER, false, Settings.FontScale);
-				TextPos += Settings.FontHeight * FVec3(0, 0, 1);
-			}
-			if (FeatureMask.bIndex && (Index >= 0))
-			{
-				FDebugDrawQueue::GetInstance().DrawDebugString(TextPos, FString::Format(TEXT("{0}"), { Index }), nullptr, FColor::Red, UE_KINDA_SMALL_NUMBER, false, Settings.FontScale);
-				TextPos += Settings.FontHeight * FVec3(0, 0, 1);
-			}
-			if (FeatureMask.bColor && (ColorIndex >= 0))
-			{
-				FDebugDrawQueue::GetInstance().DrawDebugString(TextPos, FString::Format(TEXT("{0}"), { ColorIndex }), nullptr, FColor::Red, UE_KINDA_SMALL_NUMBER, false, Settings.FontScale);
-				TextPos += Settings.FontHeight * FVec3(0, 0, 1);
-			}
-			if (FeatureMask.bIsland && (IslandIndex >= 0))
-			{
-				FDebugDrawQueue::GetInstance().DrawDebugString(TextPos, FString::Format(TEXT("{0}"), { IslandIndex }), nullptr, FColor::Red, UE_KINDA_SMALL_NUMBER, false, Settings.FontScale);
-				TextPos += Settings.FontHeight * FVec3(0, 0, 1);
-			}
 		}
 
 		void DrawJointConstraintImpl(const FRigidTransform3& SpaceTransform, const FPBDJointConstraintHandle* ConstraintHandle, Chaos::FRealSingle ColorScale, const FChaosDebugDrawJointFeatures& FeatureMask, const FChaosDebugDrawSettings& Settings)
@@ -1283,7 +1262,7 @@ namespace Chaos
 				FVec3 Xa, Xb;
 				FMatrix33 Ra, Rb;
 				ConstraintHandle->CalculateConstraintSpace(Xa, Ra, Xb, Rb);
-				DrawJointConstraintImpl(SpaceTransform, Pa, Ca, Xa, Ra, Pb, Cb, Xb, Rb, ConstraintHandle->GetConstraintIsland(), ConstraintHandle->GetConstraintLevel(), ConstraintHandle->GetConstraintColor(), ConstraintHandle->GetConstraintIndex(), ColorScale, FeatureMask, Settings);
+				DrawJointConstraintImpl(SpaceTransform, Pa, Ca, Xa, Ra, Pb, Cb, Xb, Rb, ColorScale, FeatureMask, Settings);
 			}
 		}
 
diff --git a/Engine/Source/Runtime/Experimental/Chaos/Private/Chaos/Joint/PBDJointCachedSolverGaussSeidel.cpp b/Engine/Source/Runtime/Experimental/Chaos/Private/Chaos/Joint/PBDJointCachedSolverGaussSeidel.cpp
--- a/Engine/Source/Runtime/Experimental/Chaos/Private/Chaos/Joint/PBDJointCachedSolverGaussSeidel.cpp
+++ b/Engine/Source/Runtime/Experimental/Chaos/Private/Chaos/Joint/PBDJointCachedSolverGaussSeidel.cpp
@@ -10,9 +10,6 @@
 namespace Chaos
 {
 	
-FPBDJointCachedSolver::FPBDJointCachedSolver()
-{}
-	
 /** Derived states management */
 
 void FPBDJointCachedSolver::InitDerivedState()
@@ -197,6 +194,7 @@ void FPBDJointCachedSolver::Init(
 
 	NumActiveConstraints = -1;
 	bIsActive = true;
+	bIsBroken = false;
 
 	SolverStiffness = 1.0f;
 
diff --git a/Engine/Source/Runtime/Experimental/Chaos/Private/Chaos/Joint/PBDJointContainerSolver.cpp b/Engine/Source/Runtime/Experimental/Chaos/Private/Chaos/Joint/PBDJointContainerSolver.cpp
new file mode 100644
--- /dev/null
+++ b/Engine/Source/Runtime/Experimental/Chaos/Private/Chaos/Joint/PBDJointContainerSolver.cpp
@@ -0,0 +1,427 @@
+// Copyright Epic Games, Inc. All Rights Reserved.
+#include "Chaos/Joint/PBDJointContainerSolver.h"
+#include "Chaos/Joint/ChaosJointLog.h"
+
+namespace Chaos
+{
+	namespace Private
+	{
+		///////////////////////////////////////////////////////////////////////////////////////////
+		///////////////////////////////////////////////////////////////////////////////////////////
+		///////////////////////////////////////////////////////////////////////////////////////////
+
+		// NOTE: Particles are passed to the solvers in reverse order to what they are in the container...
+		FGeometryParticleHandle* GetJointParticle(FPBDJointConstraints& Constraints, const int32 ContainerConstraintIndex, const int32 ParticleIndex)
+		{
+			check(ParticleIndex >= 0);
+			check(ParticleIndex < 2);
+
+			const int32 SwappedIndex = 1 - ParticleIndex;
+			return Constraints.GetConstrainedParticles(ContainerConstraintIndex)[SwappedIndex];
+		}
+
+		const FRigidTransform3& GetJointFrame(FPBDJointConstraints& Constraints, const int32 ContainerConstraintIndex, const int32 ParticleIndex)
+		{
+			check(ParticleIndex >= 0);
+			check(ParticleIndex < 2);
+
+			const int32 SwappedIndex = 1 - ParticleIndex;
+			return Constraints.GetConstraintSettings(ContainerConstraintIndex).ConnectorTransforms[SwappedIndex];
+		}
+
+
+		// @todo(chaos): ShockPropagation needs to handle the parent/child being in opposite order
+		FReal GetJointShockPropagationInvMassScale(const FConstraintSolverBody& Body0, const FConstraintSolverBody& Body1, const FPBDJointSolverSettings& Settings, const FPBDJointSettings& JointSettings, const int32 It, const int32 NumIts)
+		{
+			// Shock propagation is only enabled for the last iteration, and only for the QPBD solver.
+			// The standard PBD solver runs projection in the second solver phase which is mostly the same thing.
+			if (JointSettings.bShockPropagationEnabled && (It >= (NumIts - Settings.NumShockPropagationIterations)))
+			{
+				if (Body0.IsDynamic() && Body1.IsDynamic())
+				{
+					return FPBDJointUtilities::GetShockPropagationInvMassScale(Settings, JointSettings);
+				}
+			}
+			return FReal(1);
+		}
+
+		FReal GetJointIterationStiffness(const FPBDJointSolverSettings& Settings, int32 It, int32 NumIts)
+		{
+			// Linearly interpolate betwwen MinStiffness and MaxStiffness over the first few iterations,
+			// then clamp at MaxStiffness for the final NumIterationsAtMaxStiffness
+			FReal IterationStiffness = Settings.MaxSolverStiffness;
+			if (NumIts > Settings.NumIterationsAtMaxSolverStiffness)
+			{
+				const FReal Interpolant = FMath::Clamp((FReal)It / (FReal)(NumIts - Settings.NumIterationsAtMaxSolverStiffness), 0.0f, 1.0f);
+				IterationStiffness = FMath::Lerp(Settings.MinSolverStiffness, Settings.MaxSolverStiffness, Interpolant);
+			}
+			return FMath::Clamp(IterationStiffness, 0.0f, 1.0f);
+		}
+
+		bool GetJointShouldBreak(const FPBDJointSettings& JointSettings, const FReal Dt, const FVec3& LinearImpulse, const FVec3& AngularImpulse)
+		{
+			// NOTE: LinearImpulse/AngularImpulse are not really impulses - they are mass-weighted position/rotation delta, or (impulse x dt).
+			// The Threshold is a force limit, so we need to convert it to a position delta caused by that force in one timestep
+
+			bool bBreak = false;
+			if (!bBreak && JointSettings.LinearBreakForce != FLT_MAX)
+			{
+				const FReal LinearForceSq = LinearImpulse.SizeSquared() / (Dt * Dt * Dt * Dt);
+				const FReal LinearThresholdSq = FMath::Square(JointSettings.LinearBreakForce);
+				bBreak = LinearForceSq > LinearThresholdSq;
+			}
+
+			if (!bBreak && JointSettings.AngularBreakTorque != FLT_MAX)
+			{
+				const FReal AngularForceSq = AngularImpulse.SizeSquared() / (Dt * Dt * Dt * Dt);
+				const FReal AngularThresholdSq = FMath::Square(JointSettings.AngularBreakTorque);
+				bBreak = AngularForceSq > AngularThresholdSq;
+			}
+
+			return bBreak;
+		}
+
+		///////////////////////////////////////////////////////////////////////////////////////////
+		///////////////////////////////////////////////////////////////////////////////////////////
+		///////////////////////////////////////////////////////////////////////////////////////////
+
+
+		FPBDJointContainerSolver::FPBDJointContainerSolver(FPBDJointConstraints& InConstraintContainer, const int32 InPriority)
+			: FConstraintContainerSolver(InPriority)
+			, ConstraintContainer(InConstraintContainer)
+		{
+		}
+
+		FPBDJointContainerSolver::~FPBDJointContainerSolver()
+		{
+		}
+
+		bool FPBDJointContainerSolver::UseLinearSolver() const
+		{
+			return ConstraintContainer.GetSettings().bUseLinearSolver;
+		}
+
+		void FPBDJointContainerSolver::Reset(const int32 InMaxConstraints)
+		{
+			ContainerConstraintIndices.Reset(InMaxConstraints);
+
+			// NOTE: We presize the solver arrays to avoid repeated calls to Add(). We resize down later if
+			// the InMaxConstraints turns out to be an over-estimate (which is should never be, currently)
+			if (UseLinearSolver())
+			{
+				LinearConstraintSolvers.SetNum(InMaxConstraints);
+				NonLinearConstraintSolvers.Empty();
+			}
+			else
+			{
+				LinearConstraintSolvers.Empty();
+				NonLinearConstraintSolvers.SetNum(InMaxConstraints);
+			}
+		}
+
+		void FPBDJointContainerSolver::FPBDJointContainerSolver::AddConstraints()
+		{
+			Reset(ConstraintContainer.GetNumConstraints());
+
+			// @todo(chaos): we could eliminate the index array if we're solving all constraints in the scene (RBAN)
+			for (int32 ContainerConstraintIndex = 0; ContainerConstraintIndex < ConstraintContainer.GetNumConstraints(); ++ContainerConstraintIndex)
+			{
+				AddConstraint(ContainerConstraintIndex);
+			}
+		}
+
+		void FPBDJointContainerSolver::AddConstraints(const TArrayView<FPBDIslandConstraint>& IslandConstraints)
+		{
+			for (FPBDIslandConstraint& IslandConstraint : IslandConstraints)
+			{
+				// We will only ever be given constraints from our container (asserts in non-shipping)
+				const int32 ContainerConstraintIndex = IslandConstraint.GetConstraint()->AsUnsafe<FPBDJointConstraintHandle>()->GetConstraintIndex();
+
+				AddConstraint(ContainerConstraintIndex);
+			}
+		}
+
+		void FPBDJointContainerSolver::AddConstraint(const int32 InContainerConstraintIndex)
+		{
+			// If this triggers, Reset was called with the wrong constraint count
+			check(ContainerConstraintIndices.Num() < ContainerConstraintIndices.Max());
+
+			ContainerConstraintIndices.Add(InContainerConstraintIndex);
+		}
+
+		void FPBDJointContainerSolver::AddBodies(FSolverBodyContainer& SolverBodyContainer)
+		{
+			for (int32 SolverConstraintIndex = 0, SolverConstraintEndIndex = ContainerConstraintIndices.Num(); SolverConstraintIndex < SolverConstraintEndIndex; ++SolverConstraintIndex)
+			{
+				const int32 ContainerConstraintIndex = ContainerConstraintIndices[SolverConstraintIndex];
+
+				FGenericParticleHandle Particle0 = GetJointParticle(ConstraintContainer, ContainerConstraintIndex, 0);
+				FGenericParticleHandle Particle1 = GetJointParticle(ConstraintContainer, ContainerConstraintIndex, 1);
+
+				FSolverBody* SolverBody0 = SolverBodyContainer.FindOrAdd(Particle0);
+				FSolverBody* SolverBody1 = SolverBodyContainer.FindOrAdd(Particle1);
+
+				if (UseLinearSolver())
+				{
+					LinearConstraintSolvers[SolverConstraintIndex].SetSolverBodies(SolverBody0, SolverBody1);
+				}
+				else
+				{
+					NonLinearConstraintSolvers[SolverConstraintIndex].SetSolverBodies(SolverBody0, SolverBody1);
+				}
+			}
+		}
+
+		void FPBDJointContainerSolver::GatherInput(const FReal Dt)
+		{
+			GatherInput(Dt, 0, ContainerConstraintIndices.Num());
+		}
+
+		template<typename SolverType>
+		void GatherInputImpl(const FPBDJointContainerSolver& Container, TArray<SolverType>& Solvers, const FReal Dt, const int32 SolverConstraintBeginIndex, const int32 SolverConstraintEndIndex)
+		{
+			const FPBDJointSolverSettings& SolverSettings = Container.GetSettings();
+
+			for (int32 SolverConstraintIndex = SolverConstraintBeginIndex; SolverConstraintIndex < SolverConstraintEndIndex; ++SolverConstraintIndex)
+			{
+				const FPBDJointSettings& JointSettings = Container.GetConstraintSettings(SolverConstraintIndex);
+
+				const int32 ContainerConstraintIndex = Container.GetContainerConstraintIndex(SolverConstraintIndex);
+				FGenericParticleHandle Particle0 = GetJointParticle(Container.GetContainer(), ContainerConstraintIndex, 0);
+				FGenericParticleHandle Particle1 = GetJointParticle(Container.GetContainer(), ContainerConstraintIndex, 1);
+				const FRigidTransform3& Frame0 = GetJointFrame(Container.GetContainer(), ContainerConstraintIndex, 0);
+				const FRigidTransform3& Frame1 = GetJointFrame(Container.GetContainer(), ContainerConstraintIndex, 1);
+
+				Solvers[SolverConstraintIndex].Init(Dt, SolverSettings, JointSettings, Particle0->GetComRelativeTransform(Frame0), Particle1->GetComRelativeTransform(Frame1));
+			}
+		}
+
+		void FPBDJointContainerSolver::GatherInput(const FReal Dt, const int32 ConstraintBeginIndex, const int32 ConstraintEndIndex)
+		{
+			if (UseLinearSolver())
+			{
+				GatherInputImpl(*this, LinearConstraintSolvers, Dt, ConstraintBeginIndex, ConstraintEndIndex);
+			}
+			else
+			{
+				GatherInputImpl(*this, NonLinearConstraintSolvers, Dt, ConstraintBeginIndex, ConstraintEndIndex);
+			}
+		}
+
+		void FPBDJointContainerSolver::ScatterOutput(const FReal Dt)
+		{
+			ScatterOutput(Dt, 0, ContainerConstraintIndices.Num());
+		}
+
+		template<typename SolverType>
+		void ScatterOutputImpl(const FPBDJointContainerSolver& Container, TArray<SolverType>& Solvers, const FReal Dt, const int32 SolverConstraintBeginIndex, const int32 SolverConstraintEndIndex)
+		{
+			for (int32 SolverConstraintIndex = SolverConstraintBeginIndex; SolverConstraintIndex < SolverConstraintEndIndex; ++SolverConstraintIndex)
+			{
+				const int32 ContainerConstraintIndex = Container.GetContainerConstraintIndex(SolverConstraintIndex);
+				if (Dt > UE_SMALL_NUMBER)
+				{
+					SolverType& Solver = Solvers[SolverConstraintIndex];
+
+					// NOTE: Particle order was revered in the solver...
+					const FVec3 LinearImpulse = -Solver.GetNetLinearImpulse() / Dt;
+					const FVec3 AngularImpulse = -Solver.GetNetAngularImpulse() / Dt;
+					const FSolverBody* SolverBody0 = &Solver.Body0().SolverBody();
+					const FSolverBody* SolverBody1 = &Solver.Body1().SolverBody();
+					const bool bIsBroken = Solver.IsBroken();
+
+					Container.GetContainer().SetSolverResults(ContainerConstraintIndex, LinearImpulse, AngularImpulse, bIsBroken, SolverBody0, SolverBody1);
+
+					Solver.Deinit();
+				}
+				else
+				{
+					Container.GetContainer().SetSolverResults(ContainerConstraintIndex, FVec3(0), FVec3(0), false, nullptr, nullptr);
+				}
+			}
+		}
+
+		void FPBDJointContainerSolver::ScatterOutput(const FReal Dt, const int32 ConstraintBeginIndex, const int32 ConstraintEndIndex)
+		{
+			if (UseLinearSolver())
+			{
+				ScatterOutputImpl(*this, LinearConstraintSolvers, Dt, ConstraintBeginIndex, ConstraintEndIndex);
+			}
+			else
+			{
+				ScatterOutputImpl(*this, NonLinearConstraintSolvers, Dt, ConstraintBeginIndex, ConstraintEndIndex);
+			}
+		}
+
+		void FPBDJointContainerSolver::ResizeSolverArrays()
+		{
+			// We may have conservatively allocated the solver arrays. If so, reduce their size now
+			const int32 NumConstraints = GetNumConstraints();
+			if (UseLinearSolver())
+			{
+				check(LinearConstraintSolvers.Num() >= NumConstraints);
+				LinearConstraintSolvers.SetNum(NumConstraints);
+			}
+			else
+			{
+				check(NonLinearConstraintSolvers.Num() >= NumConstraints);
+				NonLinearConstraintSolvers.SetNum(NumConstraints);
+			}
+		}
+
+		// Apply position constraints for linear or non-linear solvers
+		template<typename SolverType>
+		void ApplyPositionConstraintsImpl(const FPBDJointContainerSolver& Container, TArray<SolverType>& Solvers, const FReal Dt, const int32 It, const int32 NumIts)
+		{
+			const FPBDJointSolverSettings& Settings = Container.GetSettings();
+			const FReal IterationStiffness = GetJointIterationStiffness(Settings, It, NumIts);
+
+			for (int32 SolverConstraintIndex = 0; SolverConstraintIndex < Solvers.Num(); ++SolverConstraintIndex)
+			{
+				SolverType& Solver = Solvers[SolverConstraintIndex];
+				if (!Solver.RequiresSolve())
+				{
+					continue;
+				}
+
+				const FPBDJointSettings& JointSettings = Container.GetConstraintSettings(SolverConstraintIndex);
+				Solver.Update(Dt, Settings, JointSettings);
+
+				// Set parent inverse mass scale based on current shock propagation state
+				const FReal ShockPropagationInvMassScale = GetJointShockPropagationInvMassScale(Solver.Body0(), Solver.Body1(), Settings, JointSettings, It, NumIts);
+				Solver.SetShockPropagationScales(ShockPropagationInvMassScale, FReal(1), Dt);
+
+				Solver.ApplyConstraints(Dt, IterationStiffness, Settings, JointSettings);
+
+				// @todo(ccaulfield): We should be clamping the impulse at this point. Maybe move breaking to the solver
+				if (GetJointShouldBreak(JointSettings, Dt, Solver.GetNetLinearImpulse(), Solver.GetNetAngularImpulse()))
+				{
+					Solver.SetIsBroken(true);
+				}
+			}
+		}
+
+		void FPBDJointContainerSolver::ApplyPositionConstraints(const FReal Dt, const int32 It, const int32 NumIts)
+		{
+			ResizeSolverArrays();
+
+			ApplyPositionConstraintsImpl(*this, LinearConstraintSolvers, Dt, It, NumIts);
+			ApplyPositionConstraintsImpl(*this, NonLinearConstraintSolvers, Dt, It, NumIts);
+		}
+
+		// Apply velocity constraints for linear or non-linear solvers
+		template<typename SolverType>
+		void ApplyVelocityConstraintsImpl(const FPBDJointContainerSolver& Container, TArray<SolverType>& Solvers, const FReal Dt, const int32 It, const int32 NumIts)
+		{
+			const FPBDJointSolverSettings& Settings = Container.GetSettings();
+			const FReal IterationStiffness = GetJointIterationStiffness(Settings, It, NumIts);
+
+			for (int32 SolverConstraintIndex = 0; SolverConstraintIndex < Solvers.Num(); ++SolverConstraintIndex)
+			{
+				SolverType& Solver = Solvers[SolverConstraintIndex];
+				if (!Solver.RequiresSolve())
+				{
+					continue;
+				}
+
+				const FPBDJointSettings& JointSettings = Container.GetConstraintSettings(SolverConstraintIndex);
+				Solver.Update(Dt, Settings, JointSettings);
+
+				// Set parent inverse mass scale based on current shock propagation state
+				const FReal ShockPropagationInvMassScale = GetJointShockPropagationInvMassScale(Solver.Body0(), Solver.Body1(), Settings, JointSettings, It, NumIts);
+				Solver.SetShockPropagationScales(ShockPropagationInvMassScale, FReal(1), Dt);
+
+				Solver.ApplyVelocityConstraints(Dt, IterationStiffness, Settings, JointSettings);
+
+				// @todo(chaos): should also add to net impulse and run break logic
+			}
+		}
+
+		void FPBDJointContainerSolver::ApplyVelocityConstraints(const FReal Dt, const int32 It, const int32 NumIts)
+		{
+			ApplyVelocityConstraintsImpl(*this, LinearConstraintSolvers, Dt, It, NumIts);
+			ApplyVelocityConstraintsImpl(*this, NonLinearConstraintSolvers, Dt, It, NumIts);
+		}
+
+		void FPBDJointContainerSolver::ApplyProjectionConstraints(const FReal Dt, const int32 It, const int32 NumIts)
+		{
+			ApplyLinearProjectionConstraints(Dt, It, NumIts);
+			ApplyNonLinearProjectionConstraints(Dt, It, NumIts);
+		}
+
+		void FPBDJointContainerSolver::ApplyLinearProjectionConstraints(const FReal Dt, const int32 It, const int32 NumIts)
+		{
+			const FPBDJointSolverSettings& Settings = GetSettings();
+
+			if (It == 0)
+			{
+				// Collect all the data for projection prior to the first iteration. 
+				// This must happen for all joints before we project any joints so the the initial state for each joint is not polluted by any earlier projections.
+				// @todo(chaos): if we ever support projection on other constraint types, we will need a PrepareProjection phase so that all constraint types
+				// can initialize correctly before any constraints apply their projection. For now we can just check the iteration count is zero.
+				for (int32 SolverConstraintIndex = 0; SolverConstraintIndex < LinearConstraintSolvers.Num(); ++SolverConstraintIndex)
+				{
+					FPBDJointCachedSolver& Solver = LinearConstraintSolvers[SolverConstraintIndex];
+					if (!Solver.RequiresSolve())
+					{
+						continue;
+					}
+
+					const FPBDJointSettings& JointSettings = GetConstraintSettings(SolverConstraintIndex);
+					if (JointSettings.bProjectionEnabled)
+					{
+						Solver.InitProjection(Dt, Settings, JointSettings);
+					}
+				}
+			}
+
+			for (int32 SolverConstraintIndex = 0; SolverConstraintIndex < LinearConstraintSolvers.Num(); ++SolverConstraintIndex)
+			{
+				FPBDJointCachedSolver& Solver = LinearConstraintSolvers[SolverConstraintIndex];
+				if (!Solver.RequiresSolve())
+				{
+					continue;
+				}
+
+				const FPBDJointSettings& JointSettings = GetConstraintSettings(SolverConstraintIndex);
+
+				if (It == 0)
+				{
+					Solver.ApplyTeleports(Dt, Settings, JointSettings);
+				}
+
+				const bool bLastIteration = (It == (NumIts - 1));
+				Solver.ApplyProjections(Dt, Settings, JointSettings, bLastIteration);
+			}
+		}
+
+		void FPBDJointContainerSolver::ApplyNonLinearProjectionConstraints(const FReal Dt, const int32 It, const int32 NumIts)
+		{
+			const FPBDJointSolverSettings& Settings = GetSettings();
+
+			for (int32 SolverConstraintIndex = 0; SolverConstraintIndex < NonLinearConstraintSolvers.Num(); ++SolverConstraintIndex)
+			{
+				FPBDJointSolver& Solver = NonLinearConstraintSolvers[SolverConstraintIndex];
+				if (!Solver.RequiresSolve())
+				{
+					continue;
+				}
+
+				const FPBDJointSettings& JointSettings = GetConstraintSettings(SolverConstraintIndex);
+
+				Solver.Update(Dt, Settings, JointSettings);
+
+				if (It == 0)
+				{
+					// @todo(chaos): support reverse parent/child
+					Solver.Body1().UpdateRotationDependentState();
+					Solver.UpdateMasses(FReal(0), FReal(1));
+				}
+
+				const bool bLastIteration = (It == (NumIts - 1));
+				Solver.ApplyProjections(Dt, Settings, JointSettings, bLastIteration);
+			}
+		}
+
+	}	// namespace Private
+}	// namespace Chaos
\ No newline at end of file
diff --git a/Engine/Source/Runtime/Experimental/Chaos/Private/Chaos/Joint/PBDJointSolverGaussSeidel.cpp b/Engine/Source/Runtime/Experimental/Chaos/Private/Chaos/Joint/PBDJointSolverGaussSeidel.cpp
--- a/Engine/Source/Runtime/Experimental/Chaos/Private/Chaos/Joint/PBDJointSolverGaussSeidel.cpp
+++ b/Engine/Source/Runtime/Experimental/Chaos/Private/Chaos/Joint/PBDJointSolverGaussSeidel.cpp
@@ -28,16 +28,6 @@ namespace Chaos
 	//
 	//
 
-	FPBDJointSolver::FPBDJointSolver()
-	{
-		if (bChaos_Joint_ISPC_Enabled)
-		{
-#if INTEL_ISPC
-			check(sizeof(FPBDJointSolver) == ispc::SizeofFPBDJointSolver()); 
-#endif
-		}
-	}
-
 	void FPBDJointSolver::InitDerivedState()
 	{
 		InitConnectorXs[0] = X(0) + R(0) * LocalConnectorXs[0].GetTranslation();
@@ -142,6 +132,7 @@ namespace Chaos
 		AngleTolerance = ToleranceScale * SolverSettings.AngleTolerance;
 
 		SolverStiffness = 1.0f;
+		bIsBroken = false;
 
 		LinearHardLambda = FVec3(0);
 		AngularHardLambda = FVec3(0);
@@ -211,7 +202,7 @@ namespace Chaos
 		}
 	}
 
-	void FPBDJointSolver::SetShockPropagationScales(const FReal InvMScale0, const FReal InvMScale1)
+	void FPBDJointSolver::SetShockPropagationScales(const FReal InvMScale0, const FReal InvMScale1, const FReal Dt)
 	{
 		if (InvMScales[0] != InvMScale0)
 		{
@@ -378,7 +369,8 @@ namespace Chaos
 	void FPBDJointSolver::ApplyProjections(
 		const FReal Dt,
 		const FPBDJointSolverSettings& SolverSettings,
-		const FPBDJointSettings& JointSettings)
+		const FPBDJointSettings& JointSettings,
+		const bool bLastIteration)
 	{
 		// @todo(chaos): We need to handle parent/child being the other way round
 		if (!IsDynamic(1))
@@ -405,22 +397,25 @@ namespace Chaos
 		}
 
 		// Final position fixup
-		const TVec3<EJointMotionType>& LinearMotion = JointSettings.LinearMotionTypes;
-		const bool bLinearLocked = (LinearMotion[0] == EJointMotionType::Locked) && (LinearMotion[1] == EJointMotionType::Locked) && (LinearMotion[2] == EJointMotionType::Locked);
-		if (bLinearLocked)
+		if (bLastIteration)
 		{
-			const FReal LinearProjection = FPBDJointUtilities::GetLinearProjection(SolverSettings, JointSettings);
-			const bool bLinearSoft = FPBDJointUtilities::GetSoftLinearLimitEnabled(SolverSettings, JointSettings);
-			const bool bLinearProjectionEnabled = (!bLinearSoft && JointSettings.bProjectionEnabled);
-			if (bLinearProjectionEnabled && (LinearProjection > 0))
+			const TVec3<EJointMotionType>& LinearMotion = JointSettings.LinearMotionTypes;
+			const bool bLinearLocked = (LinearMotion[0] == EJointMotionType::Locked) && (LinearMotion[1] == EJointMotionType::Locked) && (LinearMotion[2] == EJointMotionType::Locked);
+			if (bLinearLocked)
 			{
-				ApplyTranslateProjection(Dt, SolverSettings, JointSettings, LinearProjection, DP1, DR1);
-			}
+				const FReal LinearProjection = FPBDJointUtilities::GetLinearProjection(SolverSettings, JointSettings);
+				const bool bLinearSoft = FPBDJointUtilities::GetSoftLinearLimitEnabled(SolverSettings, JointSettings);
+				const bool bLinearProjectionEnabled = (!bLinearSoft && JointSettings.bProjectionEnabled);
+				if (bLinearProjectionEnabled && (LinearProjection > 0))
+				{
+					ApplyTranslateProjection(Dt, SolverSettings, JointSettings, LinearProjection, DP1, DR1);
+				}
 
-			// Add velocity correction from the net projection motion
-			if (Chaos_Joint_VelProjectionAlpha > 0.0f)
-			{
-				ApplyVelocityProjection(Dt, SolverSettings, JointSettings, Chaos_Joint_VelProjectionAlpha, DP1, DR1);
+				// Add velocity correction from the net projection motion
+				if (Chaos_Joint_VelProjectionAlpha > 0.0f)
+				{
+					ApplyVelocityProjection(Dt, SolverSettings, JointSettings, Chaos_Joint_VelProjectionAlpha, DP1, DR1);
+				}
 			}
 		}
 	}
diff --git a/Engine/Source/Runtime/Experimental/Chaos/Private/Chaos/PBDJointConstraints.cpp b/Engine/Source/Runtime/Experimental/Chaos/Private/Chaos/PBDJointConstraints.cpp
--- a/Engine/Source/Runtime/Experimental/Chaos/Private/Chaos/PBDJointConstraints.cpp
+++ b/Engine/Source/Runtime/Experimental/Chaos/Private/Chaos/PBDJointConstraints.cpp
@@ -6,6 +6,7 @@
 #include "Chaos/Joint/ChaosJointLog.h"
 #include "Chaos/Joint/ColoringGraph.h"
 #include "Chaos/Joint/JointConstraintsCVars.h"
+#include "Chaos/Joint/PBDJointContainerSolver.h"
 #include "Chaos/Particle/ParticleUtilities.h"
 #include "Chaos/ParticleHandle.h"
 #include "Chaos/PBDJointConstraintUtilities.h"
@@ -359,6 +360,7 @@ namespace Chaos
 		, NumIterationsAtMaxSolverStiffness(1)
 		, NumShockPropagationIterations(0)
 		, bUseLinearSolver(true)
+		, bSortEnabled(false)
 		, bSolvePositionLast(true)
 		, bEnableTwistLimits(true)
 		, bEnableSwingLimits(true)
@@ -400,6 +402,15 @@ namespace Chaos
 	{
 	}
 
+	TUniquePtr<FConstraintContainerSolver> FPBDJointConstraints::CreateSceneSolver(const int32 Priority)
+	{
+		return MakeUnique<Private::FPBDJointContainerSolver>(*this, Priority);
+	}
+
+	TUniquePtr<FConstraintContainerSolver> FPBDJointConstraints::CreateGroupSolver(const int32 Priority)
+	{
+		return MakeUnique<Private::FPBDJointContainerSolver>(*this, Priority);
+	}
 	
 	const FPBDJointSolverSettings& FPBDJointConstraints::GetSettings() const
 	{
@@ -500,7 +511,7 @@ namespace Chaos
 		// Update the handle for the constraint that was moved
 		if (ConstraintIndex < Handles.Num())
 		{
-			SetConstraintIndex(Handles[ConstraintIndex], ConstraintIndex);
+			Handles[ConstraintIndex]->SetConstraintIndex(ConstraintIndex);
 		}
 	}
 
@@ -538,46 +549,54 @@ namespace Chaos
 	void FPBDJointConstraints::SortConstraints()
 	{
 		// Sort constraints so that constraints with lower level (closer to a kinematic joint) are first
-		// @todo(ccaulfield): should probably also take islands/particle order into account
-		// @todo(ccaulfield): optimize (though isn't called very often)
+		// @todo(chaos): this is only needed for RBAN and should be moved to an RBAN specific solver
+		// (so we should split FPBDJointContainerSolver into Scene and Group versions)
 		SCOPE_CYCLE_COUNTER(STAT_Joints_Sort);
 
-		FHandles SortedHandles = Handles;
-		SortedHandles.StableSort([](const FConstraintContainerHandle& L, const FConstraintContainerHandle& R)
+		TArray<int32> SortedIndices;
+		SortedIndices.Reserve(Handles.Num());
+		for (int32 JointIndex = 0; JointIndex < GetNumConstraints(); ++JointIndex)
+		{
+			SortedIndices.Add(JointIndex);
+		}
+		SortedIndices.StableSort([this](const int32& LIndex, const int32& RIndex)
 			{
-				if (L.GetConstraintIsland() != R.GetConstraintIsland())
+				if (GetConstraintIsland(LIndex) != GetConstraintIsland(RIndex))
 				{
-					return L.GetConstraintIsland() < R.GetConstraintIsland();
+					return GetConstraintIsland(LIndex) < GetConstraintIsland(RIndex);
 				}
-				else if (L.GetConstraintLevel() != R.GetConstraintLevel())
+				else if (GetConstraintLevel(LIndex) != GetConstraintLevel(RIndex))
 				{
-					return L.GetConstraintLevel() < R.GetConstraintLevel();
+					return GetConstraintLevel(LIndex) < GetConstraintLevel(RIndex);
 				}
-				return L.GetConstraintColor() < R.GetConstraintColor();
+				return GetConstraintColor(LIndex) < GetConstraintColor(RIndex);
 			});
 
+		TArray<FPBDJointConstraintHandle*> SortedHandles;
 		TArray<FPBDJointSettings> SortedConstraintSettings;
 		TArray<FParticlePair> SortedConstraintParticles;
 		TArray<FPBDJointState> SortedConstraintStates;
-		SortedConstraintSettings.Reserve(SortedHandles.Num());
-		SortedConstraintParticles.Reserve(SortedHandles.Num());
-		SortedConstraintStates.Reserve(SortedHandles.Num());
+		SortedHandles.Reserve(SortedIndices.Num());
+		SortedConstraintSettings.Reserve(SortedIndices.Num());
+		SortedConstraintParticles.Reserve(SortedIndices.Num());
+		SortedConstraintStates.Reserve(SortedIndices.Num());
 
-		for (int32 SortedConstraintIndex = 0; SortedConstraintIndex < SortedHandles.Num(); ++SortedConstraintIndex)
+		for (int32 SortedConstraintIndex = 0; SortedConstraintIndex < SortedIndices.Num(); ++SortedConstraintIndex)
 		{
-			FConstraintContainerHandle* Handle = SortedHandles[SortedConstraintIndex];
-			int32 UnsortedConstraintIndex = Handle->GetConstraintIndex();
+			int32 UnsortedConstraintIndex = SortedIndices[SortedConstraintIndex];
 
+			SortedHandles.Add(Handles[UnsortedConstraintIndex]);
 			SortedConstraintSettings.Add(ConstraintSettings[UnsortedConstraintIndex]);
 			SortedConstraintParticles.Add(ConstraintParticles[UnsortedConstraintIndex]);
 			SortedConstraintStates.Add(ConstraintStates[UnsortedConstraintIndex]);
-			SetConstraintIndex(Handle, SortedConstraintIndex);
+
+			Handles[UnsortedConstraintIndex]->SetConstraintIndex(SortedConstraintIndex);
 		}
 
+		Swap(Handles, SortedHandles);
 		Swap(ConstraintSettings, SortedConstraintSettings);
 		Swap(ConstraintParticles, SortedConstraintParticles);
 		Swap(ConstraintStates, SortedConstraintStates);
-		Swap(Handles, SortedHandles);
 	}
 
 
@@ -773,34 +792,19 @@ namespace Chaos
 
 		if (bJointsDirty)
 		{
-			ColorConstraints();
-			SortConstraints();
+			if (Settings.bSortEnabled)
+			{
+				ColorConstraints();
+				SortConstraints();
+			}
 
 			bJointsDirty = false;
 		}
-
-		if (Settings.bUseLinearSolver)
-		{
-			CachedConstraintSolvers.SetNum(NumConstraints());
-		}
-		else
-		{
-			ConstraintSolvers.SetNum(NumConstraints());		
-		}
 	}
 
 	void FPBDJointConstraints::UnprepareTick()
 	{
 		SCOPE_CYCLE_COUNTER(STAT_Joints_UnprepareTick);
-
-		if (Settings.bUseLinearSolver)
-		{
-			CachedConstraintSolvers.Empty();
-		}
-		else
-		{
-			ConstraintSolvers.Empty();
-		}
 	}
 	
 	void FPBDJointConstraints::CalculateConstraintSpace(int32 ConstraintIndex, FVec3& OutX0, FMatrix33& OutR0, FVec3& OutX1, FMatrix33& OutR1) const
@@ -823,302 +827,42 @@ namespace Chaos
 	}
 
 	void FPBDJointConstraints::AddConstraintsToGraph(FPBDIslandManager& IslandManager)
-	{ 
-		IslandManager.AddContainerConstraints(*this);
-	}
-
-	void FPBDJointConstraints::AddBodies(FSolverBodyContainer& SolverBodyContainer)
 	{
-		for (int32 ConstraintIndex = 0; ConstraintIndex < NumConstraints(); ++ConstraintIndex)
+		for (int32 ConstraintIndex = 0; ConstraintIndex < GetNumConstraints(); ++ConstraintIndex)
 		{
-			AddBodies(ConstraintIndex, SolverBodyContainer);
-		}
-	}
+			FPBDJointConstraintHandle* ConstraintHandle = GetConstraintHandle(ConstraintIndex);
+			check(ConstraintHandle != nullptr);
 
-	void FPBDJointConstraints::GatherInput(const FReal Dt)
-	{
-		for (int32 ConstraintIndex = 0; ConstraintIndex < NumConstraints(); ++ConstraintIndex)
-		{
-			GatherInput(ConstraintIndex, Dt);
-		}
-	}
+			const bool bIsInGraph = ConstraintHandle->IsInConstraintGraph();
+			const bool bShouldBeInGraph = IsConstraintEnabled(ConstraintIndex) && CanEvaluate(ConstraintIndex);
 
-	void FPBDJointConstraints::ScatterOutput(const FReal Dt)
-	{
-		for (int32 ConstraintIndex = 0; ConstraintIndex < NumConstraints(); ++ConstraintIndex)
-		{
-			ScatterOutput(ConstraintIndex, Dt);
-		}
-	}
-
-	void FPBDJointConstraints::ApplyPositionConstraints(const FReal Dt, const int32 It, const int32 NumIts)
-	{
-		for (int32 ConstraintIndex = 0; ConstraintIndex < NumConstraints(); ++ConstraintIndex)
-		{
-			ApplyPositionConstraint(Dt, ConstraintIndex, It, NumIts);
-		}
-	}
-
-	void FPBDJointConstraints::ApplyVelocityConstraints(const FReal Dt, const int32 It, const int32 NumIts)
-	{
-		for (int32 ConstraintIndex = 0; ConstraintIndex < NumConstraints(); ++ConstraintIndex)
-		{
-			ApplyVelocityConstraint(Dt, ConstraintIndex, It, NumIts);
-		}
-	}
-
-	void FPBDJointConstraints::ApplyProjectionConstraints(const FReal Dt, const int32 It, const int32 NumIts)
-	{
-		if (It == 0)
-		{
-			for (int32 ConstraintIndex = 0; ConstraintIndex < NumConstraints(); ++ConstraintIndex)
+			if (bShouldBeInGraph && !bIsInGraph)
 			{
-				PrepareProjectionConstraint(Dt, ConstraintIndex, It, NumIts);
+				IslandManager.AddConstraint(ContainerId, ConstraintHandle, GetConstrainedParticles(ConstraintIndex));
 			}
-		}
-
-		for (int32 ConstraintIndex = 0; ConstraintIndex < NumConstraints(); ++ConstraintIndex)
-		{
-			ApplyProjectionConstraint(Dt, ConstraintIndex, It, NumIts);
-		}
-	}
-
-	void FPBDJointConstraints::AddBodies(const TArrayView<int32>& ConstraintIndices, FSolverBodyContainer& SolverBodyContainer)
-	{
-		for (int32 ConstraintIndex : ConstraintIndices)
-		{
-			AddBodies(ConstraintIndex, SolverBodyContainer);
-		}
-	}
-
-	void FPBDJointConstraints::GatherInput(const TArrayView<int32>& ConstraintIndices, const FReal Dt)
-	{
-		for (int32 ConstraintIndex : ConstraintIndices)
-		{
-			GatherInput(ConstraintIndex, Dt);
-		}
-	}
-
-	void FPBDJointConstraints::ScatterOutput(const TArrayView<int32>& ConstraintIndices, const FReal Dt)
-	{
-		for (int32 ConstraintIndex : ConstraintIndices)
-		{
-			ScatterOutput(ConstraintIndex, Dt);
-		}
-	}
-
-	void FPBDJointConstraints::ApplyPositionConstraints(const TArrayView<int32>& ConstraintIndices, const FReal Dt, const int32 It, const int32 NumIts)
-	{
-		for (int32 ConstraintIndex : ConstraintIndices)
-		{
-			ApplyPositionConstraint(Dt, ConstraintIndex, It, NumIts);
-		}
-	}
-
-	void FPBDJointConstraints::ApplyVelocityConstraints(const TArrayView<int32>& ConstraintIndices, const FReal Dt, const int32 It, const int32 NumIts)
-	{
-		for (int32 ConstraintIndex : ConstraintIndices)
-		{
-			ApplyVelocityConstraint(Dt, ConstraintIndex, It, NumIts);
-		}
-	}
-
-	void FPBDJointConstraints::ApplyProjectionConstraints(const TArrayView<int32>& ConstraintIndices, const FReal Dt, const int32 It, const int32 NumIts)
-	{
-		if (It == 0)
-		{
-			for (int32 ConstraintIndex : ConstraintIndices)
+			else if (bIsInGraph && !bShouldBeInGraph)
 			{
-				PrepareProjectionConstraint(Dt, ConstraintIndex, It, NumIts);
-			}
-		}
-
-		for (int32 ConstraintIndex : ConstraintIndices)
-		{
-			ApplyProjectionConstraint(Dt, ConstraintIndex, It, NumIts);
-		}
-	}
-
-	//////////////////////////////////////////////////////////////////////////
-	//
-	// Begin single-particle solve methods
-	//
-	//////////////////////////////////////////////////////////////////////////
-
-	void FPBDJointConstraints::AddBodies(const int32 ConstraintIndex, FSolverBodyContainer& SolverBodyContainer)
-	{
-		if (ConstraintStates[ConstraintIndex].bDisabled)
-		{
-			return;
-		}
-
-		int32 Index0, Index1;
-		GetConstrainedParticleIndices(ConstraintIndex, Index0, Index1);
-		FGenericParticleHandle Particle0 = FGenericParticleHandle(ConstraintParticles[ConstraintIndex][Index0]);
-		FGenericParticleHandle Particle1 = FGenericParticleHandle(ConstraintParticles[ConstraintIndex][Index1]);
-
-		// Find the solver bodies for the particles we constrain. This will add them to the container
-		// if they aren't there already, and ensure that they are populated with the latest data.
-		FSolverBody* SolverBody0 = SolverBodyContainer.FindOrAdd(Particle0);
-		FSolverBody* SolverBody1 = SolverBodyContainer.FindOrAdd(Particle1);
-
-		if (Settings.bUseLinearSolver)
-		{
-			FPBDJointCachedSolver& Solver = CachedConstraintSolvers[ConstraintIndex];
-			Solver.SetSolverBodies(SolverBody0, SolverBody1);
-		}
-		else
-		{
-			FPBDJointSolver& Solver = ConstraintSolvers[ConstraintIndex];
-			Solver.SetSolverBodies(SolverBody0, SolverBody1);
-		}
-	}
-
-	void FPBDJointConstraints::GatherInput(const int32 ConstraintIndex, const FReal Dt)
-	{
-		if (ConstraintStates[ConstraintIndex].bDisabled)
-		{
-			return;
-		}
-
-		const FPBDJointSettings& JointSettings = ConstraintSettings[ConstraintIndex];
-		const FTransformPair& JointFrames = JointSettings.ConnectorTransforms;
-
-		int32 Index0, Index1;
-		GetConstrainedParticleIndices(ConstraintIndex, Index0, Index1);
-		FGenericParticleHandle Particle0 = FGenericParticleHandle(ConstraintParticles[ConstraintIndex][Index0]);
-		FGenericParticleHandle Particle1 = FGenericParticleHandle(ConstraintParticles[ConstraintIndex][Index1]);
-
-		if (Settings.bUseLinearSolver)
-		{
-			FPBDJointCachedSolver& Solver = CachedConstraintSolvers[ConstraintIndex];
-			Solver.Init(
-				Dt,
-				Settings,
-				JointSettings,
-				FParticleUtilities::ParticleLocalToCoMLocal(Particle0, JointFrames[Index0]),
-				FParticleUtilities::ParticleLocalToCoMLocal(Particle1, JointFrames[Index1]));
-		}
-		else
-		{
-			FPBDJointSolver& Solver = ConstraintSolvers[ConstraintIndex];
-			Solver.Init(
-				Dt,
-				Settings,
-				JointSettings,
-				FParticleUtilities::ParticleLocalToCoMLocal(Particle0, JointFrames[Index0]),
-				FParticleUtilities::ParticleLocalToCoMLocal(Particle1, JointFrames[Index1]));
-		}
-
-		// Plasticity should not be turned on in the middle of simulation.
-		const bool bUseLinearPlasticity = JointSettings.LinearPlasticityLimit != FLT_MAX;
-		if (bUseLinearPlasticity)
-		{
-			const bool bIsCOMDistanceInitialized = !FMath::IsNearlyEqual(JointSettings.LinearPlasticityInitialDistanceSquared, (FReal)FLT_MAX);
-			if (!bIsCOMDistanceInitialized)
-			{
-				// Joint plasticity is baed on the distance of one of the moment arms of the joint. Typically, plasticity
-				// will get setup from the joint pivot to the child COM (centor of mass), so that is found first. However, when 
-				// the pivot is at the child COM then we fall back to the distance between thge pivot and parent COM.
-				ConstraintSettings[ConstraintIndex].LinearPlasticityInitialDistanceSquared = JointSettings.ConnectorTransforms[Index1].GetTranslation().SizeSquared();
-				if (FMath::IsNearlyZero(ConstraintSettings[ConstraintIndex].LinearPlasticityInitialDistanceSquared))
-				{
-					ConstraintSettings[ConstraintIndex].LinearPlasticityInitialDistanceSquared = JointSettings.ConnectorTransforms[Index0].GetTranslation().SizeSquared();
-				}
-				// @todo(chaos): move this to validation
-				ensureMsgf(!FMath::IsNearlyZero(ConstraintSettings[ConstraintIndex].LinearPlasticityInitialDistanceSquared), TEXT("Plasticity made inactive due to Zero length difference between parent and child rigid body."));
+				IslandManager.RemoveConstraint(ContainerId, ConstraintHandle);
 			}
 		}
 	}
 
-	void FPBDJointConstraints::ScatterOutput(const int32 ConstraintIndex, FReal Dt)
+	void FPBDJointConstraints::SetSolverResults(const int32 ConstraintIndex, const FVec3& LinearImpulse, const FVec3& AngularImpulse, bool bIsBroken, const FSolverBody* SolverBody0, const FSolverBody* SolverBody1)
 	{
-		SCOPE_CYCLE_COUNTER(STAT_Joints_Scatter);
+		ConstraintStates[ConstraintIndex].LinearImpulse = LinearImpulse;
+		ConstraintStates[ConstraintIndex].AngularImpulse = AngularImpulse;
 
-		if (ConstraintStates[ConstraintIndex].bDisabled)
+		if (bIsBroken)
 		{
-			return;
-		}
-
-		FPBDJointState& JointState = ConstraintStates[ConstraintIndex];
-
-		int32 Index0, Index1;
-		GetConstrainedParticleIndices(ConstraintIndex, Index0, Index1);
-
-		if (Settings.bUseLinearSolver)
-		{
-			FPBDJointCachedSolver& Solver = CachedConstraintSolvers[ConstraintIndex];
-			// NOTE: LinearImpulse/AngularImpulse in the solver are not really impulses - they are mass-weighted position/rotation delta, or (impulse x dt).
-			if (Dt > UE_SMALL_NUMBER)
-			{
-				if (Index0 == 0)
-				{
-					JointState.LinearImpulse = Solver.GetNetLinearImpulse() / Dt;
-					JointState.AngularImpulse = Solver.GetNetAngularImpulse() / Dt;
-				}
-				else
-				{
-					// Particles were flipped in the solver...
-					JointState.LinearImpulse = -Solver.GetNetLinearImpulse() / Dt;
-					JointState.AngularImpulse = -Solver.GetNetAngularImpulse() / Dt;
-				}
-			}
-			else
-			{
-				JointState.LinearImpulse = FVec3(0);
-				JointState.AngularImpulse = FVec3(0);
-			}
-
-			ApplyPlasticityLimits(ConstraintIndex);
-
-			// Remove our solver body reference (they are not valid between frames)
-			CachedConstraintSolvers[ConstraintIndex].Deinit();
-		}
-		else
-		{
-			FPBDJointSolver& Solver = ConstraintSolvers[ConstraintIndex];
-			// NOTE: LinearImpulse/AngularImpulse in the solver are not really impulses - they are mass-weighted position/rotation delta, or (impulse x dt).
-			if (Dt > UE_SMALL_NUMBER)
-			{
-				if (Index0 == 0)
-				{
-					JointState.LinearImpulse = Solver.GetNetLinearImpulse() / Dt;
-					JointState.AngularImpulse = Solver.GetNetAngularImpulse() / Dt;
-				}
-				else
-				{
-					// Particles were flipped in the solver...
-					JointState.LinearImpulse = -Solver.GetNetLinearImpulse() / Dt;
-					JointState.AngularImpulse = -Solver.GetNetAngularImpulse() / Dt;
-				}
-			}
-			else
-			{
-				JointState.LinearImpulse = FVec3(0);
-				JointState.AngularImpulse = FVec3(0);
-			}
-
-			ApplyPlasticityLimits(ConstraintIndex);
-
-			// Remove our solver body reference (they are not valid between frames)
-			ConstraintSolvers[ConstraintIndex].Deinit();
+			BreakConstraint(ConstraintIndex);
 		}
-	}
 
-	FReal FPBDJointConstraints::CalculateIterationStiffness(int32 It, int32 NumIts) const
-	{
-		// Linearly interpolate betwwen MinStiffness and MaxStiffness over the first few iterations,
-		// then clamp at MaxStiffness for the final NumIterationsAtMaxStiffness
-		FReal IterationStiffness = Settings.MaxSolverStiffness;
-		if (NumIts > Settings.NumIterationsAtMaxSolverStiffness)
+		if ((SolverBody0 != nullptr) && (SolverBody1 != nullptr))
 		{
-			const FReal Interpolant = FMath::Clamp((FReal)It / (FReal)(NumIts - Settings.NumIterationsAtMaxSolverStiffness), 0.0f, 1.0f);
-			IterationStiffness = FMath::Lerp(Settings.MinSolverStiffness, Settings.MaxSolverStiffness, Interpolant);
+			ApplyPlasticityLimits(ConstraintIndex, *SolverBody0, *SolverBody1);
 		}
-		return FMath::Clamp(IterationStiffness, 0.0f, 1.0f);
 	}
 
-
 	bool FPBDJointConstraints::CanEvaluate(const int32 ConstraintIndex) const
 	{
 		if (!IsConstraintEnabled(ConstraintIndex))
@@ -1138,317 +882,18 @@ namespace Chaos
 			return false;
 		}
 
-		// check valid particle and solver state
-		if (Settings.bUseLinearSolver)
-		{
-			const FPBDJointCachedSolver& Solver = CachedConstraintSolvers[ConstraintIndex];
-			if ((Particle0->Sleeping() && Particle1->Sleeping())
-				|| (Particle0->IsKinematic() && Particle1->Sleeping())
-				|| (Particle0->Sleeping() && Particle1->IsKinematic())
-				|| (FMath::IsNearlyZero(Solver.InvM(0)) && FMath::IsNearlyZero(Solver.InvM(1))))
-			{
-				return false;
-			}
-		}
-		else
-		{
-			const FPBDJointSolver& Solver = ConstraintSolvers[ConstraintIndex];
-			if ((Particle0->Sleeping() && Particle1->Sleeping())
-				|| (Particle0->IsKinematic() && Particle1->Sleeping())
-				|| (Particle0->Sleeping() && Particle1->IsKinematic())
-				|| (FMath::IsNearlyZero(Solver.InvM(0)) && FMath::IsNearlyZero(Solver.InvM(1))))
-			{
-				return false;
-			}
-		}
-		return true;
-	}
-
-	// @todo(chaos): ShockPropagation needs to handle the parent/child being in opposite order
-	FReal FPBDJointConstraints::CalculateShockPropagationInvMassScale(const FConstraintSolverBody& Body0, const FConstraintSolverBody& Body1, const FPBDJointSettings& JointSettings, const int32 It, const int32 NumIts) const
-	{
-		// Shock propagation is only enabled for the last iteration, and only for the QPBD solver.
-		// The standard PBD solver runs projection in the second solver phase which is mostly the same thing.
-		if (JointSettings.bShockPropagationEnabled && (It >= (NumIts - Settings.NumShockPropagationIterations)))
-		{
-			if (Body0.IsDynamic() && Body1.IsDynamic())
-			{
-				return FPBDJointUtilities::GetShockPropagationInvMassScale(Settings, JointSettings);
-			}
-		}
-		return FReal(1);
-	}
-
-	// This position solver iterates over each of the inner constraints (position, twist, swing) and solves them independently.
-	// This will converge slowly in some cases, particularly where resolving angular constraints violates position constraints and vice versa.
-	bool FPBDJointConstraints::ApplyPositionConstraint(const FReal Dt, const int32 ConstraintIndex, const int32 It, const int32 NumIts)
-	{
-		if (!CanEvaluate(ConstraintIndex))
-		{
-			return false;
-		}
-
-		const TVector<TGeometryParticleHandle<FReal, 3>*, 2>& Constraint = ConstraintParticles[ConstraintIndex];
-		UE_LOG(LogChaosJoint, VeryVerbose, TEXT("Solve Joint Position Constraint %d %s %s (dt = %f; it = %d / %d)"), ConstraintIndex, *Constraint[0]->ToString(), *Constraint[1]->ToString(), Dt, It, NumIts);
-
-		// @todo(chaos): store this on the Solver object and don't access the particles here
-		int32 Index0, Index1;
-		GetConstrainedParticleIndices(ConstraintIndex, Index0, Index1);
-		FGenericParticleHandle Particle0 = FGenericParticleHandle(ConstraintParticles[ConstraintIndex][Index0]);
-		FGenericParticleHandle Particle1 = FGenericParticleHandle(ConstraintParticles[ConstraintIndex][Index1]);
 		if ((Particle0->Sleeping() && Particle1->Sleeping())
 			|| (Particle0->IsKinematic() && Particle1->Sleeping())
-			|| (Particle0->Sleeping() && Particle1->IsKinematic()))
+			|| (Particle0->Sleeping() && Particle1->IsKinematic())
+			|| (Particle0->IsKinematic() && Particle1->IsKinematic()))
 		{
 			return false;
 		}
 
-		const FPBDJointSettings& JointSettings = ConstraintSettings[ConstraintIndex];
-		if (Settings.bUseLinearSolver)
-		{
-			FPBDJointCachedSolver& Solver = CachedConstraintSolvers[ConstraintIndex];
-			if ((FMath::IsNearlyZero(Solver.InvM(0)) && FMath::IsNearlyZero(Solver.InvM(1))))
-			{
-				return false;
-			}
-
-			Solver.Update(Dt, Settings, JointSettings);
-
-			// Set parent inverse mass scale based on current shock propagation state
-			const FReal ShockPropagationInvMassScale = CalculateShockPropagationInvMassScale(Solver.Body0(), Solver.Body1(), JointSettings, It, NumIts);
-			Solver.SetShockPropagationScales(ShockPropagationInvMassScale, FReal(1), Dt);
-
-			const FReal IterationStiffness = CalculateIterationStiffness(It, NumIts);
-			Solver.ApplyConstraints(Dt, IterationStiffness, Settings, JointSettings);
-
-			// @todo(ccaulfield): The break limit should really be applied to the impulse in the solver to prevent 1-frame impulses larger than the threshold
-			if ((JointSettings.LinearBreakForce!=FLT_MAX) || (JointSettings.AngularBreakTorque!=FLT_MAX))
-			{
-				ApplyBreakThreshold(Dt, ConstraintIndex, Solver.GetNetLinearImpulse(), Solver.GetNetAngularImpulse());
-			}
-		}
-		else
-		{
-			FPBDJointSolver& Solver = ConstraintSolvers[ConstraintIndex];
-			if ((FMath::IsNearlyZero(Solver.InvM(0)) && FMath::IsNearlyZero(Solver.InvM(1))))
-			{
-				return false;
-			}
-
-			Solver.Update(Dt, Settings, JointSettings);
-
-			// Set parent inverse mass scale based on current shock propagation state
-			const FReal ShockPropagationInvMassScale = CalculateShockPropagationInvMassScale(Solver.Body0(), Solver.Body1(), JointSettings, It, NumIts);
-
-			const FReal IterationStiffness = CalculateIterationStiffness(It, NumIts);
-			Solver.SetShockPropagationScales(ShockPropagationInvMassScale, FReal(1));
-
-			Solver.ApplyConstraints(Dt, IterationStiffness, Settings, JointSettings);
-
-			// @todo(ccaulfield): The break limit should really be applied to the impulse in the solver to prevent 1-frame impulses larger than the threshold
-			if ((JointSettings.LinearBreakForce!=FLT_MAX) || (JointSettings.AngularBreakTorque!=FLT_MAX))
-			{
-				ApplyBreakThreshold(Dt, ConstraintIndex, Solver.GetNetLinearImpulse(), Solver.GetNetAngularImpulse());
-			}
-		}
-
-		return true;
-	}
-
-	// QuasiPBD applies a velocity solve in phase 2
-	// Standard PBD does nothing
-	bool FPBDJointConstraints::ApplyVelocityConstraint(const FReal Dt, const int32 ConstraintIndex, const int32 It, const int32 NumIts)
-	{
-		if (!CanEvaluate(ConstraintIndex))
-		{
-			return false;
-		}
-
-		const TVector<TGeometryParticleHandle<FReal, 3>*, 2>& Constraint = ConstraintParticles[ConstraintIndex];
-		UE_LOG(LogChaosJoint, VeryVerbose, TEXT("Solve Joint Velocity Constraint %d %s %s (dt = %f; it = %d / %d)"), ConstraintIndex, *Constraint[0]->ToString(), *Constraint[1]->ToString(), Dt, It, NumIts);
-
-		// @todo(chaos): store this on the Solver object and don't access the particles here
-		int32 Index0, Index1;
-		GetConstrainedParticleIndices(ConstraintIndex, Index0, Index1);
-		FGenericParticleHandle Particle0 = FGenericParticleHandle(ConstraintParticles[ConstraintIndex][Index0]);
-		FGenericParticleHandle Particle1 = FGenericParticleHandle(ConstraintParticles[ConstraintIndex][Index1]);
-		if ((Particle0->Sleeping() && Particle1->Sleeping())
-			|| (Particle0->IsKinematic() && Particle1->Sleeping())
-			|| (Particle0->Sleeping() && Particle1->IsKinematic()))
-		{
-			return false;
-		}
-
-		const FPBDJointSettings& JointSettings = ConstraintSettings[ConstraintIndex];
-		if (Settings.bUseLinearSolver)
-		{
-			FPBDJointCachedSolver& Solver = CachedConstraintSolvers[ConstraintIndex];
-			if ((FMath::IsNearlyZero(Solver.InvM(0)) && FMath::IsNearlyZero(Solver.InvM(1))))
-			{
-				return false;
-			}
-
-			Solver.Update(Dt, Settings, JointSettings);
-			
-			// Set parent inverse mass scale based on current shock propagation state
-			const FReal ShockPropagationInvMassScale = CalculateShockPropagationInvMassScale(Solver.Body0(), Solver.Body1(), JointSettings, It, NumIts);
-			Solver.SetShockPropagationScales(ShockPropagationInvMassScale, FReal(1), Dt);
-
-			const FReal IterationStiffness = CalculateIterationStiffness(It, NumIts);
-			Solver.ApplyVelocityConstraints(Dt, IterationStiffness, Settings, JointSettings);
-
-			// @todo(ccaulfield): should probably add to net impulses in push out too...(for breaking etc)
-		}
-		else
-		{
-			FPBDJointSolver& Solver = ConstraintSolvers[ConstraintIndex];
-			if ((FMath::IsNearlyZero(Solver.InvM(0)) && FMath::IsNearlyZero(Solver.InvM(1))))
-			{
-				return false;
-			}
-
-			Solver.Update(Dt, Settings, JointSettings);
-
-			// Set parent inverse mass scale based on current shock propagation state
-			const FReal ShockPropagationInvMassScale = CalculateShockPropagationInvMassScale(Solver.Body0(), Solver.Body1(), JointSettings, It, NumIts);
-			Solver.SetShockPropagationScales(ShockPropagationInvMassScale, FReal(1));
-
-			const FReal IterationStiffness = CalculateIterationStiffness(It, NumIts);
-			Solver.ApplyVelocityConstraints(Dt, IterationStiffness, Settings, JointSettings);
-
-			// @todo(ccaulfield): should probably add to net impulses in push out too...(for breaking etc)
-		}
-
 		return true;
 	}
 
-	void FPBDJointConstraints::PrepareProjectionConstraint(const FReal Dt, const int32 ConstraintIndex, const int32 It, const int32 NumIts)
-	{
-		// Collect all the data for projection prior to the first iteration. 
-		// This must happen for all joints before we project any joints so the the initial state for each joint is not polluted by any earlier projections.
-		// @todo(chaos): if we ever support projection on other constraint types, we will need a PrepareProjection phase so that all constraint types
-		// can initialize correctly before any constraints apply their projection. For now we can just check the iteration count is zero.
-		check(It == 0);
-
-		if (Settings.bUseLinearSolver)
-		{
-			if (!CanEvaluate(ConstraintIndex))
-			{
-				return;
-			}
-
-			FPBDJointCachedSolver& Solver = CachedConstraintSolvers[ConstraintIndex];
-			const FPBDJointSettings& JointSettings = ConstraintSettings[ConstraintIndex];
-			if (JointSettings.bProjectionEnabled && (Solver.IsDynamic(0) || Solver.IsDynamic(1)))
-			{
-				Solver.InitProjection(Dt, Settings, JointSettings);
-			}
-		}
-	}
-
-	// Projection phase
-	bool FPBDJointConstraints::ApplyProjectionConstraint(const FReal Dt, const int32 ConstraintIndex, const int32 It, const int32 NumIts)
-	{
-		if (!CanEvaluate(ConstraintIndex))
-		{
-			return false;
-		}
-
-		const TVector<TGeometryParticleHandle<FReal, 3>*, 2>& Constraint = ConstraintParticles[ConstraintIndex];
-		UE_LOG(LogChaosJoint, VeryVerbose, TEXT("Project Joint Constraint %d %s %s (dt = %f; it = %d / %d)"), ConstraintIndex, *Constraint[0]->ToString(), *Constraint[1]->ToString(), Dt, It, NumIts);
-
-		const FPBDJointSettings& JointSettings = ConstraintSettings[ConstraintIndex];
-		if (!JointSettings.bProjectionEnabled)
-		{
-			return false;
-		}
-
-		// @todo(chaos): store this on the Solver object and don't access the particles here
-		int32 Index0, Index1;
-		GetConstrainedParticleIndices(ConstraintIndex, Index0, Index1);
-		FGenericParticleHandle Particle0 = FGenericParticleHandle(ConstraintParticles[ConstraintIndex][Index0]);
-		FGenericParticleHandle Particle1 = FGenericParticleHandle(ConstraintParticles[ConstraintIndex][Index1]);
-		if ((Particle0->Sleeping() && Particle1->Sleeping())
-			|| (Particle0->IsKinematic() && Particle1->Sleeping())
-			|| (Particle0->Sleeping() && Particle1->IsKinematic()))
-		{
-			return false;
-		}
-
-		if (Settings.bUseLinearSolver)
-		{
-			FPBDJointCachedSolver& Solver = CachedConstraintSolvers[ConstraintIndex];
-			if ((FMath::IsNearlyZero(Solver.InvM(0)) && FMath::IsNearlyZero(Solver.InvM(1))))
-			{
-				return false;
-			}
-
-			if (It == 0)
-			{
-				Solver.ApplyTeleports(Dt, Settings, JointSettings);
-			}
-
-			Solver.ApplyProjections(Dt, Settings, JointSettings, (It==(NumIts-1)));
-		}
-		else
-		{
-			FPBDJointSolver& Solver = ConstraintSolvers[ConstraintIndex];
-			if ((FMath::IsNearlyZero(Solver.InvM(0)) && FMath::IsNearlyZero(Solver.InvM(1))))
-			{
-				return false;
-			}
-
-			Solver.Update(Dt, Settings, JointSettings);
-
-			if (It == 0)
-			{
-				// @todo(chaos): support reverse parent/child
-				Solver.Body1().UpdateRotationDependentState();
-				Solver.UpdateMasses(FReal(0), FReal(1));
-			}
-
-			Solver.ApplyProjections(Dt, Settings, JointSettings);
-		}
-
-		return true;
-	}
-
-
-	void FPBDJointConstraints::ApplyBreakThreshold(const FReal Dt, int32 ConstraintIndex, const FVec3& LinearImpulse, const FVec3& AngularImpulse)
-	{
-		const FPBDJointSettings& JointSettings = ConstraintSettings[ConstraintIndex];
-
-		// NOTE: LinearImpulse/AngularImpulse are not really impulses - they are mass-weighted position/rotation delta, or (impulse x dt).
-		// The Threshold is a force limit, so we need to convert it to a position delta caused by that force in one timestep
-
-		bool bBreak = false;
-		if (!bBreak && JointSettings.LinearBreakForce!=FLT_MAX)
-		{
-			const FReal LinearForceSq = LinearImpulse.SizeSquared() / (Dt * Dt * Dt * Dt);
-			const FReal LinearThresholdSq = FMath::Square(JointSettings.LinearBreakForce);
-
-			UE_LOG(LogChaosJoint, VeryVerbose, TEXT("Constraint %d Linear Break Check: %f / %f at Dt = %f"), ConstraintIndex, FMath::Sqrt(LinearForceSq), FMath::Sqrt(LinearThresholdSq), Dt);
-
-			bBreak = LinearForceSq > LinearThresholdSq;
-		}
-
-		if (!bBreak && JointSettings.AngularBreakTorque!=FLT_MAX)
-		{
-			const FReal AngularForceSq = AngularImpulse.SizeSquared() / (Dt * Dt * Dt * Dt);
-			const FReal AngularThresholdSq = FMath::Square(JointSettings.AngularBreakTorque);
-			UE_LOG(LogChaosJoint, VeryVerbose, TEXT("Constraint %d Angular Break Check: %f / %f at Dt = %f"), ConstraintIndex, FMath::Sqrt(AngularForceSq), FMath::Sqrt(AngularThresholdSq), Dt);
-
-			bBreak = AngularForceSq > AngularThresholdSq;
-		}
-
-		if (bBreak)
-		{
-			BreakConstraint(ConstraintIndex);
-		}
-	}
-
-
-	void FPBDJointConstraints::ApplyPlasticityLimits(int32 ConstraintIndex)
+	void FPBDJointConstraints::ApplyPlasticityLimits(int32 ConstraintIndex, const FSolverBody& SolverBody0, const FSolverBody& SolverBody1)
 	{
 		FPBDJointSettings& JointSettings = ConstraintSettings[ConstraintIndex];
 		const bool bHasLinearPlasticityLimit = JointSettings.LinearPlasticityLimit != FLT_MAX;
@@ -1464,6 +909,26 @@ namespace Chaos
 			return;
 		}
 
+		// @todo(chaos): this should be done when the joint transforms are initialized
+		// Plasticity should not be turned on in the middle of simulation.
+		if (bHasLinearPlasticityLimit)
+		{
+			const bool bIsCOMDistanceInitialized = !FMath::IsNearlyEqual(JointSettings.LinearPlasticityInitialDistanceSquared, (FReal)FLT_MAX);
+			if (!bIsCOMDistanceInitialized)
+			{
+				// Joint plasticity is based on the distance of one of the moment arms of the joint. Typically, plasticity
+				// will get setup from the joint pivot to the child COM (centor of mass), so that is found first. However, when 
+				// the pivot is at the child COM then we fall back to the distance between thge pivot and parent COM.
+				ConstraintSettings[ConstraintIndex].LinearPlasticityInitialDistanceSquared = JointSettings.ConnectorTransforms[0].GetTranslation().SizeSquared();
+				if (FMath::IsNearlyZero(ConstraintSettings[ConstraintIndex].LinearPlasticityInitialDistanceSquared))
+				{
+					ConstraintSettings[ConstraintIndex].LinearPlasticityInitialDistanceSquared = JointSettings.ConnectorTransforms[1].GetTranslation().SizeSquared();
+				}
+				// @todo(chaos): move this to validation
+				ensureMsgf(!FMath::IsNearlyZero(ConstraintSettings[ConstraintIndex].LinearPlasticityInitialDistanceSquared), TEXT("Plasticity made inactive due to Zero length difference between parent and child rigid body."));
+			}
+		}
+
 		int32 Index0, Index1;
 		GetConstrainedParticleIndices(ConstraintIndex, Index0, Index1);
 		{
@@ -1475,11 +940,8 @@ namespace Chaos
 			}
 		}
 
-		FConstraintSolverBody& Body0 = Settings.bUseLinearSolver ? CachedConstraintSolvers[ConstraintIndex].Body(0) : ConstraintSolvers[ConstraintIndex].Body(0);
-		FConstraintSolverBody& Body1 = Settings.bUseLinearSolver ? CachedConstraintSolvers[ConstraintIndex].Body(1) : ConstraintSolvers[ConstraintIndex].Body(1);
-
 		const FTransformPair& ConstraintFramesLocal = JointSettings.ConnectorTransforms;
-		FTransformPair ConstraintFramesGlobal(ConstraintFramesLocal[Index0] * FRigidTransform3(Body0.ActorP(), Body0.ActorQ()), ConstraintFramesLocal[Index1] * FRigidTransform3(Body1.ActorP(), Body1.ActorQ()));
+		FTransformPair ConstraintFramesGlobal(ConstraintFramesLocal[Index0] * FRigidTransform3(SolverBody0.ActorP(), SolverBody0.ActorQ()), ConstraintFramesLocal[Index1] * FRigidTransform3(SolverBody1.ActorP(), SolverBody1.ActorQ()));
 		FQuat Q1 = ConstraintFramesGlobal[1].GetRotation();
 		Q1.EnforceShortestArcWith(ConstraintFramesGlobal[0].GetRotation());
 		ConstraintFramesGlobal[1].SetRotation(Q1);
@@ -1510,7 +972,7 @@ namespace Chaos
 					// Shrink and Grow are based on the distance between the joint pivot and the child. 
 					// Note, if the pivot is located at the COM of the child then shrink will not do anything. 
 					FVec3 StartDelta = ConstraintFramesLocal[Index1].InverseTransformPositionNoScale(JointSettings.LinearDrivePositionTarget);
-					FVec3 CurrentDelta = ConstraintFramesGlobal[Index1].InverseTransformPositionNoScale(Body1.P());
+					FVec3 CurrentDelta = ConstraintFramesGlobal[Index1].InverseTransformPositionNoScale(SolverBody1.P());
 
 					if (JointSettings.LinearPlasticityType == EPlasticityType::Shrink && CurrentDelta.SizeSquared() < StartDelta.SizeSquared())
 					{
diff --git a/Engine/Source/Runtime/Experimental/Chaos/Private/PBDRigidsSolver.cpp b/Engine/Source/Runtime/Experimental/Chaos/Private/PBDRigidsSolver.cpp
--- a/Engine/Source/Runtime/Experimental/Chaos/Private/PBDRigidsSolver.cpp
+++ b/Engine/Source/Runtime/Experimental/Chaos/Private/PBDRigidsSolver.cpp
@@ -892,6 +892,7 @@ namespace Chaos
 		JointsSettings.NumShockPropagationIterations = ChaosSolverJointNumShockProagationIterations;
 		JointsSettings.ShockPropagationOverride = ChaosSolverJointShockPropagation;
 		JointsSettings.bUseLinearSolver = bChaosSolverJointUseLinearSolver;
+		JointsSettings.bSortEnabled = false;
 		MEvolution->GetJointConstraints().SetSettings(JointsSettings);
 
 		// Apply CVAR overrides if set
diff --git a/Engine/Source/Runtime/Experimental/Chaos/Public/Chaos/Evolution/IndexedConstraintContainer.h b/Engine/Source/Runtime/Experimental/Chaos/Public/Chaos/Evolution/IndexedConstraintContainer.h
--- a/Engine/Source/Runtime/Experimental/Chaos/Public/Chaos/Evolution/IndexedConstraintContainer.h
+++ b/Engine/Source/Runtime/Experimental/Chaos/Public/Chaos/Evolution/IndexedConstraintContainer.h
@@ -42,6 +42,11 @@ namespace Chaos
 			return ConstraintIndex;
 		}
 
+		void SetConstraintIndex(const int32 InConstraintIndex)
+		{
+			ConstraintIndex = InConstraintIndex;
+		}
+
 		static const FConstraintHandleTypeID& StaticType()
 		{
 			static FConstraintHandleTypeID STypeID(TEXT("FIndexedConstraintHandle"), &FConstraintHandle::StaticType());
diff --git a/Engine/Source/Runtime/Experimental/Chaos/Public/Chaos/Joint/PBDJointCachedSolverGaussSeidel.h b/Engine/Source/Runtime/Experimental/Chaos/Public/Chaos/Joint/PBDJointCachedSolverGaussSeidel.h
--- a/Engine/Source/Runtime/Experimental/Chaos/Public/Chaos/Joint/PBDJointCachedSolverGaussSeidel.h
+++ b/Engine/Source/Runtime/Experimental/Chaos/Public/Chaos/Joint/PBDJointCachedSolverGaussSeidel.h
@@ -191,7 +191,9 @@ struct FAxisConstraintDatas
 			return bIsActive;
 		}
 		
-		FPBDJointCachedSolver();
+		FPBDJointCachedSolver()
+		{
+		}
 
 		void SetSolverBodies(FSolverBody* SolverBody0, FSolverBody* SolverBody1)
 		{
@@ -272,6 +274,21 @@ struct FAxisConstraintDatas
 			const FReal InvMScale1,
 			const FReal Dt);
 
+		void SetIsBroken(const bool bInIsBroken)
+		{
+			bIsBroken = bInIsBroken;
+		}
+
+		bool IsBroken() const
+		{
+			return bIsBroken;
+		}
+
+		bool RequiresSolve() const
+		{
+			return !IsBroken() && (IsDynamic(0) || IsDynamic(1));
+		}
+
 	private:
 
 		// Initialize state dependent on particle positions
@@ -584,6 +601,8 @@ struct FAxisConstraintDatas
 		FAxisConstraintDatas PositionDrives;
 		FAxisConstraintDatas RotationDrives;
 
+		bool bIsBroken;
+
 		// dummy indices
 		static constexpr int32 PointPositionConstraintIndex = 0;
 		static constexpr int32 SphericalPositionConstraintIndex = 0;
diff --git a/Engine/Source/Runtime/Experimental/Chaos/Public/Chaos/Joint/PBDJointContainerSolver.h b/Engine/Source/Runtime/Experimental/Chaos/Public/Chaos/Joint/PBDJointContainerSolver.h
new file mode 100644
--- /dev/null
+++ b/Engine/Source/Runtime/Experimental/Chaos/Public/Chaos/Joint/PBDJointContainerSolver.h
@@ -0,0 +1,67 @@
+// Copyright Epic Games, Inc. All Rights Reserved.
+#pragma once
+
+#include "Chaos/Core.h"
+#include "Chaos/Evolution/SolverConstraintContainer.h"
+#include "Chaos/Joint/PBDJointSolverGaussSeidel.h"
+#include "Chaos/Joint/PBDJointCachedSolverGaussSeidel.h"
+#include "Chaos/PBDJointConstraints.h"
+#include "Chaos/PBDJointConstraintTypes.h"
+
+namespace Chaos
+{
+	namespace Private
+	{
+		/**
+		 * Runs the solvers for a set of constraints belonging to a JointConstraints container.
+		 * 
+		 * For the main scene, each IslandGroup owns a FPBDJointContainerSolver and the list of constraints to be solved
+		 * and the order in which they are solved is determined by the constraint graph.
+		 * 
+		 * For RBAN, there is one FPBDJointContainerSolver that solves all joints in the simulation in the order that
+		 * they occur in the container.
+		*/
+		class FPBDJointContainerSolver : public FConstraintContainerSolver
+		{
+		public:
+			FPBDJointContainerSolver(FPBDJointConstraints& InConstraintContainer, const int32 InPriority);
+			~FPBDJointContainerSolver();
+
+			// FConstraintContainerSolver impl
+			virtual int32 GetNumConstraints() const override final { return ContainerConstraintIndices.Num(); }
+			virtual void Reset(const int32 InMaxCollisions) override final;
+			virtual void AddConstraints() override final;
+			virtual void AddConstraints(const TArrayView<FPBDIslandConstraint>& IslandConstraints) override final;
+			virtual void AddBodies(FSolverBodyContainer& SolverBodyContainer) override final;
+			virtual void GatherInput(const FReal Dt) override final;
+			virtual void GatherInput(const FReal Dt, const int32 BeginIndex, const int32 EndIndex) override final;
+			virtual void ScatterOutput(const FReal Dt) override final;
+			virtual void ScatterOutput(const FReal Dt, const int32 BeginIndex, const int32 EndIndex) override final;
+			virtual void ApplyPositionConstraints(const FReal Dt, const int32 It, const int32 NumIts) override final;
+			virtual void ApplyVelocityConstraints(const FReal Dt, const int32 It, const int32 NumIts) override final;
+			virtual void ApplyProjectionConstraints(const FReal Dt, const int32 It, const int32 NumIts) override final;
+
+			FPBDJointConstraints& GetContainer() const { return ConstraintContainer; }
+			const FPBDJointSolverSettings& GetSettings() const { return ConstraintContainer.GetSettings(); }
+			const FPBDJointSettings& GetConstraintSettings(const int32 InConstraintIndex) const { return ConstraintContainer.GetConstraintSettings(ContainerConstraintIndices[InConstraintIndex]); }
+			int32 GetContainerConstraintIndex(const int32 InConstraintIndex) const { return ContainerConstraintIndices[InConstraintIndex]; }
+
+		private:
+			bool UseLinearSolver() const;
+			void AddConstraint(const int32 InContainerConstraintIndex);
+			void ResizeSolverArrays();
+			void ApplyLinearProjectionConstraints(const FReal Dt, const int32 It, const int32 NumIts);
+			void ApplyNonLinearProjectionConstraints(const FReal Dt, const int32 It, const int32 NumIts);
+
+			FPBDJointConstraints& ConstraintContainer;
+
+			// Index remapping from internal index [0,NumToSolve) to index in the joint container [0,NumJointsInWorld)
+			TArray<int32> ContainerConstraintIndices;
+
+			// The linear and non-linear joint solvers. One for each joint we wish to solve in the order that they are solved
+			TArray<FPBDJointCachedSolver> LinearConstraintSolvers;
+			TArray<FPBDJointSolver> NonLinearConstraintSolvers;
+		};
+
+	}
+}
\ No newline at end of file
diff --git a/Engine/Source/Runtime/Experimental/Chaos/Public/Chaos/Joint/PBDJointSolverGaussSeidel.h b/Engine/Source/Runtime/Experimental/Chaos/Public/Chaos/Joint/PBDJointSolverGaussSeidel.h
--- a/Engine/Source/Runtime/Experimental/Chaos/Public/Chaos/Joint/PBDJointSolverGaussSeidel.h
+++ b/Engine/Source/Runtime/Experimental/Chaos/Public/Chaos/Joint/PBDJointSolverGaussSeidel.h
@@ -127,7 +127,9 @@ namespace Chaos
 			return NetAngularImpulse;
 		}
 
-		FPBDJointSolver();
+		FPBDJointSolver()
+		{
+		}
 
 		void SetSolverBodies(FSolverBody* SolverBody0, FSolverBody* SolverBody1)
 		{
@@ -174,11 +176,28 @@ namespace Chaos
 		void ApplyProjections(
 			const FReal Dt,
 			const FPBDJointSolverSettings& SolverSettings,
-			const FPBDJointSettings& JointSettings);
+			const FPBDJointSettings& JointSettings,
+			const bool bLastIteration);
 
 		void SetShockPropagationScales(
 			const FReal InvMScale0, 
-			const FReal InvMScale1);
+			const FReal InvMScale1,
+			const FReal Dt);
+
+		void SetIsBroken(const bool bInIsBroken)
+		{
+			bIsBroken = bInIsBroken;
+		}
+
+		bool IsBroken() const
+		{
+			return bIsBroken;
+		}
+
+		bool RequiresSolve() const
+		{
+			return !IsBroken() && (IsDynamic(0) || IsDynamic(1));
+		}
 
 	private:
 
@@ -725,6 +744,8 @@ namespace Chaos
 		FRotation3 LastQs[MaxConstrainedBodies];	// Rotations at the beginning of the iteration
 		FVec3 InitConstraintAxisLinearVelocities; // Linear velocities along the constraint axes at the begining of the frame, used by restitution
 		FVec3 InitConstraintAxisAngularVelocities; // Angular velocities along the constraint axes at the begining of the frame, used by restitution
+
+		bool bIsBroken;
 	};
 
 }
\ No newline at end of file
diff --git a/Engine/Source/Runtime/Experimental/Chaos/Public/Chaos/PBDJointConstraintTypes.h b/Engine/Source/Runtime/Experimental/Chaos/Public/Chaos/PBDJointConstraintTypes.h
--- a/Engine/Source/Runtime/Experimental/Chaos/Public/Chaos/PBDJointConstraintTypes.h
+++ b/Engine/Source/Runtime/Experimental/Chaos/Public/Chaos/PBDJointConstraintTypes.h
@@ -223,6 +223,9 @@ namespace Chaos
 		// Whether to use the linear or non-linear joint solver
 		bool bUseLinearSolver;
 
+		// Whether the joints need to be sorted (only required for RBAN - the world solver uses the constraint graph for ordering)
+		bool bSortEnabled;
+
 		// Whether to solve rotation then position limits (true), or vice versa
 		// Solving position last leads to less separation at the joints when limits are being forced
 		bool bSolvePositionLast;
diff --git a/Engine/Source/Runtime/Experimental/Chaos/Public/Chaos/PBDJointConstraints.h b/Engine/Source/Runtime/Experimental/Chaos/Public/Chaos/PBDJointConstraints.h
--- a/Engine/Source/Runtime/Experimental/Chaos/Public/Chaos/PBDJointConstraints.h
+++ b/Engine/Source/Runtime/Experimental/Chaos/Public/Chaos/PBDJointConstraints.h
@@ -1,18 +1,11 @@
 // Copyright Epic Games, Inc. All Rights Reserved.
 #pragma once
 
-#include "CoreMinimal.h"
-
-#include "Chaos/Array.h"
 #include "Chaos/Core.h"
-#include "Chaos/Transform.h"
-#include "Chaos/Vector.h"
 
 #include "Chaos/ConstraintHandle.h"
 #include "Chaos/Collision/CollisionApplyType.h"
 #include "Chaos/Evolution/IndexedConstraintContainer.h"
-#include "Chaos/Joint/PBDJointSolverGaussSeidel.h"
-#include "Chaos/Joint/PBDJointCachedSolverGaussSeidel.h"
 #include "Chaos/ParticleHandleFwd.h"
 #include "Chaos/PBDJointConstraintTypes.h"
 #include "Chaos/PBDJointConstraintData.h"
@@ -20,6 +13,14 @@
 
 namespace Chaos
 {
+	namespace Private
+	{
+		class FPBDJointContainerSolver;
+	}
+
+	/**
+	 * A handle to a joint constraint held in a joint container (FPBDJointConstraints) by index.
+	*/
 	class CHAOS_API FPBDJointConstraintHandle final : public TIndexedContainerConstraintHandle<FPBDJointConstraints>
 	{
 	public:
@@ -32,9 +33,6 @@ namespace Chaos
 		void SetConstraintEnabled(bool bEnabled);
 
 		void CalculateConstraintSpace(FVec3& OutXa, FMatrix33& OutRa, FVec3& OutXb, FMatrix33& OutRb) const;
-		int32 GetConstraintIsland() const;
-		int32 GetConstraintLevel() const;
-		int32 GetConstraintColor() const;
 
 		bool IsConstraintEnabled() const;
 		bool IsConstraintBroken() const;
@@ -64,6 +62,13 @@ namespace Chaos
 		void SetEnabledDuringResim(bool bEnabled);
 		EResimType ResimType() const;
 
+		UE_DEPRECATED(5.2, "No longer used")
+		int32 GetConstraintIsland() const;
+		UE_DEPRECATED(5.2, "No longer used")
+		int32 GetConstraintLevel() const;
+		UE_DEPRECATED(5.2, "No longer used")
+		int32 GetConstraintColor() const;
+
 	protected:
 		using Base::ConstraintIndex;
 		using Base::ConcreteContainer;
@@ -72,6 +77,9 @@ namespace Chaos
 		bool bAngularPlasticityInitialized;
 	};
 
+	/**
+	 * Peristent variable state for a joint
+	*/
 	class CHAOS_API FPBDJointState
 	{
 	public:
@@ -93,16 +101,15 @@ namespace Chaos
 	};
 
 	/**
-	 * A joint restricting up to 6 degrees of freedom, with linear and angular limits.
+	 * A set of joint restricting up to 6 degrees of freedom, with linear and angular limits.
 	 */
-	class CHAOS_API FPBDJointConstraints : public TPBDIndexedConstraintContainer<FPBDJointConstraints>
+	class CHAOS_API FPBDJointConstraints : public FPBDConstraintContainer
 	{
 	public:
-		using Base = TPBDIndexedConstraintContainer<FPBDJointConstraints>;
+		using Base = FPBDConstraintContainer;
 
 		using FConstraintContainerHandle = FPBDJointConstraintHandle;
 		using FConstraintHandleAllocator = TConstraintHandleAllocator<FPBDJointConstraints>;
-		using FVectorPair = TVector<FVec3, 2>;
 		using FTransformPair = TVector<FRigidTransform3, 2>;
 		using FHandles = TArray<FConstraintContainerHandle*>;
 
@@ -110,14 +117,26 @@ namespace Chaos
 
 		virtual ~FPBDJointConstraints();
 
+		/**
+		 * Get the solver settings (used by FPBDJointContainerSolver)
+		*/
 		const FPBDJointSolverSettings& GetSettings() const;
+
+		/*
+		* Modify the solver settings
+		*/
 		void SetSettings(const FPBDJointSolverSettings& InSettings);
 
+		/**
+		* Whether to use a linear or non-linear joint solver. Non-linear is more stable but much more expensive.
+		* A linear solver is used by default (see FPBDJointSolverSettings).
+		*/
 		void SetUseLinearJointSolver(const bool bInEnable) { Settings.bUseLinearSolver = bInEnable; }
 
-		//
-		// Constraint Container API
-		//
+		/**
+		 * Whether to sort the joints internally. Sort will be triggered on any tick when a joint was added. Only needed for RBAN.
+		*/
+		void SetSortEnabled(const bool bInEnable) { Settings.bSortEnabled = bInEnable; }
 
 		/**
 		 * Get the number of constraints.
@@ -198,22 +217,42 @@ namespace Chaos
 		 */
 		void FixConstraint(int32 ConstraintIndex);
 
+		/**
+		 * Set the break callback. This will be called once after the constraint solver phase on the \
+		 * tick when the constraint is broken. There is only one callback allowed - it is usually a 
+		 * method on the owner (e.g., the Evolution object) which will probably dispatch its own event.
+		*/
 		void SetBreakCallback(const FJointBreakCallback& Callback);
+
+		/**
+		 * Remove the previously assigned break callback.
+		*/
 		void ClearBreakCallback();
 
-		//
-		// Constraint API
-		//
+		/**
+		 * All of the constraints in the container, including inactive
+		*/
 		FHandles& GetConstraintHandles()
 		{
 			return Handles;
 		}
+
+		/**
+		 * All of the constraints in the container, including inactive
+		*/
 		const FHandles& GetConstConstraintHandles() const
 		{
 			return Handles;
 		}
 
+		/**
+		 * Get a joint constraint by index
+		*/
 		const FConstraintContainerHandle* GetConstraintHandle(int32 ConstraintIndex) const;
+
+		/**
+		 * Get a joint constraint by index
+		*/
 		FConstraintContainerHandle* GetConstraintHandle(int32 ConstraintIndex);
 
 		/**
@@ -221,17 +260,34 @@ namespace Chaos
 		 */
 		const FParticlePair& GetConstrainedParticles(int32 ConstraintIndex) const;
 
+		/**
+		 * Get the settings for a joint constraint by index
+		*/
 		const FPBDJointSettings& GetConstraintSettings(int32 ConstraintIndex) const;
 
+		/**
+		 * Set the settings for a joint constraint by index
+		*/
 		void SetConstraintSettings(int32 ConstraintIndex, const FPBDJointSettings& InConstraintSettings);
+
+		/**
+		 * Set the linear drive target for a constraint by index
+		*/
 		void SetLinearDrivePositionTarget(int32 ConstraintIndex, FVec3 InLinearDrivePositionTarget);
-		void SetAngularDrivePositionTarget(int32 ConstraintIndex, FRotation3 InAngularDrivePositionTarget);
 
-		int32 GetConstraintIsland(int32 ConstraintIndex) const;
-		int32 GetConstraintLevel(int32 ConstraintIndex) const;
-		int32 GetConstraintColor(int32 ConstraintIndex) const;
+		/**
+		 * Set the angular drive target for a constraint by index
+		*/
+		void SetAngularDrivePositionTarget(int32 ConstraintIndex, FRotation3 InAngularDrivePositionTarget);
 
+		/**
+		* The total linear impulse applied by the constraint
+		*/
 		FVec3 GetConstraintLinearImpulse(int32 ConstraintIndex) const;
+
+		/**
+		* The total linear angular applied by the constraint
+		*/
 		FVec3 GetConstraintAngularImpulse(int32 ConstraintIndex) const;
 
 		ESyncState GetConstraintSyncState(int32 ConstraintIndex) const;
@@ -240,49 +296,59 @@ namespace Chaos
 		void SetConstraintEnabledDuringResim(int32 ConstraintIndex, bool bEnabled);
 		
 		EResimType GetConstraintResimType(int32 ConstraintIndex) const;
-		
+
 		//
 		// FConstraintContainer Implementation
 		//
+		virtual TUniquePtr<FConstraintContainerSolver> CreateSceneSolver(const int32 Priority) override final;
+		virtual TUniquePtr<FConstraintContainerSolver> CreateGroupSolver(const int32 Priority) override final;
 		virtual int32 GetNumConstraints() const override final { return NumConstraints(); }
 		virtual void ResetConstraints() override final {}
 		virtual void AddConstraintsToGraph(FPBDIslandManager& IslandManager) override final;
 		virtual void PrepareTick() override final;
 		virtual void UnprepareTick() override final;
 
-		//
-		// TSimpleConstraintContainerSolver API - used by RBAN solvers
-		//
-		void AddBodies(FSolverBodyContainer& SolverBodyContainer);
-		void GatherInput(const FReal Dt);
-		void ScatterOutput(const FReal Dt);
-		void ApplyPositionConstraints(const FReal Dt, const int32 It, const int32 NumIts);
-		void ApplyVelocityConstraints(const FReal Dt, const int32 It, const int32 NumIts);
-		void ApplyProjectionConstraints(const FReal Dt, const int32 It, const int32 NumIts);
-
-		//
-		// TIndexedConstraintContainerSolver API - used by World solvers
-		//
-		void AddBodies(const TArrayView<int32>& ConstraintIndices, FSolverBodyContainer& SolverBodyContainer);
-		void GatherInput(const TArrayView<int32>& ConstraintIndices, const FReal Dt);
-		void ScatterOutput(const TArrayView<int32>& ConstraintIndices, const FReal Dt);
-		void ApplyPositionConstraints(const TArrayView<int32>& ConstraintIndices, const FReal Dt, const int32 It, const int32 NumIts);
-		void ApplyVelocityConstraints(const TArrayView<int32>& ConstraintIndices, const FReal Dt, const int32 It, const int32 NumIts);
-		void ApplyProjectionConstraints(const TArrayView<int32>& ConstraintIndices, const FReal Dt, const int32 It, const int32 NumIts);
+		// Called by the joint solver at the end of the constraint solver phase
+		void SetSolverResults(const int32 ConstraintIndex, const FVec3& LinearImpulse, const FVec3& AngularImpulse, const bool bIsBroken, const FSolverBody* SolverBody0, const FSolverBody* SolverBody1);
 
-		void GatherInput(const int32 ConstraintIndex, const FReal Dt);
-		void ScatterOutput(const int32 ConstraintIndex, const FReal Dt);
+		// @todo(chaos): only needed for RBAN, and should be private or moved to the solver
+		int32 GetConstraintIsland(int32 ConstraintIndex) const;
+		int32 GetConstraintLevel(int32 ConstraintIndex) const;
+		int32 GetConstraintColor(int32 ConstraintIndex) const;
 
-	protected:
-		using Base::GetConstraintIndex;
-		using Base::SetConstraintIndex;
+		// @todo(chaos): functionality no longer supported here (see FPBDJointContainerSolver)
+		UE_DEPRECATED(5.2, "Joint Solver API moved to FPBDJointContainerSolver")
+		void AddBodies(FSolverBodyContainer& SolverBodyContainer) {}
+		UE_DEPRECATED(5.2, "Joint Solver API moved to FPBDJointContainerSolver")
+		void GatherInput(const FReal Dt){}
+		UE_DEPRECATED(5.2, "Joint Solver API moved to FPBDJointContainerSolver")
+		void ScatterOutput(const FReal Dt){}
+		UE_DEPRECATED(5.2, "Joint Solver API moved to FPBDJointContainerSolver")
+		void ApplyPositionConstraints(const FReal Dt, const int32 It, const int32 NumIts) {}
+		UE_DEPRECATED(5.2, "Joint Solver API moved to FPBDJointContainerSolver")
+		void ApplyVelocityConstraints(const FReal Dt, const int32 It, const int32 NumIts) {}
+		UE_DEPRECATED(5.2, "Joint Solver API moved to FPBDJointContainerSolver")
+		void ApplyProjectionConstraints(const FReal Dt, const int32 It, const int32 NumIts) {}
+		UE_DEPRECATED(5.2, "Joint Solver API moved to FPBDJointContainerSolver")
+		void AddBodies(const TArrayView<int32>& ConstraintIndices, FSolverBodyContainer& SolverBodyContainer) {}
+		UE_DEPRECATED(5.2, "Joint Solver API moved to FPBDJointContainerSolver")
+		void GatherInput(const TArrayView<int32>& ConstraintIndices, const FReal Dt) {}
+		UE_DEPRECATED(5.2, "Joint Solver API moved to FPBDJointContainerSolver")
+		void ScatterOutput(const TArrayView<int32>& ConstraintIndices, const FReal Dt) {}
+		UE_DEPRECATED(5.2, "Joint Solver API moved to FPBDJointContainerSolver")
+		void ApplyPositionConstraints(const TArrayView<int32>& ConstraintIndices, const FReal Dt, const int32 It, const int32 NumIts) {}
+		UE_DEPRECATED(5.2, "Joint Solver API moved to FPBDJointContainerSolver")
+		void ApplyVelocityConstraints(const TArrayView<int32>& ConstraintIndices, const FReal Dt, const int32 It, const int32 NumIts) {}
+		UE_DEPRECATED(5.2, "Joint Solver API moved to FPBDJointContainerSolver")
+		void ApplyProjectionConstraints(const TArrayView<int32>& ConstraintIndices, const FReal Dt, const int32 It, const int32 NumIts) {}
+		UE_DEPRECATED(5.2, "Joint Solver API moved to FPBDJointContainerSolver")
+		void GatherInput(const int32 ConstraintIndex, const FReal Dt) {}
+		UE_DEPRECATED(5.2, "Joint Solver API moved to FPBDJointContainerSolver")
+		void ScatterOutput(const int32 ConstraintIndex, const FReal Dt) {}
 
 	private:
 		friend class FPBDJointConstraintHandle;
 
-		FReal CalculateIterationStiffness(int32 It, int32 NumIts) const;
-		FReal CalculateShockPropagationInvMassScale(const FConstraintSolverBody& Body0, const FConstraintSolverBody& Body1, const FPBDJointSettings& JointSettings, const int32 It, const int32 NumIts) const;
-
 		void GetConstrainedParticleIndices(const int32 ConstraintIndex, int32& Index0, int32& Index1) const;
 		void CalculateConstraintSpace(int32 ConstraintIndex, FVec3& OutX0, FMatrix33& OutR0, FVec3& OutX1, FMatrix33& OutR1) const;
 		
@@ -291,13 +357,7 @@ namespace Chaos
 
 		bool CanEvaluate(const int32 ConstraintIndex) const;
 
-		void AddBodies(const int32 ConstraintIndex, FSolverBodyContainer& SolverBodyContainer);
-		bool ApplyPositionConstraint(const FReal Dt, const int32 ConstraintIndex, const int32 It, const int32 NumIts);
-		bool ApplyVelocityConstraint(const FReal Dt, const int32 ConstraintIndex, const int32 It, const int32 NumIts);
-		void PrepareProjectionConstraint(const FReal Dt, const int32 ConstraintIndex, const int32 It, const int32 NumIts);
-		bool ApplyProjectionConstraint(const FReal Dt, const int32 ConstraintIndex, const int32 It, const int32 NumIts);
-		void ApplyBreakThreshold(const FReal Dt, int32 ConstraintIndex, const FVec3& LinearImpulse, const FVec3& AngularImpulse);
-		void ApplyPlasticityLimits(const int32 ConstraintIndex);
+		void ApplyPlasticityLimits(const int32 ConstraintIndex, const FSolverBody& SolverBody0, const FSolverBody& SolverBody1);
 
 		void SetConstraintBroken(int32 ConstraintIndex, bool bBroken);
 		void SetConstraintBreaking(int32 ConstraintIndex, bool bBreaking);
@@ -313,10 +373,6 @@ namespace Chaos
 		bool bJointsDirty;
 
 		FJointBreakCallback BreakCallback;
-
-		// @todo(ccaulfield): optimize storage for joint solver
-		TArray<FPBDJointSolver> ConstraintSolvers;
-		TArray<FPBDJointCachedSolver> CachedConstraintSolvers;
 	};
 
 }
