Chaos - fix deprecation version

diff --git a/Engine/Source/Runtime/Experimental/Chaos/Public/Chaos/Collision/CollisionDetector.h b/Engine/Source/Runtime/Experimental/Chaos/Public/Chaos/Collision/CollisionDetector.h
--- a/Engine/Source/Runtime/Experimental/Chaos/Public/Chaos/Collision/CollisionDetector.h
+++ b/Engine/Source/Runtime/Experimental/Chaos/Public/Chaos/Collision/CollisionDetector.h
@@ -20,16 +20,16 @@ namespace Chaos
 
 		virtual ~FCollisionDetector() {}
 
-		UE_DEPRECATED(5.3, "Moved to FPBDCollisionConstraints")
+		UE_DEPRECATED(5.2, "Moved to FPBDCollisionConstraints")
 		const FCollisionDetectorSettings& GetSettings() const;
 
-		UE_DEPRECATED(5.3, "Moved to FPBDCollisionConstraints")
+		UE_DEPRECATED(5.2, "Moved to FPBDCollisionConstraints")
 		void SetSettings(const FCollisionDetectorSettings& InSettings);
 
-		UE_DEPRECATED(5.3, "Moved to FPBDCollisionConstraints and renamed to SetCullDistance")
+		UE_DEPRECATED(5.2, "Moved to FPBDCollisionConstraints and renamed to SetCullDistance")
 		void SetBoundsExpansion(const FReal InBoundsExpansion);
 
-		UE_DEPRECATED(5.3, "No longer supported")
+		UE_DEPRECATED(5.2, "No longer supported")
 		void SetBoundsVelocityInflation(const FReal InBoundsVelocityInflation) {}
 
 		FPBDCollisionConstraints& GetCollisionContainer() { return CollisionContainer; }
diff --git a/Engine/Source/Runtime/Experimental/Chaos/Public/Chaos/Evolution/PBDMinEvolution.h b/Engine/Source/Runtime/Experimental/Chaos/Public/Chaos/Evolution/PBDMinEvolution.h
--- a/Engine/Source/Runtime/Experimental/Chaos/Public/Chaos/Evolution/PBDMinEvolution.h
+++ b/Engine/Source/Runtime/Experimental/Chaos/Public/Chaos/Evolution/PBDMinEvolution.h
@@ -84,13 +84,13 @@ namespace Chaos
 		}
 
 
-		UE_DEPRECATED(5.3, "InBoundsExtension parameter has been removed")
+		UE_DEPRECATED(5.2, "InBoundsExtension parameter has been removed")
 		FPBDMinEvolution(FRigidParticleSOAs& InParticles, TArrayCollectionArray<FVec3>& InPrevX, TArrayCollectionArray<FRotation3>& InPrevR, FCollisionDetector& InCollisionDetector, const FReal InBoundsExtension)
 			: FPBDMinEvolution(InParticles, InPrevX, InPrevR, InCollisionDetector)
 		{
 		}
 
-		UE_DEPRECATED(5.3, "BoundsExtension has been removed")
+		UE_DEPRECATED(5.2, "BoundsExtension has been removed")
 		void SetBoudsExtension(const FReal Unused)
 		{
 		}
diff --git a/Engine/Source/Runtime/Experimental/Chaos/Public/Chaos/PBDRigidParticles.h b/Engine/Source/Runtime/Experimental/Chaos/Public/Chaos/PBDRigidParticles.h
--- a/Engine/Source/Runtime/Experimental/Chaos/Public/Chaos/PBDRigidParticles.h
+++ b/Engine/Source/Runtime/Experimental/Chaos/Public/Chaos/PBDRigidParticles.h
@@ -208,16 +208,16 @@ class TPBDRigidParticles : public TRigidParticles<T, d>
 		Ar << MP << MQ << MPreV << MPreW;
 	}
 
-	UE_DEPRECATED(5.3, "To be removed")
+	UE_DEPRECATED(5.2, "To be removed")
 	FORCEINLINE TArray<TVector<T, d>>& AllP() { return MP; }
 
-	UE_DEPRECATED(5.3, "To be removed")
+	UE_DEPRECATED(5.2, "To be removed")
 	FORCEINLINE TArray<TRotation<T, d>>& AllQ() { return MQ; }
 	
-	UE_DEPRECATED(5.3, "To be removed")
+	UE_DEPRECATED(5.2, "To be removed")
 	FORCEINLINE TArray<TVector<T, d>>& AllPreV() { return MPreV; }
 	
-	UE_DEPRECATED(5.3, "To be removed")
+	UE_DEPRECATED(5.2, "To be removed")
 	FORCEINLINE TArray<TVector<T, d>>& AllPreW() { return MPreW; }
 
   private:
