Chaos - per-island iteration counts (WIP)

Partial support for per-island iteration. Also move all Island management and solving into the private namespace.

diff --git a/Engine/Source/Programs/HeadlessChaos/Private/HeadlessChaosTestGraph.cpp b/Engine/Source/Programs/HeadlessChaos/Private/HeadlessChaosTestGraph.cpp
--- a/Engine/Source/Programs/HeadlessChaos/Private/HeadlessChaosTestGraph.cpp
+++ b/Engine/Source/Programs/HeadlessChaos/Private/HeadlessChaosTestGraph.cpp
@@ -56,7 +56,7 @@ namespace ChaosTest {
 		//
 		virtual int32 GetNumConstraints() const override final { return Constraints.Num(); }
 		virtual void ResetConstraints() override final {}
-		virtual void AddConstraintsToGraph(FPBDIslandManager& IslandManager) override final;
+		virtual void AddConstraintsToGraph(Private::FPBDIslandManager& IslandManager) override final;
 		virtual void PrepareTick() override final {}
 		virtual void UnprepareTick() override final {}
 
@@ -168,7 +168,7 @@ namespace ChaosTest {
 	};
 
 	template<int32 T_TYPEID>
-	void TMockGraphConstraints<T_TYPEID>::AddConstraintsToGraph(FPBDIslandManager& IslandManager)
+	void TMockGraphConstraints<T_TYPEID>::AddConstraintsToGraph(Private::FPBDIslandManager& IslandManager)
 	{
 		for (TMockGraphConstraintHandle<T_TYPEID>* Handle : Handles)
 		{
@@ -244,7 +244,7 @@ namespace ChaosTest {
 		}
 
 		// Set up the particle graph
-		FPBDConstraintGraph Graph;
+		Private::FPBDIslandManager Graph;
 		Graph.AddConstraintContainer(ConstraintsOfType0);
 		Graph.AddConstraintContainer(ConstraintsOfType1);
 
@@ -293,10 +293,10 @@ namespace ChaosTest {
 		for (int32 ExpectedIslandIndex = 0; ExpectedIslandIndex < ExpectedIslandParticles.Num(); ++ExpectedIslandIndex)
 		{
 			const int32 CalculatedIslandIndex = CalculatedIslandIndices[ExpectedIslandIndex];
-			TArrayView<FPBDIslandParticle> CalculatedIslandParticles = Graph.GetIsland(CalculatedIslandIndex)->GetParticles();
+			TArrayView<Private::FPBDIslandParticle> CalculatedIslandParticles = Graph.GetIsland(CalculatedIslandIndex)->GetParticles();
 
 			EXPECT_EQ(CalculatedIslandParticles.Num(), ExpectedIslandParticles[ExpectedIslandIndex].Num());
-			for (FPBDIslandParticle& CalculatedIslandParticle : CalculatedIslandParticles)
+			for (Private::FPBDIslandParticle& CalculatedIslandParticle : CalculatedIslandParticles)
 			{
 				EXPECT_TRUE(ExpectedIslandParticles[ExpectedIslandIndex].Contains(CalculatedIslandParticle.GetParticle()));
 			}
@@ -305,15 +305,15 @@ namespace ChaosTest {
 		Graph.Reset();
 	}
 
-	void CheckIslandIntegrity(TArray<TSet<TGeometryParticleHandle<FReal,3>*>>& ExpectedIslandParticles, const TArray<int32> CalculatedIslandIndices, FPBDConstraintGraph& Graph)
+	void CheckIslandIntegrity(TArray<TSet<TGeometryParticleHandle<FReal,3>*>>& ExpectedIslandParticles, const TArray<int32> CalculatedIslandIndices, Private::FPBDIslandManager& Graph)
 	{
 		for (int32 ExpectedIslandIndex = 0; ExpectedIslandIndex < ExpectedIslandParticles.Num(); ++ExpectedIslandIndex)
 		{
 			const int32 CalculatedIslandIndex = CalculatedIslandIndices[ExpectedIslandIndex];
-			TArrayView<FPBDIslandParticle> CalculatedIslandParticles = Graph.GetIsland(CalculatedIslandIndex)->GetParticles();
+			TArrayView<Private::FPBDIslandParticle> CalculatedIslandParticles = Graph.GetIsland(CalculatedIslandIndex)->GetParticles();
 
 			EXPECT_EQ(CalculatedIslandParticles.Num(), ExpectedIslandParticles[ExpectedIslandIndex].Num());
-			for (FPBDIslandParticle& CalculatedIslandParticle : CalculatedIslandParticles)
+			for (Private::FPBDIslandParticle& CalculatedIslandParticle : CalculatedIslandParticles)
 			{
 				EXPECT_TRUE(ExpectedIslandParticles[ExpectedIslandIndex].Contains(CalculatedIslandParticle.GetParticle()));
 
@@ -536,7 +536,7 @@ namespace ChaosTest {
 		{
 			FIterationData& IterData = IterationData[Iteration];
 
-			FPBDConstraintGraph Graph;
+			Private::FPBDIslandManager Graph;
 			//FPBDConstraintColor GraphColor;
 			const int32 ContainerId = 0;
 
@@ -609,7 +609,7 @@ namespace ChaosTest {
 			int Index = 0;
 			for (int32 Island : CalculatedIslandIndices)
 			{
-				TArrayView<FPBDIslandConstraint> IslandConstraints = Graph.GetIsland(Island)->GetConstraints(ContainerId);
+				TArrayView<Private::FPBDIslandConstraint> IslandConstraints = Graph.GetIsland(Island)->GetConstraints(ContainerId);
 				EXPECT_EQ(IslandConstraints.Num(), IterData.ExpectedIslandEdges[Index++]);
 			}
 
@@ -637,11 +637,10 @@ namespace ChaosTest {
 	}
 
 	void HelpTickConstraints(FPBDRigidsSOAs& SOAs, const TArray<TPBDRigidParticleHandle<FReal, 3>*>& Particles,
-		FPBDConstraintGraph& Graph, TMockGraphConstraints<0> Constraints,
+		Private::FPBDIslandManager& Graph, TMockGraphConstraints<0> Constraints,
 		const TArrayCollectionArray<TSerializablePtr<FChaosPhysicsMaterial>>& PhysicsMaterials,
 		const THandleArray<FChaosPhysicsMaterial>& PhysicalMaterials)
 	{
-
 		SOAs.ClearTransientDirty();
 
 		Graph.InitializeGraph(SOAs.GetNonDisabledDynamicView());
@@ -655,7 +654,7 @@ namespace ChaosTest {
 
 			if(bSleeped)
 			{
-				for(FPBDIslandParticle& Particle : Graph.GetIsland(IslandIndex)->GetParticles())
+				for(Private::FPBDIslandParticle& Particle : Graph.GetIsland(IslandIndex)->GetParticles())
 				{
 					SOAs.DeactivateParticle(Particle.GetParticle());
 				}
@@ -734,7 +733,7 @@ namespace ChaosTest {
 				Constraints.AddConstraint({ Particles[ConstrainedParticleIndices[0]], Particles[ConstrainedParticleIndices[1]] });
 			}
 
-			FPBDConstraintGraph Graph;
+			Private::FPBDIslandManager Graph;
 			Graph.AddConstraintContainer(Constraints);
 			for (FGeometryParticleHandle* Particle : Particles)
 			{
@@ -829,7 +828,7 @@ namespace ChaosTest {
 				Constraints.AddConstraint({ Particles[ConstrainedParticleIndices[0]], Particles[ConstrainedParticleIndices[1]] });
 			}
 
-			FPBDConstraintGraph Graph;
+			Private::FPBDIslandManager Graph;
 			Graph.AddConstraintContainer(Constraints);
 			for (FGeometryParticleHandle* Particle : Particles)
 			{
@@ -936,7 +935,7 @@ namespace ChaosTest {
 				Constraints.AddConstraint({ Particles[ConstrainedParticleIndices[0]], Particles[ConstrainedParticleIndices[1]] });
 			}
 
-			FPBDConstraintGraph Graph;
+			Private::FPBDIslandManager Graph;
 			Graph.AddConstraintContainer(Constraints);
 			for (FGeometryParticleHandle* Particle : Particles)
 			{
@@ -1121,7 +1120,7 @@ namespace ChaosTest {
 		}
 		
 		// Build the connectivity graph and islands
-		FPBDConstraintGraph Graph;
+		Private::FPBDIslandManager Graph;
 		Graph.AddConstraintContainer(Constraints);
 		//FPBDConstraintColor GraphColor;
 		for (FGeometryParticleHandle* Particle : AllParticles)
@@ -1219,7 +1218,7 @@ namespace ChaosTest {
 			return true;
 		}
 
-		static void ValidateConstraintGraphParticleHandles(const FPBDConstraintGraph& ConstraintGraph)
+		static void ValidateConstraintGraphParticleHandles(const Private::FPBDIslandManager& ConstraintGraph)
 		{
 			// Are all keys, values in the ParticleToNodeIndex Valid?
 			const int32 NodeCount = ConstraintGraph.GetGraphNodes().GetMaxIndex();
@@ -1234,22 +1233,22 @@ namespace ChaosTest {
 			const int32 NumIslands = ConstraintGraph.NumIslands();
 			for (int32 IslandIndex = 0; IslandIndex < NumIslands; ++IslandIndex)
 			{
-				TArrayView<const FPBDIslandParticle> IslandParticles = ConstraintGraph.GetIsland(IslandIndex)->GetParticles();
-				for (const FPBDIslandParticle& IslandParticle : IslandParticles)
+				TArrayView<const Private::FPBDIslandParticle> IslandParticles = ConstraintGraph.GetIsland(IslandIndex)->GetParticles();
+				for (const Private::FPBDIslandParticle& IslandParticle : IslandParticles)
 				{
 					EXPECT_TRUE(IsParticleHandleValid(IslandParticle.GetParticle()));
 				}
 			}
 
 			// Are all Graph Nodes valid?
-			for (const FPBDConstraintGraph::FGraphNode& Node : ConstraintGraph.GetGraphNodes())
+			for (const Private::FPBDIslandManager::FGraphNode& Node : ConstraintGraph.GetGraphNodes())
 			{
 				ensure(IsParticleHandleValid(Node.NodeItem));
 				EXPECT_TRUE(IsParticleHandleValid(Node.NodeItem));
 			}
 
 			// Are all Graph Edges valid?
-			for (const FPBDConstraintGraph::FGraphEdge& Edge : ConstraintGraph.GetGraphEdges())
+			for (const Private::FPBDIslandManager::FGraphEdge& Edge : ConstraintGraph.GetGraphEdges())
 			{
 				EXPECT_GE(Edge.FirstNode, 0);
 				EXPECT_LT(Edge.FirstNode, NodeCount);
@@ -1322,7 +1321,7 @@ namespace ChaosTest {
 		}
 
 		// Set up the particle graph
-		FPBDConstraintGraph Graph;
+		Private::FPBDIslandManager Graph;
 		Graph.AddConstraintContainer(PositionConstraints);
 		Graph.AddConstraintContainer(SuspensionConstraints);
 		Graph.AddConstraintContainer(SpringConstraints);
diff --git a/Engine/Source/Programs/HeadlessChaos/Private/HeadlessChaosTestGraphEvolution.cpp b/Engine/Source/Programs/HeadlessChaos/Private/HeadlessChaosTestGraphEvolution.cpp
--- a/Engine/Source/Programs/HeadlessChaos/Private/HeadlessChaosTestGraphEvolution.cpp
+++ b/Engine/Source/Programs/HeadlessChaos/Private/HeadlessChaosTestGraphEvolution.cpp
@@ -25,7 +25,7 @@ namespace ChaosTest
 			, Constraints()
 			, TickCount(0)
 		{
-			IslandManager = &Evolution.GetConstraintGraph();
+			IslandManager = &Evolution.GetIslandManager();
 
 			// Bind the constraints to the evolution
 			Evolution.AddConstraintContainer(Constraints);
@@ -95,7 +95,7 @@ namespace ChaosTest
 		FPBDNullConstraints Constraints;
 		TArray<FPBDRigidParticleHandle*> ParticleHandles;
 		TArray<FPBDNullConstraintHandle*> ConstraintHandles;
-		FPBDIslandManager* IslandManager;
+		Private::FPBDIslandManager* IslandManager;
 
 		int32 TickCount;
 	};
diff --git a/Engine/Source/Programs/HeadlessChaos/Private/HeadlessChaosTestRewind.cpp b/Engine/Source/Programs/HeadlessChaos/Private/HeadlessChaosTestRewind.cpp
--- a/Engine/Source/Programs/HeadlessChaos/Private/HeadlessChaosTestRewind.cpp
+++ b/Engine/Source/Programs/HeadlessChaos/Private/HeadlessChaosTestRewind.cpp
@@ -2357,7 +2357,7 @@ namespace ChaosTest {
 
 			// Wait for sleep (active particles get added to the dirty list)
 			// NOTE: Sleep requires 20 frames of inactivity by default, plus the time for smoothed velocity to damp to zero
-			// (see FPBDConstraintGraph::SleepInactive)
+			// (see FPBDIslandManager::SleepInactive)
 			for(int Step = 0; Step < 500; ++Step)
 			{
 				TickSolverHelper(Solver);
diff --git a/Engine/Source/Programs/HeadlessChaos/Public/HeadlessChaosCollisionConstraints.h b/Engine/Source/Programs/HeadlessChaos/Public/HeadlessChaosCollisionConstraints.h
--- a/Engine/Source/Programs/HeadlessChaos/Public/HeadlessChaosCollisionConstraints.h
+++ b/Engine/Source/Programs/HeadlessChaos/Public/HeadlessChaosCollisionConstraints.h
@@ -41,7 +41,7 @@ public:
 		, CollisionConstraints(EmptyParticles, EmptyCollided, EmptyPhysicsMaterials, EmptyUniquePhysicsMaterials, nullptr)
 		, BroadPhase(EmptyParticles)
 		, CollisionDetector(BroadPhase, CollisionConstraints)
-		, ConstraintSolver()
+		, ConstraintSolver(Iterations)
 	{
 		CollisionConstraints.SetContainerId(0);
 
@@ -57,7 +57,7 @@ public:
 		, CollisionConstraints(InParticles, Collided, PerParticleMaterials, PerParticleUniqueMaterials, nullptr)
 		, BroadPhase(InParticles)
 		, CollisionDetector(BroadPhase, CollisionConstraints)
-		, ConstraintSolver()
+		, ConstraintSolver(Iterations)
 	{
 		CollisionConstraints.SetContainerId(0);
 
@@ -138,12 +138,16 @@ public:
 
 	void Apply(const FReal Dt, const int32 NumIts)
 	{
-		ConstraintSolver.ApplyPositionConstraints(Dt, NumIts);
+		Iterations.SetNumPositionIterations(NumIts);
+		ConstraintSolver.SetIterationSettings(Iterations);
+		ConstraintSolver.ApplyPositionConstraints(Dt);
 	}
 
 	void ApplyPushOut(const FReal Dt, int32 NumIts)
 	{
-		ConstraintSolver.ApplyVelocityConstraints(Dt, NumIts);
+		Iterations.SetNumVelocityIterations(NumIts);
+		ConstraintSolver.SetIterationSettings(Iterations);
+		ConstraintSolver.ApplyVelocityConstraints(Dt);
 	}
 
 	void GatherInput(FReal Dt)
@@ -171,6 +175,7 @@ public:
 	FCollisionConstraints CollisionConstraints;
 	FSpatialAccelerationBroadPhase BroadPhase;
 	FCollisionDetector CollisionDetector;
-	FPBDSceneConstraintGroupSolver ConstraintSolver;
+	Private::FPBDSceneConstraintGroupSolver ConstraintSolver;
+	Private::FIterationSettings Iterations;
 };
 }
diff --git a/Engine/Source/Runtime/Experimental/Chaos/Private/Chaos/ChaosDebugDraw.cpp b/Engine/Source/Runtime/Experimental/Chaos/Private/Chaos/ChaosDebugDraw.cpp
--- a/Engine/Source/Runtime/Experimental/Chaos/Private/Chaos/ChaosDebugDraw.cpp
+++ b/Engine/Source/Runtime/Experimental/Chaos/Private/Chaos/ChaosDebugDraw.cpp
@@ -1281,7 +1281,7 @@ namespace Chaos
 			FDebugDrawQueue::GetInstance().DrawDebugLine(Pos, Pos + 0.01f * Settings.AngVelScale * SimSpace.AngularAcceleration, FColor::Orange, false, UE_KINDA_SMALL_NUMBER, Settings.DrawPriority, Settings.LineThickness);
 		}
 
-		void DrawConstraintGraphImpl(const FRigidTransform3& SpaceTransform, const FPBDConstraintGraph& Graph, const FChaosDebugDrawSettings& Settings)
+		void DrawConstraintGraphImpl(const FRigidTransform3& SpaceTransform, const Private::FPBDIslandManager& Graph, const FChaosDebugDrawSettings& Settings)
 		{
 			auto DrawGraphCollision = [&](const FRigidTransform3& SpaceTransform, const FPBDCollisionConstraint* Constraint,  int32 IslandIndex, int32 LevelIndex, int32 ColorIndex, int32 OrderIndex, bool bIsUsed, const FChaosDebugDrawSettings& Settings)
 			{
@@ -1321,7 +1321,7 @@ namespace Chaos
 				TArray<FAABB3> IslandBounds;
 				IslandBounds.SetNum(Graph.NumIslands());
 
-				const typename FPBDIslandManager::GraphType* IslandGraph = Graph.GetIslandGraph();
+				const typename Private::FPBDIslandManager::GraphType* IslandGraph = Graph.GetIslandGraph();
 				for (const auto& GraphNode : IslandGraph->GraphNodes)
 				{
 					FConstGenericParticleHandle Particle = GraphNode.NodeItem;
@@ -1346,7 +1346,7 @@ namespace Chaos
 
 			if (bChaosDebugDebugDrawContactGraph || bChaosDebugDebugDrawContactGraphUnused || bChaosDebugDebugDrawContactGraphUsed)
 			{
-				const typename FPBDIslandManager::GraphType* IslandGraph = Graph.GetIslandGraph();
+				const typename Private::FPBDIslandManager::GraphType* IslandGraph = Graph.GetIslandGraph();
 				for (const auto& GraphEdge : IslandGraph->GraphEdges)
 				{
 					const FConstraintHandle* Constraint = GraphEdge.EdgeItem;
@@ -1745,7 +1745,7 @@ namespace Chaos
 			}
 		}
 
-		void DrawConstraintGraph(const FRigidTransform3& SpaceTransform, const FPBDConstraintGraph& Graph, const FChaosDebugDrawSettings* Settings)
+		void DrawConstraintGraph(const FRigidTransform3& SpaceTransform, const Private::FPBDIslandManager& Graph, const FChaosDebugDrawSettings* Settings)
 		{
 			if (FDebugDrawQueue::IsDebugDrawingEnabled())
 			{
diff --git a/Engine/Source/Runtime/Experimental/Chaos/Private/Chaos/Collision/SolverCollisionContainer.cpp b/Engine/Source/Runtime/Experimental/Chaos/Private/Chaos/Collision/SolverCollisionContainer.cpp
--- a/Engine/Source/Runtime/Experimental/Chaos/Private/Chaos/Collision/SolverCollisionContainer.cpp
+++ b/Engine/Source/Runtime/Experimental/Chaos/Private/Chaos/Collision/SolverCollisionContainer.cpp
@@ -326,9 +326,9 @@ namespace Chaos
 		}
 	}
 
-	void FPBDCollisionContainerSolver::AddConstraints(const TArrayView<FPBDIslandConstraint>& IslandConstraints)
+	void FPBDCollisionContainerSolver::AddConstraints(const TArrayView<Private::FPBDIslandConstraint>& IslandConstraints)
 	{
-		for (FPBDIslandConstraint& IslandConstraint : IslandConstraints)
+		for (Private::FPBDIslandConstraint& IslandConstraint : IslandConstraints)
 		{
 			// We will only ever be given constraints from our container (asserts in non-shipping)
 			FPBDCollisionConstraint& Constraint = IslandConstraint.GetConstraint()->AsUnsafe<FPBDCollisionConstraintHandle>()->GetContact();
diff --git a/Engine/Source/Runtime/Experimental/Chaos/Private/Chaos/Evolution/ConstraintGroupSolver.cpp b/Engine/Source/Runtime/Experimental/Chaos/Private/Chaos/Evolution/ConstraintGroupSolver.cpp
--- a/Engine/Source/Runtime/Experimental/Chaos/Private/Chaos/Evolution/ConstraintGroupSolver.cpp
+++ b/Engine/Source/Runtime/Experimental/Chaos/Private/Chaos/Evolution/ConstraintGroupSolver.cpp
@@ -6,258 +6,266 @@
 
 namespace Chaos
 {
-	FPBDConstraintGroupSolver::FPBDConstraintGroupSolver()
-		: TotalNumConstraints(0)
+	namespace Private
 	{
-	}
+		FPBDConstraintGroupSolver::FPBDConstraintGroupSolver()
+			: TotalNumConstraints(0)
+		{
+		}
 
-	FPBDConstraintGroupSolver::~FPBDConstraintGroupSolver()
-	{
-	}
-	
-	void FPBDConstraintGroupSolver::SetConstraintSolver(const int32 ContainerId, TUniquePtr<FConstraintContainerSolver>&& Solver)
-	{
-		if (ContainerId >= ConstraintContainerSolvers.Num())
+		FPBDConstraintGroupSolver::~FPBDConstraintGroupSolver()
 		{
-			ConstraintContainerSolvers.SetNum(ContainerId + 1);
 		}
+	
+		void FPBDConstraintGroupSolver::SetConstraintSolver(const int32 ContainerId, TUniquePtr<FConstraintContainerSolver>&& Solver)
+		{
+			if (ContainerId >= ConstraintContainerSolvers.Num())
+			{
+				ConstraintContainerSolvers.SetNum(ContainerId + 1);
+			}
 
-		ConstraintContainerSolvers[ContainerId] = MoveTemp(Solver);
+			ConstraintContainerSolvers[ContainerId] = MoveTemp(Solver);
 
-		SetConstraintSolverImpl(ContainerId);
+			SetConstraintSolverImpl(ContainerId);
 
-		SortSolverContainers();
-	}
+			SortSolverContainers();
+		}
 
-	void FPBDConstraintGroupSolver::SetConstraintSolverPriority(const int32 ContainerId, const int32 Priority)
-	{
-		if (ConstraintContainerSolvers[ContainerId] != nullptr)
+		void FPBDConstraintGroupSolver::SetConstraintSolverPriority(const int32 ContainerId, const int32 Priority)
 		{
-			if (ConstraintContainerSolvers[ContainerId]->GetPriority() != Priority)
+			if (ConstraintContainerSolvers[ContainerId] != nullptr)
 			{
-				ConstraintContainerSolvers[ContainerId]->SetPriority(Priority);
-				SortSolverContainers();
+				if (ConstraintContainerSolvers[ContainerId]->GetPriority() != Priority)
+				{
+					ConstraintContainerSolvers[ContainerId]->SetPriority(Priority);
+					SortSolverContainers();
+				}
 			}
 		}
-	}
 
-	void FPBDConstraintGroupSolver::SortSolverContainers()
-	{
-		PrioritizedConstraintContainerSolvers.Reset(ConstraintContainerSolvers.Num());
-		for (TUniquePtr<FConstraintContainerSolver>& SolverContainer : ConstraintContainerSolvers)
+		void FPBDConstraintGroupSolver::SortSolverContainers()
 		{
-			PrioritizedConstraintContainerSolvers.Add(SolverContainer.Get());
-		}
-
-		PrioritizedConstraintContainerSolvers.StableSort(
-			[](const FConstraintContainerSolver& L, const FConstraintContainerSolver& R)
+			PrioritizedConstraintContainerSolvers.Reset(ConstraintContainerSolvers.Num());
+			for (TUniquePtr<FConstraintContainerSolver>& SolverContainer : ConstraintContainerSolvers)
 			{
-				return L.GetPriority() < R.GetPriority();
-			});
-	}
+				PrioritizedConstraintContainerSolvers.Add(SolverContainer.Get());
+			}
 
-	void FPBDConstraintGroupSolver::Reset()
-	{
-		SolverBodyContainer.Reset(0);
+			PrioritizedConstraintContainerSolvers.StableSort(
+				[](const FConstraintContainerSolver& L, const FConstraintContainerSolver& R)
+				{
+					return L.GetPriority() < R.GetPriority();
+				});
+		}
 
-		for (int32 ContainerIndex = 0; ContainerIndex < ConstraintContainerSolvers.Num(); ++ContainerIndex)
+		void FPBDConstraintGroupSolver::Reset()
 		{
-			if (ConstraintContainerSolvers[ContainerIndex] != nullptr)
+			SolverBodyContainer.Reset(0);
+
+			for (int32 ContainerIndex = 0; ContainerIndex < ConstraintContainerSolvers.Num(); ++ContainerIndex)
 			{
-				ConstraintContainerSolvers[ContainerIndex]->Reset(0);
+				if (ConstraintContainerSolvers[ContainerIndex] != nullptr)
+				{
+					ConstraintContainerSolvers[ContainerIndex]->Reset(0);
+				}
 			}
-		}
 
-		TotalNumConstraints = 0;
+			TotalNumConstraints = 0;
 
-		ResetImpl();
-	}
-
-	void FPBDConstraintGroupSolver::AddConstraintsAndBodies()
-	{
-		AddConstraintsImpl();
+			ResetImpl();
+		}
 
-		for (int32 ContainerIndex = 0; ContainerIndex < ConstraintContainerSolvers.Num(); ++ContainerIndex)
+		void FPBDConstraintGroupSolver::AddConstraintsAndBodies()
 		{
-			if (ConstraintContainerSolvers[ContainerIndex] != nullptr)
+			AddConstraintsImpl();
+
+			for (int32 ContainerIndex = 0; ContainerIndex < ConstraintContainerSolvers.Num(); ++ContainerIndex)
 			{
-				ConstraintContainerSolvers[ContainerIndex]->AddBodies(SolverBodyContainer);
+				if (ConstraintContainerSolvers[ContainerIndex] != nullptr)
+				{
+					ConstraintContainerSolvers[ContainerIndex]->AddBodies(SolverBodyContainer);
 
-				TotalNumConstraints += ConstraintContainerSolvers[ContainerIndex]->GetNumConstraints();
+					TotalNumConstraints += ConstraintContainerSolvers[ContainerIndex]->GetNumConstraints();
+				}
 			}
-		}
 
-		SolverBodyContainer.Lock();
-	}
+			SolverBodyContainer.Lock();
+		}
 
-	void FPBDConstraintGroupSolver::GatherBodies(const FReal Dt)
-	{
-		SolverBodyContainer.GatherInput(Dt, 0, SolverBodyContainer.Num());
-	}
+		void FPBDConstraintGroupSolver::GatherBodies(const FReal Dt)
+		{
+			SolverBodyContainer.GatherInput(Dt, 0, SolverBodyContainer.Num());
+		}
 
-	void FPBDConstraintGroupSolver::GatherConstraints(const FReal Dt)
-	{
-		for (int32 ContainerIndex = 0; ContainerIndex < ConstraintContainerSolvers.Num(); ++ContainerIndex)
+		void FPBDConstraintGroupSolver::GatherConstraints(const FReal Dt)
 		{
-			if (ConstraintContainerSolvers[ContainerIndex] != nullptr)
+			for (int32 ContainerIndex = 0; ContainerIndex < ConstraintContainerSolvers.Num(); ++ContainerIndex)
 			{
-				ConstraintContainerSolvers[ContainerIndex]->GatherInput(Dt);
+				if (ConstraintContainerSolvers[ContainerIndex] != nullptr)
+				{
+					ConstraintContainerSolvers[ContainerIndex]->GatherInput(Dt);
+				}
 			}
 		}
-	}
-	void FPBDConstraintGroupSolver::GatherBodies(const FReal Dt, const int32 BeginIndex, const int32 EndIndex)
-	{
-		SolverBodyContainer.GatherInput(Dt, BeginIndex, EndIndex);
+		void FPBDConstraintGroupSolver::GatherBodies(const FReal Dt, const int32 BeginIndex, const int32 EndIndex)
+		{
+			SolverBodyContainer.GatherInput(Dt, BeginIndex, EndIndex);
 
-		// Allow derived class to set up some extra data (level, color, etc)
-		GatherBodiesImpl(Dt, BeginIndex, EndIndex);
-	}
+			// Allow derived class to set up some extra data (level, color, etc)
+			GatherBodiesImpl(Dt, BeginIndex, EndIndex);
+		}
 
-	template<typename LambdaType>
-	void FPBDConstraintGroupSolver::ApplyToConstraintRange(const int32 BeginConstraintIndex, const int32 EndConstraintIndex, const LambdaType& Lambda)
-	{
-		// Loop over all the solver containers until we find the one that contains BeginConstraintIndex
-		int32 BeginIndex = BeginConstraintIndex;
-		int32 EndIndex = EndConstraintIndex;
-		for (int32 ContainerIndex = 0; ContainerIndex < ConstraintContainerSolvers.Num(); ++ContainerIndex)
+		template<typename LambdaType>
+		void FPBDConstraintGroupSolver::ApplyToConstraintRange(const int32 BeginConstraintIndex, const int32 EndConstraintIndex, const LambdaType& Lambda)
 		{
-			if (FConstraintContainerSolver* ConstraintSolver = ConstraintContainerSolvers[ContainerIndex].Get())
+			// Loop over all the solver containers until we find the one that contains BeginConstraintIndex
+			int32 BeginIndex = BeginConstraintIndex;
+			int32 EndIndex = EndConstraintIndex;
+			for (int32 ContainerIndex = 0; ContainerIndex < ConstraintContainerSolvers.Num(); ++ContainerIndex)
 			{
-				const int32 NumSolverConstraints = ConstraintSolver->GetNumConstraints();
-				if (BeginIndex < NumSolverConstraints)
-				{
-					// The current range start is in this container. 
-					// Calculate how many of the constraints we should process and call the lambda
-					const int32 BeginSolverIndex = BeginIndex;
-					const int32 EndSolverIndex = FMath::Min(EndIndex, NumSolverConstraints);
-
-					Lambda(ConstraintSolver, BeginSolverIndex, EndSolverIndex);
-				}
-
-				// Remove the constraints we just processed from the range. The range begin end is now relative to the next container (or empty)
-				BeginIndex = FMath::Max(0, BeginIndex - NumSolverConstraints);
-				EndIndex = FMath::Max(0, EndIndex - NumSolverConstraints);
-				if (EndIndex <= 0)
+				if (FConstraintContainerSolver* ConstraintSolver = ConstraintContainerSolvers[ContainerIndex].Get())
 				{
-					break;
+					const int32 NumSolverConstraints = ConstraintSolver->GetNumConstraints();
+					if (BeginIndex < NumSolverConstraints)
+					{
+						// The current range start is in this container. 
+						// Calculate how many of the constraints we should process and call the lambda
+						const int32 BeginSolverIndex = BeginIndex;
+						const int32 EndSolverIndex = FMath::Min(EndIndex, NumSolverConstraints);
+
+						Lambda(ConstraintSolver, BeginSolverIndex, EndSolverIndex);
+					}
+
+					// Remove the constraints we just processed from the range. The range begin end is now relative to the next container (or empty)
+					BeginIndex = FMath::Max(0, BeginIndex - NumSolverConstraints);
+					EndIndex = FMath::Max(0, EndIndex - NumSolverConstraints);
+					if (EndIndex <= 0)
+					{
+						break;
+					}
 				}
 			}
-		}
 
-		// We should have processed the whole range
-		check(EndIndex == 0);
-	}
+			// We should have processed the whole range
+			check(EndIndex == 0);
+		}
 
 
-	void FPBDConstraintGroupSolver::GatherConstraints(const FReal Dt, const int32 BeginConstraintIndex, const int32 EndConstraintIndex)
-	{
-		ApplyToConstraintRange(BeginConstraintIndex, EndConstraintIndex,
-			[this, Dt](FConstraintContainerSolver* ConstraintSolver, const int32 BeginSolverIndex, const int32 EndSolverIndex)
-			{
-				ConstraintSolver->GatherInput(Dt, BeginSolverIndex, EndSolverIndex);
-			});
-	}
+		void FPBDConstraintGroupSolver::GatherConstraints(const FReal Dt, const int32 BeginConstraintIndex, const int32 EndConstraintIndex)
+		{
+			ApplyToConstraintRange(BeginConstraintIndex, EndConstraintIndex,
+				[this, Dt](FConstraintContainerSolver* ConstraintSolver, const int32 BeginSolverIndex, const int32 EndSolverIndex)
+				{
+					ConstraintSolver->GatherInput(Dt, BeginSolverIndex, EndSolverIndex);
+				});
+		}
 
-	void FPBDConstraintGroupSolver::ScatterBodies(const FReal Dt)
-	{
-		SolverBodyContainer.ScatterOutput(0, SolverBodyContainer.Num());
-	}
+		void FPBDConstraintGroupSolver::ScatterBodies(const FReal Dt)
+		{
+			SolverBodyContainer.ScatterOutput(0, SolverBodyContainer.Num());
+		}
 
-	void FPBDConstraintGroupSolver::ScatterConstraints(const FReal Dt)
-	{
-		for (int32 ContainerIndex = 0; ContainerIndex < ConstraintContainerSolvers.Num(); ++ContainerIndex)
+		void FPBDConstraintGroupSolver::ScatterConstraints(const FReal Dt)
 		{
-			if (ConstraintContainerSolvers[ContainerIndex] != nullptr)
+			for (int32 ContainerIndex = 0; ContainerIndex < ConstraintContainerSolvers.Num(); ++ContainerIndex)
 			{
-				ConstraintContainerSolvers[ContainerIndex]->ScatterOutput(Dt);
+				if (ConstraintContainerSolvers[ContainerIndex] != nullptr)
+				{
+					ConstraintContainerSolvers[ContainerIndex]->ScatterOutput(Dt);
+				}
 			}
 		}
-	}
-
-	void FPBDConstraintGroupSolver::ScatterBodies(const FReal Dt, const int32 BeginBodyIndex, const int32 EndBodyIndex)
-	{
-		SolverBodyContainer.ScatterOutput(BeginBodyIndex, EndBodyIndex);
-	}
 
-	void FPBDConstraintGroupSolver::ScatterConstraints(const FReal Dt, const int32 BeginConstraintIndex, const int32 EndConstraintIndex)
-	{
-		ApplyToConstraintRange(BeginConstraintIndex, EndConstraintIndex,
-			[this, Dt](FConstraintContainerSolver* ConstraintSolver, const int32 BeginSolverIndex, const int32 EndSolverIndex)
-			{
-				ConstraintSolver->ScatterOutput(Dt, BeginSolverIndex, EndSolverIndex);
-			});
-	}
+		void FPBDConstraintGroupSolver::ScatterBodies(const FReal Dt, const int32 BeginBodyIndex, const int32 EndBodyIndex)
+		{
+			SolverBodyContainer.ScatterOutput(BeginBodyIndex, EndBodyIndex);
+		}
 
+		void FPBDConstraintGroupSolver::ScatterConstraints(const FReal Dt, const int32 BeginConstraintIndex, const int32 EndConstraintIndex)
+		{
+			ApplyToConstraintRange(BeginConstraintIndex, EndConstraintIndex,
+				[this, Dt](FConstraintContainerSolver* ConstraintSolver, const int32 BeginSolverIndex, const int32 EndSolverIndex)
+				{
+					ConstraintSolver->ScatterOutput(Dt, BeginSolverIndex, EndSolverIndex);
+				});
+		}
 
-	void FPBDConstraintGroupSolver::ApplyPositionConstraints(const FReal Dt, const int32 NumIts)
-	{
-		// NOTE: We loop over prioritized solvers here
-		for (int32 It = 0; It < NumIts; ++It)
+		void FPBDConstraintGroupSolver::ApplyPositionConstraints(const FReal Dt)
 		{
-			for (int32 ContainerIndex = 0; ContainerIndex < PrioritizedConstraintContainerSolvers.Num(); ++ContainerIndex)
+			const int32 NumIts = Iterations.GetNumPositionIterations();
+
+			// NOTE: We loop over prioritized solvers here
+			for (int32 It = 0; It < NumIts; ++It)
 			{
-				if (PrioritizedConstraintContainerSolvers[ContainerIndex] != nullptr)
+				for (int32 ContainerIndex = 0; ContainerIndex < PrioritizedConstraintContainerSolvers.Num(); ++ContainerIndex)
 				{
-					PrioritizedConstraintContainerSolvers[ContainerIndex]->ApplyPositionConstraints(Dt, It, NumIts);
+					if (PrioritizedConstraintContainerSolvers[ContainerIndex] != nullptr)
+					{
+						PrioritizedConstraintContainerSolvers[ContainerIndex]->ApplyPositionConstraints(Dt, It, NumIts);
+					}
 				}
 			}
-		}
 
-		// Calculate the velocity from the net change in position after applying position constraints
-		SolverBodyContainer.SetImplicitVelocities(Dt);
-	}
+			// Calculate the velocity from the net change in position after applying position constraints
+			SolverBodyContainer.SetImplicitVelocities(Dt);
+		}
 
-	void FPBDConstraintGroupSolver::ApplyVelocityConstraints(const FReal Dt, const int32 NumIts)
-	{
-		// NOTE: We loop over prioritized solvers here
-		for (int32 It = 0; It < NumIts; ++It)
+		void FPBDConstraintGroupSolver::ApplyVelocityConstraints(const FReal Dt)
 		{
-			for (int32 ContainerIndex = 0; ContainerIndex < PrioritizedConstraintContainerSolvers.Num(); ++ContainerIndex)
+			const int32 NumIts = Iterations.GetNumVelocityIterations();
+
+			// NOTE: We loop over prioritized solvers here
+			for (int32 It = 0; It < NumIts; ++It)
 			{
-				if (PrioritizedConstraintContainerSolvers[ContainerIndex] != nullptr)
+				for (int32 ContainerIndex = 0; ContainerIndex < PrioritizedConstraintContainerSolvers.Num(); ++ContainerIndex)
 				{
-					PrioritizedConstraintContainerSolvers[ContainerIndex]->ApplyVelocityConstraints(Dt, It, NumIts);
+					if (PrioritizedConstraintContainerSolvers[ContainerIndex] != nullptr)
+					{
+						PrioritizedConstraintContainerSolvers[ContainerIndex]->ApplyVelocityConstraints(Dt, It, NumIts);
+					}
 				}
 			}
 		}
-	}
-
-	void FPBDConstraintGroupSolver::ApplyProjectionConstraints(const FReal Dt, const int32 NumIts)
-	{
-		// Update the body transforms from the deltas calculated in the constraint solve phases 1 and 2
-		// NOTE: deliberately not updating the world-space inertia as it is not used by joint projection
-		// and no other constraints currently implement projection
-		SolverBodyContainer.ApplyCorrections();
 
-		// NOTE: We loop over prioritized solvers here
-		for (int32 It = 0; It < NumIts; ++It)
+		void FPBDConstraintGroupSolver::ApplyProjectionConstraints(const FReal Dt)
 		{
-			for (int32 ContainerIndex = 0; ContainerIndex < PrioritizedConstraintContainerSolvers.Num(); ++ContainerIndex)
+			const int32 NumIts = Iterations.GetNumProjectionIterations();
+
+			// Update the body transforms from the deltas calculated in the constraint solve phases 1 and 2
+			// NOTE: deliberately not updating the world-space inertia as it is not used by joint projection
+			// and no other constraints currently implement projection
+			SolverBodyContainer.ApplyCorrections();
+
+			// NOTE: We loop over prioritized solvers here
+			for (int32 It = 0; It < NumIts; ++It)
 			{
-				if (PrioritizedConstraintContainerSolvers[ContainerIndex] != nullptr)
+				for (int32 ContainerIndex = 0; ContainerIndex < PrioritizedConstraintContainerSolvers.Num(); ++ContainerIndex)
 				{
-					PrioritizedConstraintContainerSolvers[ContainerIndex]->ApplyProjectionConstraints(Dt, It, NumIts);
+					if (PrioritizedConstraintContainerSolvers[ContainerIndex] != nullptr)
+					{
+						PrioritizedConstraintContainerSolvers[ContainerIndex]->ApplyProjectionConstraints(Dt, It, NumIts);
+					}
 				}
 			}
 		}
-	}
 
-	//
-	//
-	//
-	//
-	//
+		//
+		//
+		//
+		//
+		//
 
-	void FPBDSceneConstraintGroupSolver::AddConstraintsImpl()
-	{
-		// The Scene Solver solves all constraints, so just add everything
-		for (int32 ContainerIndex = 0; ContainerIndex < ConstraintContainerSolvers.Num(); ++ContainerIndex)
+		void FPBDSceneConstraintGroupSolver::AddConstraintsImpl()
 		{
-			if (ConstraintContainerSolvers[ContainerIndex] != nullptr)
+			// The Scene Solver solves all constraints, so just add everything
+			for (int32 ContainerIndex = 0; ContainerIndex < ConstraintContainerSolvers.Num(); ++ContainerIndex)
 			{
-				ConstraintContainerSolvers[ContainerIndex]->AddConstraints();
+				if (ConstraintContainerSolvers[ContainerIndex] != nullptr)
+				{
+					ConstraintContainerSolvers[ContainerIndex]->AddConstraints();
+				}
 			}
 		}
-	}
 
-}
\ No newline at end of file
+	}	// namespace Private
+}	// namespace Chaos
diff --git a/Engine/Source/Runtime/Experimental/Chaos/Private/Chaos/Evolution/PBDMinEvolution.cpp b/Engine/Source/Runtime/Experimental/Chaos/Private/Chaos/Evolution/PBDMinEvolution.cpp
--- a/Engine/Source/Runtime/Experimental/Chaos/Private/Chaos/Evolution/PBDMinEvolution.cpp
+++ b/Engine/Source/Runtime/Experimental/Chaos/Private/Chaos/Evolution/PBDMinEvolution.cpp
@@ -58,9 +58,7 @@ namespace Chaos
 		, CollisionDetector(InCollisionDetector)
 		, ParticlePrevXs(InPrevX)
 		, ParticlePrevRs(InPrevR)
-		, NumPositionIterations(0)
-		, NumVelocityIterations(0)
-		, NumProjectionIterations(0)
+		, ConstraintSolver(Private::FIterationSettings(0,0,0))
 		, Gravity(FVec3(0))
 		, SimulationSpaceSettings()
 	{
@@ -390,20 +388,20 @@ namespace Chaos
 	{
 		SCOPE_CYCLE_COUNTER(STAT_MinEvolution_ApplyConstraintsPhase1);
 
-		ConstraintSolver.ApplyPositionConstraints(Dt, NumPositionIterations);
+		ConstraintSolver.ApplyPositionConstraints(Dt);
 	}
 
 	void FPBDMinEvolution::ApplyConstraintsPhase2(FReal Dt)
 	{
 		SCOPE_CYCLE_COUNTER(STAT_MinEvolution_ApplyConstraintsPhase2);
 
-		ConstraintSolver.ApplyVelocityConstraints(Dt, NumVelocityIterations);
+		ConstraintSolver.ApplyVelocityConstraints(Dt);
 	}
 
 	void FPBDMinEvolution::ApplyConstraintsPhase3(FReal Dt)
 	{
 		SCOPE_CYCLE_COUNTER(STAT_MinEvolution_ApplyConstraintsPhase3);
 
-		ConstraintSolver.ApplyProjectionConstraints(Dt, NumProjectionIterations);
+		ConstraintSolver.ApplyProjectionConstraints(Dt);
 	}
 }
diff --git a/Engine/Source/Runtime/Experimental/Chaos/Private/Chaos/Island/IslandGraph.cpp b/Engine/Source/Runtime/Experimental/Chaos/Private/Chaos/Island/IslandGraph.cpp
--- a/Engine/Source/Runtime/Experimental/Chaos/Private/Chaos/Island/IslandGraph.cpp
+++ b/Engine/Source/Runtime/Experimental/Chaos/Private/Chaos/Island/IslandGraph.cpp
@@ -979,6 +979,6 @@ void FIslandGraph<NodeType, EdgeType, IslandType, OwnerType>::RemoveAllAwakeEdge
 	} 
 }
 
-template class FIslandGraph<FGeometryParticleHandle*, FConstraintHandleHolder, FPBDIsland*, FPBDIslandManager>;
+template class FIslandGraph<FGeometryParticleHandle*, FConstraintHandleHolder, Private::FPBDIsland*, Private::FPBDIslandManager>;
 template class FIslandGraph<int32, int32, int32, TNullIslandGraphOwner<int32, int32>>;
 }
diff --git a/Engine/Source/Runtime/Experimental/Chaos/Private/Chaos/Island/IslandGroup.cpp b/Engine/Source/Runtime/Experimental/Chaos/Private/Chaos/Island/IslandGroup.cpp
--- a/Engine/Source/Runtime/Experimental/Chaos/Private/Chaos/Island/IslandGroup.cpp
+++ b/Engine/Source/Runtime/Experimental/Chaos/Private/Chaos/Island/IslandGroup.cpp
@@ -8,85 +8,96 @@
 
 namespace Chaos
 {
-	FPBDIslandConstraintGroupSolver::FPBDIslandConstraintGroupSolver(FPBDIslandManager& InIslandManager)
-		: FPBDConstraintGroupSolver()
-		, IslandManager(InIslandManager)
-		, NumParticles(0)
-		, NumConstraints(0)
+	namespace Private
 	{
-	}
-
-	void FPBDIslandConstraintGroupSolver::SetConstraintSolverImpl(int32 ContainerId)
-	{
-		if (ContainerId >= NumContainerConstraints.Num())
+		FPBDIslandConstraintGroupSolver::FPBDIslandConstraintGroupSolver(FPBDIslandManager& InIslandManager)
+			: FPBDConstraintGroupSolver()
+			, IslandManager(InIslandManager)
+			, NumParticles(0)
+			, NumConstraints(0)
 		{
-			NumContainerConstraints.SetNumZeroed(ContainerId + 1);
 		}
-	}
 
-	void FPBDIslandConstraintGroupSolver::ResetImpl()
-	{
-		NumParticles = 0;
-		NumConstraints = 0;
-		for (int32 ContainerIndex = 0; ContainerIndex < ConstraintContainerSolvers.Num(); ++ContainerIndex)
+		void FPBDIslandConstraintGroupSolver::SetConstraintSolverImpl(int32 ContainerId)
 		{
-			NumContainerConstraints[ContainerIndex] = 0;
+			if (ContainerId >= NumContainerConstraints.Num())
+			{
+				NumContainerConstraints.SetNumZeroed(ContainerId + 1);
+			}
 		}
 
-		Islands.Reset();
-	}
-
-	void FPBDIslandConstraintGroupSolver::AddIsland(FPBDIsland* Island)
-	{
-		if (Island)
+		void FPBDIslandConstraintGroupSolver::ResetImpl()
 		{
-			Islands.Add(Island);
-
-			NumParticles += Island->GetNumParticles();
-
+			NumParticles = 0;
+			NumConstraints = 0;
 			for (int32 ContainerIndex = 0; ContainerIndex < ConstraintContainerSolvers.Num(); ++ContainerIndex)
 			{
-				const int32 NumIslandConstraints = Island->GetConstraints(ContainerIndex).Num();
-				NumContainerConstraints[ContainerIndex] += NumIslandConstraints;
-				NumConstraints += NumIslandConstraints;
+				NumContainerConstraints[ContainerIndex] = 0;
 			}
-		}
-	}
 
-	void FPBDIslandConstraintGroupSolver::AddConstraintsImpl()
-	{
-		// Initialize buffer sizes for solver particles
-		SolverBodyContainer.Reset(NumParticles);
+			Islands.Reset();
+		}
 
-		// Add all the constraints to the solvers, but do not collect data from the constraints yet (each constraint type has its own solver). 
-		// This also sets up the SolverBody array, but does not gather the body data.
-		for (int32 ContainerIndex = 0; ContainerIndex < ConstraintContainerSolvers.Num(); ++ContainerIndex)
+		void FPBDIslandConstraintGroupSolver::AddIsland(FPBDIsland* Island)
 		{
-			if (ConstraintContainerSolvers[ContainerIndex] != nullptr)
+			if (Island)
 			{
-				// Initialize buffer sizes for constraint solvers
-				ConstraintContainerSolvers[ContainerIndex]->Reset(NumContainerConstraints[ContainerIndex]);
+				Islands.Add(Island);
 
-				for (FPBDIsland* Island : Islands)
+				NumParticles += Island->GetNumParticles();
+
+				for (int32 ContainerIndex = 0; ContainerIndex < ConstraintContainerSolvers.Num(); ++ContainerIndex)
 				{
-					TArrayView<FPBDIslandConstraint> IslandConstraints = MakeArrayView(Island->GetConstraints(ContainerIndex));
-					ConstraintContainerSolvers[ContainerIndex]->AddConstraints(IslandConstraints);
+					const int32 NumIslandConstraints = Island->GetConstraints(ContainerIndex).Num();
+					NumContainerConstraints[ContainerIndex] += NumIslandConstraints;
+					NumConstraints += NumIslandConstraints;
 				}
 			}
 		}
-	}
 
-	void FPBDIslandConstraintGroupSolver::GatherBodiesImpl(const FReal Dt, const int32 BeginBodyIndex, const int32 EndBodyIndex)
-	{
-		// @todo(chaos): optimize
-		for (int32 SolverBodyIndex = BeginBodyIndex; SolverBodyIndex < EndBodyIndex; ++SolverBodyIndex)
+		void FPBDIslandConstraintGroupSolver::AddConstraintsImpl()
+		{
+			// Initialize buffer sizes for solver particles
+			SolverBodyContainer.Reset(NumParticles);
+
+			// Add all the constraints to the solvers, but do not collect data from the constraints yet (each constraint type has its own solver). 
+			// This also sets up the SolverBody array, but does not gather the body data.
+			for (int32 ContainerIndex = 0; ContainerIndex < ConstraintContainerSolvers.Num(); ++ContainerIndex)
+			{
+				if (ConstraintContainerSolvers[ContainerIndex] != nullptr)
+				{
+					// Initialize buffer sizes for constraint solvers
+					ConstraintContainerSolvers[ContainerIndex]->Reset(NumContainerConstraints[ContainerIndex]);
+
+					for (FPBDIsland* Island : Islands)
+					{
+						TArrayView<FPBDIslandConstraint> IslandConstraints = MakeArrayView(Island->GetConstraints(ContainerIndex));
+						ConstraintContainerSolvers[ContainerIndex]->AddConstraints(IslandConstraints);
+					}
+				}
+			}
+		}
+
+		void FPBDIslandConstraintGroupSolver::GatherBodiesImpl(const FReal Dt, const int32 BeginBodyIndex, const int32 EndBodyIndex)
 		{
-			FSolverBody& SolverBody = SolverBodyContainer.GetSolverBody(SolverBodyIndex);
-			FGeometryParticleHandle* Particle = SolverBodyContainer.GetParticle(SolverBodyIndex);
-			const int32 Level = IslandManager.GetParticleLevel(Particle);
-			SolverBody.SetLevel(Level);
+			// @todo(chaos): optimize
+			for (int32 SolverBodyIndex = BeginBodyIndex; SolverBodyIndex < EndBodyIndex; ++SolverBodyIndex)
+			{
+				FSolverBody& SolverBody = SolverBodyContainer.GetSolverBody(SolverBodyIndex);
+				FGeometryParticleHandle* Particle = SolverBodyContainer.GetParticle(SolverBodyIndex);
+				const int32 Level = IslandManager.GetParticleLevel(Particle);
+				SolverBody.SetLevel(Level);
+			}
 		}
-	}
 
+		void FPBDIslandConstraintGroupSolver::SetIterationSettings(const FIterationSettings& InDefaultIterations)
+		{
+			Iterations = InDefaultIterations;
+			for (FPBDIsland* Island : Islands)
+			{
+				Iterations = FIterationSettings::Merge(Iterations, Island->GetIterationSettings());
+			}
+		}
 
-}
\ No newline at end of file
+	}	// namespace Private
+}	// namespace Chaos
\ No newline at end of file
diff --git a/Engine/Source/Runtime/Experimental/Chaos/Private/Chaos/Island/IslandGroupManager.cpp b/Engine/Source/Runtime/Experimental/Chaos/Private/Chaos/Island/IslandGroupManager.cpp
--- a/Engine/Source/Runtime/Experimental/Chaos/Private/Chaos/Island/IslandGroupManager.cpp
+++ b/Engine/Source/Runtime/Experimental/Chaos/Private/Chaos/Island/IslandGroupManager.cpp
@@ -53,499 +53,504 @@ namespace Chaos
 
 	}
 
-
-	FPBDIslandGroupManager::FPBDIslandGroupManager(FPBDIslandManager& InIslandManager)
-		: IslandManager(InIslandManager)
-		, IslandGroups()
-		, NumActiveGroups(0)
-		, NumWorkerThreads(0)
-		, TargetNumBodiesPerTask(0)
-		, TargetNumConstraintsPerTask(0)
-		, NumPositionIterations(0)
-		, NumVelocityIterations(0)
-		, NumProjectionIterations(0)
-	{
-		// Check for use of the "-onethread" command line arg, and physics threading disabled (GetNumWorkerThreads() is not affected by these)
-		// @todo(chaos): is the number of worker threads a good indicator of how many threads we get in the solver loop? (Currently uses ParallelFor)
-		NumWorkerThreads = (FApp::ShouldUseThreadingForPerformance() && !GSingleThreadedPhysics) ? FMath::Min(FTaskGraphInterface::Get().GetNumWorkerThreads(), Chaos::MaxNumWorkers) : 0;
-		const int32 MaxIslandGroups = (CVars::GIslandGroupsMaxWorkers > 0) ? CVars::GIslandGroupsMaxWorkers : TNumericLimits<int32>::Max();
-		const int32 NumIslandGroups = FMath::Clamp(FMath::CeilToInt32(FReal(NumWorkerThreads) * CVars::GIslandGroupsWorkerMultiplier), 1, MaxIslandGroups);
-
-		IslandGroups.Reserve(NumIslandGroups);
-		for (int32 GroupIndex = 0; GroupIndex < NumIslandGroups; ++GroupIndex)
-		{
-			IslandGroups.Emplace(MakeUnique<FPBDIslandConstraintGroupSolver>(IslandManager));
-		}
-	}
-
-	FPBDIslandGroupManager::~FPBDIslandGroupManager()
+	namespace Private
 	{
 
-	}
-
-	void FPBDIslandGroupManager::AddConstraintContainer(FPBDConstraintContainer& ConstraintContainer, const int32 Priority)
-	{
-		const int32 ContainerId = ConstraintContainer.GetContainerId();
-
-		for (int32 GroupIndex = 0; GroupIndex < IslandGroups.Num(); ++GroupIndex)
+		FPBDIslandGroupManager::FPBDIslandGroupManager(FPBDIslandManager& InIslandManager)
+			: IslandManager(InIslandManager)
+			, IslandGroups()
+			, NumActiveGroups(0)
+			, NumWorkerThreads(0)
+			, TargetNumBodiesPerTask(0)
+			, TargetNumConstraintsPerTask(0)
+			, Iterations(0, 0, 0)
 		{
-			IslandGroups[GroupIndex]->SetConstraintSolver(ContainerId, ConstraintContainer.CreateGroupSolver(Priority));
+			// Check for use of the "-onethread" command line arg, and physics threading disabled (GetNumWorkerThreads() is not affected by these)
+			// @todo(chaos): is the number of worker threads a good indicator of how many threads we get in the solver loop? (Currently uses ParallelFor)
+			NumWorkerThreads = (FApp::ShouldUseThreadingForPerformance() && !GSingleThreadedPhysics) ? FMath::Min(FTaskGraphInterface::Get().GetNumWorkerThreads(), Chaos::MaxNumWorkers) : 0;
+			const int32 MaxIslandGroups = (CVars::GIslandGroupsMaxWorkers > 0) ? CVars::GIslandGroupsMaxWorkers : TNumericLimits<int32>::Max();
+			const int32 NumIslandGroups = FMath::Clamp(FMath::CeilToInt32(FReal(NumWorkerThreads) * CVars::GIslandGroupsWorkerMultiplier), 1, MaxIslandGroups);
+
+			IslandGroups.Reserve(NumIslandGroups);
+			for (int32 GroupIndex = 0; GroupIndex < NumIslandGroups; ++GroupIndex)
+			{
+				IslandGroups.Emplace(MakeUnique<FPBDIslandConstraintGroupSolver>(IslandManager));
+			}
 		}
-	}
 
-	void FPBDIslandGroupManager::SetConstraintContainerPriority(const int32 ContainerId, const int32 Priority)
-	{
-		for (int32 GroupIndex = 0; GroupIndex < IslandGroups.Num(); ++GroupIndex)
+		FPBDIslandGroupManager::~FPBDIslandGroupManager()
 		{
-			IslandGroups[GroupIndex]->SetConstraintSolverPriority(ContainerId, Priority);
+
 		}
-	}
 
-	int32 FPBDIslandGroupManager::BuildGroups()
-	{
-		// The set of islands with some work to do
-		// NOTE: We do not add islands with no constraints even if they have particles. There is no need since the Particles'
-		// predicted positions/rotations will not be changed by the solver. (Zero-constraint islands may contain one isolated Particle).
-		TArray<FPBDIsland*> Islands;
-		for (int32 IslandIndex = 0; IslandIndex < IslandManager.NumIslands(); ++IslandIndex)
+		void FPBDIslandGroupManager::AddConstraintContainer(FPBDConstraintContainer& ConstraintContainer, const int32 Priority)
 		{
-			FPBDIsland* Island = IslandManager.GetIsland(IslandIndex);
-			if (!Island->IsSleeping() && (Island->GetNumConstraints() > 0))
+			const int32 ContainerId = ConstraintContainer.GetContainerId();
+
+			for (int32 GroupIndex = 0; GroupIndex < IslandGroups.Num(); ++GroupIndex)
 			{
-				Islands.Add(Island);
+				IslandGroups[GroupIndex]->SetConstraintSolver(ContainerId, ConstraintContainer.CreateGroupSolver(Priority));
 			}
 		}
 
-		// Sort islands by constraint count - largest first
-		// @todo(chaos): Do we need this?
-		if (Islands.Num() > 1)
+		void FPBDIslandGroupManager::SetConstraintContainerPriority(const int32 ContainerId, const int32 Priority)
 		{
-			Algo::Sort(Islands, 
-				[](const FPBDIsland* L, const FPBDIsland* R) -> bool
-				{
-					return L->GetNumConstraints() > R->GetNumConstraints();
-				});
-		}
-
-		// Figure out how many bodies and constraints we'd like per group for gather, scatter and solve
-		const int32 MaxGroups = IslandGroups.Num();
-		int32 NumAllBodies = 0;
-		int32 NumAllConstraints = 0;
-		for (FPBDIsland* Island : Islands)
-		{
-			NumAllBodies += Island->GetNumParticles();
-			NumAllConstraints += Island->GetNumConstraints();
-		}
-		TargetNumConstraintsPerTask = NumAllConstraints / MaxGroups + 1;
-		TargetNumBodiesPerTask = NumAllBodies / MaxGroups + 1;
-
-		// We want to gather or solve a minimum number of bodies and constraints per task if possible.
-		// There's not point creating 100 tasks if each task only processes one constraint - the task overhead
-		// would outweight the benefits of going wide.
-		// @todo(chaos): we may want to consider separating the gather task count from the number of island groups by adding a second multiplier.
-		TargetNumConstraintsPerTask = FMath::Max(TargetNumConstraintsPerTask, CVars::GIslandGroupsMinConstraintsPerWorker);
-		TargetNumBodiesPerTask = FMath::Max(TargetNumConstraintsPerTask, CVars::GIslandGroupsMinBodiesPerWorker);
-
-		// Reset all the groups
-		for (TUniquePtr<FPBDIslandConstraintGroupSolver>& IslandGroup : IslandGroups)
-		{
-			IslandGroup->Reset();
+			for (int32 GroupIndex = 0; GroupIndex < IslandGroups.Num(); ++GroupIndex)
+			{
+				IslandGroups[GroupIndex]->SetConstraintSolverPriority(ContainerId, Priority);
+			}
 		}
-		NumActiveGroups = 0;
 
-		// Add each Island to the first group with enough space, or the group with the fewest constraint if none have enough space
-		// @todo(chaos): optimize - when a group is full we should move it to the back so we don't keep visiting it
-		for (FPBDIsland* Island : Islands)
+		int32 FPBDIslandGroupManager::BuildGroups()
 		{
-			const int32 NumIslandConstraints = Island->GetNumConstraints();
-
-			int32 InsertGroupIndex = INDEX_NONE;
-			int32 SmallestGroupSize = TNumericLimits<int32>::Max();
-			for (int32 GroupIndex = 0; GroupIndex < IslandGroups.Num(); ++GroupIndex)
+			// The set of islands with some work to do
+			// NOTE: We do not add islands with no constraints even if they have particles. There is no need since the Particles'
+			// predicted positions/rotations will not be changed by the solver. (Zero-constraint islands may contain one isolated Particle).
+			TArray<FPBDIsland*> Islands;
+			for (int32 IslandIndex = 0; IslandIndex < IslandManager.NumIslands(); ++IslandIndex)
 			{
-				const int32 NumGroupConstraints = IslandGroups[GroupIndex]->GetNumConstraints();
-
-				if (NumGroupConstraints < SmallestGroupSize)
+				FPBDIsland* Island = IslandManager.GetIsland(IslandIndex);
+				if (!Island->IsSleeping() && (Island->GetNumConstraints() > 0))
 				{
-					InsertGroupIndex = GroupIndex;
-					SmallestGroupSize = NumGroupConstraints;
+					Islands.Add(Island);
 				}
+			}
 
-				if (NumGroupConstraints + NumIslandConstraints < TargetNumConstraintsPerTask)
-				{
-					InsertGroupIndex = GroupIndex;
-					break;
-				}
+			// Sort islands by constraint count - largest first
+			// @todo(chaos): Do we need this?
+			if (Islands.Num() > 1)
+			{
+				Algo::Sort(Islands, 
+					[](const FPBDIsland* L, const FPBDIsland* R) -> bool
+					{
+						return L->GetNumConstraints() > R->GetNumConstraints();
+					});
 			}
 
-			check(InsertGroupIndex != INDEX_NONE);
-			IslandGroups[InsertGroupIndex]->AddIsland(Island);
+			// Figure out how many bodies and constraints we'd like per group for gather, scatter and solve
+			const int32 MaxGroups = IslandGroups.Num();
+			int32 NumAllBodies = 0;
+			int32 NumAllConstraints = 0;
+			for (FPBDIsland* Island : Islands)
+			{
+				NumAllBodies += Island->GetNumParticles();
+				NumAllConstraints += Island->GetNumConstraints();
+			}
+			TargetNumConstraintsPerTask = NumAllConstraints / MaxGroups + 1;
+			TargetNumBodiesPerTask = NumAllBodies / MaxGroups + 1;
+
+			// We want to gather or solve a minimum number of bodies and constraints per task if possible.
+			// There's not point creating 100 tasks if each task only processes one constraint - the task overhead
+			// would outweight the benefits of going wide.
+			// @todo(chaos): we may want to consider separating the gather task count from the number of island groups by adding a second multiplier.
+			TargetNumConstraintsPerTask = FMath::Max(TargetNumConstraintsPerTask, CVars::GIslandGroupsMinConstraintsPerWorker);
+			TargetNumBodiesPerTask = FMath::Max(TargetNumConstraintsPerTask, CVars::GIslandGroupsMinBodiesPerWorker);
+
+			// Reset all the groups
+			for (TUniquePtr<FPBDIslandConstraintGroupSolver>& IslandGroup : IslandGroups)
+			{
+				IslandGroup->Reset();
+			}
+			NumActiveGroups = 0;
 
-			NumActiveGroups = FMath::Max(NumActiveGroups, InsertGroupIndex + 1);
-		}
+			// Add each Island to the first group with enough space, or the group with the fewest constraint if none have enough space
+			// @todo(chaos): optimize - when a group is full we should move it to the back so we don't keep visiting it
+			for (FPBDIsland* Island : Islands)
+			{
+				const int32 NumIslandConstraints = Island->GetNumConstraints();
 
-		return NumActiveGroups;
-	}
+				int32 InsertGroupIndex = INDEX_NONE;
+				int32 SmallestGroupSize = TNumericLimits<int32>::Max();
+				for (int32 GroupIndex = 0; GroupIndex < IslandGroups.Num(); ++GroupIndex)
+				{
+					const int32 NumGroupConstraints = IslandGroups[GroupIndex]->GetNumConstraints();
 
-	void FPBDIslandGroupManager::BuildGatherBatches(TArray<FIslandGroupRange>& BodyRanges, TArray<FIslandGroupRange>& ConstraintRanges)
-	{
-		BodyRanges.Reset();
-		BodyRanges.Reserve(IslandGroups.Num());
-		ConstraintRanges.Reset();
-		ConstraintRanges.Reserve(IslandGroups.Num());
-
-		// If a group has a lot of constraints in it (compared to the target number per task) split the gather into multiple tasks
-		// The goal is to end up with the same number of tasks as we have task threads (see comments in SolveParallelTasks for more info)
-		for (int32 GroupIndex = 0; GroupIndex < NumActiveGroups; ++GroupIndex)
-		{
-			FPBDIslandConstraintGroupSolver* IslandGroup = GetGroup(GroupIndex);
+					if (NumGroupConstraints < SmallestGroupSize)
+					{
+						InsertGroupIndex = GroupIndex;
+						SmallestGroupSize = NumGroupConstraints;
+					}
 
-			const int32 NumConstraints = IslandGroup->GetNumSolverConstraints();
-			const int32 NumConstraintGatherTasks = FMath::Max(1, FMath::DivideAndRoundNearest(NumConstraints, TargetNumConstraintsPerTask));
-			const int32 NumConstraintsPerTask = FMath::DivideAndRoundUp(NumConstraints, NumConstraintGatherTasks);
+					if (NumGroupConstraints + NumIslandConstraints < TargetNumConstraintsPerTask)
+					{
+						InsertGroupIndex = GroupIndex;
+						break;
+					}
+				}
 
-			for (int32 BeginIndex = 0; BeginIndex < NumConstraints; BeginIndex += NumConstraintsPerTask)
-			{
-				const int32 EndIndex = FMath::Min(BeginIndex + NumConstraintsPerTask, NumConstraints);
+				check(InsertGroupIndex != INDEX_NONE);
+				IslandGroups[InsertGroupIndex]->AddIsland(Island);
 
-				ConstraintRanges.Emplace(IslandGroup, BeginIndex, EndIndex);
+				NumActiveGroups = FMath::Max(NumActiveGroups, InsertGroupIndex + 1);
 			}
 
-			const int32 NumBodies = IslandGroup->GetNumSolverBodies();
-			const int32 NumBodyGatherTasks = FMath::Max(1, FMath::DivideAndRoundNearest(NumBodies, TargetNumBodiesPerTask));
-			const int32 NumBodiesPerTask = FMath::DivideAndRoundUp(NumBodies, NumBodyGatherTasks);
-
-			for (int32 BeginIndex = 0; BeginIndex < NumBodies; BeginIndex += NumBodiesPerTask)
+			// Allow each island group to determine how many iterations to run
+			for (TUniquePtr<FPBDIslandConstraintGroupSolver>& IslandGroup : IslandGroups)
 			{
-				const int32 EndIndex = FMath::Min(BeginIndex + NumBodiesPerTask, NumBodies);
-
-				BodyRanges.Emplace(IslandGroup, BeginIndex, EndIndex);
+				IslandGroup->SetIterationSettings(Iterations);
 			}
+
+			return NumActiveGroups;
 		}
-	}
 
-	void FPBDIslandGroupManager::SolveGroupConstraints(int32 GroupIndex, const FReal Dt)
-	{
-		CSV_SCOPED_TIMING_STAT(Chaos, ApplyConstraints);
+		void FPBDIslandGroupManager::BuildGatherBatches(TArray<FIslandGroupRange>& BodyRanges, TArray<FIslandGroupRange>& ConstraintRanges)
+		{
+			BodyRanges.Reset();
+			BodyRanges.Reserve(IslandGroups.Num());
+			ConstraintRanges.Reset();
+			ConstraintRanges.Reserve(IslandGroups.Num());
+
+			// If a group has a lot of constraints in it (compared to the target number per task) split the gather into multiple tasks
+			// The goal is to end up with the same number of tasks as we have task threads (see comments in SolveParallelTasks for more info)
+			for (int32 GroupIndex = 0; GroupIndex < NumActiveGroups; ++GroupIndex)
+			{
+				FPBDIslandConstraintGroupSolver* IslandGroup = GetGroup(GroupIndex);
 
-		FPBDIslandConstraintGroupSolver* IslandGroup = GetGroup(GroupIndex);
+				const int32 NumConstraints = IslandGroup->GetNumSolverConstraints();
+				const int32 NumConstraintGatherTasks = FMath::Max(1, FMath::DivideAndRoundNearest(NumConstraints, TargetNumConstraintsPerTask));
+				const int32 NumConstraintsPerTask = FMath::DivideAndRoundUp(NumConstraints, NumConstraintGatherTasks);
 
-		{
-			SCOPE_CYCLE_COUNTER(STAT_Evolution_ApplyConstraintsPhase1);
-			CSV_SCOPED_ISLANDGROUP_TIMING_STAT(PerIslandSolve_ApplyTotalSerialized, GroupIndex);
+				for (int32 BeginIndex = 0; BeginIndex < NumConstraints; BeginIndex += NumConstraintsPerTask)
+				{
+					const int32 EndIndex = FMath::Min(BeginIndex + NumConstraintsPerTask, NumConstraints);
 
-			IslandGroup->ApplyPositionConstraints(Dt, NumPositionIterations);
-		}
-		{
-			SCOPE_CYCLE_COUNTER(STAT_Evolution_ApplyConstraintsPhase2);
-			CSV_SCOPED_ISLANDGROUP_TIMING_STAT(PerIslandSolve_ApplyPushOutTotalSerialized, GroupIndex);
+					ConstraintRanges.Emplace(IslandGroup, BeginIndex, EndIndex);
+				}
 
-			IslandGroup->ApplyVelocityConstraints(Dt, NumVelocityIterations);
-		}
-		{
-			SCOPE_CYCLE_COUNTER(STAT_Evolution_ApplyConstraintsPhase3);
-			CSV_SCOPED_ISLANDGROUP_TIMING_STAT(PerIslandSolve_ApplyProjectionTotalSerialized, GroupIndex);
+				const int32 NumBodies = IslandGroup->GetNumSolverBodies();
+				const int32 NumBodyGatherTasks = FMath::Max(1, FMath::DivideAndRoundNearest(NumBodies, TargetNumBodiesPerTask));
+				const int32 NumBodiesPerTask = FMath::DivideAndRoundUp(NumBodies, NumBodyGatherTasks);
 
-			IslandGroup->ApplyProjectionConstraints(Dt, NumProjectionIterations);
-		}
-	}
+				for (int32 BeginIndex = 0; BeginIndex < NumBodies; BeginIndex += NumBodiesPerTask)
+				{
+					const int32 EndIndex = FMath::Min(BeginIndex + NumBodiesPerTask, NumBodies);
 
-	void FPBDIslandGroupManager::SolveSerial(const FReal Dt)
-	{
-		SCOPE_CYCLE_COUNTER(STAT_SolveConstraints);
+					BodyRanges.Emplace(IslandGroup, BeginIndex, EndIndex);
+				}
+			}
+		}
 
-		for (int32 GroupIndex = 0; GroupIndex < NumActiveGroups; ++GroupIndex)
+		void FPBDIslandGroupManager::SolveGroupConstraints(int32 GroupIndex, const FReal Dt)
 		{
+			CSV_SCOPED_TIMING_STAT(Chaos, ApplyConstraints);
+
 			FPBDIslandConstraintGroupSolver* IslandGroup = GetGroup(GroupIndex);
 
 			{
-				SCOPE_CYCLE_COUNTER(STAT_Evolution_Gather);
-				CSV_SCOPED_ISLANDGROUP_TIMING_STAT(PerIslandSolve_GatherTotalSerialized, GroupIndex);
+				SCOPE_CYCLE_COUNTER(STAT_Evolution_ApplyConstraintsPhase1);
+				CSV_SCOPED_ISLANDGROUP_TIMING_STAT(PerIslandSolve_ApplyTotalSerialized, GroupIndex);
 
-				IslandGroup->AddConstraintsAndBodies();
-				IslandGroup->GatherBodies(Dt, 0, IslandGroup->GetNumSolverBodies());
-				IslandGroup->GatherConstraints(Dt, 0, IslandGroup->GetNumConstraints());
+				IslandGroup->ApplyPositionConstraints(Dt);
 			}
+			{
+				SCOPE_CYCLE_COUNTER(STAT_Evolution_ApplyConstraintsPhase2);
+				CSV_SCOPED_ISLANDGROUP_TIMING_STAT(PerIslandSolve_ApplyPushOutTotalSerialized, GroupIndex);
 
-			SolveGroupConstraints(GroupIndex, Dt);
-
+				IslandGroup->ApplyVelocityConstraints(Dt);
+			}
 			{
-				SCOPE_CYCLE_COUNTER(STAT_Evolution_Scatter);
-				CSV_SCOPED_ISLANDGROUP_TIMING_STAT(PerIslandSolve_ScatterTotalSerialized, GroupIndex);
-				IslandGroup->ScatterConstraints(Dt);
-				IslandGroup->ScatterBodies(Dt);
+				SCOPE_CYCLE_COUNTER(STAT_Evolution_ApplyConstraintsPhase3);
+				CSV_SCOPED_ISLANDGROUP_TIMING_STAT(PerIslandSolve_ApplyProjectionTotalSerialized, GroupIndex);
+
+				IslandGroup->ApplyProjectionConstraints(Dt);
 			}
 		}
-	}
-
-	void FPBDIslandGroupManager::SolveParallelFor(const FReal Dt)
-	{
-		SCOPE_CYCLE_COUNTER(STAT_SolveConstraints);
-
-		TArray<FIslandGroupRange> BodyRanges;
-		TArray<FIslandGroupRange> ConstraintRanges;
 
+		void FPBDIslandGroupManager::SolveSerial(const FReal Dt)
 		{
-			CSV_SCOPED_TIMING_STAT(PhysicsVerbose, StepSolver_Gather);
+			SCOPE_CYCLE_COUNTER(STAT_SolveConstraints);
+
+			for (int32 GroupIndex = 0; GroupIndex < NumActiveGroups; ++GroupIndex)
+			{
+				FPBDIslandConstraintGroupSolver* IslandGroup = GetGroup(GroupIndex);
 
-			// Add all the bodies and constraints to the groups in the required order
-			PhysicsParallelFor(NumActiveGroups,
-				[this](const int32 GroupIndex)
 				{
-					SCOPE_CYCLE_COUNTER(STAT_Evolution_Gather_Prepare);
+					SCOPE_CYCLE_COUNTER(STAT_Evolution_Gather);
 					CSV_SCOPED_ISLANDGROUP_TIMING_STAT(PerIslandSolve_GatherTotalSerialized, GroupIndex);
 
-					FPBDIslandConstraintGroupSolver* IslandGroup = GetGroup(GroupIndex);
 					IslandGroup->AddConstraintsAndBodies();
-				});
+					IslandGroup->GatherBodies(Dt, 0, IslandGroup->GetNumSolverBodies());
+					IslandGroup->GatherConstraints(Dt, 0, IslandGroup->GetNumConstraints());
+				}
 
-			// Generate body and constraint ranges for parallel gather
-			// NOTE: Must be after AddConstraintsAndBodies for all groups
-			BuildGatherBatches(BodyRanges, ConstraintRanges);
+				SolveGroupConstraints(GroupIndex, Dt);
 
-			// Gather all body data using all worker threads
-			PhysicsParallelFor(BodyRanges.Num(),
-				[this, &BodyRanges, Dt](const int32 RangeIndex)
 				{
-					SCOPE_CYCLE_COUNTER(STAT_Evolution_Gather_Bodies);
-					CSV_SCOPED_ISLANDGROUP_TIMING_STAT(PerIslandSolve_GatherTotalSerialized, RangeIndex);
+					SCOPE_CYCLE_COUNTER(STAT_Evolution_Scatter);
+					CSV_SCOPED_ISLANDGROUP_TIMING_STAT(PerIslandSolve_ScatterTotalSerialized, GroupIndex);
+					IslandGroup->ScatterConstraints(Dt);
+					IslandGroup->ScatterBodies(Dt);
+				}
+			}
+		}
 
-					FIslandGroupRange& Range = BodyRanges[RangeIndex];
-					Range.IslandGroup->GatherBodies(Dt, Range.BeginIndex, Range.EndIndex);
-				});
+		void FPBDIslandGroupManager::SolveParallelFor(const FReal Dt)
+		{
+			SCOPE_CYCLE_COUNTER(STAT_SolveConstraints);
 
-			// Gather all constraint data
-			PhysicsParallelFor(ConstraintRanges.Num(),
-				[this, &ConstraintRanges, Dt](const int32 RangeIndex)
-				{
-					SCOPE_CYCLE_COUNTER(STAT_Evolution_Gather_Constraints);
-					CSV_SCOPED_ISLANDGROUP_TIMING_STAT(PerIslandSolve_GatherTotalSerialized, RangeIndex);
+			TArray<FIslandGroupRange> BodyRanges;
+			TArray<FIslandGroupRange> ConstraintRanges;
 
-					FIslandGroupRange& Range = ConstraintRanges[RangeIndex];
-					Range.IslandGroup->GatherConstraints(Dt, Range.BeginIndex, Range.EndIndex);
-				});
-		}
-		
-		// Solve all the constraints, each group in parallel
-		PhysicsParallelFor(NumActiveGroups,
-			[this, Dt](const int32 GroupIndex)
 			{
-				FPBDIslandConstraintGroupSolver* IslandGroup = GetGroup(GroupIndex);
-
-				SolveGroupConstraints(GroupIndex, Dt);
-			});
+				CSV_SCOPED_TIMING_STAT(PhysicsVerbose, StepSolver_Gather);
 
-		{
-			SCOPE_CYCLE_COUNTER(STAT_Evolution_Scatter);
+				// Add all the bodies and constraints to the groups in the required order
+				PhysicsParallelFor(NumActiveGroups,
+					[this](const int32 GroupIndex)
+					{
+						SCOPE_CYCLE_COUNTER(STAT_Evolution_Gather_Prepare);
+						CSV_SCOPED_ISLANDGROUP_TIMING_STAT(PerIslandSolve_GatherTotalSerialized, GroupIndex);
 
-			// Scatter constraints using all worker threads
-			PhysicsParallelFor(ConstraintRanges.Num(),
-				[this, &ConstraintRanges, Dt](const int32 RangeIndex)
-				{
-					SCOPE_CYCLE_COUNTER(STAT_Evolution_Scatter_Constraints);
-					CSV_SCOPED_ISLANDGROUP_TIMING_STAT(PerIslandSolve_ScatterTotalSerialized, RangeIndex);
+						FPBDIslandConstraintGroupSolver* IslandGroup = GetGroup(GroupIndex);
+						IslandGroup->AddConstraintsAndBodies();
+					});
 
-					FIslandGroupRange& Range = ConstraintRanges[RangeIndex];
-					Range.IslandGroup->ScatterConstraints(Dt, Range.BeginIndex, Range.EndIndex);
-				});
+				// Generate body and constraint ranges for parallel gather
+				// NOTE: Must be after AddConstraintsAndBodies for all groups
+				BuildGatherBatches(BodyRanges, ConstraintRanges);
 
-			// Scatter bodies using all worker threads
-			PhysicsParallelFor(BodyRanges.Num(),
-				[this, &BodyRanges, Dt](const int32 RangeIndex)
-				{
-					SCOPE_CYCLE_COUNTER(STAT_Evolution_Scatter_Bodies);
-					CSV_SCOPED_ISLANDGROUP_TIMING_STAT(PerIslandSolve_ScatterTotalSerialized, RangeIndex);
+				// Gather all body data using all worker threads
+				PhysicsParallelFor(BodyRanges.Num(),
+					[this, &BodyRanges, Dt](const int32 RangeIndex)
+					{
+						SCOPE_CYCLE_COUNTER(STAT_Evolution_Gather_Bodies);
+						CSV_SCOPED_ISLANDGROUP_TIMING_STAT(PerIslandSolve_GatherTotalSerialized, RangeIndex);
 
-					FIslandGroupRange& Range = BodyRanges[RangeIndex];
-					Range.IslandGroup->ScatterBodies(Dt, Range.BeginIndex, Range.EndIndex);
-				});
-		}
-	}
+						FIslandGroupRange& Range = BodyRanges[RangeIndex];
+						Range.IslandGroup->GatherBodies(Dt, Range.BeginIndex, Range.EndIndex);
+					});
 
+				// Gather all constraint data
+				PhysicsParallelFor(ConstraintRanges.Num(),
+					[this, &ConstraintRanges, Dt](const int32 RangeIndex)
+					{
+						SCOPE_CYCLE_COUNTER(STAT_Evolution_Gather_Constraints);
+						CSV_SCOPED_ISLANDGROUP_TIMING_STAT(PerIslandSolve_GatherTotalSerialized, RangeIndex);
 
-	void FPBDIslandGroupManager::SolveParallelTasks(const FReal Dt)
-	{
-		SCOPE_CYCLE_COUNTER(STAT_SolveConstraints);
-
-		FGraphEventArray GroupSolveEvents;
-		GroupSolveEvents.Reserve(NumActiveGroups);
-
-		// Each IslandGroup is independent of the otyher so they can all be solved in parallel.
-		// In addition, the Gather and Scatter work within each group is highly parallelizable.
-		//
-		// In order to utilize all workers when we have a small number of groups some are much
-		// larger than others, we allow each group to generate multiple gather/scatter tasks. 
-		// The number of tasks each group generates depends on how many constraints it has and 
-		// TargetNumConstraintsPerTask, which was calculated in BuildGroups such that we end 
-		// up with the right number of gather tasks below.
-		//
-		for (int32 GroupIndex = 0; GroupIndex < NumActiveGroups; ++GroupIndex)
-		{
-			// NOTE: This task spawns other child tasks. We must run AddConstraintsAndBodies before we know how many
-			// Gather and Scatter tasks we want to create.
-			FGraphEventRef GroupSolveEvent = FFunctionGraphTask::CreateAndDispatchWhenReady(
-				[this, GroupIndex, Dt](ENamedThreads::Type CurrentThread, const FGraphEventRef& MyCompletionGraphEvent)
+						FIslandGroupRange& Range = ConstraintRanges[RangeIndex];
+						Range.IslandGroup->GatherConstraints(Dt, Range.BeginIndex, Range.EndIndex);
+					});
+			}
+		
+			// Solve all the constraints, each group in parallel
+			PhysicsParallelFor(NumActiveGroups,
+				[this, Dt](const int32 GroupIndex)
 				{
 					FPBDIslandConstraintGroupSolver* IslandGroup = GetGroup(GroupIndex);
 
-					// Add bodies and constraints to the solvers but do no real work. Cannot be further parallelized
-					{
-						SCOPE_CYCLE_COUNTER(STAT_Evolution_Gather_Prepare);
+					SolveGroupConstraints(GroupIndex, Dt);
+				});
 
-						IslandGroup->AddConstraintsAndBodies();
-					}
+			{
+				SCOPE_CYCLE_COUNTER(STAT_Evolution_Scatter);
 
-					// How many gather tasks should we be creating for this group?
-					// NOTE: TargetNumBodiesPerTask and TargetNumConstraintsPerTask were calculated in BuildGroups
-					// so that we end up with a reasonable number of Gather/Scatter tasks for the system's worker threads.
-					const int32 NumBodies = IslandGroup->GetNumSolverBodies();
-					const int32 NumBodyGatherTasks = FMath::Max(1, FMath::DivideAndRoundNearest(NumBodies, TargetNumBodiesPerTask));
-					const int32 NumBodiesPerTask = FMath::DivideAndRoundUp(NumBodies, NumBodyGatherTasks);
-					const int32 NumConstraints = IslandGroup->GetNumSolverConstraints();
-					const int32 NumConstraintGatherTasks = FMath::Max(1, FMath::DivideAndRoundNearest(NumConstraints, TargetNumConstraintsPerTask));
-					const int32 NumConstraintsPerTask = FMath::DivideAndRoundUp(NumConstraints, NumConstraintGatherTasks);
-
-					// Gather Bodies. Can run as wide as we like
-					FGraphEventArray GatherBodyEvents;
-					GatherBodyEvents.Reserve(NumBodyGatherTasks);
-					for (int32 BodyBeginIndex = 0; BodyBeginIndex < NumBodies; BodyBeginIndex += NumBodiesPerTask)
+				// Scatter constraints using all worker threads
+				PhysicsParallelFor(ConstraintRanges.Num(),
+					[this, &ConstraintRanges, Dt](const int32 RangeIndex)
 					{
-						const int32 BodyEndIndex = FMath::Min(BodyBeginIndex + NumBodiesPerTask, NumBodies);
+						SCOPE_CYCLE_COUNTER(STAT_Evolution_Scatter_Constraints);
+						CSV_SCOPED_ISLANDGROUP_TIMING_STAT(PerIslandSolve_ScatterTotalSerialized, RangeIndex);
 
-						FGraphEventRef GatherBodyEvent = FFunctionGraphTask::CreateAndDispatchWhenReady(
-							[IslandGroup, BodyBeginIndex, BodyEndIndex, Dt]()
-							{
-								CSV_SCOPED_TIMING_STAT(PhysicsVerbose, StepSolver_Gather);
+						FIslandGroupRange& Range = ConstraintRanges[RangeIndex];
+						Range.IslandGroup->ScatterConstraints(Dt, Range.BeginIndex, Range.EndIndex);
+					});
 
-								IslandGroup->GatherBodies(Dt, BodyBeginIndex, BodyEndIndex);
-							}
-							, GET_STATID(STAT_Evolution_Gather_Bodies));
+				// Scatter bodies using all worker threads
+				PhysicsParallelFor(BodyRanges.Num(),
+					[this, &BodyRanges, Dt](const int32 RangeIndex)
+					{
+						SCOPE_CYCLE_COUNTER(STAT_Evolution_Scatter_Bodies);
+						CSV_SCOPED_ISLANDGROUP_TIMING_STAT(PerIslandSolve_ScatterTotalSerialized, RangeIndex);
 
-						GatherBodyEvents.Add(GatherBodyEvent);
-					}
+						FIslandGroupRange& Range = BodyRanges[RangeIndex];
+						Range.IslandGroup->ScatterBodies(Dt, Range.BeginIndex, Range.EndIndex);
+					});
+			}
+		}
 
-					// Wait for Gather Bodies
-					FGraphEventRef GatherBodiesCompletionEvent = TGraphTask<FNullGraphTask>::CreateTask(&GatherBodyEvents).ConstructAndDispatchWhenReady(TStatId(), ENamedThreads::AnyThread);
 
-					// Gather Constraints after body gather is complete. Can run as wide as wek like
-					FGraphEventArray GatherConstraintEvents;
-					GatherConstraintEvents.Reserve(NumConstraintGatherTasks);
-					for (int32 ConstraintBeginIndex = 0; ConstraintBeginIndex < NumConstraints; ConstraintBeginIndex += NumConstraintsPerTask)
+		void FPBDIslandGroupManager::SolveParallelTasks(const FReal Dt)
+		{
+			SCOPE_CYCLE_COUNTER(STAT_SolveConstraints);
+
+			FGraphEventArray GroupSolveEvents;
+			GroupSolveEvents.Reserve(NumActiveGroups);
+
+			// Each IslandGroup is independent of the otyher so they can all be solved in parallel.
+			// In addition, the Gather and Scatter work within each group is highly parallelizable.
+			//
+			// In order to utilize all workers when we have a small number of groups some are much
+			// larger than others, we allow each group to generate multiple gather/scatter tasks. 
+			// The number of tasks each group generates depends on how many constraints it has and 
+			// TargetNumConstraintsPerTask, which was calculated in BuildGroups such that we end 
+			// up with the right number of gather tasks below.
+			//
+			for (int32 GroupIndex = 0; GroupIndex < NumActiveGroups; ++GroupIndex)
+			{
+				// NOTE: This task spawns other child tasks. We must run AddConstraintsAndBodies before we know how many
+				// Gather and Scatter tasks we want to create.
+				FGraphEventRef GroupSolveEvent = FFunctionGraphTask::CreateAndDispatchWhenReady(
+					[this, GroupIndex, Dt](ENamedThreads::Type CurrentThread, const FGraphEventRef& MyCompletionGraphEvent)
 					{
-						const int32 ConstraintEndIndex = FMath::Min(ConstraintBeginIndex + NumConstraintsPerTask, NumConstraints);
+						FPBDIslandConstraintGroupSolver* IslandGroup = GetGroup(GroupIndex);
 
-						FGraphEventRef GatherConstraintsEvent = FFunctionGraphTask::CreateAndDispatchWhenReady(
-							[this, IslandGroup, ConstraintBeginIndex, ConstraintEndIndex, Dt]()
-							{
-								CSV_SCOPED_TIMING_STAT(PhysicsVerbose, StepSolver_Gather);
+						// Add bodies and constraints to the solvers but do no real work. Cannot be further parallelized
+						{
+							SCOPE_CYCLE_COUNTER(STAT_Evolution_Gather_Prepare);
 
-								IslandGroup->GatherConstraints(Dt, ConstraintBeginIndex, ConstraintEndIndex);
-							}
-							, GET_STATID(STAT_Evolution_Gather_Constraints)
-							, GatherBodiesCompletionEvent);
+							IslandGroup->AddConstraintsAndBodies();
+						}
 
-						GatherConstraintEvents.Add(GatherConstraintsEvent);
-					}
+						// How many gather tasks should we be creating for this group?
+						// NOTE: TargetNumBodiesPerTask and TargetNumConstraintsPerTask were calculated in BuildGroups
+						// so that we end up with a reasonable number of Gather/Scatter tasks for the system's worker threads.
+						const int32 NumBodies = IslandGroup->GetNumSolverBodies();
+						const int32 NumBodyGatherTasks = FMath::Max(1, FMath::DivideAndRoundNearest(NumBodies, TargetNumBodiesPerTask));
+						const int32 NumBodiesPerTask = FMath::DivideAndRoundUp(NumBodies, NumBodyGatherTasks);
+						const int32 NumConstraints = IslandGroup->GetNumSolverConstraints();
+						const int32 NumConstraintGatherTasks = FMath::Max(1, FMath::DivideAndRoundNearest(NumConstraints, TargetNumConstraintsPerTask));
+						const int32 NumConstraintsPerTask = FMath::DivideAndRoundUp(NumConstraints, NumConstraintGatherTasks);
+
+						// Gather Bodies. Can run as wide as we like
+						FGraphEventArray GatherBodyEvents;
+						GatherBodyEvents.Reserve(NumBodyGatherTasks);
+						for (int32 BodyBeginIndex = 0; BodyBeginIndex < NumBodies; BodyBeginIndex += NumBodiesPerTask)
+						{
+							const int32 BodyEndIndex = FMath::Min(BodyBeginIndex + NumBodiesPerTask, NumBodies);
 
-					// Wait for Gather Constraints
-					FGraphEventRef GatherConstraintsCompletionEvent = TGraphTask<FNullGraphTask>::CreateTask(&GatherConstraintEvents).ConstructAndDispatchWhenReady(TStatId(), ENamedThreads::AnyThread);
+							FGraphEventRef GatherBodyEvent = FFunctionGraphTask::CreateAndDispatchWhenReady(
+								[IslandGroup, BodyBeginIndex, BodyEndIndex, Dt]()
+								{
+									CSV_SCOPED_TIMING_STAT(PhysicsVerbose, StepSolver_Gather);
 
-					// Solve Constraints after gathering is complete. Single task - no further parallelization (without coloring)
-					// @todo(chaos): add coloring support
-					FGraphEventRef SolveConstraintsCompletionEvent = FFunctionGraphTask::CreateAndDispatchWhenReady(
-						[this, GroupIndex, Dt]()
-						{
-							SolveGroupConstraints(GroupIndex, Dt);
+									IslandGroup->GatherBodies(Dt, BodyBeginIndex, BodyEndIndex);
+								}
+								, GET_STATID(STAT_Evolution_Gather_Bodies));
+
+							GatherBodyEvents.Add(GatherBodyEvent);
 						}
-						, TStatId()
-						, GatherConstraintsCompletionEvent);
 
+						// Wait for Gather Bodies
+						FGraphEventRef GatherBodiesCompletionEvent = TGraphTask<FNullGraphTask>::CreateTask(&GatherBodyEvents).ConstructAndDispatchWhenReady(TStatId(), ENamedThreads::AnyThread);
 
-					// Scatter constraints and bodies after the solve completes. No dependencies between any tasks at this point
-					FGraphEventArray ScatterEvents;
-					ScatterEvents.Reserve(NumBodyGatherTasks + NumConstraintGatherTasks);
+						// Gather Constraints after body gather is complete. Can run as wide as wek like
+						FGraphEventArray GatherConstraintEvents;
+						GatherConstraintEvents.Reserve(NumConstraintGatherTasks);
+						for (int32 ConstraintBeginIndex = 0; ConstraintBeginIndex < NumConstraints; ConstraintBeginIndex += NumConstraintsPerTask)
+						{
+							const int32 ConstraintEndIndex = FMath::Min(ConstraintBeginIndex + NumConstraintsPerTask, NumConstraints);
 
-					// Scatter Constraints
-					for (int32 ConstraintBeginIndex = 0; ConstraintBeginIndex < NumConstraints; ConstraintBeginIndex += NumConstraintsPerTask)
-					{
-						const int32 ConstraintEndIndex = FMath::Min(ConstraintBeginIndex + NumConstraintsPerTask, NumConstraints);
+							FGraphEventRef GatherConstraintsEvent = FFunctionGraphTask::CreateAndDispatchWhenReady(
+								[this, IslandGroup, ConstraintBeginIndex, ConstraintEndIndex, Dt]()
+								{
+									CSV_SCOPED_TIMING_STAT(PhysicsVerbose, StepSolver_Gather);
 
-						FGraphEventRef ScatterConstraintsEvent = FFunctionGraphTask::CreateAndDispatchWhenReady(
-							[this, IslandGroup, ConstraintBeginIndex, ConstraintEndIndex, Dt]()
-							{
-								CSV_SCOPED_TIMING_STAT(PhysicsVerbose, StepSolver_Scatter);
+									IslandGroup->GatherConstraints(Dt, ConstraintBeginIndex, ConstraintEndIndex);
+								}
+								, GET_STATID(STAT_Evolution_Gather_Constraints)
+								, GatherBodiesCompletionEvent);
+
+							GatherConstraintEvents.Add(GatherConstraintsEvent);
+						}
+
+						// Wait for Gather Constraints
+						FGraphEventRef GatherConstraintsCompletionEvent = TGraphTask<FNullGraphTask>::CreateTask(&GatherConstraintEvents).ConstructAndDispatchWhenReady(TStatId(), ENamedThreads::AnyThread);
 
-								IslandGroup->ScatterConstraints(Dt, ConstraintBeginIndex, ConstraintEndIndex);
+						// Solve Constraints after gathering is complete. Single task - no further parallelization (without coloring)
+						// @todo(chaos): add coloring support
+						FGraphEventRef SolveConstraintsCompletionEvent = FFunctionGraphTask::CreateAndDispatchWhenReady(
+							[this, GroupIndex, Dt]()
+							{
+								SolveGroupConstraints(GroupIndex, Dt);
 							}
-							, GET_STATID(STAT_Evolution_Scatter_Constraints)
-							, SolveConstraintsCompletionEvent);
+							, TStatId()
+							, GatherConstraintsCompletionEvent);
 
-						ScatterEvents.Add(ScatterConstraintsEvent);
-					}
 
-					// Scatter Bodies
-					for (int32 BodyBeginIndex = 0; BodyBeginIndex < NumBodies; BodyBeginIndex += NumBodiesPerTask)
-					{
-						const int32 BodyEndIndex = FMath::Min(BodyBeginIndex + NumBodiesPerTask, NumBodies);
+						// Scatter constraints and bodies after the solve completes. No dependencies between any tasks at this point
+						FGraphEventArray ScatterEvents;
+						ScatterEvents.Reserve(NumBodyGatherTasks + NumConstraintGatherTasks);
 
-						FGraphEventRef ScatterBodyEvent = FFunctionGraphTask::CreateAndDispatchWhenReady(
-							[this, IslandGroup, BodyBeginIndex, BodyEndIndex, Dt]()
-							{
-								CSV_SCOPED_TIMING_STAT(PhysicsVerbose, StepSolver_Scatter);
+						// Scatter Constraints
+						for (int32 ConstraintBeginIndex = 0; ConstraintBeginIndex < NumConstraints; ConstraintBeginIndex += NumConstraintsPerTask)
+						{
+							const int32 ConstraintEndIndex = FMath::Min(ConstraintBeginIndex + NumConstraintsPerTask, NumConstraints);
 
-								IslandGroup->ScatterBodies(Dt, BodyBeginIndex, BodyEndIndex);
-							}
-							, GET_STATID(STAT_Evolution_Gather_Bodies)
-							, SolveConstraintsCompletionEvent);
+							FGraphEventRef ScatterConstraintsEvent = FFunctionGraphTask::CreateAndDispatchWhenReady(
+								[this, IslandGroup, ConstraintBeginIndex, ConstraintEndIndex, Dt]()
+								{
+									CSV_SCOPED_TIMING_STAT(PhysicsVerbose, StepSolver_Scatter);
 
-						ScatterEvents.Add(ScatterBodyEvent);
-					}
+									IslandGroup->ScatterConstraints(Dt, ConstraintBeginIndex, ConstraintEndIndex);
+								}
+								, GET_STATID(STAT_Evolution_Scatter_Constraints)
+								, SolveConstraintsCompletionEvent);
 
-					// Completion event
-					FGraphEventRef ScatterCompletionEvent = TGraphTask<FNullGraphTask>::CreateTask(&ScatterEvents).ConstructAndDispatchWhenReady(TStatId(), ENamedThreads::AnyThread);
-					MyCompletionGraphEvent->DontCompleteUntil(ScatterCompletionEvent);
-				});
+							ScatterEvents.Add(ScatterConstraintsEvent);
+						}
 
-			GroupSolveEvents.Add(GroupSolveEvent);
-		}
+						// Scatter Bodies
+						for (int32 BodyBeginIndex = 0; BodyBeginIndex < NumBodies; BodyBeginIndex += NumBodiesPerTask)
+						{
+							const int32 BodyEndIndex = FMath::Min(BodyBeginIndex + NumBodiesPerTask, NumBodies);
 
-		// Wait for all tasks to complete and help out
-		FTaskGraphInterface::Get().WaitUntilTasksComplete(GroupSolveEvents);
-	}
+							FGraphEventRef ScatterBodyEvent = FFunctionGraphTask::CreateAndDispatchWhenReady(
+								[this, IslandGroup, BodyBeginIndex, BodyEndIndex, Dt]()
+								{
+									CSV_SCOPED_TIMING_STAT(PhysicsVerbose, StepSolver_Scatter);
 
-	void FPBDIslandGroupManager::SetNumIterations(const int32 InNumPositionIterations, const int32 InNumVelocityIterations, const int32 InNumProjectionIterations)
-	{
-		NumPositionIterations = InNumPositionIterations;
-		NumVelocityIterations = InNumVelocityIterations;
-		NumProjectionIterations = InNumProjectionIterations;
-	}
+									IslandGroup->ScatterBodies(Dt, BodyBeginIndex, BodyEndIndex);
+								}
+								, GET_STATID(STAT_Evolution_Gather_Bodies)
+								, SolveConstraintsCompletionEvent);
 
-	void FPBDIslandGroupManager::Solve(const FReal Dt)
-	{
-#if CSV_PROFILER
-		GroupStats.Reset();
-		GroupStats.AddDefaulted(IslandGroups.Num());
-#endif
-
-		// @todo(chaos): Remove SolveParallelFor when SolveParallelTasks has been thoroughly tested
-		const bool bSingleThreaded = GSingleThreadedPhysics || (IslandGroups.Num() == 1);
-		const int32 ParallelMode = bSingleThreaded ? 0 : CVars::GIslandGroupsParallelMode;
-		switch (ParallelMode)
+							ScatterEvents.Add(ScatterBodyEvent);
+						}
+
+						// Completion event
+						FGraphEventRef ScatterCompletionEvent = TGraphTask<FNullGraphTask>::CreateTask(&ScatterEvents).ConstructAndDispatchWhenReady(TStatId(), ENamedThreads::AnyThread);
+						MyCompletionGraphEvent->DontCompleteUntil(ScatterCompletionEvent);
+					});
+
+				GroupSolveEvents.Add(GroupSolveEvent);
+			}
+
+			// Wait for all tasks to complete and help out
+			FTaskGraphInterface::Get().WaitUntilTasksComplete(GroupSolveEvents);
+		}
+
+		void FPBDIslandGroupManager::SetIterationSettings(const FIterationSettings& InIterations)
 		{
-		case 0:
-			SolveSerial(Dt);
-			break;
-		case 1:
-			SolveParallelFor(Dt);
-			break;
-		case 2:
-			SolveParallelTasks(Dt);
-			break;
+			Iterations = InIterations;
 		}
 
-#if CSV_PROFILER
-		FIslandGroupStats FlattenedStats = FIslandGroupStats::Flatten(GroupStats);
-		FlattenedStats.ReportStats();
-#endif
-	}
-}
\ No newline at end of file
+		void FPBDIslandGroupManager::Solve(const FReal Dt)
+		{
+	#if CSV_PROFILER
+			GroupStats.Reset();
+			GroupStats.AddDefaulted(IslandGroups.Num());
+	#endif
+
+			// @todo(chaos): Remove SolveParallelFor when SolveParallelTasks has been thoroughly tested
+			const bool bSingleThreaded = GSingleThreadedPhysics || (IslandGroups.Num() == 1);
+			const int32 ParallelMode = bSingleThreaded ? 0 : CVars::GIslandGroupsParallelMode;
+			switch (ParallelMode)
+			{
+			case 0:
+				SolveSerial(Dt);
+				break;
+			case 1:
+				SolveParallelFor(Dt);
+				break;
+			case 2:
+				SolveParallelTasks(Dt);
+				break;
+			}
+
+	#if CSV_PROFILER
+			FIslandGroupStats FlattenedStats = FIslandGroupStats::Flatten(GroupStats);
+			FlattenedStats.ReportStats();
+	#endif
+		}
+	}	// namespace Private
+}	// namespace Chaos
diff --git a/Engine/Source/Runtime/Experimental/Chaos/Private/Chaos/Island/IslandManager.cpp b/Engine/Source/Runtime/Experimental/Chaos/Private/Chaos/Island/IslandManager.cpp
--- a/Engine/Source/Runtime/Experimental/Chaos/Private/Chaos/Island/IslandManager.cpp
+++ b/Engine/Source/Runtime/Experimental/Chaos/Private/Chaos/Island/IslandManager.cpp
@@ -16,7 +16,6 @@
 
 namespace Chaos
 {
-
 /** Cvar to enable/disable the island sleeping */
 bool bChaosSolverSleepEnabled = true;
 FAutoConsoleVariableRef CVarChaosSolverSleepEnabled(TEXT("p.Chaos.Solver.SleepEnabled"), bChaosSolverSleepEnabled, TEXT(""));
@@ -39,8 +38,8 @@ FAutoConsoleVariableRef CVarChaosSolverValidateGraph(TEXT("p.Chaos.Solver.Valida
 bool bChaosSolverPersistentGraph = true;
 FAutoConsoleVariableRef CVarChaosSolverPersistentGraph(TEXT("p.Chaos.Solver.PersistentGraph"), bChaosSolverPersistentGraph, TEXT(""));
 
-extern int32 GSingleThreadedPhysics;
-
+namespace Private
+{
 	
 /** Check if a particle is dynamic or sleeping */
 FORCEINLINE bool IsDynamicParticle(const FGeometryParticleHandle* ParticleHandle)
@@ -1283,4 +1282,5 @@ void FPBDIslandManager::DebugCheckIslands() const
 }
 #endif
 
-}
\ No newline at end of file
+}	// namespace Private
+}	// namespace Chaos
\ No newline at end of file
diff --git a/Engine/Source/Runtime/Experimental/Chaos/Private/Chaos/Island/SolverIsland.cpp b/Engine/Source/Runtime/Experimental/Chaos/Private/Chaos/Island/SolverIsland.cpp
--- a/Engine/Source/Runtime/Experimental/Chaos/Private/Chaos/Island/SolverIsland.cpp
+++ b/Engine/Source/Runtime/Experimental/Chaos/Private/Chaos/Island/SolverIsland.cpp
@@ -5,165 +5,168 @@
 
 namespace Chaos
 {
+	namespace Private
+	{
 	
-	FPBDIsland::FPBDIsland(const int32 MaxConstraintContainers)
-	: IslandIndex(INDEX_NONE)
-	, bIsSleeping(false)
-	, bNeedsResim(false)
-	, bIsPersistent(true)
-	, SleepCounter(0)
-	, IslandParticles()
-	, IslandConstraintsByType()
-	, NumConstraints(0)
-{
-	IslandConstraintsByType.SetNum(MaxConstraintContainers);
-}
+		FPBDIsland::FPBDIsland(const int32 MaxConstraintContainers)
+		: IslandIndex(INDEX_NONE)
+		, bIsSleeping(false)
+		, bNeedsResim(false)
+		, bIsPersistent(true)
+		, SleepCounter(0)
+		, IslandParticles()
+		, IslandConstraintsByType()
+		, NumConstraints(0)
+		{
+			IslandConstraintsByType.SetNum(MaxConstraintContainers);
+		}
 
-void FPBDIsland::Reuse()
-{
-	check(GetNumParticles() == 0);
-	check(GetNumConstraints() == 0);
+		void FPBDIsland::Reuse()
+		{
+			check(GetNumParticles() == 0);
+			check(GetNumConstraints() == 0);
 
-	SetSleepCounter(0);
-	SetIsSleepingChanged(false);
-}
+			SetSleepCounter(0);
+			SetIsSleepingChanged(false);
+		}
 
-void FPBDIsland::UpdateParticles()
-{
-	for (FPBDIslandParticle& IslandParticle : IslandParticles)
-	{
-		FGenericParticleHandle ParticleHandle = IslandParticle.GetParticle();
-		if (ParticleHandle.IsValid() && ParticleHandle->IsDynamic())
+		void FPBDIsland::UpdateParticles()
 		{
-			ParticleHandle->SetIslandIndex(IslandIndex);
+			for (FPBDIslandParticle& IslandParticle : IslandParticles)
+			{
+				FGenericParticleHandle ParticleHandle = IslandParticle.GetParticle();
+				if (ParticleHandle.IsValid() && ParticleHandle->IsDynamic())
+				{
+					ParticleHandle->SetIslandIndex(IslandIndex);
+				}
+			}
 		}
-	}
-}
 	
-void FPBDIsland::ClearParticles()
-{
-	IslandParticles.Reset();
-}
+		void FPBDIsland::ClearParticles()
+		{
+			IslandParticles.Reset();
+		}
 
-void FPBDIsland::AddParticle(FGenericParticleHandle ParticleHandle)
-{
-	if (ParticleHandle.IsValid())
-	{
-		if (ParticleHandle->IsDynamic())
+		void FPBDIsland::AddParticle(FGenericParticleHandle ParticleHandle)
 		{
-			ParticleHandle->SetIslandIndex(IslandIndex);
+			if (ParticleHandle.IsValid())
+			{
+				if (ParticleHandle->IsDynamic())
+				{
+					ParticleHandle->SetIslandIndex(IslandIndex);
+				}
+				IslandParticles.Add(ParticleHandle->Handle());
+			}
 		}
-		IslandParticles.Add(ParticleHandle->Handle());
-	}
-}
 
-void FPBDIsland::ReserveParticles(const int32 NumParticles)
-{
-	ClearParticles();
-	IslandParticles.Reserve(NumParticles);
-}
+		void FPBDIsland::ReserveParticles(const int32 NumParticles)
+		{
+			ClearParticles();
+			IslandParticles.Reserve(NumParticles);
+		}
 
-void FPBDIsland::AddConstraint(FConstraintHandle* ConstraintHandle, const int32 Level, const int32 Color, const uint32 SubSortKey)
-{
-	if (ConstraintHandle)
-	{
-		const int32 ContainerId = ConstraintHandle->GetContainerId();
+		void FPBDIsland::AddConstraint(FConstraintHandle* ConstraintHandle, const int32 Level, const int32 Color, const uint32 SubSortKey)
+		{
+			if (ConstraintHandle)
+			{
+				const int32 ContainerId = ConstraintHandle->GetContainerId();
 
-		// Array should already be appropriately sized for all containers
-		IslandConstraintsByType[ContainerId].Emplace(ConstraintHandle, Level, Color, SubSortKey);
+				// Array should already be appropriately sized for all containers
+				IslandConstraintsByType[ContainerId].Emplace(ConstraintHandle, Level, Color, SubSortKey);
 
-		++NumConstraints;
-	}
-}
+				++NumConstraints;
+			}
+		}
 
-void FPBDIsland::ClearConstraints()
-{
-	for (TArray<FPBDIslandConstraint>& Constraints : IslandConstraintsByType)
-	{
-		Constraints.Reset();
-	}
-	NumConstraints = 0;
-}
+		void FPBDIsland::ClearConstraints()
+		{
+			for (TArray<FPBDIslandConstraint>& Constraints : IslandConstraintsByType)
+			{
+				Constraints.Reset();
+			}
+			NumConstraints = 0;
+		}
 
-void FPBDIsland::SortConstraints()
-{
-	for (TArray<FPBDIslandConstraint>& IslandConstraints : IslandConstraintsByType)
-	{
-		if (IslandConstraints.Num() > 1)
+		void FPBDIsland::SortConstraints()
 		{
-			Algo::Sort(IslandConstraints,
-				[](const FPBDIslandConstraint& L, const FPBDIslandConstraint& R) -> bool
+			for (TArray<FPBDIslandConstraint>& IslandConstraints : IslandConstraintsByType)
+			{
+				if (IslandConstraints.Num() > 1)
 				{
-					return L.GetSortKey() < R.GetSortKey();
-				});
+					Algo::Sort(IslandConstraints,
+						[](const FPBDIslandConstraint& L, const FPBDIslandConstraint& R) -> bool
+						{
+							return L.GetSortKey() < R.GetSortKey();
+						});
+				}
+			}
 		}
-	}
-}
 
-void FPBDIsland::PropagateSleepState(FPBDRigidsSOAs& ParticleSOAs)
-{
-	bool bNeedRebuild = false;
-	for (FPBDIslandParticle& IslandParticle : IslandParticles)
-	{
-		FGeometryParticleHandle* Particle = IslandParticle.GetParticle();
-		if (Particle->CastToRigidParticle() && !Particle->CastToRigidParticle()->Disabled())
+		void FPBDIsland::PropagateSleepState(FPBDRigidsSOAs& ParticleSOAs)
 		{
-			// If not sleeping we activate the sleeping particles
-			if (!bIsSleeping && Particle->Sleeping())
+			bool bNeedRebuild = false;
+			for (FPBDIslandParticle& IslandParticle : IslandParticles)
 			{
-				ParticleSOAs.ActivateParticle(Particle, true);
-
-				// When we wake particles, we have skipped their integrate step which causes some issues:
-				//	- we have zero velocity (no gravity or external forces applied)
-				//	- the world transforms cached in the ShapesArray will be at the last post-integrate positions
-				//	  which doesn't match what the velocity is telling us
-				// This causes problems for the solver - essentially we have an "initial overlap" situation.
-				// @todo(chaos): We could just run (partial) integrate here for this particle, but we don't know about the Evolution - fix this
-				for (const TUniquePtr<FPerShapeData>& Shape : Particle->ShapesArray())
+				FGeometryParticleHandle* Particle = IslandParticle.GetParticle();
+				if (Particle->CastToRigidParticle() && !Particle->CastToRigidParticle()->Disabled())
 				{
-					Shape->UpdateLeafWorldTransform(Particle);
+					// If not sleeping we activate the sleeping particles
+					if (!bIsSleeping && Particle->Sleeping())
+					{
+						ParticleSOAs.ActivateParticle(Particle, true);
+
+						// When we wake particles, we have skipped their integrate step which causes some issues:
+						//	- we have zero velocity (no gravity or external forces applied)
+						//	- the world transforms cached in the ShapesArray will be at the last post-integrate positions
+						//	  which doesn't match what the velocity is telling us
+						// This causes problems for the solver - essentially we have an "initial overlap" situation.
+						// @todo(chaos): We could just run (partial) integrate here for this particle, but we don't know about the Evolution - fix this
+						for (const TUniquePtr<FPerShapeData>& Shape : Particle->ShapesArray())
+						{
+							Shape->UpdateLeafWorldTransform(Particle);
+						}
+
+						bNeedRebuild = true;
+					}
+					// If sleeping we deactivate the dynamic particles
+					else if (bIsSleeping && !Particle->Sleeping())
+					{
+						ParticleSOAs.DeactivateParticle(Particle, true);
+						bNeedRebuild = true;
+					}
 				}
-
-				bNeedRebuild = true;
 			}
-			// If sleeping we deactivate the dynamic particles
-			else if (bIsSleeping && !Particle->Sleeping())
+
+			if (bNeedRebuild)
 			{
-				ParticleSOAs.DeactivateParticle(Particle, true);
-				bNeedRebuild = true;
+				ParticleSOAs.RebuildViews();
 			}
-		}
-	}
-
-	if (bNeedRebuild)
-	{
-		ParticleSOAs.RebuildViews();
-	}
 
-	// Island constraints are updating their sleeping flag + awaken one 
-	for (TArray<FPBDIslandConstraint>& TypedIslandConstraints : IslandConstraintsByType)
-	{
-		for (FPBDIslandConstraint& IslandConstraint : TypedIslandConstraints)
-		{
-			IslandConstraint.GetConstraint()->SetIsSleeping(bIsSleeping);
+			// Island constraints are updating their sleeping flag + awaken one 
+			for (TArray<FPBDIslandConstraint>& TypedIslandConstraints : IslandConstraintsByType)
+			{
+				for (FPBDIslandConstraint& IslandConstraint : TypedIslandConstraints)
+				{
+					IslandConstraint.GetConstraint()->SetIsSleeping(bIsSleeping);
+				}
+			}
 		}
-	}
-}
 
-void FPBDIsland::UpdateSyncState(FPBDRigidsSOAs& Particles)
-{
-	bNeedsResim = false;
-	for (FPBDIslandParticle& IslandParticle : IslandParticles)
-	{
-		// If even one particle is soft/hard desync we must resim the entire island (when resim is used)
-		// seems cheap enough so just always do it, if slow pass resim template in here
-		FGeometryParticleHandle* Particle = IslandParticle.GetParticle();
-		if (Particle->SyncState() != ESyncState::InSync)
+		void FPBDIsland::UpdateSyncState(FPBDRigidsSOAs& Particles)
 		{
-			bNeedsResim = true;
-			break;
+			bNeedsResim = false;
+			for (FPBDIslandParticle& IslandParticle : IslandParticles)
+			{
+				// If even one particle is soft/hard desync we must resim the entire island (when resim is used)
+				// seems cheap enough so just always do it, if slow pass resim template in here
+				FGeometryParticleHandle* Particle = IslandParticle.GetParticle();
+				if (Particle->SyncState() != ESyncState::InSync)
+				{
+					bNeedsResim = true;
+					break;
+				}
+			}
 		}
-	}
-}
 
-}
\ No newline at end of file
+	}	// namespace Private
+}	// namespace Chaos
diff --git a/Engine/Source/Runtime/Experimental/Chaos/Private/Chaos/Joint/PBDJointContainerSolver.cpp b/Engine/Source/Runtime/Experimental/Chaos/Private/Chaos/Joint/PBDJointContainerSolver.cpp
--- a/Engine/Source/Runtime/Experimental/Chaos/Private/Chaos/Joint/PBDJointContainerSolver.cpp
+++ b/Engine/Source/Runtime/Experimental/Chaos/Private/Chaos/Joint/PBDJointContainerSolver.cpp
@@ -130,9 +130,9 @@ namespace Chaos
 			}
 		}
 
-		void FPBDJointContainerSolver::AddConstraints(const TArrayView<FPBDIslandConstraint>& IslandConstraints)
+		void FPBDJointContainerSolver::AddConstraints(const TArrayView<Private::FPBDIslandConstraint>& IslandConstraints)
 		{
-			for (FPBDIslandConstraint& IslandConstraint : IslandConstraints)
+			for (Private::FPBDIslandConstraint& IslandConstraint : IslandConstraints)
 			{
 				// We will only ever be given constraints from our container (asserts in non-shipping)
 				const int32 ContainerConstraintIndex = IslandConstraint.GetConstraint()->AsUnsafe<FPBDJointConstraintHandle>()->GetConstraintIndex();
diff --git a/Engine/Source/Runtime/Experimental/Chaos/Private/Chaos/PBDCollisionConstraints.cpp b/Engine/Source/Runtime/Experimental/Chaos/Private/Chaos/PBDCollisionConstraints.cpp
--- a/Engine/Source/Runtime/Experimental/Chaos/Private/Chaos/PBDCollisionConstraints.cpp
+++ b/Engine/Source/Runtime/Experimental/Chaos/Private/Chaos/PBDCollisionConstraints.cpp
@@ -360,7 +360,7 @@ namespace Chaos
 		}
 	}
 
-	void FPBDCollisionConstraints::AddConstraintsToGraph(FPBDIslandManager& IslandManager)
+	void FPBDCollisionConstraints::AddConstraintsToGraph(Private::FPBDIslandManager& IslandManager)
 	{
 		// Debugging/diagnosing: if we have collisions disabled, remove all collisions from the graph and don't add any more
 		if (!GetCollisionsEnabled())
diff --git a/Engine/Source/Runtime/Experimental/Chaos/Private/Chaos/PBDJointConstraints.cpp b/Engine/Source/Runtime/Experimental/Chaos/Private/Chaos/PBDJointConstraints.cpp
--- a/Engine/Source/Runtime/Experimental/Chaos/Private/Chaos/PBDJointConstraints.cpp
+++ b/Engine/Source/Runtime/Experimental/Chaos/Private/Chaos/PBDJointConstraints.cpp
@@ -826,7 +826,7 @@ namespace Chaos
 		OutR1 = FRotation3(Q1 * XL1.GetRotation()).ToMatrix();
 	}
 
-	void FPBDJointConstraints::AddConstraintsToGraph(FPBDIslandManager& IslandManager)
+	void FPBDJointConstraints::AddConstraintsToGraph(Private::FPBDIslandManager& IslandManager)
 	{
 		for (int32 ConstraintIndex = 0; ConstraintIndex < GetNumConstraints(); ++ConstraintIndex)
 		{
diff --git a/Engine/Source/Runtime/Experimental/Chaos/Private/Chaos/PBDNullConstraints.cpp b/Engine/Source/Runtime/Experimental/Chaos/Private/Chaos/PBDNullConstraints.cpp
--- a/Engine/Source/Runtime/Experimental/Chaos/Private/Chaos/PBDNullConstraints.cpp
+++ b/Engine/Source/Runtime/Experimental/Chaos/Private/Chaos/PBDNullConstraints.cpp
@@ -4,7 +4,7 @@
 
 namespace Chaos
 {
-	void FPBDNullConstraints::AddConstraintsToGraph(FPBDIslandManager& IslandManager)
+	void FPBDNullConstraints::AddConstraintsToGraph(Private::FPBDIslandManager& IslandManager)
 	{
 		IslandManager.AddContainerConstraints(*this);
 	}
diff --git a/Engine/Source/Runtime/Experimental/Chaos/Private/Chaos/PBDPositionConstraints.cpp b/Engine/Source/Runtime/Experimental/Chaos/Private/Chaos/PBDPositionConstraints.cpp
--- a/Engine/Source/Runtime/Experimental/Chaos/Private/Chaos/PBDPositionConstraints.cpp
+++ b/Engine/Source/Runtime/Experimental/Chaos/Private/Chaos/PBDPositionConstraints.cpp
@@ -10,7 +10,7 @@ namespace Chaos
 		return ConcreteContainer()->GetConstrainedParticles(ConstraintIndex);
 	}
 
-	void FPBDPositionConstraints::AddConstraintsToGraph(FPBDIslandManager& IslandManager)
+	void FPBDPositionConstraints::AddConstraintsToGraph(Private::FPBDIslandManager& IslandManager)
 	{
 		IslandManager.AddContainerConstraints(*this);
 	}
diff --git a/Engine/Source/Runtime/Experimental/Chaos/Private/Chaos/PBDRigidDynamicSpringConstraints.cpp b/Engine/Source/Runtime/Experimental/Chaos/Private/Chaos/PBDRigidDynamicSpringConstraints.cpp
--- a/Engine/Source/Runtime/Experimental/Chaos/Private/Chaos/PBDRigidDynamicSpringConstraints.cpp
+++ b/Engine/Source/Runtime/Experimental/Chaos/Private/Chaos/PBDRigidDynamicSpringConstraints.cpp
@@ -97,7 +97,7 @@ namespace Chaos
 		}
 	}
 
-	void FPBDRigidDynamicSpringConstraints::AddConstraintsToGraph(FPBDIslandManager& IslandManager)
+	void FPBDRigidDynamicSpringConstraints::AddConstraintsToGraph(Private::FPBDIslandManager& IslandManager)
 	{
 		IslandManager.AddContainerConstraints(*this);
 	}
diff --git a/Engine/Source/Runtime/Experimental/Chaos/Private/Chaos/PBDRigidSpringConstraints.cpp b/Engine/Source/Runtime/Experimental/Chaos/Private/Chaos/PBDRigidSpringConstraints.cpp
--- a/Engine/Source/Runtime/Experimental/Chaos/Private/Chaos/PBDRigidSpringConstraints.cpp
+++ b/Engine/Source/Runtime/Experimental/Chaos/Private/Chaos/PBDRigidSpringConstraints.cpp
@@ -124,7 +124,7 @@ namespace Chaos
 		return SpringSettings[ConstraintIndex].Stiffness * Delta / CombinedInvMass;
 	}
 
-	void FPBDRigidSpringConstraints::AddConstraintsToGraph(FPBDIslandManager& IslandManager)
+	void FPBDRigidSpringConstraints::AddConstraintsToGraph(Private::FPBDIslandManager& IslandManager)
 	{
 		IslandManager.AddContainerConstraints(*this);
 	}
diff --git a/Engine/Source/Runtime/Experimental/Chaos/Private/Chaos/PBDRigidsEvolution.cpp b/Engine/Source/Runtime/Experimental/Chaos/Private/Chaos/PBDRigidsEvolution.cpp
--- a/Engine/Source/Runtime/Experimental/Chaos/Private/Chaos/PBDRigidsEvolution.cpp
+++ b/Engine/Source/Runtime/Experimental/Chaos/Private/Chaos/PBDRigidsEvolution.cpp
@@ -187,7 +187,7 @@ namespace Chaos
 	};
 
 	FPBDRigidsEvolutionBase::FPBDRigidsEvolutionBase(FPBDRigidsSOAs& InParticles, THandleArray<FChaosPhysicsMaterial>& InSolverPhysicsMaterials, bool InIsSingleThreaded)
-	    : IslandGroupManager(ConstraintGraph)
+	    : IslandGroupManager(IslandManager)
 		, Particles(InParticles)
 		, SolverPhysicsMaterials(InSolverPhysicsMaterials)
 		, InternalAcceleration(nullptr)
@@ -196,9 +196,6 @@ namespace Chaos
 		, bIsSingleThreaded(InIsSingleThreaded)
 		, bCanStartAsyncTasks(true)
 		, LatestExternalTimestampConsumed_Internal(-1)
-		, NumPositionIterations(0)
-		, NumVelocityIterations(0)
-		, NumProjectionIterations(0)
 		, SpatialCollectionFactory(new FDefaultCollectionFactory())
 	{
 		Particles.GetParticleHandles().AddArray(&PhysicsMaterials);
@@ -1030,7 +1027,7 @@ namespace Chaos
 		// because Kinematic particles are only in the graph if referenced by a constraint.
 		if (bIsDynamic && !bWasDynamic && !Particle->Disabled())
 		{
-			ConstraintGraph.AddParticle(Particle);
+			IslandManager.AddParticle(Particle);
 		}
 
 		// If we switched to/from dynamic the inertia conditioning may need to be refreshed
@@ -1064,7 +1061,7 @@ namespace Chaos
 			Particles.DisableParticle(Particle);
 			RemoveParticleFromAccelerationStructure(*Particle);
 			DisableConstraints(Particle);
-			ConstraintGraph.RemoveParticle(Particle);
+			IslandManager.RemoveParticle(Particle);
 		}
 	}
 
diff --git a/Engine/Source/Runtime/Experimental/Chaos/Private/Chaos/PBDRigidsEvolutionGBF.cpp b/Engine/Source/Runtime/Experimental/Chaos/Private/Chaos/PBDRigidsEvolutionGBF.cpp
--- a/Engine/Source/Runtime/Experimental/Chaos/Private/Chaos/PBDRigidsEvolutionGBF.cpp
+++ b/Engine/Source/Runtime/Experimental/Chaos/Private/Chaos/PBDRigidsEvolutionGBF.cpp
@@ -241,13 +241,13 @@ void FPBDRigidsEvolutionGBF::ReloadParticlesCache()
 	FEvolutionResimCache* ResimCache = GetCurrentStepResimCache();
 	if ((ResimCache != nullptr) && ResimCache->IsResimming())
 	{
-		for (int32 IslandIndex = 0; IslandIndex < GetConstraintGraph().NumIslands(); ++IslandIndex)
+		for (int32 IslandIndex = 0; IslandIndex < GetIslandManager().NumIslands(); ++IslandIndex)
 		{
-			FPBDIsland* Island = GetConstraintGraph().GetIsland(IslandIndex);
+			Private::FPBDIsland* Island = GetIslandManager().GetIsland(IslandIndex);
 			bool bIsUsingCache = false;
-			if (!Island->IsSleeping() && !GetConstraintGraph().IslandNeedsResim(IslandIndex))
+			if (!Island->IsSleeping() && !GetIslandManager().IslandNeedsResim(IslandIndex))
 			{
-				for (FPBDIslandParticle& IslandParticle : Island->GetParticles())
+				for (Private::FPBDIslandParticle& IslandParticle : Island->GetParticles())
 				{
 					if (auto Rigid = IslandParticle.GetParticle()->CastToRigidParticle())
 					{
@@ -263,7 +263,7 @@ void FPBDRigidsEvolutionGBF::ReloadParticlesCache()
 
 void FPBDRigidsEvolutionGBF::BuildDisabledParticles(const int32 Island, TArray<TArray<FPBDRigidParticleHandle*>>& DisabledParticles, TArray<bool>& SleepedIslands)
 {
-	for (FPBDIslandParticle& IslandParticle : GetConstraintGraph().GetIsland(Island)->GetParticles())
+	for (Private::FPBDIslandParticle& IslandParticle : GetIslandManager().GetIsland(Island)->GetParticles())
 	{
 		// If a dynamic particle is moving slowly enough for long enough, disable it.
 		// @todo(mlentine): Find a good way of not doing this when we aren't using this functionality
@@ -293,7 +293,7 @@ void FPBDRigidsEvolutionGBF::BuildDisabledParticles(const int32 Island, TArray<T
 	}
 
 	// Turn off if not moving
-	SleepedIslands[Island] = GetConstraintGraph().SleepInactive(Island, PhysicsMaterials, SolverPhysicsMaterials);
+	SleepedIslands[Island] = GetIslandManager().SleepInactive(Island, PhysicsMaterials, SolverPhysicsMaterials);
 }
 
 int32 DrawAwake = 0;
@@ -437,9 +437,9 @@ void FPBDRigidsEvolutionGBF::AdvanceOneTimeStepImpl(const FReal Dt, const FSubSt
 
 
 	TArray<bool> SleepedIslands;
-	SleepedIslands.SetNum(GetConstraintGraph().NumIslands());
+	SleepedIslands.SetNum(GetIslandManager().NumIslands());
 	TArray<TArray<FPBDRigidParticleHandle*>> DisabledParticles;
-	DisabledParticles.SetNum(GetConstraintGraph().NumIslands());
+	DisabledParticles.SetNum(GetIslandManager().NumIslands());
 	if(Dt > 0)
 	{
 		// Solve all the constraints
@@ -447,10 +447,10 @@ void FPBDRigidsEvolutionGBF::AdvanceOneTimeStepImpl(const FReal Dt, const FSubSt
 			SCOPE_CYCLE_COUNTER(STAT_Evolution_ParallelSolve);
 			CSV_SCOPED_TIMING_STAT(PhysicsVerbose, StepSolver_PerIslandSolve);
 
-			IslandGroupManager.SetNumIterations(NumPositionIterations, NumVelocityIterations, NumProjectionIterations);
 			IslandGroupManager.Solve(Dt);
 		}
 
+		// Post-solve CCD fixup to prevent the constraint solve from pushing CCD objects our of the world
 		{
 			SCOPE_CYCLE_COUNTER(STAT_Evolution_CCDCorrection);
 			CSV_SCOPED_TIMING_STAT(PhysicsVerbose, CCDCorrection);
@@ -464,10 +464,10 @@ void FPBDRigidsEvolutionGBF::AdvanceOneTimeStepImpl(const FReal Dt, const FSubSt
 			// @todo(chaos): improve this - batching will be poor when some island are much larger than others. 
 			// We can't just loop over groups because there are some islands with no constraints (and usually a single particle) 
 			// that do not get added to a group but still need otbe checked for sleeping
-			const int32 NumIslands = GetConstraintGraph().NumIslands();
+			const int32 NumIslands = GetIslandManager().NumIslands();
 			PhysicsParallelFor(NumIslands, [&](const int32 IslandIndex)
 			{
-				FPBDIsland* Island = GetConstraintGraph().GetIsland(IslandIndex);
+					Private::FPBDIsland* Island = GetIslandManager().GetIsland(IslandIndex);
 				if (!Island->IsSleeping() && !Island->IsUsingCache())
 				{
 					BuildDisabledParticles(IslandIndex, DisabledParticles, SleepedIslands);
@@ -501,11 +501,11 @@ void FPBDRigidsEvolutionGBF::AdvanceOneTimeStepImpl(const FReal Dt, const FSubSt
 	{
 		SCOPE_CYCLE_COUNTER(STAT_Evolution_DeactivateSleep);
 		CSV_SCOPED_TIMING_STAT(PhysicsVerbose, StepSolver_DeactivateSleep);
-		for (int32 Island = 0; Island < GetConstraintGraph().NumIslands(); ++Island)
+		for (int32 Island = 0; Island < GetIslandManager().NumIslands(); ++Island)
 		{
 			if (SleepedIslands[Island])
 			{
-				GetConstraintGraph().SleepIsland(Particles, Island);
+				GetIslandManager().SleepIsland(Particles, Island);
 			}
 			
 			for (const auto Particle : DisabledParticles[Island])
@@ -525,7 +525,7 @@ void FPBDRigidsEvolutionGBF::AdvanceOneTimeStepImpl(const FReal Dt, const FSubSt
 	ParticleUpdatePosition(Particles.GetDirtyParticlesView(), Dt);
 
 	// Clean up the transient data from the constraint graph (e.g., Islands get cleared here)
-	GetConstraintGraph().EndTick();
+	GetIslandManager().EndTick();
 
 	if (bChaos_Solver_TestMode)
 	{
@@ -677,7 +677,7 @@ FPBDRigidsEvolutionGBF::~FPBDRigidsEvolutionGBF()
 {
 	// This is really only needed to ensure proper cleanup (we verify that constraints have been removed from 
 	// the graph in the destructor). This can be optimized if it's a problem but it shouldn't be
-	GetConstraintGraph().Reset();
+	GetIslandManager().Reset();
 }
 
 void FPBDRigidsEvolutionGBF::Serialize(FChaosArchive& Ar)
@@ -881,7 +881,7 @@ void FPBDRigidsEvolutionGBF::DestroyTransientConstraints(FGeometryParticleHandle
 	if (Particle != nullptr)
 	{
 		// Remove all the particle's collisions from the graph
-		GetConstraintGraph().RemoveParticleConstraints(Particle, CollisionConstraints.GetContainerId());
+		GetIslandManager().RemoveParticleConstraints(Particle, CollisionConstraints.GetContainerId());
 
 		// Mark all collision constraints for destruction
 		DestroyParticleCollisionsInAllocator(Particle);
diff --git a/Engine/Source/Runtime/Experimental/Chaos/Private/Chaos/PBDSuspensionConstraints.cpp b/Engine/Source/Runtime/Experimental/Chaos/Private/Chaos/PBDSuspensionConstraints.cpp
--- a/Engine/Source/Runtime/Experimental/Chaos/Private/Chaos/PBDSuspensionConstraints.cpp
+++ b/Engine/Source/Runtime/Experimental/Chaos/Private/Chaos/PBDSuspensionConstraints.cpp
@@ -114,7 +114,7 @@ namespace Chaos
 		}
 	}
 
-	void FPBDSuspensionConstraints::AddConstraintsToGraph(FPBDIslandManager& IslandManager)
+	void FPBDSuspensionConstraints::AddConstraintsToGraph(Private::FPBDIslandManager& IslandManager)
 	{
 		IslandManager.AddContainerConstraints(*this);
 	}
diff --git a/Engine/Source/Runtime/Experimental/Chaos/Private/PBDRigidsSolver.cpp b/Engine/Source/Runtime/Experimental/Chaos/Private/PBDRigidsSolver.cpp
--- a/Engine/Source/Runtime/Experimental/Chaos/Private/PBDRigidsSolver.cpp
+++ b/Engine/Source/Runtime/Experimental/Chaos/Private/PBDRigidsSolver.cpp
@@ -1545,7 +1545,7 @@ CSV_CUSTOM_STAT(PhysicsCounters, Name, Value, ECsvCustomStatOp::Set);
 		CHAOS_COUNTER_STAT(NumGeometryCollectionBodies, (int32)GetEvolution()->GetParticles().GetGeometryCollectionParticles().Size());
 
 		// Constraint counts
-		CHAOS_COUNTER_STAT(NumIslands, GetEvolution()->GetConstraintGraph().NumIslands());
+		CHAOS_COUNTER_STAT(NumIslands, GetEvolution()->GetIslandManager().NumIslands());
 		CHAOS_COUNTER_STAT(NumIslandGroups, GetEvolution()->GetIslandGroupManager().GetNumActiveGroups());
 		CHAOS_COUNTER_STAT(NumContacts, NumCollisionConstraints());
 		CHAOS_COUNTER_STAT(NumJoints, NumJointConstraints());
@@ -1674,7 +1674,7 @@ CSV_CUSTOM_STAT(PhysicsCounters, Name, Value, ECsvCustomStatOp::Set);
 		}
 		if (ChaosSolverDrawIslands == 1)
 		{
-			DebugDraw::DrawConstraintGraph(FRigidTransform3(), GetEvolution()->GetConstraintGraph(), &ChaosSolverDebugDebugDrawSettings);
+			DebugDraw::DrawConstraintGraph(FRigidTransform3(), GetEvolution()->GetIslandManager(), &ChaosSolverDebugDebugDrawSettings);
 		}
 		if (ChaosSolverDrawClusterConstraints == 1)
 		{
diff --git a/Engine/Source/Runtime/Experimental/Chaos/Public/Chaos/ChaosDebugDraw.h b/Engine/Source/Runtime/Experimental/Chaos/Public/Chaos/ChaosDebugDraw.h
--- a/Engine/Source/Runtime/Experimental/Chaos/Public/Chaos/ChaosDebugDraw.h
+++ b/Engine/Source/Runtime/Experimental/Chaos/Public/Chaos/ChaosDebugDraw.h
@@ -215,7 +215,7 @@ namespace Chaos
 		CHAOS_API void DrawJointConstraints(const FRigidTransform3& SpaceTransform, const FPBDJointConstraints& Constraints, FRealSingle ColorScale, const FChaosDebugDrawJointFeatures& FeatureMask = FChaosDebugDrawJointFeatures::MakeDefault(), const FChaosDebugDrawSettings* Settings = nullptr);
 		CHAOS_API void DrawSimulationSpace(const FSimulationSpace& SimSpace, const FChaosDebugDrawSettings* Settings = nullptr);
 		CHAOS_API void DrawShape(const FRigidTransform3& ShapeTransform, const FImplicitObject* Implicit, const FShapeOrShapesArray& Shapes, const FColor& Color, const FChaosDebugDrawSettings* Settings = nullptr);
-		CHAOS_API void DrawConstraintGraph(const FRigidTransform3& ShapeTransform, const FPBDConstraintGraph& Graph, const FChaosDebugDrawSettings* Settings = nullptr);
+		CHAOS_API void DrawConstraintGraph(const FRigidTransform3& ShapeTransform, const Private::FPBDIslandManager& Graph, const FChaosDebugDrawSettings* Settings = nullptr);
 		CHAOS_API void DrawConnectionGraph(const FRigidClustering& Clustering, const FChaosDebugDrawSettings* Settings = nullptr);
 		CHAOS_API void DrawCollidingShapes(const FRigidTransform3& SpaceTransform, const FPBDCollisionConstraints& Collisions, FRealSingle ColorScale, const FRealSingle Duration = 0.f, const FChaosDebugDrawSettings* Settings = nullptr);
 		CHAOS_API void DrawCollidingShapes(const FRigidTransform3& SpaceTransform, const FPBDCollisionConstraint& Collision, FRealSingle ColorScale, const FRealSingle Duration = 0.f, const FChaosDebugDrawSettings* Settings = nullptr);
diff --git a/Engine/Source/Runtime/Experimental/Chaos/Public/Chaos/ChaosDebugDrawDeclares.h b/Engine/Source/Runtime/Experimental/Chaos/Public/Chaos/ChaosDebugDrawDeclares.h
--- a/Engine/Source/Runtime/Experimental/Chaos/Public/Chaos/ChaosDebugDrawDeclares.h
+++ b/Engine/Source/Runtime/Experimental/Chaos/Public/Chaos/ChaosDebugDrawDeclares.h
@@ -8,6 +8,11 @@
 
 namespace Chaos
 {
+	namespace Private
+	{
+		class FPBDIslandManager;
+	}
+
 	class FConstraintHandle;
 
 	class FPBDCollisionConstraints;
@@ -20,10 +25,5 @@ namespace Chaos
 
 	class FPBDJointConstraints;
 
-	class FPBDIslandManager;
-
 	class FSimulationSpace;
-
-
-	using FPBDConstraintGraph = FPBDIslandManager;
 }
diff --git a/Engine/Source/Runtime/Experimental/Chaos/Public/Chaos/Collision/SolverCollisionContainer.h b/Engine/Source/Runtime/Experimental/Chaos/Public/Chaos/Collision/SolverCollisionContainer.h
--- a/Engine/Source/Runtime/Experimental/Chaos/Public/Chaos/Collision/SolverCollisionContainer.h
+++ b/Engine/Source/Runtime/Experimental/Chaos/Public/Chaos/Collision/SolverCollisionContainer.h
@@ -9,10 +9,13 @@
 
 namespace Chaos
 {
-	class FPBDIslandConstraint;
+	namespace Private
+	{
+		class FPBDIsland;
+		class FPBDIslandConstraint;
+	}
 	class FPBDCollisionConstraint;
 	class FPBDCollisionConstraints;
-	class FPBDIsland;
 	class FSolverBody;
 	class FSolverBodyContainer;
 
@@ -67,7 +70,7 @@ namespace Chaos
 		// IslandGroup API
 		//
 		virtual void AddConstraints() override final;
-		virtual void AddConstraints(const TArrayView<FPBDIslandConstraint>& ConstraintHandles) override final;
+		virtual void AddConstraints(const TArrayView<Private::FPBDIslandConstraint>& ConstraintHandles) override final;
 		virtual void AddBodies(FSolverBodyContainer& SolverBodyContainer) override final;
 		virtual void GatherInput(const FReal Dt) override final;
 		virtual void GatherInput(const FReal Dt, const int32 BeginIndex, const int32 EndIndex) override final;
diff --git a/Engine/Source/Runtime/Experimental/Chaos/Public/Chaos/Evolution/ConstraintGroupSolver.h b/Engine/Source/Runtime/Experimental/Chaos/Public/Chaos/Evolution/ConstraintGroupSolver.h
--- a/Engine/Source/Runtime/Experimental/Chaos/Public/Chaos/Evolution/ConstraintGroupSolver.h
+++ b/Engine/Source/Runtime/Experimental/Chaos/Public/Chaos/Evolution/ConstraintGroupSolver.h
@@ -2,6 +2,7 @@
 #pragma once
 
 #include "Chaos/Core.h"
+#include "Chaos/Evolution/IterationSettings.h"
 #include "Chaos/Evolution/SolverBodyContainer.h"
 #include "Containers/Array.h"
 
@@ -10,169 +11,204 @@ namespace Chaos
 	class FConstraintContainerSolver;
 	class FPBDConstraintContainer;
 
-	/**
-	 * All the data required to solver a set of constraints.
-	 * All constraints in the groups are solved in sequence in a single thread. We can create one of these for 
-	 * each Island (or group of Islands, or subset of a color of constraints) for parallelism.
-	*/
-	class CHAOS_API FPBDConstraintGroupSolver
+	namespace Private
 	{
-	public:
-		UE_NONCOPYABLE(FPBDConstraintGroupSolver);
-
-		FPBDConstraintGroupSolver();
-		virtual ~FPBDConstraintGroupSolver();
-
 		/**
-		 * Get the number of solver bodies we have. This is all the bodies referenced by all constraints and is only non-zero
-		 * after calling AddConstraintsAndBodies(). NOTE: it may be less than the number of particles in the IslandGroup 
-		 * (@see FPBDISlandGroup) when there are no constraints in some islands.
+		 * All the data required to solver a set of constraints.
+		 * All constraints in the groups are solved in sequence in a single thread. We can create one of these for 
+		 * each Island (or group of Islands, or subset of a color of constraints) for parallelism.
 		*/
-		inline int32 GetNumSolverBodies() const
+		class CHAOS_API FPBDConstraintGroupSolver
 		{
-			return SolverBodyContainer.Num();
-		}
+		public:
+			UE_NONCOPYABLE(FPBDConstraintGroupSolver);
+
+			FPBDConstraintGroupSolver();
+			virtual ~FPBDConstraintGroupSolver();
+
+			/**
+			 * Get the iterations settings for this group of islands
+			*/
+			FIterationSettings GetIterationSettings() const
+			{
+				return Iterations;
+			}
+
+			/**
+			 * Set the iterations settings for this group of islands
+			*/
+			virtual void SetIterationSettings(const FIterationSettings& InIterations)
+			{
+				Iterations = InIterations;
+			}
+
+			/**
+			 * Get the number of solver bodies we have. This is all the bodies referenced by all constraints and is only non-zero
+			 * after calling AddConstraintsAndBodies(). NOTE: it may be less than the number of particles in the IslandGroup 
+			 * (@see FPBDISlandGroup) when there are no constraints in some islands.
+			*/
+			inline int32 GetNumSolverBodies() const
+			{
+				return SolverBodyContainer.Num();
+			}
+
+			/**
+			 * Get the number of constraint solver we have, after AddConstraintsAndBodies() has been called.
+			*/
+			inline int32 GetNumSolverConstraints() const
+			{
+				return TotalNumConstraints;
+			}
+
+			/**
+			 * Attach a constraint solver to the specified ContainerId. This must be for the same constraint type as the container with that Id.
+			*/
+			void SetConstraintSolver(const int32 ContainerId, TUniquePtr<FConstraintContainerSolver>&& Solver);
+
+			/**
+			 * Set the solver priority of the specified constraint type
+			*/
+			void SetConstraintSolverPriority(const int32 ContainerId, const int32 Priority);
+
+			/**
+			 * Reset all state - called once per tick
+			*/
+			void Reset();
+
+			/**
+			 * Set up the constraints solvers and body containers with pointers to their constraint and particles, but do not collect any data.
+			*/
+			void AddConstraintsAndBodies();
+
+			/**
+			 * Collect all the data for all solver bodies from their respective particles.
+			*/
+			void GatherBodies(const FReal Dt);
+
+			/**
+			 * Collect all the data for the specified range of solver bodies from their respective particles. Will be called from multiple
+			 * threads with different non-overlapping indices.
+			*/
+			void GatherBodies(const FReal Dt, const int32 BeginBodyIndex, const int32 EndBodyIndex);
+
+			/**
+			 * Collect all the data for all constraint solvers from their respective constraints.
+			*/
+			void GatherConstraints(const FReal Dt);
+
+			/**
+			 * Collect all the data for the specified range of constraint solvers from their respective constraints. Will be called from multiple
+			 * threads with different non-overlapping indices.
+			*/
+			void GatherConstraints(const FReal Dt, const int32 BeginConstraintIndex, const int32 EndConstraintIndex);
+
+			/**
+			 * Apply positional constraints, and set the body velocities
+			*/
+			void ApplyPositionConstraints(const FReal Dt);
+
+			/**
+			 * Apply any velocity constraints, and update the body velocities
+			*/
+			void ApplyVelocityConstraints(const FReal Dt);
+
+			/**
+			 * Apply projection to attempt to fix up any errors left over from the Position and Velocity saolver phases
+			*/
+			void ApplyProjectionConstraints(const FReal Dt);
+
+			/**
+			 * Push results from all solver bodies back to their respective particles.
+			*/
+			void ScatterBodies(const FReal Dt);
+
+			/**
+			 * Push results from the specified range of solver bodies back to their respective particles. Will be called from multiple
+			 * threads with different non-overlapping indices.
+			*/
+			void ScatterBodies(const FReal Dt, const int32 BeginBodyIndex, const int32 EndBodyIndex);
+
+			/**
+			 * Push results from all constraint solvers back to their respective constraints.
+			*/
+			void ScatterConstraints(const FReal Dt);
+
+			/**
+			 * Push results from the specified range of constraint solvers back to their respective constraints.Will be called from multiple
+			 * threads with different non-overlapping indices. The range covers all constraint types.
+			*/
+			void ScatterConstraints(const FReal Dt, const int32 BeginConstraintIndex, const int32 EndConstraintIndex);
+
+		protected:
+
+			// Loop over the range of constraints and apply the lambda. The range is treats constraints of all types as a contiguous list
+			// and the function loops over constraint types and maps the range into the solver ranges for each type.
+			template<typename LambdaType>
+			void ApplyToConstraintRange(const int32 BeginConstraintIndex, const int32 EndConstraintIndex, const LambdaType& Lambda);
+
+			// Sort the solvers based on Level and other criteria for more stable solving
+			void SortSolverContainers();
+
+			// Allow a derived class to perform per-tick reset
+			virtual void ResetImpl() {}
+
+			// Allow a derived class to add data based on the number of constraint types we need to support
+			virtual void SetConstraintSolverImpl(const int32 ContainerId) {}
+
+			// Allow a derived class to add constraints solvers
+			virtual void AddConstraintsImpl() {}
+
+			// Allow a derived class to populate some of the SolverBody data (Level and Color)
+			virtual void GatherBodiesImpl(const FReal Dt, const int32 BeginBodyIndex, const int32 EndBodyIndex) {}
+
+			// All the bodies used by all the constraints solved by this group
+			FSolverBodyContainer SolverBodyContainer;
+
+			// A solver for each constraint type in the group
+			TArray<TUniquePtr<FConstraintContainerSolver>> ConstraintContainerSolvers;
+			int32 TotalNumConstraints;
+
+			// Sorted by priority for use in solve methods
+			TArray<FConstraintContainerSolver*> PrioritizedConstraintContainerSolvers;
+
+			FIterationSettings Iterations;
+		};
+
 
-		/**
-		 * Get the number of constraint solver we have, after AddConstraintsAndBodies() has been called.
+		//
+		//
+		//
+		//
+		//
+
+		/**
+		 * A Constraint Solver that solves all constraints in the scene in sequence. Used by RBAN and tests, as
+		 * opposed to the FPBDIslandConstraintGroupSolver used by the main scene and has a constraint graph to batch constraints
+		 * into non-interacting islands for parallelization.
+		 * 
+		 * There will only be one (or zero) FPBDSceneConstraintGroupSolver per simulation world.
 		*/
-		inline int32 GetNumSolverConstraints() const
+		class CHAOS_API FPBDSceneConstraintGroupSolver : public FPBDConstraintGroupSolver
 		{
-			return TotalNumConstraints;
-		}
-
-		/**
-		 * Attach a constraint solver to the specified ContainerId. This must be for the same constraint type as the container with that Id.
-		*/
-		void SetConstraintSolver(const int32 ContainerId, TUniquePtr<FConstraintContainerSolver>&& Solver);
-
-		/**
-		 * Set the solver priority of the specified constraint type
-		*/
-		void SetConstraintSolverPriority(const int32 ContainerId, const int32 Priority);
-
-		/**
-		 * Reset all state - called once per tick
-		*/
-		void Reset();
-
-		/**
-		 * Set up the constraints solvers and body containers with pointers to their constraint and particles, but do not collect any data.
-		*/
-		void AddConstraintsAndBodies();
-
-		/**
-		 * Collect all the data for all solver bodies from their respective particles.
-		*/
-		void GatherBodies(const FReal Dt);
-
-		/**
-		 * Collect all the data for the specified range of solver bodies from their respective particles. Will be called from multiple
-		 * threads with different non-overlapping indices.
-		*/
-		void GatherBodies(const FReal Dt, const int32 BeginBodyIndex, const int32 EndBodyIndex);
-
-		/**
-		 * Collect all the data for all constraint solvers from their respective constraints.
-		*/
-		void GatherConstraints(const FReal Dt);
-
-		/**
-		 * Collect all the data for the specified range of constraint solvers from their respective constraints. Will be called from multiple
-		 * threads with different non-overlapping indices.
-		*/
-		void GatherConstraints(const FReal Dt, const int32 BeginConstraintIndex, const int32 EndConstraintIndex);
-
-		/**
-		 * Apply positional constraints, and set the body velocities
-		*/
-		void ApplyPositionConstraints(const FReal Dt, const int32 NumIts);
-
-		/**
-		 * Apply any velocity constraints, and update the body velocities
-		*/
-		void ApplyVelocityConstraints(const FReal Dt, const int32 NumIts);
-
-		/**
-		 * Apply projection to attempt to fix up any errors left over from the Position and Velocity saolver phases
-		*/
-		void ApplyProjectionConstraints(const FReal Dt, const int32 NumIts);
-
-		/**
-		 * Push results from all solver bodies back to their respective particles.
-		*/
-		void ScatterBodies(const FReal Dt);
-
-		/**
-		 * Push results from the specified range of solver bodies back to their respective particles. Will be called from multiple
-		 * threads with different non-overlapping indices.
-		*/
-		void ScatterBodies(const FReal Dt, const int32 BeginBodyIndex, const int32 EndBodyIndex);
+		public:
+			FPBDSceneConstraintGroupSolver(const FIterationSettings& InIterations)
+			{
+				Iterations = InIterations;
+			}
 
-		/**
-		 * Push results from all constraint solvers back to their respective constraints.
-		*/
-		void ScatterConstraints(const FReal Dt);
-
-		/**
-		 * Push results from the specified range of constraint solvers back to their respective constraints.Will be called from multiple
-		 * threads with different non-overlapping indices. The range covers all constraint types.
-		*/
-		void ScatterConstraints(const FReal Dt, const int32 BeginConstraintIndex, const int32 EndConstraintIndex);
-
-	protected:
-
-		// Loop over the range of constraints and apply the lambda. The range is treats constraints of all types as a contiguous list
-		// and the function loops over constraint types and maps the range into the solver ranges for each type.
-		template<typename LambdaType>
-		void ApplyToConstraintRange(const int32 BeginConstraintIndex, const int32 EndConstraintIndex, const LambdaType& Lambda);
+		protected:
+			// The Scene Group Solver adds all constraints from all containers
+			virtual void AddConstraintsImpl() override final;
+		};
 
-		// Sort the solvers based on Level and other criteria for more stable solving
-		void SortSolverContainers();
+	} // namespace Private
 
-		// Allow a derived class to perform per-tick reset
-		virtual void ResetImpl() {}
 
-		// Allow a derived class to add data based on the number of constraint types we need to support
-		virtual void SetConstraintSolverImpl(const int32 ContainerId) {}
+	//UE_DEPRECATED(5.2, "FPBDConstraintGroupSolver is for internal use only")
+	//using FPBDConstraintGroupSolver = Private::FPBDConstraintGroupSolver;
 
-		// Allow a derived class to add constraints solvers
-		virtual void AddConstraintsImpl() {}
 
-		// Allow a derived class to populate some of the SolverBody data (Level and Color)
-		virtual void GatherBodiesImpl(const FReal Dt, const int32 BeginBodyIndex, const int32 EndBodyIndex) {}
-
-		// All the bodies used by all the constraints solved by this group
-		FSolverBodyContainer SolverBodyContainer;
-
-		// A solver for each constraint type in the group
-		TArray<TUniquePtr<FConstraintContainerSolver>> ConstraintContainerSolvers;
-		int32 TotalNumConstraints;
-
-		// Sorted by priority for use in solve methods
-		TArray<FConstraintContainerSolver*> PrioritizedConstraintContainerSolvers;
-	};
-
-
-	//
-	//
-	//
-	//
-	//
-
-	/**
-	 * A Constraint Solver that solves all constraints in the scene in sequence. Used by RBAN and tests, as
-	 * opposed to the FPBDIslandConstraintGroupSolver used by the main scene and has a constraint graph to batch constraints
-	 * into non-interacting islands for parallelization.
-	 * 
-	 * There will only be one (or zero) FPBDSceneConstraintGroupSolver per simulation world.
-	*/
-	class CHAOS_API FPBDSceneConstraintGroupSolver : public FPBDConstraintGroupSolver
-	{
-	public:
-	protected:
-		// The Scene Group Solver adds all constraints from all containers
-		virtual void AddConstraintsImpl() override final;
-	};
+	//UE_DEPRECATED(5.2, "FPBDSceneConstraintGroupSolver is for internal use only")
+	//using FPBDSceneConstraintGroupSolver = Private::FPBDSceneConstraintGroupSolver;
 
-}
\ No newline at end of file
+}	// namespace Chaos
\ No newline at end of file
diff --git a/Engine/Source/Runtime/Experimental/Chaos/Public/Chaos/Evolution/IndexedConstraintContainer.h b/Engine/Source/Runtime/Experimental/Chaos/Public/Chaos/Evolution/IndexedConstraintContainer.h
--- a/Engine/Source/Runtime/Experimental/Chaos/Public/Chaos/Evolution/IndexedConstraintContainer.h
+++ b/Engine/Source/Runtime/Experimental/Chaos/Public/Chaos/Evolution/IndexedConstraintContainer.h
@@ -153,9 +153,9 @@ namespace Chaos
 			ensure(false);
 		}
 
-		virtual void AddConstraints(const TArrayView<FPBDIslandConstraint>& Constraints) override final
+		virtual void AddConstraints(const TArrayView<Private::FPBDIslandConstraint>& Constraints) override final
 		{
-			for (const FPBDIslandConstraint& Constraint : Constraints)
+			for (const Private::FPBDIslandConstraint& Constraint : Constraints)
 			{
 				// We will only ever be given constraints from our container (asserts in non-shipping)
 				const FIndexedConstraintHandle* IndexedConstraintHandle = Constraint.GetConstraint()->AsUnsafe<FIndexedConstraintHandle>();
@@ -236,19 +236,6 @@ namespace Chaos
 		{
 		}
 
-		//virtual void SetConstraintEnabled(FConstraintHandle* ConstraintHandle, bool bEnabled) override final
-		//{
-		//	const FIndexedConstraintHandle* IndexedConstraintHandle = ConstraintHandle->As<FIndexedConstraintHandle>();
-		//	check(IndexedConstraintHandle != nullptr);
-		//	SetConstraintEnabled(IndexedConstraintHandle->GetConstraintIndex(), bEnabled);
-		//}
-		//virtual bool IsConstraintEnabled(FConstraintHandle* ConstraintHandle) const override final
-		//{
-		//	const FIndexedConstraintHandle* IndexedConstraintHandle = ConstraintHandle->As<FIndexedConstraintHandle>();
-		//	check(IndexedConstraintHandle != nullptr);
-		//	return IsConstraintEnabled(IndexedConstraintHandle->GetConstraintIndex());
-		//}
-
 		virtual void SetConstraintEnabled(int32 ConstraintIndex, bool bEnabled) { }
 		virtual bool IsConstraintEnabled(int32 ConstraintIndex) const { return true; }
 
diff --git a/Engine/Source/Runtime/Experimental/Chaos/Public/Chaos/Evolution/IterationSettings.h b/Engine/Source/Runtime/Experimental/Chaos/Public/Chaos/Evolution/IterationSettings.h
new file mode 100644
--- /dev/null
+++ b/Engine/Source/Runtime/Experimental/Chaos/Public/Chaos/Evolution/IterationSettings.h
@@ -0,0 +1,67 @@
+// Copyright Epic Games, Inc. All Rights Reserved.
+#pragma once
+
+#include "Chaos/Core.h"
+
+#include <limits>
+
+namespace Chaos
+{
+	namespace Private
+	{
+		/**
+		 * Iteration counts for use by bodies and constraints. 
+		 * A simulation island will use the maximum number of iterations of any body or constraint in the island.
+		*/
+		class CHAOS_API FIterationSettings
+		{
+		private:
+			using FIndexType = int8;
+			FIndexType NumPositionIterations;
+			FIndexType NumVelocityIterations;
+			FIndexType NumProjectionIterations;
+
+		public:
+			static const int32 InvalidIterations = INDEX_NONE;
+			static const int32 MaxPositionIterations = std::numeric_limits<FIndexType>::max();
+			static const int32 MaxVelocityIterations = std::numeric_limits<FIndexType>::max();
+			static const int32 MaxProjectionIterations = std::numeric_limits<FIndexType>::max();
+
+			static FIterationSettings Empty()
+			{
+				return FIterationSettings(InvalidIterations, InvalidIterations, InvalidIterations);
+			}
+
+			FIterationSettings()
+				: NumPositionIterations(InvalidIterations)
+				, NumVelocityIterations(InvalidIterations)
+				, NumProjectionIterations(InvalidIterations)
+			{
+			}
+
+			FIterationSettings(const int32 InNumPositionIterations, const int32 InNumVelocityIterations, const int32 InNumProjectionInterations)
+			{
+				SetNumPositionIterations(InNumPositionIterations);
+				SetNumVelocityIterations(InNumVelocityIterations);
+				SetNumProjectionIterations(InNumProjectionInterations);
+			}
+
+			int32 GetNumPositionIterations() const { return NumPositionIterations; }
+			int32 GetNumVelocityIterations() const { return NumVelocityIterations; }
+			int32 GetNumProjectionIterations() const { return NumProjectionIterations; }
+
+			void SetNumPositionIterations(const int32 InNum) { NumPositionIterations = FIndexType(FMath::Clamp(InNum, InvalidIterations, MaxPositionIterations)); }
+			void SetNumVelocityIterations(const int32 InNum) { NumVelocityIterations = FIndexType(FMath::Clamp(InNum, InvalidIterations, MaxVelocityIterations)); }
+			void SetNumProjectionIterations(const int32 InNum) { NumProjectionIterations = FIndexType(FMath::Clamp(InNum, InvalidIterations, MaxProjectionIterations)); }
+
+			static FIterationSettings Merge(const FIterationSettings& L, const FIterationSettings& R)
+			{
+				FIterationSettings Result;
+				Result.NumPositionIterations = FMath::Max(L.NumPositionIterations, R.NumPositionIterations);
+				Result.NumVelocityIterations = FMath::Max(L.NumVelocityIterations, R.NumVelocityIterations);
+				Result.NumProjectionIterations = FMath::Max(L.NumProjectionIterations, R.NumProjectionIterations);
+				return Result;
+			}
+		};
+	}
+}
\ No newline at end of file
diff --git a/Engine/Source/Runtime/Experimental/Chaos/Public/Chaos/Evolution/PBDMinEvolution.h b/Engine/Source/Runtime/Experimental/Chaos/Public/Chaos/Evolution/PBDMinEvolution.h
--- a/Engine/Source/Runtime/Experimental/Chaos/Public/Chaos/Evolution/PBDMinEvolution.h
+++ b/Engine/Source/Runtime/Experimental/Chaos/Public/Chaos/Evolution/PBDMinEvolution.h
@@ -7,6 +7,7 @@
 #include "Chaos/Evolution/ConstraintGroupSolver.h"
 #include "Chaos/Evolution/SimulationSpace.h"
 #include "Chaos/Evolution/SolverBodyContainer.h"
+#include "Chaos/Evolution/IterationSettings.h"
 #include "Chaos/ParticleHandleFwd.h"
 
 
@@ -44,17 +45,23 @@ namespace Chaos
 
 		void SetNumPositionIterations(const int32 NumIts)
 		{
-			NumPositionIterations = NumIts;
+			Private::FIterationSettings Iterations = ConstraintSolver.GetIterationSettings();
+			Iterations.SetNumPositionIterations(NumIts);
+			ConstraintSolver.SetIterationSettings(Iterations);
 		}
 
 		void SetNumVelocityIterations(const int32 NumIts)
 		{
-			NumVelocityIterations = NumIts;
+			Private::FIterationSettings Iterations = ConstraintSolver.GetIterationSettings();
+			Iterations.SetNumVelocityIterations(NumIts);
+			ConstraintSolver.SetIterationSettings(Iterations);
 		}
 
 		void SetNumProjectionIterations(const int32 NumIts)
 		{
-			NumProjectionIterations = NumIts;
+			Private::FIterationSettings Iterations = ConstraintSolver.GetIterationSettings();
+			Iterations.SetNumProjectionIterations(NumIts);
+			ConstraintSolver.SetIterationSettings(Iterations);
 		}
 
 		void SetGravity(const FVec3& G)
@@ -115,11 +122,8 @@ namespace Chaos
 		TArrayCollectionArray<FRotation3>& ParticlePrevRs;
 
 		TArray<FPBDConstraintContainer*> ConstraintContainers;
-		FPBDSceneConstraintGroupSolver ConstraintSolver;
+		Private::FPBDSceneConstraintGroupSolver ConstraintSolver;
 
-		int32 NumPositionIterations;
-		int32 NumVelocityIterations;
-		int32 NumProjectionIterations;
 		FVec3 Gravity;
 		FSimulationSpaceSettings SimulationSpaceSettings;
 		FSimulationSpace SimulationSpace;
diff --git a/Engine/Source/Runtime/Experimental/Chaos/Public/Chaos/Evolution/SolverConstraintContainer.h b/Engine/Source/Runtime/Experimental/Chaos/Public/Chaos/Evolution/SolverConstraintContainer.h
--- a/Engine/Source/Runtime/Experimental/Chaos/Public/Chaos/Evolution/SolverConstraintContainer.h
+++ b/Engine/Source/Runtime/Experimental/Chaos/Public/Chaos/Evolution/SolverConstraintContainer.h
@@ -5,10 +5,13 @@
 
 namespace Chaos
 {
+	namespace Private
+	{
+		class FPBDIslandConstraint;
+	}
 	class FConstraintHandleHolder;
-	class FPBDIslandConstraint;
 	class FSolverBodyContainer;
-	class FPBDIsland;
+
 
 	/** 
 	 * Base class for all the solver for a set of constraints of a specific type.
@@ -63,7 +66,7 @@ namespace Chaos
 		 * NOTE: this should not do any actual data gathering - it should just add to the list of constraints in this group. All data
 		 * gathering is handled in GatherInput.
 		*/
-		virtual void AddConstraints(const TArrayView<FPBDIslandConstraint>& Constraints) = 0;
+		virtual void AddConstraints(const TArrayView<Private::FPBDIslandConstraint>& Constraints) = 0;
 
 		/**
 		 * Add all the required bodies to the body container (required for the constraints added with AddConstraints)
@@ -127,7 +130,7 @@ namespace Chaos
 			// We solve all constraints in the container in the order it prefers so nothing to do here
 		}
 
-		virtual void AddConstraints(const TArrayView<FPBDIslandConstraint>& Constraints) override final
+		virtual void AddConstraints(const TArrayView<Private::FPBDIslandConstraint>& Constraints) override final
 		{
 			// This solver container is for use with the a non-graph evolution. It will not call this function.
 			ensure(false);
diff --git a/Engine/Source/Runtime/Experimental/Chaos/Public/Chaos/Island/IslandGroup.h b/Engine/Source/Runtime/Experimental/Chaos/Public/Chaos/Island/IslandGroup.h
--- a/Engine/Source/Runtime/Experimental/Chaos/Public/Chaos/Island/IslandGroup.h
+++ b/Engine/Source/Runtime/Experimental/Chaos/Public/Chaos/Island/IslandGroup.h
@@ -7,67 +7,83 @@
 
 namespace Chaos
 {
-	class FPBDConstraintContainer;
-	class FPBDIsland;
-	class FPBDIslandManager;
-
-	/**
-	* A set of constraints that will ne solved in sequence on a single thread. This will usually be the constraints
-	* from several islands, but may be a sub set of constraints of a single island when coloring is enabled for
-	* that island.
-	*/
-	class CHAOS_API FPBDIslandConstraintGroupSolver : public FPBDConstraintGroupSolver
+	namespace Private
 	{
-	public:
-		UE_NONCOPYABLE(FPBDIslandConstraintGroupSolver);
+		class FPBDIsland;
+		class FPBDIslandManager;
+	}
+	class FPBDConstraintContainer;
 
+	namespace Private
+	{
 		/**
-		* Init the island group 
+		* A set of constraints that will ne solved in sequence on a single thread. This will usually be the constraints
+		* from several islands, but may be a sub set of constraints of a single island when coloring is enabled for
+		* that island.
 		*/
-		FPBDIslandConstraintGroupSolver(FPBDIslandManager& InIslandManager);
+		class CHAOS_API FPBDIslandConstraintGroupSolver : public FPBDConstraintGroupSolver
+		{
+		public:
+			UE_NONCOPYABLE(FPBDIslandConstraintGroupSolver);
 
-		/**
-		* Add island to the group
-		* @param Island Island to be added
-		*/
-		void AddIsland(FPBDIsland* Island);
+			/**
+			* Init the island group 
+			*/
+			FPBDIslandConstraintGroupSolver(FPBDIslandManager& InIslandManager);
 
-		/**
-		* Return the islands within the group
-		*/
-		FORCEINLINE const TArray<FPBDIsland*>& GetIslands() { return Islands; }
+			/**
+			* Add island to the group
+			* @param Island Island to be added
+			*/
+			void AddIsland(FPBDIsland* Island);
 
-		/**
-		 * The number of particles (dynamics, kinematic and static) in the island.
-		*/
-		inline int32 GetNumParticles() const { return NumParticles; }
+			/**
+			* Return the islands within the group
+			*/
+			FORCEINLINE const TArray<FPBDIsland*>& GetIslands() { return Islands; }
 
-		/**
-		 * The number of constraints of all types in the island.
-		*/
-		inline int32 GetNumConstraints() const { return NumConstraints; }
+			/**
+			 * The number of particles (dynamics, kinematic and static) in the island.
+			*/
+			inline int32 GetNumParticles() const { return NumParticles; }
 
-		/**
-		 * The number of constraints of the specific container type in the island.
-		*/
-		inline int32 GetNumConstraints(const int32 ContainerId) const { return NumContainerConstraints[ContainerId]; }
-
-	protected:
-		// Base class overrides
-		virtual void SetConstraintSolverImpl(const int32 ContainerId) override final;
-		virtual void ResetImpl() override final;
-		virtual void AddConstraintsImpl() override final;
-		virtual void GatherBodiesImpl(const FReal Dt, const int32 BeginBodyIndex, const int32 EndBodyIndex) override final;
-
-	private:
-		FPBDIslandManager& IslandManager;
-
-		// Item counters used to initialize the solver containers
-		static constexpr int32 NumExpectedConstraintTypes = 5;
-		int32 NumParticles;
-		int32 NumConstraints;
-		TArray<int32, TInlineAllocator<NumExpectedConstraintTypes>> NumContainerConstraints;	// Per ContainerId
-
-		TArray<FPBDIsland*> Islands;
-	};
-}
\ No newline at end of file
+			/**
+			 * The number of constraints of all types in the island.
+			*/
+			inline int32 GetNumConstraints() const { return NumConstraints; }
+
+			/**
+			 * The number of constraints of the specific container type in the island.
+			*/
+			inline int32 GetNumConstraints(const int32 ContainerId) const { return NumContainerConstraints[ContainerId]; }
+
+			/**
+			 * Update the iteration settings based on the islands in the group, using at least as many as specified in InIterations
+			*/
+			virtual void SetIterationSettings(const FIterationSettings& InIterations) override;
+
+		protected:
+			// Base class overrides
+			virtual void SetConstraintSolverImpl(const int32 ContainerId) override final;
+			virtual void ResetImpl() override final;
+			virtual void AddConstraintsImpl() override final;
+			virtual void GatherBodiesImpl(const FReal Dt, const int32 BeginBodyIndex, const int32 EndBodyIndex) override final;
+
+		private:
+			FPBDIslandManager& IslandManager;
+
+			// Item counters used to initialize the solver containers
+			static constexpr int32 NumExpectedConstraintTypes = 5;
+			int32 NumParticles;
+			int32 NumConstraints;
+			TArray<int32, TInlineAllocator<NumExpectedConstraintTypes>> NumContainerConstraints;	// Per ContainerId
+
+			TArray<FPBDIsland*> Islands;
+		};
+
+	}	// namespace Private
+
+
+	using FPBDIslandConstraintGroupSolver UE_DEPRECATED(5.2, "Internal class moved to Private namespace") = Private::FPBDIslandConstraintGroupSolver;
+
+}	// namespace Chaos
\ No newline at end of file
diff --git a/Engine/Source/Runtime/Experimental/Chaos/Public/Chaos/Island/IslandGroupManager.h b/Engine/Source/Runtime/Experimental/Chaos/Public/Chaos/Island/IslandGroupManager.h
--- a/Engine/Source/Runtime/Experimental/Chaos/Public/Chaos/Island/IslandGroupManager.h
+++ b/Engine/Source/Runtime/Experimental/Chaos/Public/Chaos/Island/IslandGroupManager.h
@@ -5,9 +5,12 @@
 
 namespace Chaos
 {
+	namespace Private
+	{
+		class FPBDIsland;
+		class FPBDIslandManager;
+	}
 	class FPBDConstraintContainer;
-	class FPBDIsland;
-	class FPBDIslandManager;
 
 #if CSV_PROFILER
 	#define CSV_CUSTOM_STAT_ISLANDGROUP_HELPER(Stat) CSV_CUSTOM_STAT(PhysicsVerbose, Stat, Stats[Stat] * 1000.0, ECsvCustomStatOp::Set);
@@ -67,109 +70,123 @@ namespace Chaos
 
 #endif
 
-
-	/**
-	 * Assigns Islands to IslandGroups, attempting to have a roughly equal number of constraints per IslandGroup.
-	 * Each Island Group may be solved in parallel with the others and contains its own set of ConstraintContainerSolvers
-	 * (one for each type of constraint) so that solver data access may be cache efficient during the solver phases.
-	*/
-	class CHAOS_API FPBDIslandGroupManager
+	namespace Private
 	{
-	public:
-		UE_NONCOPYABLE(FPBDIslandGroupManager)
-
-		FPBDIslandGroupManager(FPBDIslandManager& InIslandManager);
-		~FPBDIslandGroupManager();
-
-		/**
-		 * Register a constraint type with the manager. This will create a ConstraintContainerSolver for the ConstraintsContainer.
-		 * Constraints are solved in Priority order with lower priorities first, so higher priorities "win".
-		*/
-		void AddConstraintContainer(FPBDConstraintContainer& ConstraintContainer, const int32 Priority = 0);
-
-		/**
-		 * Set the priority for the specified container (that must have been pre-registered with AddConstraintContainer).
-		*/
-		void SetConstraintContainerPriority(const int32 ContainerId, const int32 Priority);
 
 		/**
-		 * The number of groups with at least one constraint in them
+		 * Assigns Islands to IslandGroups, attempting to have a roughly equal number of constraints per IslandGroup.
+		 * Each Island Group may be solved in parallel with the others and contains its own set of ConstraintContainerSolvers
+		 * (one for each type of constraint) so that solver data access may be cache efficient during the solver phases.
 		*/
-		inline int32 GetNumActiveGroups() const { return NumActiveGroups; }
+		class CHAOS_API FPBDIslandGroupManager
+		{
+		public:
+			UE_NONCOPYABLE(FPBDIslandGroupManager)
+
+			FPBDIslandGroupManager(FPBDIslandManager& InIslandManager);
+			~FPBDIslandGroupManager();
+
+			/**
+			 * Register a constraint type with the manager. This will create a ConstraintContainerSolver for the ConstraintsContainer.
+			 * Constraints are solved in Priority order with lower priorities first, so higher priorities "win".
+			*/
+			void AddConstraintContainer(FPBDConstraintContainer& ConstraintContainer, const int32 Priority = 0);
+
+			/**
+			 * Set the priority for the specified container (that must have been pre-registered with AddConstraintContainer).
+			*/
+			void SetConstraintContainerPriority(const int32 ContainerId, const int32 Priority);
+
+			/**
+			 * The number of groups with at least one constraint in them
+			*/
+			inline int32 GetNumActiveGroups() const { return NumActiveGroups; }
+
+			/**
+			 * Get the specified group
+			*/
+			inline FPBDIslandConstraintGroupSolver* GetGroup(const int32 GroupIndex) { return IslandGroups[GroupIndex].Get(); }
+
+			/**
+			 * Pull all the active islands from the IslandManager and assign to groups.
+			 * @return 
+			*/
+			int32 BuildGroups();
+
+			/**
+			 * Set the default number of iterations to perform in the Solve step (can be increased on a per-island basis by any dynamic body).
+			*/
+			void SetIterationSettings(const FIterationSettings& InIterations);
+			void SetNumPositionIterations(const int32 InNumIterations) { Iterations.SetNumPositionIterations(InNumIterations); }
+			void SetNumVelocityIterations(const int32 InNumIterations) { Iterations.SetNumVelocityIterations(InNumIterations); }
+			void SetNumProjectionIterations(const int32 InNumIterations) { Iterations.SetNumProjectionIterations(InNumIterations); }
+
+			/**
+			 * Get the default iteration settings. These are the minimum number of iterations that will be run. An island may run more
+			 * depending on the settings on the dynamic bodies in that island.
+			*/
+			const FIterationSettings& GetIterationSettings() const { return Iterations; }
+
+			/**
+			 * Solve all constraints.
+			*/
+			void Solve(const FReal Dt);
+
+
+		private:
+			// Used to specify a range of bodies or constraints in an island group for more even parallelization of Gather and Scatter
+			struct FIslandGroupRange
+			{
+				FIslandGroupRange(FPBDIslandConstraintGroupSolver* InIslandGroup, const int32 InBeginIndex, const int32 InEndIndex)
+					: BeginIndex(InBeginIndex)
+					, EndIndex(InEndIndex)
+					, IslandGroup(InIslandGroup)
+				{
+				}
 
-		/**
-		 * Get the specified group
-		*/
-		inline FPBDIslandConstraintGroupSolver* GetGroup(const int32 GroupIndex) { return IslandGroups[GroupIndex].Get(); }
+				// Ranges are sorted by range size
+				friend bool operator<(const FIslandGroupRange& L, const FIslandGroupRange& R)
+				{
+					return (L.EndIndex - L.BeginIndex) > (R.EndIndex - R.BeginIndex);
+				}
 
-		/**
-		 * Pull all the active islands from the IslandManager and assign to groups.
-		 * @return 
-		*/
-		int32 BuildGroups();
+				int32 BeginIndex;
+				int32 EndIndex;
+				FPBDIslandConstraintGroupSolver* IslandGroup;
+			};
 
-		/**
-		 * Set the number of iterations to perform in the Solve step.
-		*/
-		void SetNumIterations(const int32 InNumPositionIterations, const int32 InNumVelocityIterations, const int32 InNumProjectionIterations);
+			void SolveSerial(const FReal Dt);
+			void SolveParallelFor(const FReal Dt);
+			void SolveParallelTasks(const FReal Dt);
 
-		/**
-		 * Solve all constraints.
-		*/
-		void Solve(const FReal Dt);
+			void BuildGatherBatches(TArray<FIslandGroupRange>& BodyRanges, TArray<FIslandGroupRange>& ConstraintRanges);
+			void SolveGroupConstraints(const int32 GroupIndex, const FReal Dt);
 
 
-	private:
-		// Used to specify a range of bodies or constraints in an island group for more even parallelization of Gather and Scatter
-		struct FIslandGroupRange
-		{
-			FIslandGroupRange(FPBDIslandConstraintGroupSolver* InIslandGroup, const int32 InBeginIndex, const int32 InEndIndex)
-				: BeginIndex(InBeginIndex)
-				, EndIndex(InEndIndex)
-				, IslandGroup(InIslandGroup)
-			{
-			}
+			FPBDIslandManager& IslandManager;
+			TArray<TUniquePtr<FPBDIslandConstraintGroupSolver>> IslandGroups;
+			int32 NumActiveGroups;
+			int32 NumWorkerThreads;
+			int32 TargetNumBodiesPerTask;
+			int32 TargetNumConstraintsPerTask;
+			FIterationSettings Iterations;
 
-			// Ranges are sorted by range size
-			friend bool operator<(const FIslandGroupRange& L, const FIslandGroupRange& R)
+	#if CSV_PROFILER
+			double& GetThreadStatAccumulator(const int32 ThreadIndex, const FIslandGroupStats::EPerIslandStat StatId)
 			{
-				return (L.EndIndex - L.BeginIndex) > (R.EndIndex - R.BeginIndex);
+				if (!GroupStats.IsValidIndex(ThreadIndex))
+				{
+					GroupStats.SetNum(ThreadIndex + 1);
+				}
+				return GroupStats[ThreadIndex].Stats[StatId];
 			}
-
-			int32 BeginIndex;
-			int32 EndIndex;
-			FPBDIslandConstraintGroupSolver* IslandGroup;
+			TArray<FIslandGroupStats> GroupStats;
+	#endif
 		};
 
-		void SolveSerial(const FReal Dt);
-		void SolveParallelFor(const FReal Dt);
-		void SolveParallelTasks(const FReal Dt);
-
-		void BuildGatherBatches(TArray<FIslandGroupRange>& BodyRanges, TArray<FIslandGroupRange>& ConstraintRanges);
-		void SolveGroupConstraints(const int32 GroupIndex, const FReal Dt);
-
+	}	// namespace Private
 
-		FPBDIslandManager& IslandManager;
-		TArray<TUniquePtr<FPBDIslandConstraintGroupSolver>> IslandGroups;
-		int32 NumActiveGroups;
-		int32 NumWorkerThreads;
-		int32 TargetNumBodiesPerTask;
-		int32 TargetNumConstraintsPerTask;
-		int32 NumPositionIterations;
-		int32 NumVelocityIterations;
-		int32 NumProjectionIterations;
 
-#if CSV_PROFILER
-		double& GetThreadStatAccumulator(const int32 ThreadIndex, const FIslandGroupStats::EPerIslandStat StatId)
-		{
-			if (!GroupStats.IsValidIndex(ThreadIndex))
-			{
-				GroupStats.SetNum(ThreadIndex + 1);
-			}
-			return GroupStats[ThreadIndex].Stats[StatId];
-		}
-		TArray<FIslandGroupStats> GroupStats;
-#endif
-	};
+	using FPBDIslandGroupManager UE_DEPRECATED(5.2, "Internal class moved to Private namespace") = Private::FPBDIslandGroupManager;
 
-}
\ No newline at end of file
+}	// namespace Chaos
diff --git a/Engine/Source/Runtime/Experimental/Chaos/Public/Chaos/Island/IslandManager.h b/Engine/Source/Runtime/Experimental/Chaos/Public/Chaos/Island/IslandManager.h
--- a/Engine/Source/Runtime/Experimental/Chaos/Public/Chaos/Island/IslandManager.h
+++ b/Engine/Source/Runtime/Experimental/Chaos/Public/Chaos/Island/IslandManager.h
@@ -31,6 +31,9 @@ class FChaosPhysicsMaterial;
 
 using FPBDRigidParticles = TPBDRigidParticles<FReal, 3>;
 
+namespace Private
+{
+
 /** Island manager responsible to create the list of solver islands that will be persistent over time */
 class CHAOS_API FPBDIslandManager
 {
@@ -359,4 +362,9 @@ protected:
 	bool									bEndTickCalled;
 };
 
-}
\ No newline at end of file
+}	// namespace Private
+
+
+	using FPBDIslandManager UE_DEPRECATED(5.2, "Internal class moved to Private namespace") = Private::FPBDIslandManager;
+
+}	// namespace Chaos
\ No newline at end of file
diff --git a/Engine/Source/Runtime/Experimental/Chaos/Public/Chaos/Island/SolverIsland.h b/Engine/Source/Runtime/Experimental/Chaos/Public/Chaos/Island/SolverIsland.h
--- a/Engine/Source/Runtime/Experimental/Chaos/Public/Chaos/Island/SolverIsland.h
+++ b/Engine/Source/Runtime/Experimental/Chaos/Public/Chaos/Island/SolverIsland.h
@@ -1,241 +1,264 @@
 // Copyright Epic Games, Inc. All Rights Reserved.
 #pragma once
 
-#include "Containers/Array.h"
-#include "Chaos/ParticleHandle.h"
 #include "Chaos/ConstraintHandle.h"
+#include "Chaos/Evolution/IterationSettings.h"
+#include "Chaos/ParticleHandle.h"
+
+#include "Containers/Array.h"
 
 namespace Chaos
 {
-	
 	/** Forward Declaration */
-	class FPBDIslandManager;
-	class FPBDRigidsSOAs;
-
-	/**
-	 * Data for a particle in an Island
-	*/
-	class CHAOS_API FPBDIslandParticle
+	namespace Private
 	{
-	public:
-		FPBDIslandParticle(FGeometryParticleHandle* InParticleHandle)
-			: ParticleHandle(InParticleHandle)
-		{
-		}
-
-		const FGeometryParticleHandle* GetParticle() const { return ParticleHandle; }
-		FGeometryParticleHandle* GetParticle() { return ParticleHandle; }
-
-	private:
-		FGeometryParticleHandle* ParticleHandle;
-	};
-
-	/**
-	* Data for a constraint in an island. This includes the Level and a sort key (which is build from the level and some 
-	* externally generated key that should be consistent between frames).
-	*/
-	class CHAOS_API FPBDIslandConstraint
-	{
-	public:
-		FPBDIslandConstraint(FConstraintHandle* InConstraintHandle, const int32 InLevel, const int32 InColor, const uint32 InSubSortKey)
-			: ConstraintHandle(InConstraintHandle)
-		{
-			// Sort by level, then sort key
-			SortKey = (uint64(InLevel) << 32) | uint64(InSubSortKey);
-		}
-
-		const FConstraintHandle* GetConstraint() const { return ConstraintHandle; }
-		FConstraintHandle* GetConstraint() { return ConstraintHandle; }
+		class FPBDIslandManager;
+	}
+	class FPBDRigidsSOAs;
 
-		int32 GetLevel() const { return int32((SortKey >> 32) & 0xFFFFF); }
-		
-		int32 GetColor() const { return INDEX_NONE; }
-		
-		uint64 GetSortKey() const { return SortKey; }
-
-	private:
-		uint64 SortKey;
-		FConstraintHandle* ConstraintHandle;
-	};
-
-	/**
-	* A set of constraints and particles that form a connected set in the constraint graph. These constraints
-	* must be solved sequentially (if not colored). FPBDIsland objects are collected into 
-	*/
-	class CHAOS_API FPBDIsland
+	namespace Private
 	{
-	public:
-
 		/**
-		* Init the solver island with an island manager and an index
-		* @param IslandManager Island manager that is storing the solver islands
-		* @param IslandIndex Index of the solver island within that list
+		 * Data for a particle in an Island
 		*/
-		FPBDIsland(const int32 MaxConstraintContainers);
+		class CHAOS_API FPBDIslandParticle
+		{
+		public:
+			FPBDIslandParticle(FGeometryParticleHandle* InParticleHandle)
+				: ParticleHandle(InParticleHandle)
+			{
+			}
 
-		/**
-		 * This island is about to be reused as a new island. It should be empty, but we need to reset
-		 * sleep state and other persistent state.
-		 */
-		void Reuse();
+			const FGeometryParticleHandle* GetParticle() const { return ParticleHandle; }
+			FGeometryParticleHandle* GetParticle() { return ParticleHandle; }
 
-		/**
-		 * The island index as seen by the systems external to the IslandManager.
-		 * The IslandManager uses this to index into the IslandIndexing array to get the
-		 * internal island index.
-		*/
-		int32 GetIslandIndex() const { return IslandIndex; }
+		private:
+			FGeometryParticleHandle* ParticleHandle;
+		};
 
 		/**
-		 * Set the external island index.
-		 * @see GetIslandIndex.
+		* Data for a constraint in an island. This includes the Level and a sort key (which is build from the level and some 
+		* externally generated key that should be consistent between frames).
 		*/
-		void SetIslandIndex(const int32 InIslandIndex) { IslandIndex = InIslandIndex; }
+		class CHAOS_API FPBDIslandConstraint
+		{
+		public:
+			FPBDIslandConstraint(FConstraintHandle* InConstraintHandle, const int32 InLevel, const int32 InColor, const uint32 InSubSortKey)
+				: ConstraintHandle(InConstraintHandle)
+			{
+				// Sort by level, then sort key
+				SortKey = (uint64(InLevel) << 32) | uint64(InSubSortKey);
+			}
 
-		/**
-		* Clear all the particles from the island
-		*/
-		void ClearParticles();
+			const FConstraintHandle* GetConstraint() const { return ConstraintHandle; }
+			FConstraintHandle* GetConstraint() { return ConstraintHandle; }
 
-		/**
-		* Add particle to the island
-		* @param ParticleHandle ParticleHandle to be added
-		*/
-		void AddParticle(FGenericParticleHandle ParticleHandle);
+			int32 GetLevel() const { return int32((SortKey >> 32) & 0xFFFFF); }
+		
+			int32 GetColor() const { return INDEX_NONE; }
+		
+			uint64 GetSortKey() const { return SortKey; }
 
-		/**
-		* Reset the particles list and reserve a number of particles in memory
-		* @param NumParticles Number of particles to be reserved
-		*/
-		void ReserveParticles(const int32 NumParticles);
+		private:
+			uint64 SortKey;
+			FConstraintHandle* ConstraintHandle;
+		};
 
 		/**
-		* Update the particles island index to match the graph index
+		* A set of constraints and particles that form a connected set in the constraint graph. These constraints
+		* must be solved sequentially (if not colored). FPBDIsland objects are collected into 
 		*/
-		void UpdateParticles();
+		class CHAOS_API FPBDIsland
+		{
+		public:
+
+			/**
+			* Init the solver island with an island manager and an index
+			* @param IslandManager Island manager that is storing the solver islands
+			* @param IslandIndex Index of the solver island within that list
+			*/
+			FPBDIsland(const int32 MaxConstraintContainers);
+
+			/**
+			 * This island is about to be reused as a new island. It should be empty, but we need to reset
+			 * sleep state and other persistent state.
+			 */
+			void Reuse();
+
+			/**
+			 * The island index as seen by the systems external to the IslandManager.
+			 * The IslandManager uses this to index into the IslandIndexing array to get the
+			 * internal island index.
+			*/
+			int32 GetIslandIndex() const { return IslandIndex; }
+
+			/**
+			 * Set the external island index.
+			 * @see GetIslandIndex.
+			*/
+			void SetIslandIndex(const int32 InIslandIndex) { IslandIndex = InIslandIndex; }
+
+			/**
+			 * The iteration settings for this island
+			*/
+			const FIterationSettings& GetIterationSettings() const { return Iterations; }
+
+			/**
+			* Clear all the particles from the island
+			*/
+			void ClearParticles();
+
+			/**
+			* Add particle to the island
+			* @param ParticleHandle ParticleHandle to be added
+			*/
+			void AddParticle(FGenericParticleHandle ParticleHandle);
+
+			/**
+			* Reset the particles list and reserve a number of particles in memory
+			* @param NumParticles Number of particles to be reserved
+			*/
+			void ReserveParticles(const int32 NumParticles);
+
+			/**
+			* Update the particles island index to match the graph index
+			*/
+			void UpdateParticles();
 	
-		/**
-		* Remove all the constraints from the solver island
-		*/
-		void ClearConstraints();
-
-		/**
-		* Add constraint to the island
-		* @param ConstraintHandle ConstraintHandle to be added
-		*/
-		void AddConstraint(FConstraintHandle* ConstraintHandle, const int32 Level, const int32 Color, const uint32 SubSortKey);
-
-		/**
-		* Sort the islands constraints
-		*/
-		void SortConstraints();
-
-		/**
-		* Return the list of particles within the solver island
-		*/
-		FORCEINLINE TArrayView<const FPBDIslandParticle> GetParticles() const { return MakeArrayView(IslandParticles); }
-		FORCEINLINE TArrayView<FPBDIslandParticle> GetParticles() { return MakeArrayView(IslandParticles); }
-
-		/**
-		* Return the list of constraints of the specified type within the island
-		*/
-		TArrayView<const FPBDIslandConstraint> GetConstraints(const int32 ContainerId) const { return MakeArrayView(IslandConstraintsByType[ContainerId]); }
-		TArrayView<FPBDIslandConstraint> GetConstraints(const int32 ContainerId) { return MakeArrayView(IslandConstraintsByType[ContainerId]); }
-
-		const int32 GetNumConstraintContainers() const { return IslandConstraintsByType.Num(); }
-
-		/**
-		* Get the number of particles within the island
-		*/
-		FORCEINLINE int32 GetNumParticles() const { return IslandParticles.Num(); }
-
-		/**
-		* Get the number of constraints (from the specified container) within the island
-		*/
-		FORCEINLINE int32 GetNumConstraints(const int32 ContainerId) const { return IslandConstraintsByType[ContainerId].Num(); }
-
-		/**
-		* Get the number of constraints (of all types) within the island
-		*/
-		FORCEINLINE int32 GetNumConstraints() const { return NumConstraints; }
-
-		/**
-		* Members accessors
-		*/
-		FORCEINLINE bool IsSleeping() const { return bIsSleeping; }
-		FORCEINLINE void SetIsSleeping(const bool bInIsSleepingIn) { bIsSleeping = bInIsSleepingIn; }
-		FORCEINLINE bool IsSleepingChanged() const { return bIsSleepingChanged; }
-		FORCEINLINE void SetIsSleepingChanged(const bool bInIsSleepingChanged) { bIsSleepingChanged = bInIsSleepingChanged; }
-
-		FORCEINLINE bool IsPersistent() const { return bIsPersistent; }
-		FORCEINLINE void SetIsPersistent(const bool bIsPersistentIn) { bIsPersistent = bIsPersistentIn; }
-		FORCEINLINE bool NeedsResim() const { return bNeedsResim; }
-		FORCEINLINE void SetNeedsResim(const bool bNeedsResimIn) { bNeedsResim = bNeedsResimIn; }
-		FORCEINLINE int32 GetSleepCounter() const { return SleepCounter; }
-		FORCEINLINE void SetSleepCounter(const int32 SleepCounterIn) { SleepCounter = SleepCounterIn; }
-		FORCEINLINE void SetIsUsingCache(const bool bIsUsingCacheIn ) { bIsUsingCache = bIsUsingCacheIn; }
-		FORCEINLINE bool IsUsingCache() const { return bIsUsingCache; }
+			/**
+			* Remove all the constraints from the solver island
+			*/
+			void ClearConstraints();
+
+			/**
+			* Add constraint to the island
+			* @param ConstraintHandle ConstraintHandle to be added
+			*/
+			void AddConstraint(FConstraintHandle* ConstraintHandle, const int32 Level, const int32 Color, const uint32 SubSortKey);
+
+			/**
+			* Sort the islands constraints
+			*/
+			void SortConstraints();
+
+			/**
+			* Return the list of particles within the solver island
+			*/
+			FORCEINLINE TArrayView<const FPBDIslandParticle> GetParticles() const { return MakeArrayView(IslandParticles); }
+			FORCEINLINE TArrayView<FPBDIslandParticle> GetParticles() { return MakeArrayView(IslandParticles); }
+
+			/**
+			* Return the list of constraints of the specified type within the island
+			*/
+			TArrayView<const FPBDIslandConstraint> GetConstraints(const int32 ContainerId) const { return MakeArrayView(IslandConstraintsByType[ContainerId]); }
+			TArrayView<FPBDIslandConstraint> GetConstraints(const int32 ContainerId) { return MakeArrayView(IslandConstraintsByType[ContainerId]); }
+
+			const int32 GetNumConstraintContainers() const { return IslandConstraintsByType.Num(); }
+
+			/**
+			* Get the number of particles within the island
+			*/
+			FORCEINLINE int32 GetNumParticles() const { return IslandParticles.Num(); }
+
+			/**
+			* Get the number of constraints (from the specified container) within the island
+			*/
+			FORCEINLINE int32 GetNumConstraints(const int32 ContainerId) const { return IslandConstraintsByType[ContainerId].Num(); }
+
+			/**
+			* Get the number of constraints (of all types) within the island
+			*/
+			FORCEINLINE int32 GetNumConstraints() const { return NumConstraints; }
+
+			/**
+			* Members accessors
+			*/
+			FORCEINLINE bool IsSleeping() const { return bIsSleeping; }
+			FORCEINLINE void SetIsSleeping(const bool bInIsSleepingIn) { bIsSleeping = bInIsSleepingIn; }
+			FORCEINLINE bool IsSleepingChanged() const { return bIsSleepingChanged; }
+			FORCEINLINE void SetIsSleepingChanged(const bool bInIsSleepingChanged) { bIsSleepingChanged = bInIsSleepingChanged; }
+
+			FORCEINLINE bool IsPersistent() const { return bIsPersistent; }
+			FORCEINLINE void SetIsPersistent(const bool bIsPersistentIn) { bIsPersistent = bIsPersistentIn; }
+			FORCEINLINE bool NeedsResim() const { return bNeedsResim; }
+			FORCEINLINE void SetNeedsResim(const bool bNeedsResimIn) { bNeedsResim = bNeedsResimIn; }
+			FORCEINLINE int32 GetSleepCounter() const { return SleepCounter; }
+			FORCEINLINE void SetSleepCounter(const int32 SleepCounterIn) { SleepCounter = SleepCounterIn; }
+			FORCEINLINE void SetIsUsingCache(const bool bIsUsingCacheIn ) { bIsUsingCache = bIsUsingCacheIn; }
+			FORCEINLINE bool IsUsingCache() const { return bIsUsingCache; }
 	
-		/**
-		 * Set the sleep state of all particles and constraints based on the island sleep state
-		*/
-		void PropagateSleepState(FPBDRigidsSOAs& Particles);
+			/**
+			 * Set the sleep state of all particles and constraints based on the island sleep state
+			*/
+			void PropagateSleepState(FPBDRigidsSOAs& Particles);
 
-		void UpdateSyncState(FPBDRigidsSOAs& Particles);
+			void UpdateSyncState(FPBDRigidsSOAs& Particles);
 
-		// Debug functionality
-		bool DebugContainsParticle(const FConstGenericParticleHandle& ParticleHandle) const
-		{
-			for (const FPBDIslandParticle& IslandParticle : IslandParticles)
+			// Debug functionality
+			bool DebugContainsParticle(const FConstGenericParticleHandle& ParticleHandle) const
 			{
-				if (IslandParticle.GetParticle() == ParticleHandle->Handle())
+				for (const FPBDIslandParticle& IslandParticle : IslandParticles)
 				{
-					return true;
+					if (IslandParticle.GetParticle() == ParticleHandle->Handle())
+					{
+						return true;
+					}
 				}
+				return false;
 			}
-			return false;
-		}
-		bool DebugContainsConstraint(FConstraintHandle* ConstraintHandle) const
-		{
-			for (const FPBDIslandConstraint& IslandConstraint : IslandConstraintsByType[ConstraintHandle->GetContainerId()])
+			bool DebugContainsConstraint(FConstraintHandle* ConstraintHandle) const
 			{
-				if (ConstraintHandle == IslandConstraint.GetConstraint())
+				for (const FPBDIslandConstraint& IslandConstraint : IslandConstraintsByType[ConstraintHandle->GetContainerId()])
 				{
-					return true;
+					if (ConstraintHandle == IslandConstraint.GetConstraint())
+					{
+						return true;
+					}
 				}
+				return false;
 			}
-			return false;
-		}
 
-	private:
+		private:
+
+			int32 IslandIndex;
 
-		int32 IslandIndex;
+			FIterationSettings Iterations;
 
-		/** Flag to check if an island is awake or sleeping */
-		bool bIsSleeping = false;
+			/** Flag to check if an island is awake or sleeping */
+			bool bIsSleeping = false;
 
-		/** Flag to check if an island need to be re-simulated or not */
-		bool bNeedsResim = false;
+			/** Flag to check if an island need to be re-simulated or not */
+			bool bNeedsResim = false;
 
-		/** Flag to check if an island is persistent over time */
-		bool bIsPersistent = true;
+			/** Flag to check if an island is persistent over time */
+			bool bIsPersistent = true;
 
-		/** Flag to check if the sleeping state has changed or not */
-		bool bIsSleepingChanged = false;
+			/** Flag to check if the sleeping state has changed or not */
+			bool bIsSleepingChanged = false;
 
-		/** Sleep counter to trigger island sleeping */
-		int32 SleepCounter = 0;
+			/** Sleep counter to trigger island sleeping */
+			int32 SleepCounter = 0;
 
-		/** List of all the island particles handles */
-		TArray<FPBDIslandParticle> IslandParticles;
+			/** List of all the island particles handles */
+			TArray<FPBDIslandParticle> IslandParticles;
 
-		/** List of all the island constraints handles */
-		static constexpr int32 NumExpectedConstraintTypes = 5;
-		TArray<TArray<FPBDIslandConstraint>, TInlineAllocator<NumExpectedConstraintTypes>> IslandConstraintsByType;
-		int32 NumConstraints;
+			/** List of all the island constraints handles */
+			static constexpr int32 NumExpectedConstraintTypes = 5;
+			TArray<TArray<FPBDIslandConstraint>, TInlineAllocator<NumExpectedConstraintTypes>> IslandConstraintsByType;
+			int32 NumConstraints;
+	
+			/** Check if the island is using the cache or not */
+			bool bIsUsingCache = false;
+	
+		};
+
+	}	// namespace Private
+
+
+	using FPBDIslandParticle UE_DEPRECATED(5.2, "Internal class moved to Private namespace") = Private::FPBDIslandParticle;
+	
+	using FPBDIslandConstraint UE_DEPRECATED(5.2, "Internal class moved to Private namespace") = Private::FPBDIslandConstraint;
 	
-		/** Check if the island is using the cache or not */
-		bool bIsUsingCache = false;
+	using FPBDIsland UE_DEPRECATED(5.2, "Internal class moved to Private namespace") = Private::FPBDIsland;
 	
-	};
-}
\ No newline at end of file
+}	// namespace Chaos
diff --git a/Engine/Source/Runtime/Experimental/Chaos/Public/Chaos/Joint/PBDJointContainerSolver.h b/Engine/Source/Runtime/Experimental/Chaos/Public/Chaos/Joint/PBDJointContainerSolver.h
--- a/Engine/Source/Runtime/Experimental/Chaos/Public/Chaos/Joint/PBDJointContainerSolver.h
+++ b/Engine/Source/Runtime/Experimental/Chaos/Public/Chaos/Joint/PBDJointContainerSolver.h
@@ -31,7 +31,7 @@ namespace Chaos
 			virtual int32 GetNumConstraints() const override final { return ContainerConstraintIndices.Num(); }
 			virtual void Reset(const int32 InMaxCollisions) override final;
 			virtual void AddConstraints() override final;
-			virtual void AddConstraints(const TArrayView<FPBDIslandConstraint>& IslandConstraints) override final;
+			virtual void AddConstraints(const TArrayView<Private::FPBDIslandConstraint>& IslandConstraints) override final;
 			virtual void AddBodies(FSolverBodyContainer& SolverBodyContainer) override final;
 			virtual void GatherInput(const FReal Dt) override final;
 			virtual void GatherInput(const FReal Dt, const int32 BeginIndex, const int32 EndIndex) override final;
diff --git a/Engine/Source/Runtime/Experimental/Chaos/Public/Chaos/PBDCollisionConstraints.h b/Engine/Source/Runtime/Experimental/Chaos/Public/Chaos/PBDCollisionConstraints.h
--- a/Engine/Source/Runtime/Experimental/Chaos/Public/Chaos/PBDCollisionConstraints.h
+++ b/Engine/Source/Runtime/Experimental/Chaos/Public/Chaos/PBDCollisionConstraints.h
@@ -132,7 +132,7 @@ public:
 	//
 	virtual int32 GetNumConstraints() const override final { return NumConstraints(); }
 	virtual void ResetConstraints() override final { Reset(); }
-	virtual void AddConstraintsToGraph(FPBDIslandManager& IslandManager) override final;
+	virtual void AddConstraintsToGraph(Private::FPBDIslandManager& IslandManager) override final;
 	virtual void PrepareTick() override final {}
 	virtual void UnprepareTick() override final {}
 
diff --git a/Engine/Source/Runtime/Experimental/Chaos/Public/Chaos/PBDConstraintContainer.h b/Engine/Source/Runtime/Experimental/Chaos/Public/Chaos/PBDConstraintContainer.h
--- a/Engine/Source/Runtime/Experimental/Chaos/Public/Chaos/PBDConstraintContainer.h
+++ b/Engine/Source/Runtime/Experimental/Chaos/Public/Chaos/PBDConstraintContainer.h
@@ -9,10 +9,11 @@
 
 namespace Chaos
 {
+	namespace Private
+	{
+		class FPBDIslandManager;
+	}
 	class FConstraintContainerSolver;
-	class FPBDIslandManager;
-	using FPBDConstraintGraph = FPBDIslandManager;
-
 
 	/**
 	 * Base class for containers of constraints.
@@ -98,7 +99,7 @@ namespace Chaos
 		/**
 		 * Add all the constraints in the container to the graph
 		*/
-		virtual void AddConstraintsToGraph(FPBDIslandManager& IslandManager) = 0;
+		virtual void AddConstraintsToGraph(Private::FPBDIslandManager& IslandManager) = 0;
 
 	protected:
 		FConstraintHandleTypeID ConstraintHandleType;
diff --git a/Engine/Source/Runtime/Experimental/Chaos/Public/Chaos/PBDConstraintGraph.h b/Engine/Source/Runtime/Experimental/Chaos/Public/Chaos/PBDConstraintGraph.h
--- a/Engine/Source/Runtime/Experimental/Chaos/Public/Chaos/PBDConstraintGraph.h
+++ b/Engine/Source/Runtime/Experimental/Chaos/Public/Chaos/PBDConstraintGraph.h
@@ -5,8 +5,7 @@
 
 namespace Chaos
 {
-	// @todo(chaos): deprecate FPBDConstraintGraph
-	using FPBDConstraintGraph = FPBDIslandManager;
+	using FPBDConstraintGraph UE_DEPRECATED(5.2, "Internal class moved to Private namespace") = Private::FPBDIslandManager;
 }
 
 
diff --git a/Engine/Source/Runtime/Experimental/Chaos/Public/Chaos/PBDJointConstraints.h b/Engine/Source/Runtime/Experimental/Chaos/Public/Chaos/PBDJointConstraints.h
--- a/Engine/Source/Runtime/Experimental/Chaos/Public/Chaos/PBDJointConstraints.h
+++ b/Engine/Source/Runtime/Experimental/Chaos/Public/Chaos/PBDJointConstraints.h
@@ -16,6 +16,7 @@ namespace Chaos
 	namespace Private
 	{
 		class FPBDJointContainerSolver;
+		class FPBDIslandManager;
 	}
 
 	/**
@@ -304,7 +305,7 @@ namespace Chaos
 		virtual TUniquePtr<FConstraintContainerSolver> CreateGroupSolver(const int32 Priority) override final;
 		virtual int32 GetNumConstraints() const override final { return NumConstraints(); }
 		virtual void ResetConstraints() override final {}
-		virtual void AddConstraintsToGraph(FPBDIslandManager& IslandManager) override final;
+		virtual void AddConstraintsToGraph(Private::FPBDIslandManager& IslandManager) override final;
 		virtual void PrepareTick() override final;
 		virtual void UnprepareTick() override final;
 
diff --git a/Engine/Source/Runtime/Experimental/Chaos/Public/Chaos/PBDNullConstraints.h b/Engine/Source/Runtime/Experimental/Chaos/Public/Chaos/PBDNullConstraints.h
--- a/Engine/Source/Runtime/Experimental/Chaos/Public/Chaos/PBDNullConstraints.h
+++ b/Engine/Source/Runtime/Experimental/Chaos/Public/Chaos/PBDNullConstraints.h
@@ -76,7 +76,7 @@ namespace Chaos
 		//
 		virtual int32 GetNumConstraints() const override final { return NumConstraints(); }
 		virtual void ResetConstraints() override final {}
-		virtual void AddConstraintsToGraph(FPBDIslandManager& IslandManager) override final;
+		virtual void AddConstraintsToGraph(Private::FPBDIslandManager& IslandManager) override final;
 		virtual void PrepareTick() override final {}
 		virtual void UnprepareTick() override final {}
 
diff --git a/Engine/Source/Runtime/Experimental/Chaos/Public/Chaos/PBDPositionConstraints.h b/Engine/Source/Runtime/Experimental/Chaos/Public/Chaos/PBDPositionConstraints.h
--- a/Engine/Source/Runtime/Experimental/Chaos/Public/Chaos/PBDPositionConstraints.h
+++ b/Engine/Source/Runtime/Experimental/Chaos/Public/Chaos/PBDPositionConstraints.h
@@ -177,7 +177,7 @@ namespace Chaos
 		//
 		virtual int32 GetNumConstraints() const override final { return NumConstraints(); }
 		virtual void ResetConstraints() override final {}
-		virtual void AddConstraintsToGraph(FPBDIslandManager& IslandManager) override final;
+		virtual void AddConstraintsToGraph(Private::FPBDIslandManager& IslandManager) override final;
 		virtual void PrepareTick() override final {}
 		virtual void UnprepareTick() override final {}
 
diff --git a/Engine/Source/Runtime/Experimental/Chaos/Public/Chaos/PBDRigidDynamicSpringConstraints.h b/Engine/Source/Runtime/Experimental/Chaos/Public/Chaos/PBDRigidDynamicSpringConstraints.h
--- a/Engine/Source/Runtime/Experimental/Chaos/Public/Chaos/PBDRigidDynamicSpringConstraints.h
+++ b/Engine/Source/Runtime/Experimental/Chaos/Public/Chaos/PBDRigidDynamicSpringConstraints.h
@@ -196,7 +196,7 @@ namespace Chaos
 		//
 		virtual int32 GetNumConstraints() const override final { return NumConstraints(); }
 		virtual void ResetConstraints() override final {}
-		virtual void AddConstraintsToGraph(FPBDIslandManager& IslandManager) override final;
+		virtual void AddConstraintsToGraph(Private::FPBDIslandManager& IslandManager) override final;
 		virtual void PrepareTick() override final {}
 		virtual void UnprepareTick() override final {}
 
diff --git a/Engine/Source/Runtime/Experimental/Chaos/Public/Chaos/PBDRigidSpringConstraints.h b/Engine/Source/Runtime/Experimental/Chaos/Public/Chaos/PBDRigidSpringConstraints.h
--- a/Engine/Source/Runtime/Experimental/Chaos/Public/Chaos/PBDRigidSpringConstraints.h
+++ b/Engine/Source/Runtime/Experimental/Chaos/Public/Chaos/PBDRigidSpringConstraints.h
@@ -158,7 +158,7 @@ namespace Chaos
 		//
 		virtual int32 GetNumConstraints() const override final { return NumConstraints(); }
 		virtual void ResetConstraints() override final {}
-		virtual void AddConstraintsToGraph(FPBDIslandManager& IslandManager) override final;
+		virtual void AddConstraintsToGraph(Private::FPBDIslandManager& IslandManager) override final;
 		virtual void PrepareTick() override final {}
 		virtual void UnprepareTick() override final {}
 
diff --git a/Engine/Source/Runtime/Experimental/Chaos/Public/Chaos/PBDRigidsEvolution.h b/Engine/Source/Runtime/Experimental/Chaos/Public/Chaos/PBDRigidsEvolution.h
--- a/Engine/Source/Runtime/Experimental/Chaos/Public/Chaos/PBDRigidsEvolution.h
+++ b/Engine/Source/Runtime/Experimental/Chaos/Public/Chaos/PBDRigidsEvolution.h
@@ -350,7 +350,7 @@ public:
 		const int32 ContainerId = ConstraintContainers.Add(&InContainer);
 		InContainer.SetContainerId(ContainerId);
 
-		GetConstraintGraph().AddConstraintContainer(InContainer);
+		GetIslandManager().AddConstraintContainer(InContainer);
 		IslandGroupManager.AddConstraintContainer(InContainer, Priority);
 	}
 
@@ -359,7 +359,7 @@ public:
 	*/
 	CHAOS_API void SetNumPositionIterations(int32 InNumIterations)
 	{
-		NumPositionIterations = InNumIterations;
+		IslandGroupManager.SetNumPositionIterations(InNumIterations);
 	}
 
 	/**
@@ -367,7 +367,7 @@ public:
 	*/
 	CHAOS_API int32 GetNumPositionIterations() const
 	{
-		return NumPositionIterations;
+		return IslandGroupManager.GetIterationSettings().GetNumPositionIterations();
 	}
 
 	/**
@@ -375,7 +375,7 @@ public:
 	*/
 	CHAOS_API void SetNumVelocityIterations(int32 InNumIterations)
 	{
-		NumVelocityIterations = InNumIterations;
+		IslandGroupManager.SetNumVelocityIterations(InNumIterations);
 	}
 
 	/**
@@ -383,7 +383,7 @@ public:
 	*/
 	CHAOS_API int32 GetNumVelocityIterations() const
 	{
-		return NumVelocityIterations;
+		return IslandGroupManager.GetIterationSettings().GetNumVelocityIterations();
 	}
 
 	/**
@@ -391,7 +391,7 @@ public:
 	*/
 	CHAOS_API void SetNumProjectionIterations(int32 InNumIterations)
 	{
-		NumProjectionIterations = InNumIterations;
+		IslandGroupManager.SetNumProjectionIterations(InNumIterations);
 	}
 
 	/**
@@ -399,7 +399,7 @@ public:
 	*/
 	CHAOS_API int32 GetNumProjectionIterations() const
 	{
-		return NumProjectionIterations;
+		return IslandGroupManager.GetIterationSettings().GetNumProjectionIterations();
 	}
 
 	/**
@@ -435,7 +435,7 @@ public:
 		{
 			if (Rigid->IsDynamic() && !Rigid->Disabled())
 			{
-				ConstraintGraph.AddParticle(Particle);
+				IslandManager.AddParticle(Particle);
 			}
 		}
 
@@ -451,7 +451,7 @@ public:
 	{
 		Particles.EnableParticle(Particle);
 		EnableConstraints(Particle);
-		ConstraintGraph.AddParticle(Particle);
+		IslandManager.AddParticle(Particle);
 		DirtyParticle(*Particle);
 	}
 
@@ -463,7 +463,7 @@ public:
 		RemoveParticleFromAccelerationStructure(*Particle);
 		Particles.DisableParticle(Particle);
 		DisableConstraints(Particle);
-		ConstraintGraph.RemoveParticle(Particle);
+		IslandManager.RemoveParticle(Particle);
 	}
 
 	/**
@@ -482,11 +482,11 @@ public:
 		if (Particle->IsInConstraintGraph())
 		{
 			// Remove the particle from the constraint graph. This should remove all the constraints too
-			ConstraintGraph.RemoveParticle(Particle);
+			IslandManager.RemoveParticle(Particle);
 
 			// Re-add the particle to the constraint graph
 			// (we could add a RemoveParticleConstraints method to the graph, but removing and adding a particle isn't too bad)
-			ConstraintGraph.AddParticle(Particle);
+			IslandManager.AddParticle(Particle);
 		}
 	}
 	
@@ -556,7 +556,7 @@ public:
 
 		RemoveParticleFromAccelerationStructure(*Particle);
 		DisconnectConstraints(TSet<FGeometryParticleHandle*>({ Particle }));
-		ConstraintGraph.RemoveParticle(Particle);
+		IslandManager.RemoveParticle(Particle);
 		Particles.DestroyParticle(Particle);
 	}
 
@@ -565,7 +565,7 @@ public:
 	 */
 	CHAOS_API void ReserveParticles(const int32 Num)
 	{
-		if (const int32 NumNew = ConstraintGraph.ReserveParticles(Num))
+		if (const int32 NumNew = IslandManager.ReserveParticles(Num))
 		{
 			InternalAccelerationQueue.PendingData.Reserve(InternalAccelerationQueue.Num() + NumNew);
 			AsyncAccelerationQueue.PendingData.Reserve(AsyncAccelerationQueue.Num() + NumNew);
@@ -583,7 +583,7 @@ public:
 	{
 		if (ConstraintHandle->IsInConstraintGraph())
 		{
-			ConstraintGraph.RemoveConstraint(ConstraintHandle->GetContainerId(), ConstraintHandle);
+			IslandManager.RemoveConstraint(ConstraintHandle->GetContainerId(), ConstraintHandle);
 		}
 	}
 
@@ -660,7 +660,7 @@ public:
 	CHAOS_API void ResetConstraints()
 	{
 		// Remove all the constraints from the graph
-		GetConstraintGraph().RemoveConstraints();
+		GetIslandManager().RemoveConstraints();
 		
 		// Clear all particle lists of collisions and constraints
 		// (this could be performed by the constraint containers
@@ -866,9 +866,14 @@ public:
 	/* Ticks computation of acceleration structures. Normally handled by Advance, but if not advancing can be called to incrementally build structures.*/
 	CHAOS_API void ComputeIntermediateSpatialAcceleration(bool bBlock = false);
 
-	CHAOS_API const FPBDConstraintGraph& GetConstraintGraph() const { return ConstraintGraph; }
-	CHAOS_API FPBDConstraintGraph& GetConstraintGraph() { return ConstraintGraph; }
-	CHAOS_API const FPBDIslandGroupManager& GetIslandGroupManager() const { return IslandGroupManager; }
+	UE_DEPRECATED(5.2, "Renamed to GetIslandManager")
+	CHAOS_API const Private::FPBDIslandManager& GetConstraintGraph() const { return IslandManager; }
+	UE_DEPRECATED(5.2, "Renamed to GetIslandManager")
+	CHAOS_API Private::FPBDIslandManager& GetConstraintGraph() { return IslandManager; }
+	
+	CHAOS_API Private::FPBDIslandManager& GetIslandManager() { return IslandManager; }
+	CHAOS_API const Private::FPBDIslandManager& GetIslandManager() const { return IslandManager; }
+	CHAOS_API const Private::FPBDIslandGroupManager& GetIslandGroupManager() const { return IslandGroupManager; }
 
 
 	void SetResim(bool bInResim) { bIsResim = bInResim; }
@@ -960,7 +965,7 @@ protected:
 	{
 		// Update the current state of the graph based on existing particles and constraints.
 		// Any new particles (from this tick) should have been added when they were enabled.
-		ConstraintGraph.InitializeGraph(Particles.GetNonDisabledDynamicView());
+		IslandManager.InitializeGraph(Particles.GetNonDisabledDynamicView());
 
 		// Add all constraints to the graph.
 		// NOTE: in PersistentGraph mode, only new constraints need to be added and expired ones should be removed.
@@ -970,14 +975,14 @@ protected:
 		// of particles, but we could probably make this cleaner.
 		for (FPBDConstraintContainer* ConstraintContainer : ConstraintContainers)
 		{
-			ConstraintContainer->AddConstraintsToGraph(GetConstraintGraph());
+			ConstraintContainer->AddConstraintsToGraph(GetIslandManager());
 		}
 	}
 
 	void CreateIslands()
 	{
 		// Package the constraints and particles into islands
-		ConstraintGraph.UpdateIslands(Particles);
+		IslandManager.UpdateIslands(Particles);
 	}
 	
 	void FlushInternalAccelerationQueue();
@@ -992,8 +997,8 @@ protected:
 	FKinematicUpdateRule KinematicUpdate;
 	FCaptureRewindRule CaptureRewindData;
 	TArray<FPBDConstraintContainer*> ConstraintContainers;
-	FPBDIslandManager ConstraintGraph;
-	FPBDIslandGroupManager IslandGroupManager;
+	Private::FPBDIslandManager IslandManager;
+	Private::FPBDIslandGroupManager IslandGroupManager;
 	TArrayCollectionArray<TSerializablePtr<FChaosPhysicsMaterial>> PhysicsMaterials;
 	TArrayCollectionArray<TUniquePtr<FChaosPhysicsMaterial>> PerParticlePhysicsMaterials;
 	TArrayCollectionArray<int32> ParticleDisableCount;
@@ -1086,9 +1091,6 @@ protected:
 	};
 	FGraphEventRef AccelerationStructureTaskComplete;
 
-	int32 NumPositionIterations;
-	int32 NumVelocityIterations;
-	int32 NumProjectionIterations;
 	TUniquePtr<ISpatialAccelerationCollectionFactory> SpatialCollectionFactory;
 
 	FAccelerationStructure* GetFreeSpatialAcceleration_Internal();
diff --git a/Engine/Source/Runtime/Experimental/Chaos/Public/Chaos/PBDSuspensionConstraints.h b/Engine/Source/Runtime/Experimental/Chaos/Public/Chaos/PBDSuspensionConstraints.h
--- a/Engine/Source/Runtime/Experimental/Chaos/Public/Chaos/PBDSuspensionConstraints.h
+++ b/Engine/Source/Runtime/Experimental/Chaos/Public/Chaos/PBDSuspensionConstraints.h
@@ -230,7 +230,7 @@ namespace Chaos
 		//
 		virtual int32 GetNumConstraints() const override final { return NumConstraints(); }
 		virtual void ResetConstraints() override final {}
-		virtual void AddConstraintsToGraph(FPBDIslandManager& IslandManager) override final;
+		virtual void AddConstraintsToGraph(Private::FPBDIslandManager& IslandManager) override final;
 		virtual void PrepareTick() override final {}
 		virtual void UnprepareTick() override final {}
 
diff --git a/Engine/Source/Runtime/Experimental/Chaos/Public/PhysicsProxy/FieldSystemProxyHelper.h b/Engine/Source/Runtime/Experimental/Chaos/Public/PhysicsProxy/FieldSystemProxyHelper.h
--- a/Engine/Source/Runtime/Experimental/Chaos/Public/PhysicsProxy/FieldSystemProxyHelper.h
+++ b/Engine/Source/Runtime/Experimental/Chaos/Public/PhysicsProxy/FieldSystemProxyHelper.h
@@ -200,7 +200,7 @@ namespace Chaos
 			RigidHandle->SetObjectStateLowLevel(FieldState);
 			if (!RigidHandle->Disabled())
 			{
-				RigidSolver->GetEvolution()->GetConstraintGraph().AddParticle(RigidHandle);
+				RigidSolver->GetEvolution()->GetIslandManager().AddParticle(RigidHandle);
 			}
 		}
 	}
