ChaosFlesh : Add quasistatic option for PBD evolution.

diff --git a/Engine/Plugins/Experimental/ChaosFlesh/Source/ChaosFleshEngine/Private/ChaosFlesh/ChaosDeformableSolverComponent.cpp b/Engine/Plugins/Experimental/ChaosFlesh/Source/ChaosFleshEngine/Private/ChaosFlesh/ChaosDeformableSolverComponent.cpp
--- a/Engine/Plugins/Experimental/ChaosFlesh/Source/ChaosFleshEngine/Private/ChaosFlesh/ChaosDeformableSolverComponent.cpp
+++ b/Engine/Plugins/Experimental/ChaosFlesh/Source/ChaosFleshEngine/Private/ChaosFlesh/ChaosDeformableSolverComponent.cpp
@@ -207,6 +207,7 @@ void UDeformableSolverComponent::Reset()
 		, bDoSelfCollision
 		, bUseGridBasedConstraints
 		, GridDx
+		, bDoQuasistatics
 		}));
 
 	
diff --git a/Engine/Plugins/Experimental/ChaosFlesh/Source/ChaosFleshEngine/Public/ChaosFlesh/ChaosDeformableSolverComponent.h b/Engine/Plugins/Experimental/ChaosFlesh/Source/ChaosFleshEngine/Public/ChaosFlesh/ChaosDeformableSolverComponent.h
--- a/Engine/Plugins/Experimental/ChaosFlesh/Source/ChaosFleshEngine/Public/ChaosFlesh/ChaosDeformableSolverComponent.h
+++ b/Engine/Plugins/Experimental/ChaosFlesh/Source/ChaosFleshEngine/Public/ChaosFlesh/ChaosDeformableSolverComponent.h
@@ -92,6 +92,9 @@ public:
 	UPROPERTY(EditAnywhere, Category = Chaos)
 		float GridDx = 25.;
 
+	UPROPERTY(EditAnywhere, Category = Chaos)
+		bool bDoQuasistatics = false;
+
 	//UPROPERTY(EditAnywhere, Category = Chaos)
 	//	bool bWaitForParallelDeformableTask = true;
 
diff --git a/Engine/Plugins/Experimental/ChaosFlesh/Source/ChaosFleshNodes/Private/Dataflow/ChaosFleshKinematicInitializationNodes.cpp b/Engine/Plugins/Experimental/ChaosFlesh/Source/ChaosFleshNodes/Private/Dataflow/ChaosFleshKinematicInitializationNodes.cpp
--- a/Engine/Plugins/Experimental/ChaosFlesh/Source/ChaosFleshNodes/Private/Dataflow/ChaosFleshKinematicInitializationNodes.cpp
+++ b/Engine/Plugins/Experimental/ChaosFlesh/Source/ChaosFleshNodes/Private/Dataflow/ChaosFleshKinematicInitializationNodes.cpp
@@ -227,11 +227,12 @@ void FSetVerticesKinematicDataflowNode::Evaluate(Dataflow::FContext& Context, co
 					BoundWeights.Add(1.f);
 				}
 			}
-			if (BoundVerts.Num() > 0)
-			{
-				GeometryCollection::Facades::FKinematicBindingFacade Kinematics(InCollection);
-				Kinematics.AddKinematicBinding(Kinematics.SetBoneBindings(INDEX_NONE, BoundVerts, BoundWeights));
-			}	
+			
+		}
+		if (BoundVerts.Num() > 0)
+		{
+			GeometryCollection::Facades::FKinematicBindingFacade Kinematics(InCollection);
+			Kinematics.AddKinematicBinding(Kinematics.SetBoneBindings(INDEX_NONE, BoundVerts, BoundWeights));
 		}
 		SetValue<DataType>(Context, InCollection, &Collection);
 	}
diff --git a/Engine/Source/Runtime/Experimental/Chaos/Private/Chaos/Deformable/ChaosDeformableSolver.cpp b/Engine/Source/Runtime/Experimental/Chaos/Private/Chaos/Deformable/ChaosDeformableSolver.cpp
--- a/Engine/Source/Runtime/Experimental/Chaos/Private/Chaos/Deformable/ChaosDeformableSolver.cpp
+++ b/Engine/Source/Runtime/Experimental/Chaos/Private/Chaos/Deformable/ChaosDeformableSolver.cpp
@@ -70,7 +70,13 @@ namespace Chaos::Softs
 		MObjects = TArrayCollectionArray<const UObject*>();
 		FSolverParticles LocalParticlesDummy;
 		FSolverRigidParticles RigidParticles;
-		Evolution.Reset(new FPBDEvolution(MoveTemp(LocalParticlesDummy), MoveTemp(RigidParticles), {}, Property.NumSolverIterations));
+		Evolution.Reset(new FPBDEvolution(MoveTemp(LocalParticlesDummy), MoveTemp(RigidParticles), {}, 
+			Property.NumSolverIterations, (FSolverReal)0.,
+			/*SelfCollisionsThickness = */(FSolverReal)0.,
+			/*CoefficientOfFriction = */(FSolverReal)0.,
+			/*FSolverReal Damping = */(FSolverReal)0.04,
+			/*FSolverReal LocalDamping = */(FSolverReal)0.,
+			Property.bDoQuasistatics));
 		Evolution->Particles().AddArray(&MObjects);
 		if (Property.bDoSelfCollision || Property.CacheToFile)
 		{
@@ -204,12 +210,9 @@ namespace Chaos::Softs
 
 				for (int32 vdx : BoundVerts)
 				{
-					if (BoneIndex != INDEX_NONE)
-					{
-						int32 ParticleIndex = Range[0] + vdx;
-						Evolution->Particles().InvM(ParticleIndex) = 0.f;
-						Evolution->Particles().PAndInvM(ParticleIndex).InvM = 0.f;
-					}
+					int32 ParticleIndex = Range[0] + vdx;
+					Evolution->Particles().InvM(ParticleIndex) = 0.f;
+					Evolution->Particles().PAndInvM(ParticleIndex).InvM = 0.f;
 				}
 			}
 		}
@@ -251,7 +254,7 @@ namespace Chaos::Softs
 
 			FXPBDCorotatedConstraints<FSolverReal, FSolverParticles>* CorotatedConstraint =
 				new FXPBDCorotatedConstraints<FSolverReal, FSolverParticles>(
-					Evolution->Particles(), Elements, /*bRecordMetric = */false, (FSolverReal)100000.0);
+					Evolution->Particles(), Elements, /*bRecordMetric = */false, (FSolverReal)1000000.0);
 
 			int32 InitIndex = Evolution->AddConstraintInitRange(1, true);
 			Evolution->ConstraintInits()[InitIndex] =
@@ -328,6 +331,7 @@ namespace Chaos::Softs
 						}
 
 						typedef GeometryCollection::Facades::FVertexBoneWeightsFacade FWeightsFacade;
+						bool bParticleTouched = false;
 						FWeightsFacade WeightsFacade(Rest);
 						if (WeightsFacade.IsValid())
 						{
@@ -337,6 +341,7 @@ namespace Chaos::Softs
 							{
 								if(FleshInputBuffer)
 								{
+									bParticleTouched = true;
 									TArray<int32> BoneIndices = WeightsFacade.GetBoneIndices()[ObjectVertexIndex];
 									TArray<float> BoneWeights = WeightsFacade.GetBoneWeights()[ObjectVertexIndex];
 
@@ -376,7 +381,7 @@ namespace Chaos::Softs
 								}
 							}
 						}
-						else
+						if (!bParticleTouched)
 						{
 							MParticles.X(Index) = GlobalTransform.TransformPosition(ChaosVert(Vertex[Index - Range[0]]));
 							MParticles.PAndInvM(Index).P = MParticles.X(Index);
diff --git a/Engine/Source/Runtime/Experimental/Chaos/Private/Chaos/PBDEvolution.cpp b/Engine/Source/Runtime/Experimental/Chaos/Private/Chaos/PBDEvolution.cpp
--- a/Engine/Source/Runtime/Experimental/Chaos/Private/Chaos/PBDEvolution.cpp
+++ b/Engine/Source/Runtime/Experimental/Chaos/Private/Chaos/PBDEvolution.cpp
@@ -80,7 +80,8 @@ FPBDEvolution::FPBDEvolution(
 	FSolverReal SelfCollisionThickness,
 	FSolverReal CoefficientOfFriction,
 	FSolverReal Damping,
-	FSolverReal LocalDamping)
+	FSolverReal LocalDamping,
+	bool bDoQuasistatics)
 	: MParticles(MoveTemp(InParticles))
 	, MParticlesActiveView(MParticles)
 	, MCollisionParticles(MoveTemp(InGeometryParticles))
@@ -96,6 +97,7 @@ FPBDEvolution::FPBDEvolution(
 	, MDamping(Damping)
 	, MLocalDamping(LocalDamping)
 	, MTime(0)
+	, bDoQuasistatics(bDoQuasistatics)
 {
 	// Add group arrays
 	TArrayCollection::AddArray(&MGroupGravityAccelerations);
@@ -270,47 +272,67 @@ void FPBDEvolution::PreIterationUpdate(
 			DampingIntegrated = Dt;
 		}
 
-		const int32 RangeSize = Range - Offset;
-		PhysicsParallelFor(RangeSize,
-			[this, &Offset, &ForceRule, &Gravity, &VelocityAndPressureField, &DampVelocityRule, DampingPowDt, DampingIntegrated, Dt](int32 i)
-			{
-				const int32 Index = Offset + i;
-				if (MParticles.InvM(Index) != (FSolverReal)0.)  // Process dynamic particles
+		if (bDoQuasistatics)
+		{
+			const int32 RangeSize = Range - Offset;
+			PhysicsParallelFor(RangeSize,
+				[this, &Offset, &ForceRule, &Gravity, &VelocityAndPressureField, &DampVelocityRule, DampingPowDt, DampingIntegrated, Dt](int32 i)
 				{
-					// Init forces with GravityForces
-					MParticles.Acceleration(Index) = Gravity;
-
-					// Force Rule
-					if (bForceRule)
+					const int32 Index = Offset + i;
+					if (MParticles.InvM(Index) != (FSolverReal)0.)  // Process dynamic particles
 					{
-						ForceRule(MParticles, Dt, Index); // F += M * A
+						MParticles.P(Index) = MParticles.X(Index);
 					}
-
-					// Velocity Field
-					if (bVelocityField)
+					else  // Process kinematic particles
 					{
-						VelocityAndPressureField.Apply(MParticles, Dt, Index);
+						MKinematicUpdate(MParticles, Dt, MTime, Index);
 					}
+				}, RangeSize < MinParallelBatchSize);
+		} 
+		else
+		{
+			const int32 RangeSize = Range - Offset;
+			PhysicsParallelFor(RangeSize,
+				[this, &Offset, &ForceRule, &Gravity, &VelocityAndPressureField, &DampVelocityRule, DampingPowDt, DampingIntegrated, Dt](int32 i)
+				{
+					const int32 Index = Offset + i;
+					if (MParticles.InvM(Index) != (FSolverReal)0.)  // Process dynamic particles
+					{
+						// Init forces with GravityForces
+						MParticles.Acceleration(Index) = Gravity;
 
-					// Euler Step Velocity
-					MParticles.V(Index) += MParticles.Acceleration(Index) * Dt;
+						// Force Rule
+						if (bForceRule)
+						{
+							ForceRule(MParticles, Dt, Index); // F += M * A
+						}
 
-					// Damp Velocity Rule
-					if (bDampVelocityRule)
-					{
-						DampVelocityRule.ApplyFast(MParticles, Dt, Index);
-					}
+						// Velocity Field
+						if (bVelocityField)
+						{
+							VelocityAndPressureField.Apply(MParticles, Dt, Index);
+						}
 
-					// Euler Step with point damping integration
-					MParticles.P(Index) = MParticles.X(Index) + MParticles.V(Index) * DampingIntegrated;
+						// Euler Step Velocity
+						MParticles.V(Index) += MParticles.Acceleration(Index) * Dt;
 
-					MParticles.V(Index) *= DampingPowDt;
-				}
-				else  // Process kinematic particles
-				{
-					MKinematicUpdate(MParticles, Dt, MTime, Index);
-				}
-			}, RangeSize < MinParallelBatchSize);
+						// Damp Velocity Rule
+						if (bDampVelocityRule)
+						{
+							DampVelocityRule.ApplyFast(MParticles, Dt, Index);
+						}
+
+						// Euler Step with point damping integration
+						MParticles.P(Index) = MParticles.X(Index) + MParticles.V(Index) * DampingIntegrated;
+
+						MParticles.V(Index) *= DampingPowDt;
+					}
+					else  // Process kinematic particles
+					{
+						MKinematicUpdate(MParticles, Dt, MTime, Index);
+					}
+				}, RangeSize < MinParallelBatchSize);
+		}
 	}
 }
 
diff --git a/Engine/Source/Runtime/Experimental/Chaos/Public/Chaos/Deformable/ChaosDeformableSolverTypes.h b/Engine/Source/Runtime/Experimental/Chaos/Public/Chaos/Deformable/ChaosDeformableSolverTypes.h
--- a/Engine/Source/Runtime/Experimental/Chaos/Public/Chaos/Deformable/ChaosDeformableSolverTypes.h
+++ b/Engine/Source/Runtime/Experimental/Chaos/Public/Chaos/Deformable/ChaosDeformableSolverTypes.h
@@ -23,7 +23,8 @@ namespace Chaos::Softs
 			bool InbUseFloor = true,
 			bool InbDoSelfCollision = false,
 			bool InbUseGridBasedConstraints = false,
-			FSolverReal InGridDx = (FSolverReal)1. )
+			FSolverReal InGridDx = (FSolverReal)1. ,
+			bool InbDoQuasistatics = false)
 			: NumSolverSubSteps(InNumSolverSubSteps)
 			, NumSolverIterations(InNumSolverIterations)
 			, FixTimeStep(InFixTimeStep)
@@ -34,6 +35,7 @@ namespace Chaos::Softs
 			, bDoSelfCollision(InbDoSelfCollision)
 			, bUseGridBasedConstraints(InbUseGridBasedConstraints)
 			, GridDx(InGridDx)
+			, bDoQuasistatics(InbDoQuasistatics)
 		{}
 
 		FDeformableSolverProperties(const FDeformableSolverProperties& InProp)
@@ -47,6 +49,7 @@ namespace Chaos::Softs
 			, bDoSelfCollision(InProp.bDoSelfCollision)
 			, bUseGridBasedConstraints(InProp.bUseGridBasedConstraints)
 			, GridDx(InProp.GridDx)
+			, bDoQuasistatics(InProp.bDoQuasistatics)
 		{}
 
 		int32 NumSolverSubSteps = 5;
@@ -59,6 +62,7 @@ namespace Chaos::Softs
 		bool bDoSelfCollision = false;
 		bool bUseGridBasedConstraints = false;
 		FSolverReal GridDx = (FSolverReal)1.;
+		bool bDoQuasistatics = false;
 	};
 
 
diff --git a/Engine/Source/Runtime/Experimental/Chaos/Public/Chaos/PBDEvolution.h b/Engine/Source/Runtime/Experimental/Chaos/Public/Chaos/PBDEvolution.h
--- a/Engine/Source/Runtime/Experimental/Chaos/Public/Chaos/PBDEvolution.h
+++ b/Engine/Source/Runtime/Experimental/Chaos/Public/Chaos/PBDEvolution.h
@@ -25,7 +25,8 @@ class CHAOS_API FPBDEvolution : public TArrayCollection
 		FSolverReal SelfCollisionsThickness = (FSolverReal)0.,
 		FSolverReal CoefficientOfFriction = (FSolverReal)0.,
 		FSolverReal Damping = (FSolverReal)0.04,
-		FSolverReal LocalDamping = (FSolverReal)0.);
+		FSolverReal LocalDamping = (FSolverReal)0.,
+		bool bDoQuasistatics = false);
 	~FPBDEvolution() {}
 
 	// Advance one time step. Filter the input time step if specified.
@@ -220,6 +221,7 @@ private:
 	FSolverReal MDamping;
 	FSolverReal MLocalDamping;
 	FSolverReal MTime;
+	bool bDoQuasistatics = false;
 };
 
 }  // End namespace Chaos::Softs
