Geometry collection : optimize collision generation of geometry collection
- Added TriangleMesh GetVertexSetAsArray using an array instead of a set - better perf when calling reserve on very large collections
- optimized ReindexMaterials to use the N+M version of the index remapping ( instead of the N*M )

diff --git a/Engine/Source/Runtime/Experimental/Chaos/Private/Chaos/TriangleMesh.cpp b/Engine/Source/Runtime/Experimental/Chaos/Private/Chaos/TriangleMesh.cpp
--- a/Engine/Source/Runtime/Experimental/Chaos/Private/Chaos/TriangleMesh.cpp
+++ b/Engine/Source/Runtime/Experimental/Chaos/Private/Chaos/TriangleMesh.cpp
@@ -171,6 +171,26 @@ void FTriangleMesh::GetVertexSet(TSet<int32>& VertexSet) const
 	}
 }
 
+void FTriangleMesh::GetVertexSetAsArray(TArray<int32>& VertexSet) const
+{
+	TBitArray<> VisitedVertices; // using local index
+	VisitedVertices.Init(false, MNumIndices);
+	VertexSet.Reserve(MNumIndices);
+	for (const TVec3<int32>& Element : MElements)
+	{
+		for (int32 Corner = 0; Corner < 3; Corner++)
+		{
+			const int32 VertexIndex = Element[Corner];
+			FBitReference VisitedRef = VisitedVertices[VertexIndex-MStartIdx];
+			if (VisitedRef == false)
+			{
+				VisitedRef = true;
+				VertexSet.Add(VertexIndex);
+			}
+		}
+	}
+}
+
 const TMap<int32, TSet<int32>>& FTriangleMesh::GetPointToNeighborsMap() const
 {
 	if (MPointToNeighborsMap.Num())
diff --git a/Engine/Source/Runtime/Experimental/Chaos/Private/GeometryCollection/GeometryCollection.cpp b/Engine/Source/Runtime/Experimental/Chaos/Private/GeometryCollection/GeometryCollection.cpp
--- a/Engine/Source/Runtime/Experimental/Chaos/Private/GeometryCollection/GeometryCollection.cpp
+++ b/Engine/Source/Runtime/Experimental/Chaos/Private/GeometryCollection/GeometryCollection.cpp
@@ -488,18 +488,23 @@ void FGeometryCollection::ReindexMaterials(FManagedArrayCollection& InCollection
 	const int32 NumSections = InCollection.NumElements(FGeometryCollection::MaterialGroup);
 	const int32 NumFaceElements = InCollection.NumElements(FGeometryCollection::FacesGroup);
 
-	int Idx = 0;
-	for (int Section=0; Section < NumSections; Section++)
+	// since we know the number of triangles per section we can precompute the start offset of each sections 
+	// this avoid nested loop that result in N * M operations (N=sections M=faces) and we can process it in (N + M) operations instead 
+	TArray<int32> OffsetPerSection;
+	OffsetPerSection.AddUninitialized(NumSections);
+	int32 SectionOffset = 0;
+	for (int Section = 0; Section < NumSections; Section++)
 	{
-		for (int FaceElement = 0; FaceElement < NumFaceElements; FaceElement++)
-		{
-			int32 ID = (MaterialID)[FaceElement];
-	
-			if (Section == ID)
-			{
-				(MaterialIndex)[Idx++] = FaceElement;
-			}
-		}
+		OffsetPerSection[Section] = SectionOffset;
+		SectionOffset += Sections[Section].NumTriangles;
+	}
+
+	for (int FaceElement = 0; FaceElement < NumFaceElements; FaceElement++)
+	{
+		const int32 SectionID = MaterialID[FaceElement];
+		int32& SectionOffsetRef = OffsetPerSection[SectionID];
+		//ensure(MaterialIndex[SectionOffsetRef] == FaceElement);
+		MaterialIndex[SectionOffsetRef++] = FaceElement;
 	}
 
 	// delete unused material sections
diff --git a/Engine/Source/Runtime/Experimental/Chaos/Private/GeometryCollection/GeometryCollectionCollisionStructureManager.cpp b/Engine/Source/Runtime/Experimental/Chaos/Private/GeometryCollection/GeometryCollectionCollisionStructureManager.cpp
--- a/Engine/Source/Runtime/Experimental/Chaos/Private/GeometryCollection/GeometryCollectionCollisionStructureManager.cpp
+++ b/Engine/Source/Runtime/Experimental/Chaos/Private/GeometryCollection/GeometryCollectionCollisionStructureManager.cpp
@@ -45,13 +45,13 @@ FCollisionStructureManager::NewSimplicial(
 	{
 		Chaos::FReal Extent = 0;
 		int32 LSVCounter = 0;
-		TSet<int32> Indices;
-		TriMesh.GetVertexSet(Indices);
+		TArray<int32> IndicesArray;
+		TriMesh.GetVertexSetAsArray(IndicesArray);
 		TArray<Chaos::FVec3> OutsideVertices;
 
 		bool bFullCopy = true;
 		int32 LocalCollisionParticlesMax = CollisionParticlesMaxInput > 0 ? FMath::Min(CollisionParticlesMaxInput, CollisionParticlesMax) : CollisionParticlesMax;
-		if (bCollisionParticlesUseImplicitCulling!=0 && Implicit && Indices.Num()>LocalCollisionParticlesMax)
+		if (bCollisionParticlesUseImplicitCulling!=0 && Implicit && IndicesArray.Num()>LocalCollisionParticlesMax)
 		{
 			Extent = Implicit->HasBoundingBox() ? Implicit->BoundingBox().Extents().Size() : 1.f;
 
@@ -60,8 +60,8 @@ FCollisionStructureManager::NewSimplicial(
 			//
 			//  Remove particles inside the levelset. (I think this is useless) 
 			//
-			OutsideVertices.AddUninitialized(Indices.Num());
-			for (int32 Idx : Indices)
+			OutsideVertices.AddUninitialized(IndicesArray.Num());
+			for (int32 Idx : IndicesArray)
 			{
 				const Chaos::FVec3& SamplePoint = Vertices.X(Idx);
 				if (Implicit->SignedDistance(SamplePoint) > Threshold)
@@ -79,8 +79,7 @@ FCollisionStructureManager::NewSimplicial(
 		if(bFullCopy)
 		{
 			FBox Bounds(ForceInitToZero);
-			TArray<int32> IndicesArray = Indices.Array();
-			OutsideVertices.AddUninitialized(Indices.Num());
+			OutsideVertices.AddUninitialized(IndicesArray.Num());
 			for (int32 Idx=0;Idx<IndicesArray.Num();Idx++)
 			{
 				Bounds += FVector(Vertices.X(IndicesArray[Idx]));
@@ -120,7 +119,7 @@ FCollisionStructureManager::NewSimplicial(
 
 		Simplicial->UpdateAccelerationStructures();
 
-		UE_LOG(LogChaos, Log, TEXT("NewSimplicial: InitialSize: %d, ImplicitExterior: %d, FullCopy: %d, FinalSize: %d"), Indices.Num(), LSVCounter, (int32)bFullCopy, NumParticles);
+		UE_LOG(LogChaos, Log, TEXT("NewSimplicial: InitialSize: %d, ImplicitExterior: %d, FullCopy: %d, FinalSize: %d"), IndicesArray.Num(), LSVCounter, (int32)bFullCopy, NumParticles);
 		return Simplicial;
 	}
 	UE_LOG(LogChaos, Log, TEXT("NewSimplicial::Empty"));
diff --git a/Engine/Source/Runtime/Experimental/Chaos/Private/PhysicsProxy/GeometryCollectionPhysicsProxy.cpp b/Engine/Source/Runtime/Experimental/Chaos/Private/PhysicsProxy/GeometryCollectionPhysicsProxy.cpp
--- a/Engine/Source/Runtime/Experimental/Chaos/Private/PhysicsProxy/GeometryCollectionPhysicsProxy.cpp
+++ b/Engine/Source/Runtime/Experimental/Chaos/Private/PhysicsProxy/GeometryCollectionPhysicsProxy.cpp
@@ -3692,7 +3692,7 @@ void FGeometryCollectionPhysicsProxy::InitializeSharedCollisionStructures(
 
 				FBox InstanceBoundingBox(EForceInit::ForceInitToZero);
 				{ // tmp scope
-					TSet<int32> VertsAdded;
+					TArray<int32> VertsAdded;
 					VertsAdded.Reserve(BiggestNumElements);
 					for (TPBDRigidParticleHandle<FReal, 3>* Child : ChildrenIndices)
 					{
@@ -3707,7 +3707,7 @@ void FGeometryCollectionPhysicsProxy::InitializeSharedCollisionStructures(
 							// transforming the mass space particles into the parent mass-space.
 							const FTransform ChildMassToClusterMass = (CollectionMassToLocal[ChildTransformIdx] * CollectionSpaceTransforms[ChildTransformIdx]).GetRelativeTransform(CollectionMassToLocal[ClusterTransformIdx] * CollectionSpaceTransforms[ClusterTransformIdx]);
 
-							ChildMesh->GetVertexSet(VertsAdded);
+							ChildMesh->GetVertexSetAsArray(VertsAdded);
 							for (const int32 VertIdx : VertsAdded)
 							{
 								//Update particles so they are in the cluster's mass space
diff --git a/Engine/Source/Runtime/Experimental/Chaos/Public/Chaos/TriangleMesh.h b/Engine/Source/Runtime/Experimental/Chaos/Public/Chaos/TriangleMesh.h
--- a/Engine/Source/Runtime/Experimental/Chaos/Public/Chaos/TriangleMesh.h
+++ b/Engine/Source/Runtime/Experimental/Chaos/Public/Chaos/TriangleMesh.h
@@ -48,6 +48,9 @@ namespace Chaos
 		/** Returns the unique set of vertices used by this triangle mesh. */
 		CHAOS_API void GetVertexSet(TSet<int32>& VertexSet) const;
 
+		/** Returns the unique set of vertices used by this triangle mesh. */
+		CHAOS_API void GetVertexSetAsArray(TArray<int32>& VertexSet) const;
+
 		/**
 		 * Extends the vertex range.
 		 *
