Introduce the FPhysicsObject opaque type which provides a unified interface to use when interacting with physics bodies.

- The Chaos::FPhysicsObject should be used in place of the FBodyInstance. Pass the handle (in an array view) to the appropriate Chaos::FPhysicsObjectInterface to get the same functionality the single particle physics proxy currently provides.
- Extracted out the FChaosScopedSceneLock out from inside PhysInterface_Chaos so it can be reused.
- Automatically create the Chaos::FPhysicsObject for single particle physics proxies and the geometry collection physics proxy.
- To better match the way the interface works, switched the geometry collection GT particle to be a TPBDRigidParticle instead of a TGeometryParticle though the extra state transfer is incomplete.

diff --git a/Engine/Source/Runtime/Engine/Classes/PhysicsEngine/PhysicsObjectExternalInterface.h b/Engine/Source/Runtime/Engine/Classes/PhysicsEngine/PhysicsObjectExternalInterface.h
new file mode 100644
--- /dev/null
+++ b/Engine/Source/Runtime/Engine/Classes/PhysicsEngine/PhysicsObjectExternalInterface.h
@@ -0,0 +1,55 @@
+// Copyright Epic Games, Inc. All Rights Reserved.
+
+#pragma once
+
+#include "Chaos/PhysicsObject.h"
+#include "Chaos/PhysicsObjectInterface.h"
+#include "Physics/Experimental/ChaosScopedSceneLock.h"
+#include "PhysicsObjectPhysicsCoreInterface.h"
+#include "Containers/ArrayView.h"
+#include <type_traits>
+
+template<EPhysicsInterfaceScopedLockType LockType>
+class ENGINE_API FLockedPhysicsObjectExternalInterface
+{
+public:
+	using InterfaceType = std::conditional_t<LockType == EPhysicsInterfaceScopedLockType::Read, Chaos::FReadPhysicsObjectInterface_External, Chaos::FWritePhysicsObjectInterface_External>;
+
+	FLockedPhysicsObjectExternalInterface(FChaosScene* Scene, InterfaceType&& InInterface)
+		: Lock(Scene, LockType)
+		, Interface(InInterface)
+	{
+	}
+
+	FLockedPhysicsObjectExternalInterface(FLockedPhysicsObjectExternalInterface& Other) = delete;
+	FLockedPhysicsObjectExternalInterface& operator=(FLockedPhysicsObjectExternalInterface& Other) = delete;
+
+	FLockedPhysicsObjectExternalInterface(FLockedPhysicsObjectExternalInterface&& Other) = default;
+	FLockedPhysicsObjectExternalInterface& operator=(FLockedPhysicsObjectExternalInterface&& Other) = default;
+
+	InterfaceType& GetInterface() { return Interface; }
+	InterfaceType* operator->() { return &Interface; }
+
+	// The FORCEINLINE gets around some linker issues when trying to call Destroy.
+	FORCEINLINE void Release()
+	{
+		Lock.Release();
+	}
+
+private:
+	FScopedSceneLock_Chaos Lock;
+	InterfaceType Interface;
+};
+
+using FLockedReadPhysicsObjectExternalInterface = FLockedPhysicsObjectExternalInterface<EPhysicsInterfaceScopedLockType::Read>;
+using FLockedWritePhysicsObjectExternalInterface = FLockedPhysicsObjectExternalInterface<EPhysicsInterfaceScopedLockType::Write>;
+
+class ENGINE_API FPhysicsObjectExternalInterface : public Chaos::FPhysicsObjectInterface, public PhysicsObjectPhysicsCoreInterface
+{
+public:
+	static FLockedReadPhysicsObjectExternalInterface LockRead(TArrayView<Chaos::FPhysicsObjectHandle> InObjects);
+	static FLockedWritePhysicsObjectExternalInterface LockWrite(TArrayView<Chaos::FPhysicsObjectHandle> InObjects);
+
+	static Chaos::FReadPhysicsObjectInterface_External GetRead_AssumesLocked();
+	static Chaos::FWritePhysicsObjectInterface_External GetWrite_AssumesLocked();
+};
\ No newline at end of file
diff --git a/Engine/Source/Runtime/Engine/Private/Collision/CollisionConversions.cpp b/Engine/Source/Runtime/Engine/Private/Collision/CollisionConversions.cpp
--- a/Engine/Source/Runtime/Engine/Private/Collision/CollisionConversions.cpp
+++ b/Engine/Source/Runtime/Engine/Private/Collision/CollisionConversions.cpp
@@ -163,7 +163,7 @@ static void SetHitResultFromShapeAndFaceIndex(const FPhysicsShape& Shape,  const
 				if (ActorProxy->GetType() == EPhysicsProxyType::GeometryCollectionType)
 				{
 					const FGeometryCollectionPhysicsProxy* ConcreteProxy = static_cast<const FGeometryCollectionPhysicsProxy*>(ActorProxy);
-					OutResult.Item = ConcreteProxy->GetItemIndexFromGTParticle_External(&Actor).GetItemIndex();
+					OutResult.Item = ConcreteProxy->GetItemIndexFromGTParticle_External(Actor.CastToRigidParticle()).GetItemIndex();
 				}
 			}
 		}
diff --git a/Engine/Source/Runtime/Engine/Private/PhysicsEngine/Experimental/ChaosScopedSceneLock.cpp b/Engine/Source/Runtime/Engine/Private/PhysicsEngine/Experimental/ChaosScopedSceneLock.cpp
new file mode 100644
--- /dev/null
+++ b/Engine/Source/Runtime/Engine/Private/PhysicsEngine/Experimental/ChaosScopedSceneLock.cpp
@@ -0,0 +1,207 @@
+// Copyright Epic Games, Inc. All Rights Reserved.
+
+#include "Physics/Experimental/ChaosScopedSceneLock.h"
+
+#include "Chaos/ChaosEngineInterface.h"
+#include "Physics/Experimental/PhysScene_Chaos.h"
+#include "PhysicsEngine/ConstraintInstance.h"
+#include "PhysicsEngine/PhysicsObjectExternalInterface.h"
+#include "Components/SkeletalMeshComponent.h"
+#include "PBDRigidsSolver.h"
+
+FScopedSceneLock_Chaos::FScopedSceneLock_Chaos(FPhysicsActorHandle const* InActorHandle, EPhysicsInterfaceScopedLockType InLockType)
+	: LockType(InLockType)
+{
+	auto Scene = GetSceneForActor(InActorHandle);
+	Solver = Scene ? Scene->GetSolver() : nullptr;
+	LockScene();
+}
+
+FScopedSceneLock_Chaos::FScopedSceneLock_Chaos(FPhysicsActorHandle const* InActorHandleA, FPhysicsActorHandle const* InActorHandleB, EPhysicsInterfaceScopedLockType InLockType)
+	: LockType(InLockType)
+{
+	FChaosScene* SceneA = GetSceneForActor(InActorHandleA);
+	FChaosScene* SceneB = GetSceneForActor(InActorHandleB);
+	FChaosScene* Scene = nullptr;
+
+	if (SceneA == SceneB)
+	{
+		Scene = SceneA;
+	}
+	else if (!SceneA || !SceneB)
+	{
+		Scene = SceneA ? SceneA : SceneB;
+	}
+	else
+	{
+		UE_LOG(LogPhysics, Warning, TEXT("Attempted to aquire a physics scene lock for two paired actors that were not in the same scene. Skipping lock"));
+	}
+
+	Solver = Scene ? Scene->GetSolver() : nullptr;
+	LockScene();
+}
+
+FScopedSceneLock_Chaos::FScopedSceneLock_Chaos(FPhysicsConstraintHandle const* InConstraintHandle, EPhysicsInterfaceScopedLockType InLockType)
+	: Solver(nullptr)
+	, LockType(InLockType)
+{
+	if (InConstraintHandle)
+	{
+		auto Scene = GetSceneForActor(InConstraintHandle);
+		Solver = Scene ? Scene->GetSolver() : nullptr;
+	}
+#if CHAOS_CHECKED
+	if (!Solver)
+	{
+		UE_LOG(LogPhysics, Warning, TEXT("Failed to find Scene for constraint. Skipping lock"));
+	}
+#endif
+	LockScene();
+}
+
+FScopedSceneLock_Chaos::FScopedSceneLock_Chaos(USkeletalMeshComponent* InSkelMeshComp, EPhysicsInterfaceScopedLockType InLockType)
+	: LockType(InLockType)
+{
+	Solver = nullptr;
+
+	if (InSkelMeshComp)
+	{
+		for (FBodyInstance* BI : InSkelMeshComp->Bodies)
+		{
+			auto Scene = GetSceneForActor(&BI->GetPhysicsActorHandle());
+			if (Scene)
+			{
+				Solver = Scene->GetSolver();
+				break;
+			}
+		}
+	}
+
+	LockScene();
+}
+
+FScopedSceneLock_Chaos::FScopedSceneLock_Chaos(Chaos::FPhysicsObjectHandle InObjectA, Chaos::FPhysicsObjectHandle InObjectB, EPhysicsInterfaceScopedLockType InLockType)
+	: LockType(InLockType)
+{
+	FChaosScene* SceneA = static_cast<FChaosScene*>(FPhysicsObjectExternalInterface::GetScene({ &InObjectA, 1 }));
+	FChaosScene* SceneB = static_cast<FChaosScene*>(FPhysicsObjectExternalInterface::GetScene({ &InObjectB, 1 }));
+	FChaosScene* Scene = nullptr;
+
+	if (SceneA == SceneB)
+	{
+		Scene = SceneA;
+	}
+	else if (!SceneA || !SceneB)
+	{
+		Scene = SceneA ? SceneA : SceneB;
+	}
+	else
+	{
+		UE_LOG(LogPhysics, Warning, TEXT("Attempted to aquire a physics scene lock for two paired physics objects that were not in the same scene. Skipping lock"));
+	}
+
+	Solver = Scene ? Scene->GetSolver() : nullptr;
+	LockScene();
+}
+
+FScopedSceneLock_Chaos::FScopedSceneLock_Chaos(FChaosScene* InScene, EPhysicsInterfaceScopedLockType InLockType)
+	: Solver(InScene ? InScene->GetSolver() : nullptr)
+	, LockType(InLockType)
+{
+	LockScene();
+}
+
+FScopedSceneLock_Chaos::FScopedSceneLock_Chaos(FScopedSceneLock_Chaos&& Other)
+{
+	*this = MoveTemp(Other);
+}
+
+FScopedSceneLock_Chaos& FScopedSceneLock_Chaos::operator=(FScopedSceneLock_Chaos&& Other)
+{
+	Solver = Other.Solver;
+	LockType = Other.LockType;
+	bHasLock = Other.bHasLock;
+
+	Other.bHasLock = false;
+	Other.Solver = nullptr;
+	return *this;
+}
+
+FScopedSceneLock_Chaos::~FScopedSceneLock_Chaos()
+{
+	Release();
+}
+
+void FScopedSceneLock_Chaos::Release()
+{
+	if (bHasLock)
+	{
+		UnlockScene();
+	}
+}
+
+void FScopedSceneLock_Chaos::LockScene()
+{
+	if (!Solver)
+	{
+		return;
+	}
+
+	switch (LockType)
+	{
+	case EPhysicsInterfaceScopedLockType::Read:
+		Solver->GetExternalDataLock_External().ReadLock();
+		break;
+	case EPhysicsInterfaceScopedLockType::Write:
+		Solver->GetExternalDataLock_External().WriteLock();
+		break;
+	}
+
+	bHasLock = true;
+}
+
+void FScopedSceneLock_Chaos::UnlockScene()
+{
+	if (!Solver)
+	{
+		return;
+	}
+
+	switch (LockType)
+	{
+	case EPhysicsInterfaceScopedLockType::Read:
+		Solver->GetExternalDataLock_External().ReadUnlock();
+		break;
+	case EPhysicsInterfaceScopedLockType::Write:
+		Solver->GetExternalDataLock_External().WriteUnlock();
+		break;
+	}
+
+	bHasLock = false;
+}
+
+FChaosScene* FScopedSceneLock_Chaos::GetSceneForActor(FPhysicsActorHandle const* InActorHandle)
+{
+	if (InActorHandle)
+	{
+		return static_cast<FPhysScene*>(FChaosEngineInterface::GetCurrentScene(*InActorHandle));
+	}
+
+	return nullptr;
+}
+
+FChaosScene* FScopedSceneLock_Chaos::GetSceneForActor(FPhysicsConstraintHandle const* InConstraintHandle)
+{
+	if (InConstraintHandle && InConstraintHandle->IsValid() && InConstraintHandle->Constraint->IsType(Chaos::EConstraintType::JointConstraintType))
+	{
+		Chaos::FJointConstraint* Constraint = static_cast<Chaos::FJointConstraint*>(InConstraintHandle->Constraint);
+
+		FConstraintInstanceBase* ConstraintInstance = (Constraint) ? FPhysicsUserData_Chaos::Get<FConstraintInstanceBase>(Constraint->GetUserData()) : nullptr;
+		if (ConstraintInstance)
+		{
+			return ConstraintInstance->GetPhysicsScene();
+		}
+	}
+
+	return nullptr;
+}
\ No newline at end of file
diff --git a/Engine/Source/Runtime/Engine/Private/PhysicsEngine/Experimental/PhysInterface_Chaos.cpp b/Engine/Source/Runtime/Engine/Private/PhysicsEngine/Experimental/PhysInterface_Chaos.cpp
--- a/Engine/Source/Runtime/Engine/Private/PhysicsEngine/Experimental/PhysInterface_Chaos.cpp
+++ b/Engine/Source/Runtime/Engine/Private/PhysicsEngine/Experimental/PhysInterface_Chaos.cpp
@@ -41,6 +41,7 @@
 #include "Math/UnrealMathUtility.h"
 #include "PBDRigidsSolver.h"
 #include "PhysicsEngine/BodySetup.h"
+#include "Physics/Experimental/ChaosScopedSceneLock.h"
 #include "Physics/PhysicsFiltering.h"
 #include "PhysicsInterfaceUtilsCore.h"
 #include "PhysicalMaterials/PhysicalMaterialMask.h"
@@ -493,167 +494,6 @@ void FPhysInterface_Chaos::UpdateDriveTarget_AssumesLocked(const FPhysicsConstra
 	}
 }
 
-
-
-enum class EPhysicsInterfaceScopedLockType : uint8
-{
-	Read,
-	Write
-};
-
-struct FScopedSceneLock_Chaos
-{
-	FScopedSceneLock_Chaos(FPhysicsActorHandle const * InActorHandle, EPhysicsInterfaceScopedLockType InLockType)
-		: LockType(InLockType)
-	{
-		auto Scene = GetSceneForActor(InActorHandle);
-		Solver = Scene ? Scene->GetSolver() : nullptr;
-		LockScene();
-	}
-
-	FScopedSceneLock_Chaos(FPhysicsActorHandle const * InActorHandleA, FPhysicsActorHandle const * InActorHandleB, EPhysicsInterfaceScopedLockType InLockType)
-		: LockType(InLockType)
-	{
-		FPhysScene_Chaos* SceneA = GetSceneForActor(InActorHandleA);
-		FPhysScene_Chaos* SceneB = GetSceneForActor(InActorHandleB);
-		FPhysScene_Chaos* Scene = nullptr;
-
-		if(SceneA == SceneB)
-		{
-			Scene = SceneA;
-		}
-		else if(!SceneA || !SceneB)
-		{
-			Scene = SceneA ? SceneA : SceneB;
-		}
-		else
-		{
-			UE_LOG(LogPhysics, Warning, TEXT("Attempted to aquire a physics scene lock for two paired actors that were not in the same scene. Skipping lock"));
-		}
-
-		Solver = Scene ? Scene->GetSolver() : nullptr;
-		LockScene();
-	}
-
-	FScopedSceneLock_Chaos(FPhysicsConstraintHandle const * InConstraintHandle, EPhysicsInterfaceScopedLockType InLockType)
-		: Solver(nullptr)
-		, LockType(InLockType)
-	{
-		if (InConstraintHandle)
-		{
-			auto Scene = GetSceneForActor(InConstraintHandle);
-			Solver = Scene ? Scene->GetSolver() : nullptr;
-		}
-#if CHAOS_CHECKED
-		if (!Solver)
-		{
-			UE_LOG(LogPhysics, Warning, TEXT("Failed to find Scene for constraint. Skipping lock"));
-		}
-#endif
-		LockScene();
-	}
-
-	FScopedSceneLock_Chaos(USkeletalMeshComponent* InSkelMeshComp, EPhysicsInterfaceScopedLockType InLockType)
-		: LockType(InLockType)
-	{
-		Solver = nullptr;
-
-		if(InSkelMeshComp)
-		{
-			for(FBodyInstance* BI : InSkelMeshComp->Bodies)
-			{
-				auto Scene = GetSceneForActor(&BI->GetPhysicsActorHandle());
-				if(Scene)
-				{
-					Solver = Scene->GetSolver();
-					break;
-				}
-			}
-		}
-
-		LockScene();
-	}
-
-	FScopedSceneLock_Chaos(FPhysScene_Chaos* InScene, EPhysicsInterfaceScopedLockType InLockType)
-		: Solver(InScene ? InScene->GetSolver() : nullptr)
-		, LockType(InLockType)
-	{
-		LockScene();
-	}
-
-	~FScopedSceneLock_Chaos()
-	{
-		UnlockScene();
-	}
-
-private:
-
-	void LockScene()
-	{
-		if(!Solver)
-		{
-			return;
-		}
-
-		switch(LockType)
-		{
-		case EPhysicsInterfaceScopedLockType::Read:
-			Solver->GetExternalDataLock_External().ReadLock();
-			break;
-		case EPhysicsInterfaceScopedLockType::Write:
-			Solver->GetExternalDataLock_External().WriteLock();
-			break;
-		}
-	}
-
-	void UnlockScene()
-	{
-		if(!Solver)
-		{
-			return;
-		}
-
-		switch(LockType)
-		{
-		case EPhysicsInterfaceScopedLockType::Read:
-			Solver->GetExternalDataLock_External().ReadUnlock();
-			break;
-		case EPhysicsInterfaceScopedLockType::Write:
-			Solver->GetExternalDataLock_External().WriteUnlock();
-			break;
-		}
-	}
-
-	FPhysScene_Chaos* GetSceneForActor(FPhysicsActorHandle const * InActorHandle)
-	{
-		if(InActorHandle)
-		{
-			return static_cast<FPhysScene*>(FChaosEngineInterface::GetCurrentScene(*InActorHandle));
-		}
-		
-		return nullptr;
-	}
-
-	FPhysScene_Chaos* GetSceneForActor(FPhysicsConstraintHandle const* InConstraintHandle)
-	{		
-		if (InConstraintHandle && InConstraintHandle->IsValid() && InConstraintHandle->Constraint->IsType(Chaos::EConstraintType::JointConstraintType))
-		{
-			Chaos::FJointConstraint* Constraint = static_cast<Chaos::FJointConstraint*>(InConstraintHandle->Constraint);
-
-			FConstraintInstanceBase* ConstraintInstance = (Constraint) ? FPhysicsUserData_Chaos::Get<FConstraintInstanceBase>(Constraint->GetUserData()) : nullptr;
-			if (ConstraintInstance)
-			{
-				return ConstraintInstance->GetPhysicsScene();
-			}
-		}
-	
-		return nullptr;
-	}
-
-	Chaos::FPBDRigidsSolver* Solver;
-	EPhysicsInterfaceScopedLockType LockType;
-};
-
 bool FPhysInterface_Chaos::ExecuteOnUnbrokenConstraintReadOnly(const FPhysicsConstraintHandle& InConstraintRef, TFunctionRef<void(const FPhysicsConstraintHandle&)> Func)
 {
     if (!IsBroken(InConstraintRef))
diff --git a/Engine/Source/Runtime/Engine/Private/PhysicsEngine/Experimental/PhysScene_Chaos.cpp b/Engine/Source/Runtime/Engine/Private/PhysicsEngine/Experimental/PhysScene_Chaos.cpp
--- a/Engine/Source/Runtime/Engine/Private/PhysicsEngine/Experimental/PhysScene_Chaos.cpp
+++ b/Engine/Source/Runtime/Engine/Private/PhysicsEngine/Experimental/PhysScene_Chaos.cpp
@@ -652,9 +652,9 @@ void FPhysScene_Chaos::RemoveObject(FGeometryCollectionPhysicsProxy* InObject)
 {
 	Chaos::FPhysicsSolver* Solver = InObject->GetSolver<Chaos::FPhysicsSolver>();
 
-	for (TUniquePtr<Chaos::TGeometryParticle<Chaos::FReal, 3>>& GTParticleUnique : InObject->GetExternalParticles())
+	for (TUniquePtr<Chaos::TPBDRigidParticle<Chaos::FReal, 3>>& GTParticleUnique : InObject->GetExternalParticles())
 	{
-		Chaos::TGeometryParticle<Chaos::FReal, 3>* GTParticle = GTParticleUnique.Get();
+		Chaos::TPBDRigidParticle<Chaos::FReal, 3>* GTParticle = GTParticleUnique.Get();
 		RemoveActorFromAccelerationStructureImp(GTParticle);
 		if(Solver)
 		{
diff --git a/Engine/Source/Runtime/Engine/Private/PhysicsEngine/PhysicsObjectExternalInterface.cpp b/Engine/Source/Runtime/Engine/Private/PhysicsEngine/PhysicsObjectExternalInterface.cpp
new file mode 100644
--- /dev/null
+++ b/Engine/Source/Runtime/Engine/Private/PhysicsEngine/PhysicsObjectExternalInterface.cpp
@@ -0,0 +1,22 @@
+// Copyright Epic Games, Inc. All Rights Reserved.
+#include "PhysicsEngine/PhysicsObjectExternalInterface.h"
+
+FLockedReadPhysicsObjectExternalInterface FPhysicsObjectExternalInterface::LockRead(TArrayView<Chaos::FPhysicsObjectHandle> InObjects)
+{
+	return FLockedReadPhysicsObjectExternalInterface{ GetScene(InObjects), CreateReadInterface<Chaos::EThreadContext::External>() };
+}
+
+FLockedWritePhysicsObjectExternalInterface FPhysicsObjectExternalInterface::LockWrite(TArrayView<Chaos::FPhysicsObjectHandle> InObjects)
+{
+	return FLockedWritePhysicsObjectExternalInterface{ GetScene(InObjects), CreateWriteInterface<Chaos::EThreadContext::External>() };
+}
+
+Chaos::FReadPhysicsObjectInterface_External FPhysicsObjectExternalInterface::GetRead_AssumesLocked()
+{
+	return CreateReadInterface<Chaos::EThreadContext::External>();
+}
+
+Chaos::FWritePhysicsObjectInterface_External FPhysicsObjectExternalInterface::GetWrite_AssumesLocked()
+{
+	return CreateWriteInterface<Chaos::EThreadContext::External>();
+}
\ No newline at end of file
diff --git a/Engine/Source/Runtime/Engine/Public/Physics/Experimental/ChaosScopedSceneLock.h b/Engine/Source/Runtime/Engine/Public/Physics/Experimental/ChaosScopedSceneLock.h
new file mode 100644
--- /dev/null
+++ b/Engine/Source/Runtime/Engine/Public/Physics/Experimental/ChaosScopedSceneLock.h
@@ -0,0 +1,51 @@
+// Copyright Epic Games, Inc. All Rights Reserved.
+
+#pragma once
+
+#include "PhysicsInterfaceDeclaresCore.h"
+#include "Chaos/PhysicsObject.h"
+
+enum class EPhysicsInterfaceScopedLockType : uint8
+{
+	Read,
+	Write
+};
+
+class USkeletalMeshComponent;
+class FChaosScene;
+
+namespace Chaos
+{
+	class FPBDRigidsSolver;
+}
+
+struct ENGINE_API FScopedSceneLock_Chaos
+{
+	FScopedSceneLock_Chaos(FPhysicsActorHandle const* InActorHandle, EPhysicsInterfaceScopedLockType InLockType);
+	FScopedSceneLock_Chaos(FPhysicsActorHandle const* InActorHandleA, FPhysicsActorHandle const* InActorHandleB, EPhysicsInterfaceScopedLockType InLockType);
+	FScopedSceneLock_Chaos(FPhysicsConstraintHandle const* InConstraintHandle, EPhysicsInterfaceScopedLockType InLockType);
+	FScopedSceneLock_Chaos(USkeletalMeshComponent* InSkelMeshComp, EPhysicsInterfaceScopedLockType InLockType);
+	FScopedSceneLock_Chaos(Chaos::FPhysicsObjectHandle InObjectA, Chaos::FPhysicsObjectHandle InObjectB, EPhysicsInterfaceScopedLockType InLockType);
+	FScopedSceneLock_Chaos(FChaosScene* InScene, EPhysicsInterfaceScopedLockType InLockType);
+	~FScopedSceneLock_Chaos();
+
+	FScopedSceneLock_Chaos(FScopedSceneLock_Chaos& Other) = delete;
+	FScopedSceneLock_Chaos& operator=(FScopedSceneLock_Chaos& Other) = delete;
+
+	FScopedSceneLock_Chaos(FScopedSceneLock_Chaos&& Other);
+	FScopedSceneLock_Chaos& operator=(FScopedSceneLock_Chaos&& Other);
+
+	void Release();
+
+private:
+
+	void LockScene();
+	void UnlockScene();
+
+	FChaosScene* GetSceneForActor(FPhysicsActorHandle const* InActorHandle);
+	FChaosScene* GetSceneForActor(FPhysicsConstraintHandle const* InConstraintHandle);
+
+	bool bHasLock = false;
+	Chaos::FPBDRigidsSolver* Solver;
+	EPhysicsInterfaceScopedLockType LockType;
+};
\ No newline at end of file
diff --git a/Engine/Source/Runtime/Experimental/Chaos/Private/Chaos/PhysicsObject.cpp b/Engine/Source/Runtime/Experimental/Chaos/Private/Chaos/PhysicsObject.cpp
new file mode 100644
--- /dev/null
+++ b/Engine/Source/Runtime/Experimental/Chaos/Private/Chaos/PhysicsObject.cpp
@@ -0,0 +1,59 @@
+// Copyright Epic Games, Inc. All Rights Reserved.
+
+#include "Chaos/PhysicsObject.h"
+#include "Chaos/PhysicsObjectInternal.h"
+#include "Chaos/Framework/PhysicsProxyBase.h"
+
+namespace Chaos
+{
+	void FPhysicsObjectDeleter::operator()(FPhysicsObjectHandle p)
+	{
+		delete p;
+	}
+
+	FPhysicsObjectUniquePtr FPhysicsObjectFactory::CreatePhysicsObject(IPhysicsProxyBase* InProxy, int32 InBodyIndex, const FName& InBodyName)
+	{
+		return FPhysicsObjectUniquePtr{new FPhysicsObject(InProxy, InBodyIndex, InBodyName)};
+	}
+
+	bool FPhysicsObject::IsValid() const
+	{
+		return Proxy != nullptr && !Proxy->GetMarkedDeleted();
+	}
+
+	bool FPhysicsObject::HasChildren() const
+	{
+		EPhysicsProxyType ProxyType = Proxy->GetType();
+		switch (ProxyType)
+		{
+		case EPhysicsProxyType::GeometryCollectionType:
+		{
+			FGeometryDynamicCollection& Collection = static_cast<FGeometryCollectionPhysicsProxy*>(Proxy)->GetExternalCollection();
+			return !Collection.Children[BodyIndex].IsEmpty();
+		}
+		default:
+			break;
+		}
+		return false;
+	}
+
+	FPhysicsObjectHandle FPhysicsObject::GetParentObject() const
+	{
+		EPhysicsProxyType ProxyType = Proxy->GetType();
+		switch (ProxyType)
+		{
+		case EPhysicsProxyType::GeometryCollectionType:
+		{
+			FGeometryCollectionPhysicsProxy* GeometryCollectionProxy = static_cast<FGeometryCollectionPhysicsProxy*>(Proxy);
+			FGeometryDynamicCollection& Collection = GeometryCollectionProxy->GetExternalCollection();
+			if (int32 Index = Collection.Parent[BodyIndex]; Index != INDEX_NONE)
+			{
+				return GeometryCollectionProxy->GetPhysicsObjectByIndex(Index);
+			}
+		}
+		default:
+			break;
+		}
+		return nullptr;
+	}
+}
\ No newline at end of file
diff --git a/Engine/Source/Runtime/Experimental/Chaos/Private/Chaos/PhysicsObjectInterface.cpp b/Engine/Source/Runtime/Experimental/Chaos/Private/Chaos/PhysicsObjectInterface.cpp
new file mode 100644
--- /dev/null
+++ b/Engine/Source/Runtime/Experimental/Chaos/Private/Chaos/PhysicsObjectInterface.cpp
@@ -0,0 +1,533 @@
+// Copyright Epic Games, Inc. All Rights Reserved.
+
+#include "Chaos/PhysicsObjectInterface.h"
+#include "Chaos/PhysicsObjectInternal.h"
+#include "PhysicsProxy/SingleParticlePhysicsProxy.h"
+#include "Math/Transform.h"
+
+namespace
+{
+	template<Chaos::EThreadContext Id>
+	void SetParticleStateHelper(Chaos::FPhysicsObjectHandle PhysicsObject, Chaos::EObjectStateType State)
+	{
+		if (!PhysicsObject)
+		{
+			return;
+		}
+
+		IPhysicsProxyBase* Proxy = PhysicsObject->PhysicsProxy();
+		Chaos::TThreadParticle<Id>* Particle = PhysicsObject->GetParticle<Id>();
+		if (!Particle || !Proxy)
+		{
+			return;
+		}
+
+		if (Chaos::TThreadRigidParticle<Id>* Rigid = Particle->CastToRigidParticle())
+		{
+			if constexpr (Id == Chaos::EThreadContext::External)
+			{
+				if (Proxy->GetType() == EPhysicsProxyType::SingleParticleProxy)
+				{
+					// Easiest way to maintain the same behavior as what we currently have for the single particle case on the game thread.
+					static_cast<Chaos::FSingleParticlePhysicsProxy*>(Proxy)->GetGameThreadAPI().SetObjectState(State);
+				}
+				else
+				{
+					Rigid->SetObjectState(State, false, false);
+
+					// In the case of the geometry collection, it won't marshal the state from the game thread to the physics thread
+					// so we need to do it for it manually. 
+					if (Proxy->GetType() == EPhysicsProxyType::GeometryCollectionType)
+					{
+						if (Chaos::FPhysicsSolverBase* Solver = Proxy->GetSolverBase())
+						{
+							Solver->EnqueueCommandImmediate(
+								[PhysicsObject, State]() {
+									SetParticleStateHelper<Chaos::EThreadContext::Internal>(PhysicsObject, State);
+								}
+							);
+						}
+					}
+				}
+			}
+			else
+			{
+				if (Chaos::FPBDRigidsSolver* Solver = Proxy->GetSolver<Chaos::FPBDRigidsSolver>())
+				{
+					if (Chaos::FPBDRigidsSolver::FPBDRigidsEvolution* Evolution = Solver->GetEvolution())
+					{
+						Evolution->SetParticleObjectState(Rigid, State);
+					}
+				}
+			}
+		}
+	}
+}
+
+FName FClosestPhysicsObjectResult::HitName() const
+{
+	if (!PhysicsObject)
+	{
+		return NAME_None;
+	}
+	return Chaos::FPhysicsObjectInterface::GetName(PhysicsObject);
+}
+
+namespace Chaos
+{
+	template<EThreadContext Id>
+	FPhysicsObjectHandle FReadPhysicsObjectInterface<Id>::GetRootObject(FPhysicsObjectHandle Object)
+	{
+		if (!Object)
+		{
+			return nullptr;
+		}
+
+		return Object->GetRootObject<Id>();
+	}
+
+	template<EThreadContext Id>
+	FTransform FReadPhysicsObjectInterface<Id>::GetTransform(FPhysicsObjectHandle Object)
+	{
+		return FTransform{ GetR(Object), GetX(Object) };
+	}
+
+	template<EThreadContext Id>
+	FVector FReadPhysicsObjectInterface<Id>::GetX(FPhysicsObjectHandle Object)
+	{
+		if (!Object)
+		{
+			return FVector::Zero();
+		}
+
+		if (TThreadParticle<Id>* Particle = Object->GetParticle<Id>())
+		{
+			return Particle->X();
+		}
+
+		return FVector::Zero();
+	}
+
+	template<EThreadContext Id>
+	FVector FReadPhysicsObjectInterface<Id>::GetCoM(FPhysicsObjectHandle Object)
+	{
+		if (!Object)
+		{
+			return FVector::Zero();
+		}
+
+		if (TThreadParticle<Id>* Particle = Object->GetParticle<Id>())
+		{
+			if (Chaos::TThreadRigidParticle<Id>* Rigid = Particle->CastToRigidParticle())
+			{
+				return Rigid->CenterOfMass();
+			}
+		}
+
+		return FVector::Zero();
+	}
+
+	template<EThreadContext Id>
+	FVector FReadPhysicsObjectInterface<Id>::GetWorldCoM(FPhysicsObjectHandle Object)
+	{
+		return GetX(Object) + GetR(Object).RotateVector(GetCoM(Object));
+	}
+
+	template<EThreadContext Id>
+	FQuat FReadPhysicsObjectInterface<Id>::GetR(FPhysicsObjectHandle Object)
+	{
+		if (!Object)
+		{
+			return FQuat::Identity;
+		}
+
+		if (TThreadParticle<Id>* Particle = Object->GetParticle<Id>())
+		{
+			return Particle->R();
+		}
+
+		return FQuat::Identity;
+	}
+
+	template<EThreadContext Id>
+	bool FReadPhysicsObjectInterface<Id>::AreAllValid(TArrayView<FPhysicsObjectHandle> InObjects)
+	{
+		bool bCheck = !InObjects.IsEmpty();
+		for (FPhysicsObjectHandle Object : InObjects)
+		{
+			bCheck &= (Object != nullptr && Object->IsValid());
+		}
+		return bCheck;
+	}
+
+	template<EThreadContext Id>
+	bool FReadPhysicsObjectInterface<Id>::AreAllKinematic(TArrayView<FPhysicsObjectHandle> InObjects)
+	{
+		bool bCheck = !InObjects.IsEmpty();
+		for (FPhysicsObjectHandle Object : InObjects)
+		{
+			bCheck &= (Object && Object->IsValid() && Object->ObjectState<Id>() == EObjectStateType::Kinematic);
+		}
+		return bCheck;
+	}
+
+	template<EThreadContext Id>
+	bool FReadPhysicsObjectInterface<Id>::AreAllSleeping(TArrayView<FPhysicsObjectHandle> InObjects)
+	{
+		bool bCheck = !InObjects.IsEmpty();
+		for (FPhysicsObjectHandle Object : InObjects)
+		{
+			bCheck &= (Object && Object->IsValid() && Object->ObjectState<Id>() == EObjectStateType::Sleeping);
+		}
+		return bCheck;
+	}
+
+	template<EThreadContext Id>
+	bool FReadPhysicsObjectInterface<Id>::AreAllRigidBody(TArrayView<FPhysicsObjectHandle> InObjects)
+	{
+		bool bCheck = !InObjects.IsEmpty();
+		for (FPhysicsObjectHandle Object : InObjects)
+		{
+			bCheck &= (Object && Object->IsValid() && Object->ObjectState<Id>() != EObjectStateType::Static);
+		}
+		return bCheck;
+	}
+
+	template<EThreadContext Id>
+	bool FReadPhysicsObjectInterface<Id>::AreAllDynamic(TArrayView<FPhysicsObjectHandle> InObjects)
+	{
+		bool bCheck = !InObjects.IsEmpty();
+		for (FPhysicsObjectHandle Object : InObjects)
+		{
+			bCheck &= (Object && Object->IsValid() && Object->ObjectState<Id>() == EObjectStateType::Dynamic);
+		}
+		return bCheck;
+	}
+
+	template<EThreadContext Id>
+	bool FReadPhysicsObjectInterface<Id>::AreAllDisabled(TArrayView<FPhysicsObjectHandle> InObjects)
+	{
+		bool bDisabled = !InObjects.IsEmpty();
+		for (FPhysicsObjectHandle Object : InObjects)
+		{
+			bool bParticleDisabled = false;
+			if (Object)
+			{
+				if (TThreadParticle<Id>* Particle = Object->GetParticle<Id>())
+				{
+					bParticleDisabled = FPhysicsObject::IsParticleDisabled<Id>(Particle);
+				}
+			}
+			bDisabled &= bParticleDisabled;
+		}
+		return bDisabled;
+	}
+
+	template<EThreadContext Id>
+	float FReadPhysicsObjectInterface<Id>::GetMass(TArrayView<FPhysicsObjectHandle> InObjects)
+	{
+		float Mass = 0.f;
+		for (FPhysicsObjectHandle Object : InObjects)
+		{
+			if (Object)
+			{
+				if (TThreadParticle<Id>* Particle = Object->GetParticle<Id>())
+				{
+					if (TThreadRigidParticle<Id>* Rigid = Particle->CastToRigidParticle())
+					{
+						Mass += static_cast<float>(Rigid->M());
+					}
+				}
+			}
+		}
+		return Mass;
+	}
+
+	template<EThreadContext Id>
+	FBox FReadPhysicsObjectInterface<Id>::GetBounds(TArrayView<FPhysicsObjectHandle> InObjects)
+	{
+		FBox RetBox(ForceInit);
+		for (FPhysicsObjectHandle Object : InObjects)
+		{
+			if (!Object)
+			{
+				continue;
+			}
+
+			const TThreadParticle<Id>* Particle = Object->GetParticle<Id>();
+			if (!Particle)
+			{
+				continue;
+			}
+
+			FBox ParticleBox(ForceInit);
+			if (const FImplicitObject* Geometry = Particle->Geometry().Get(); Geometry && Geometry->HasBoundingBox())
+			{
+				Chaos::FAABB3 Box = Geometry->BoundingBox();
+				ParticleBox = FBox{ Box.Min(), Box.Max() };
+			}
+
+			if (ParticleBox.IsValid)
+			{
+				RetBox += ParticleBox;
+			}
+		}
+		return RetBox;
+	}
+
+	template<EThreadContext Id>
+	FClosestPhysicsObjectResult FReadPhysicsObjectInterface<Id>::GetClosestPhysicsBodyFromLocation(TArrayView<FPhysicsObjectHandle> InObjects, const FVector& WorldLocation)
+	{
+		FClosestPhysicsObjectResult AggregateResult;
+		for (FPhysicsObjectHandle Object : InObjects)
+		{
+			if (!Object)
+			{
+				continue;
+			}
+
+			const TThreadParticle<Id>* Particle = Object->GetParticle<Id>();
+			if (!Particle)
+			{
+				continue;
+			}
+
+			const FTransform WorldTransform = GetTransform(Object);
+			const FVector LocalLocation = WorldTransform.InverseTransformPosition(WorldLocation);
+
+			FClosestPhysicsObjectResult Result;
+
+			if (const FImplicitObject* Geometry = Particle->Geometry().Get())
+			{
+				Result.PhysicsObject = Object;
+
+				Chaos::FVec3 Normal;
+				Result.ClosestDistance = static_cast<double>(Geometry->PhiWithNormal(LocalLocation, Normal));
+				Result.ClosestLocation = WorldTransform.TransformPosition(LocalLocation - Result.ClosestDistance * Normal);
+			}
+
+			if (!Result)
+			{
+				continue;
+			}
+
+			if (!AggregateResult || Result.ClosestDistance < AggregateResult.ClosestDistance)
+			{
+				AggregateResult = Result;
+			}
+		}
+		return AggregateResult;
+	}
+
+	template<EThreadContext Id>
+	void FWritePhysicsObjectInterface<Id>::PutToSleep(TArrayView<FPhysicsObjectHandle> InObjects)
+	{
+		for (FPhysicsObjectHandle Object : InObjects)
+		{
+			if (!Object)
+			{
+				continue;
+			}
+
+			EObjectStateType State = Object->ObjectState<Id>();
+			if (State == EObjectStateType::Dynamic || State == EObjectStateType::Sleeping)
+			{
+				SetParticleStateHelper<Id>(Object, EObjectStateType::Sleeping);
+			}
+		}
+	}
+
+	template<EThreadContext Id>
+	void FWritePhysicsObjectInterface<Id>::WakeUp(TArrayView<FPhysicsObjectHandle> InObjects)
+	{
+		for (FPhysicsObjectHandle Object : InObjects)
+		{
+			if (!Object)
+			{
+				continue;
+			}
+
+			if (TThreadParticle<Id>* Particle = Object->GetParticle<Id>())
+			{
+				EObjectStateType State = Object->ObjectState<Id>();
+				if (State == EObjectStateType::Dynamic || State == EObjectStateType::Sleeping)
+				{
+					SetParticleStateHelper<Id>(Object, EObjectStateType::Dynamic);
+					if constexpr (Id == EThreadContext::External)
+					{
+						if (TThreadRigidParticle<Id>* Rigid = Particle->CastToRigidParticle())
+						{
+							Rigid->ClearEvents();
+						}
+					}
+				}
+			}
+		}
+	}
+
+	template<EThreadContext Id>
+	void FWritePhysicsObjectInterface<Id>::AddForce(TArrayView<FPhysicsObjectHandle> InObjects, const FVector& Force, bool bInvalidate)
+	{
+		for (FPhysicsObjectHandle Object : InObjects)
+		{
+			if (!Object)
+			{
+				continue;
+			}
+
+			if (TThreadParticle<Id>* Particle = Object->GetParticle<Id>())
+			{
+				if (Chaos::TThreadRigidParticle<Id>* Rigid = Particle->CastToRigidParticle())
+				{
+					if (Rigid->ObjectState() == EObjectStateType::Sleeping || Rigid->ObjectState() == EObjectStateType::Dynamic)
+					{
+						if (bInvalidate)
+						{
+							SetParticleStateHelper<Id>(Object, EObjectStateType::Dynamic);
+						}
+
+						Rigid->AddForce(Force, bInvalidate);
+					}
+				}
+			}
+		}
+	}
+
+	template<EThreadContext Id>
+	void FWritePhysicsObjectInterface<Id>::AddTorque(TArrayView<FPhysicsObjectHandle> InObjects, const FVector& Torque, bool bInvalidate)
+	{
+		for (FPhysicsObjectHandle Object : InObjects)
+		{
+			if (!Object)
+			{
+				continue;
+			}
+
+			if (TThreadParticle<Id>* Particle = Object->GetParticle<Id>())
+			{
+				if (Chaos::TThreadRigidParticle<Id>* Rigid = Particle->CastToRigidParticle())
+				{
+					if (Rigid->ObjectState() == EObjectStateType::Sleeping || Rigid->ObjectState() == EObjectStateType::Dynamic)
+					{
+						if (bInvalidate)
+						{
+							SetParticleStateHelper<Id>(Object, EObjectStateType::Dynamic);
+						}
+
+						Rigid->AddTorque(Torque, bInvalidate);
+					}
+				}
+			}
+		}
+	}
+
+	bool FPhysicsObjectInterface::HasChildren(FPhysicsObjectHandle Object)
+	{
+		return Object ? Object->HasChildren() : false;
+	}
+
+	void FPhysicsObjectInterface::SetName(FPhysicsObjectHandle Object, const FName& InName)
+	{
+		if (!Object)
+		{
+			return;
+		}
+
+		Object->SetName(InName);
+	}
+
+	FName FPhysicsObjectInterface::GetName(FPhysicsObjectHandle Object)
+	{
+		if (!Object)
+		{
+			return NAME_None;
+		}
+
+		return Object->GetBodyName();
+	}
+
+	void FPhysicsObjectInterface::SetId(FPhysicsObjectHandle Object, int32 InId)
+	{
+		if (!Object)
+		{
+			return;
+		}
+
+		Object->SetBodyIndex(InId);
+	}
+
+	int32 FPhysicsObjectInterface::GetId(FPhysicsObjectHandle Object)
+	{
+		if (!Object)
+		{
+			return INDEX_NONE;
+		}
+
+		return Object->GetBodyIndex();
+	}
+
+	FPBDRigidsSolver* FPhysicsObjectInterface::GetSolver(TArrayView<FPhysicsObjectHandle> InObjects)
+	{
+		FPBDRigidsSolver* RetSolver = nullptr;
+		for (FPhysicsObjectHandle Object : InObjects)
+		{
+			if (!Object)
+			{
+				continue;
+			}
+
+			FPBDRigidsSolver* Solver = nullptr;
+			if (IPhysicsProxyBase* Proxy = Object->PhysicsProxy())
+			{
+				Solver = Proxy->GetSolver<FPBDRigidsSolver>();
+			}
+
+			if (!Solver)
+			{
+				return nullptr;
+			}
+			else if (!RetSolver)
+			{
+				RetSolver = Solver;
+			}
+			else if (Solver != RetSolver)
+			{
+				return nullptr;
+			}
+		}
+		return RetSolver;
+	}
+
+	IPhysicsProxyBase* FPhysicsObjectInterface::GetProxy(TArrayView<FPhysicsObjectHandle> InObjects)
+	{
+		IPhysicsProxyBase* RetProxy = nullptr;
+		for (FPhysicsObjectHandle Object : InObjects)
+		{
+			if (!Object)
+			{
+				continue;
+			}
+
+			IPhysicsProxyBase* Proxy = Object->PhysicsProxy();
+			if (!Proxy)
+			{
+				return nullptr;
+			}
+			else if (!RetProxy)
+			{
+				RetProxy = Proxy;
+			}
+			else if (Proxy != RetProxy)
+			{
+				return nullptr;
+			}
+		}
+		return RetProxy;
+	}
+
+	template class FReadPhysicsObjectInterface<EThreadContext::External>;
+	template class FReadPhysicsObjectInterface<EThreadContext::Internal>;
+
+	template class FWritePhysicsObjectInterface<EThreadContext::External>;
+	template class FWritePhysicsObjectInterface<EThreadContext::Internal>;
+}
\ No newline at end of file
diff --git a/Engine/Source/Runtime/Experimental/Chaos/Private/Chaos/PhysicsObjectInternal.h b/Engine/Source/Runtime/Experimental/Chaos/Private/Chaos/PhysicsObjectInternal.h
new file mode 100644
--- /dev/null
+++ b/Engine/Source/Runtime/Experimental/Chaos/Private/Chaos/PhysicsObjectInternal.h
@@ -0,0 +1,150 @@
+// Copyright Epic Games, Inc. All Rights Reserved.
+#pragma once
+#include "Chaos/PhysicsObject.h"
+#include "Chaos/ParticleHandleFwd.h"
+#include "Containers/Set.h"
+#include "Framework/Threading.h"
+
+#include "PhysicsProxy/SingleParticlePhysicsProxy.h"
+#include "PhysicsProxy/GeometryCollectionPhysicsProxy.h"
+
+namespace Chaos
+{
+	struct FPhysicsObject
+	{
+	public:
+		bool IsValid() const;
+
+		void SetBodyIndex(int32 InBodyIndex) { BodyIndex = InBodyIndex; }
+		int32 GetBodyIndex() const { return BodyIndex; }
+
+		void SetName(const FName& InBodyName) { BodyName = InBodyName; }
+		const FName& GetBodyName() const { return BodyName; }
+
+		template<EThreadContext Id>
+		EObjectStateType ObjectState()
+		{
+			TThreadParticle<Id>* Particle = GetParticle<Id>();
+			if (!Particle)
+			{
+				return EObjectStateType::Uninitialized;
+			}
+
+			if (TThreadRigidParticle<Id>* Rigid = Particle->CastToRigidParticle())
+			{
+				return Rigid->ObjectState();
+			}
+			else
+			{
+				return Particle->ObjectState();
+			}
+		}
+
+		template<EThreadContext Id>
+		FPhysicsObjectHandle GetRootObject()
+		{
+			EPhysicsProxyType ProxyType = Proxy->GetType();
+			if (ProxyType == EPhysicsProxyType::SingleParticleProxy)
+			{
+				return this;
+			}
+
+			TThreadParticle<Id>* Particle = GetParticle<Id>();
+			FPhysicsObjectHandle CurrentObject = this;
+			FPhysicsObjectHandle CurrentParent = GetParentObject();
+			while (CurrentParent && IsParticleDisabled<Id>(Particle))
+			{
+				Particle = CurrentParent->GetParticle<Id>();
+				CurrentObject = CurrentParent;
+				CurrentParent = CurrentParent->GetParentObject();
+			}
+			return CurrentObject;
+		}
+
+		FPhysicsObjectHandle GetParentObject() const;
+
+		template <EThreadContext Id>
+		TThreadParticle<Id>* GetRootParticle()
+		{
+			FPhysicsObjectHandle Root = GetRootObject<Id>();
+			if (!Root)
+			{
+				return nullptr;
+			}
+			return Root->GetParticle<Id>();
+		}
+
+		template <EThreadContext Id>
+		TThreadParticle<Id>* GetParticle()
+		{
+			if (Proxy)
+			{
+				EPhysicsProxyType ProxyType = Proxy->GetType();
+				if constexpr (Id == EThreadContext::External)
+				{
+					switch (ProxyType)
+					{
+					case EPhysicsProxyType::GeometryCollectionType:
+						return static_cast<FGeometryCollectionPhysicsProxy*>(Proxy)->GetParticleByIndex_External(BodyIndex);
+					case EPhysicsProxyType::SingleParticleProxy:
+						return static_cast<FSingleParticlePhysicsProxy*>(Proxy)->GetParticle_LowLevel();
+					default:
+						break;
+					}
+				}
+				else
+				{
+					switch (ProxyType)
+					{
+					case EPhysicsProxyType::GeometryCollectionType:
+						return static_cast<FGeometryCollectionPhysicsProxy*>(Proxy)->GetParticleByIndex_Internal(BodyIndex);
+					case EPhysicsProxyType::SingleParticleProxy:
+						return static_cast<FSingleParticlePhysicsProxy*>(Proxy)->GetHandle_LowLevel();
+					default:
+						break;
+					}
+				}
+			}
+
+			return (TThreadParticle<Id>*)nullptr;
+		}
+
+		template <EThreadContext Id>
+		static bool IsParticleDisabled(TThreadParticle<Id>* Particle)
+		{
+			if (TThreadRigidParticle<Id>* Rigid = Particle->CastToRigidParticle())
+			{
+				return Rigid->Disabled();
+			}
+			return false;
+		}
+
+		IPhysicsProxyBase* PhysicsProxy() { return Proxy; }
+		const IPhysicsProxyBase* PhysicsProxy() const { return Proxy; }
+
+		bool HasChildren() const;
+
+		friend class FPhysicsObjectFactory;
+	protected:
+		FPhysicsObject(IPhysicsProxyBase* InProxy, int32 InBodyIndex = INDEX_NONE, const FName& InBodyName = NAME_None)
+			: Proxy(InProxy)
+			, BodyIndex(InBodyIndex)
+			, BodyName(InBodyName)
+		{}
+
+	private:
+		IPhysicsProxyBase* Proxy = nullptr;
+		int32 BodyIndex = INDEX_NONE;
+		FName BodyName = NAME_None;
+	};
+
+	/**
+	 * This class helps us limit the creation of the FPhysicsObject to only internal use cases. Realistically, only the physics proxies
+	 * should be creating physics objects.
+	 */
+	class FPhysicsObjectFactory
+	{
+	public:
+		static FPhysicsObjectUniquePtr CreatePhysicsObject(IPhysicsProxyBase* InProxy, int32 InBodyIndex = INDEX_NONE, const FName& InBodyName = NAME_None);
+	};
+}
\ No newline at end of file
diff --git a/Engine/Source/Runtime/Experimental/Chaos/Private/Chaos/PhysicsObjectInternalInterface.cpp b/Engine/Source/Runtime/Experimental/Chaos/Private/Chaos/PhysicsObjectInternalInterface.cpp
new file mode 100644
--- /dev/null
+++ b/Engine/Source/Runtime/Experimental/Chaos/Private/Chaos/PhysicsObjectInternalInterface.cpp
@@ -0,0 +1,18 @@
+// Copyright Epic Games, Inc. All Rights Reserved.
+
+#pragma once
+
+#include "Chaos/PhysicsObjectInternalInterface.h"
+
+namespace Chaos
+{
+	FReadPhysicsObjectInterface_Internal FPhysicsObjectInternalInterface::GetRead()
+	{
+		return CreateReadInterface<EThreadContext::Internal>();
+	}
+
+	FWritePhysicsObjectInterface_Internal FPhysicsObjectInternalInterface::GetWrite()
+	{
+		return CreateWriteInterface<EThreadContext::Internal>();
+	}
+}
\ No newline at end of file
diff --git a/Engine/Source/Runtime/Experimental/Chaos/Private/PhysicsProxy/GeometryCollectionPhysicsProxy.cpp b/Engine/Source/Runtime/Experimental/Chaos/Private/PhysicsProxy/GeometryCollectionPhysicsProxy.cpp
--- a/Engine/Source/Runtime/Experimental/Chaos/Private/PhysicsProxy/GeometryCollectionPhysicsProxy.cpp
+++ b/Engine/Source/Runtime/Experimental/Chaos/Private/PhysicsProxy/GeometryCollectionPhysicsProxy.cpp
@@ -31,6 +31,8 @@
 #include "Modules/ModuleManager.h"
 #include "Chaos/PullPhysicsDataImp.h"
 #include "Chaos/PBDRigidsEvolution.h"
+#include "Chaos/PhysicsObjectInternal.h"
+#include "Chaos/PhysicsObjectInterface.h"
 #include "GeometryCollection/Facades/CollectionAnchoringFacade.h"
 
 #ifndef TODO_REIMPLEMENT_INIT_COMMANDS
@@ -407,7 +409,7 @@ void FGeometryCollectionPhysicsProxy::Initialize(Chaos::FPBDRigidsEvolutionBase
 	// Attach the external particles to the gamethread collection
 	if (DynamicCollection.HasAttribute(FGeometryCollection::ParticlesAttribute, FTransformCollection::TransformGroup))
 		DynamicCollection.RemoveAttribute(FGeometryCollection::ParticlesAttribute, FTransformCollection::TransformGroup);
-	DynamicCollection.AddExternalAttribute<TUniquePtr<Chaos::FGeometryParticle>>(FGeometryCollection::ParticlesAttribute, FTransformCollection::TransformGroup, GTParticles);
+	DynamicCollection.AddExternalAttribute<TUniquePtr<FParticle>>(FGeometryCollection::ParticlesAttribute, FTransformCollection::TransformGroup, GTParticles);
 
 
 	NumParticles = DynamicCollection.NumElements(FGeometryCollection::TransformGroup);
@@ -415,6 +417,13 @@ void FGeometryCollectionPhysicsProxy::Initialize(Chaos::FPBDRigidsEvolutionBase
 	SolverClusterID.Init(nullptr, NumParticles);
 	SolverClusterHandles.Init(nullptr, NumParticles);
 	SolverParticleHandles.Init(nullptr, NumParticles);
+	
+	PhysicsObjects.Empty();
+	PhysicsObjects.Reserve(NumParticles);
+	for (int32 Index = 0; Index < NumParticles; ++Index)
+	{
+		PhysicsObjects.Emplace(Chaos::FPhysicsObjectFactory::CreatePhysicsObject(this, Index, FName(Parameters.RestCollection->BoneName[Index])));
+	}
 
 	// compatibility requirement to make sure we at least initialize GameThreadPerFrameData properly
 	GameThreadPerFrameData.SetWorldTransform(Parameters.WorldTransform);
@@ -447,18 +456,20 @@ void FGeometryCollectionPhysicsProxy::Initialize(Chaos::FPBDRigidsEvolutionBase
 		GameThreadCollection.RemoveAttribute(FGeometryCollection::ParticlesAttribute, FTransformCollection::TransformGroup);
 	}
 		
-	GameThreadCollection.AddExternalAttribute<TUniquePtr<Chaos::FGeometryParticle>>(FGeometryCollection::ParticlesAttribute, FTransformCollection::TransformGroup, GTParticles);
+	GameThreadCollection.AddExternalAttribute<TUniquePtr<FParticle>>(FGeometryCollection::ParticlesAttribute, FTransformCollection::TransformGroup, GTParticles);
 	
 
 	const FVector Scale = Parameters.WorldTransform.GetScale3D();
+	const FVector::FReal MassScale = Scale.X * Scale.Y * Scale.Z;
+	const TManagedArray<float>& Mass = Parameters.RestCollection->GetAttribute<float>("Mass", FTransformCollection::TransformGroup);
 
 	TArray<int32> ChildrenToCheckForParentFix;
 	if(ensure(NumTransforms == GameThreadCollection.Implicits.Num() && NumTransforms == GTParticles.Num())) // Implicits are in the transform group so this invariant should always hold
 	{
 		for(int32 Index = 0; Index < NumTransforms; ++Index)
 		{
-			GTParticles[Index] = Chaos::FGeometryParticle::CreateParticle();
-			Chaos::FGeometryParticle* P = GTParticles[Index].Get();
+			GTParticles[Index] = FParticle::CreateParticle();
+			FParticle* P = GTParticles[Index].Get();
 			GTParticlesToTransformGroupIndex.Add(P, Index);
 
 			GTParticles[Index]->SetUniqueIdx(Evolution->GenerateUniqueIdx());
@@ -466,6 +477,7 @@ void FGeometryCollectionPhysicsProxy::Initialize(Chaos::FPBDRigidsEvolutionBase
 			const FTransform& T = Parameters.WorldTransform * GameThreadCollection.Transform[Index];
 			P->SetX(T.GetTranslation(), false);
 			P->SetR(T.GetRotation(), false);
+			P->SetM(Mass[Index] * MassScale);
 			P->SetUserData(Parameters.UserData);
 			P->SetProxy(this);
 
@@ -542,7 +554,7 @@ void FGeometryCollectionPhysicsProxy::Initialize(Chaos::FPBDRigidsEvolutionBase
 							{
 								using FImplicitObjectTransformed = Chaos::TImplicitObjectTransformed<Chaos::FReal, 3>;
 
-								Chaos::FGeometryParticle* ChildParticle = GTParticles[ChildIndex].Get();
+								FParticle* ChildParticle = GTParticles[ChildIndex].Get();
 								const FGeometryDynamicCollection::FSharedImplicit& ChildImplicit = GameThreadCollection.Implicits[ChildIndex];
 								if (ChildImplicit)
 								{
@@ -600,7 +612,7 @@ void FGeometryCollectionPhysicsProxy::Initialize(Chaos::FPBDRigidsEvolutionBase
 		// Phase 3 : finalization of shapes
 		for(int32 Index = 0; Index < NumTransforms; ++Index)
 		{
-			Chaos::FGeometryParticle* P = GTParticles[Index].Get();
+			FParticle* P = GTParticles[Index].Get();
 			const Chaos::FShapesArray& Shapes = P->ShapesArray();
 			const int32 NumShapes = Shapes.Num();
 			for(int32 ShapeIndex = 0; ShapeIndex < NumShapes; ++ShapeIndex)
@@ -660,6 +672,11 @@ void FGeometryCollectionPhysicsProxy::Initialize(Chaos::FPBDRigidsEvolutionBase
 		{
 			bool bIsRoot = (GameThreadCollection.Parent[TransformIndex] == INDEX_NONE);
 			GameThreadCollection.Active[TransformIndex] = bIsRoot;
+
+			if (FParticle* P = GTParticles[TransformIndex].Get())
+			{
+				P->SetDisabled(!bIsRoot);
+			}
 		}
 	}
 }
@@ -1080,7 +1097,7 @@ void FGeometryCollectionPhysicsProxy::InitializeBodiesPT(Chaos::FPBDRigidsSolver
 					CreationParameters.Scale = Parameters.WorldTransform.GetScale3D();
 
 					// Hook the handle up with the GT particle
-					Chaos::FGeometryParticle* GTParticle = GTParticles[TransformGroupIndex].Get();
+					FParticle* GTParticle = GTParticles[TransformGroupIndex].Get();
 
 					Chaos::FUniqueIdx ExistingIndex = GTParticle->UniqueIdx();
 					Chaos::FPBDRigidClusteredParticleHandle* Handle = BuildClusters_Internal(TransformGroupIndex, RigidChildren, RigidChildrenTransformGroupIndex, CreationParameters, &ExistingIndex);
@@ -2334,23 +2351,28 @@ void FGeometryCollectionPhysicsProxy::SetClusteredParticleKinematicTarget_Intern
 	}
 }
 
+static EObjectStateTypeEnum GetObjectStateEnumFromObjectState(Chaos::EObjectStateType State)
+{
+	switch (State)
+	{
+	case Chaos::EObjectStateType::Kinematic:
+		return EObjectStateTypeEnum::Chaos_Object_Kinematic;
+	case Chaos::EObjectStateType::Static:
+		return EObjectStateTypeEnum::Chaos_Object_Static;
+	case Chaos::EObjectStateType::Sleeping:
+		return EObjectStateTypeEnum::Chaos_Object_Sleeping;
+	case Chaos::EObjectStateType::Dynamic:
+	case Chaos::EObjectStateType::Uninitialized:
+	default:
+		return EObjectStateTypeEnum::Chaos_Object_Dynamic;
+	}
+}
+
 static EObjectStateTypeEnum GetObjectStateFromHandle(const Chaos::TPBDRigidClusteredParticleHandle<Chaos::FReal, 3>* Handle)
 {
 	if (!Handle->Sleeping())
 	{
-		switch (Handle->ObjectState())
-		{
-		case Chaos::EObjectStateType::Kinematic:
-			return EObjectStateTypeEnum::Chaos_Object_Kinematic;
-		case Chaos::EObjectStateType::Static:
-			return EObjectStateTypeEnum::Chaos_Object_Static;
-		case Chaos::EObjectStateType::Sleeping:
-			return EObjectStateTypeEnum::Chaos_Object_Sleeping;
-		case Chaos::EObjectStateType::Dynamic:
-		case Chaos::EObjectStateType::Uninitialized:
-		default:
-			return EObjectStateTypeEnum::Chaos_Object_Dynamic;
-		}
+		return GetObjectStateEnumFromObjectState(Handle->ObjectState());
 	}
 	return EObjectStateTypeEnum::Chaos_Object_Sleeping;
 }
@@ -2398,9 +2420,10 @@ void FGeometryCollectionPhysicsProxy::BufferPhysicsResults_External(Chaos::FDirt
 	
 	for (int32 TransformGroupIndex = 0; TransformGroupIndex < TargetResults.Transforms.Num(); ++TransformGroupIndex)
 	{
-		const Chaos::FGeometryParticle& GTParticle = *GTParticles[TransformGroupIndex];
+		FParticle& GTParticle = *GTParticles[TransformGroupIndex];
 		TargetResults.ParticleXs[TransformGroupIndex] = GTParticle.X();
 		TargetResults.ParticleRs[TransformGroupIndex] = GTParticle.R();
+		TargetResults.States[TransformGroupIndex].DynamicState = static_cast<uint8>(GetObjectStateEnumFromObjectState(GTParticle.ObjectState()));
 
 		if (LinearVelocities && AngularVelocities)
 		{
@@ -2703,10 +2726,13 @@ bool FGeometryCollectionPhysicsProxy::PullFromPhysicsState(const Chaos::FDirtyGe
 		// first : copy the non interpolate-able values
 		for (int32 TransformGroupIndex = 0; TransformGroupIndex < NumTransforms; ++TransformGroupIndex)
 		{
+			FParticle& GTParticle = *GTParticles[TransformGroupIndex];
+
 			const bool bIsActive = !TargetResults.States[TransformGroupIndex].DisabledState;
 			if (GameThreadCollection.Active[TransformGroupIndex] != bIsActive)
 			{
 				GameThreadCollection.Active[TransformGroupIndex] = bIsActive;
+				GTParticle.SetDisabled(!bIsActive);
 				bIsCollectionDirty = true;
 			}
 			
@@ -2714,6 +2740,7 @@ bool FGeometryCollectionPhysicsProxy::PullFromPhysicsState(const Chaos::FDirtyGe
 			if (GameThreadCollection.DynamicState[TransformGroupIndex] != NewState)
 			{
 				GameThreadCollection.DynamicState[TransformGroupIndex] = NewState;
+				GTParticle.SetObjectState(static_cast<Chaos::EObjectStateType>(NewState));
 				bIsCollectionDirty = true;
 			}
 
@@ -2768,7 +2795,7 @@ bool FGeometryCollectionPhysicsProxy::PullFromPhysicsState(const Chaos::FDirtyGe
 			{
 				if (!TargetResults.States[TransformGroupIndex].DisabledState)
 				{
-					Chaos::FGeometryParticle& GTParticle = *GTParticles[TransformGroupIndex];
+					FParticle& GTParticle = *GTParticles[TransformGroupIndex];
 
 					const Chaos::FVec3 OldX = GTParticle.X();
 					const Chaos::FVec3 NewX = FMath::Lerp(Prev.ParticleXs[TransformGroupIndex], Next.ParticleXs[TransformGroupIndex], *Alpha);
@@ -2814,7 +2841,7 @@ bool FGeometryCollectionPhysicsProxy::PullFromPhysicsState(const Chaos::FDirtyGe
 			{
 				if (!TargetResults.States[TransformGroupIndex].DisabledState)
 				{
-					Chaos::FGeometryParticle& GTParticle = *GTParticles[TransformGroupIndex];
+					FParticle& GTParticle = *GTParticles[TransformGroupIndex];
 
 					const Chaos::FVec3 OldX = GTParticle.X();
 					const Chaos::FVec3& NewX = TargetResults.ParticleXs[TransformGroupIndex];
@@ -2891,7 +2918,7 @@ void FGeometryCollectionPhysicsProxy::UpdateFilterData_External(const FCollision
 	int32 NumTransforms = GameThreadCollection.Transform.Num();
 	for (int32 Index = 0; Index < NumTransforms; ++Index)
 	{
-		Chaos::FGeometryParticle* P = GTParticles[Index].Get();
+		FParticle* P = GTParticles[Index].Get();
 		const Chaos::FShapesArray& Shapes = P->ShapesArray();
 		const int32 NumShapes = Shapes.Num();
 		for (int32 ShapeIndex = 0; ShapeIndex < NumShapes; ++ShapeIndex)
@@ -4088,6 +4115,43 @@ void FGeometryCollectionPhysicsProxy::FieldForcesUpdateCallback(Chaos::FPBDRigid
 	}
 }
 
+Chaos::FPhysicsObjectHandle FGeometryCollectionPhysicsProxy::GetPhysicsObjectByIndex(int32 Index) const
+{
+	if (!PhysicsObjects.IsValidIndex(Index))
+	{
+		return nullptr;
+	}
+
+	return PhysicsObjects[Index].Get();
+}
+
+FGeometryCollectionPhysicsProxy::FParticle* FGeometryCollectionPhysicsProxy::GetParticleByIndex_External(int32 Index)
+{
+	return const_cast<FGeometryCollectionPhysicsProxy::FParticle*>(const_cast<const FGeometryCollectionPhysicsProxy*>(this)->GetParticleByIndex_External(Index));
+}
+
+const FGeometryCollectionPhysicsProxy::FParticle* FGeometryCollectionPhysicsProxy::GetParticleByIndex_External(int32 Index) const
+{
+	if (Index < 0 || Index >= GTParticles.Num())
+	{
+		return nullptr;
+	}
+	return GTParticles[Index].Get();
+}
+
+FGeometryCollectionPhysicsProxy::FParticleHandle* FGeometryCollectionPhysicsProxy::GetParticleByIndex_Internal(int32 Index)
+{
+	return const_cast<FGeometryCollectionPhysicsProxy::FParticleHandle*>(const_cast<const FGeometryCollectionPhysicsProxy*>(this)->GetParticleByIndex_Internal(Index));
+}
+
+const FGeometryCollectionPhysicsProxy::FParticleHandle* FGeometryCollectionPhysicsProxy::GetParticleByIndex_Internal(int32 Index) const
+{
+	if (!SolverParticleHandles.IsValidIndex(Index))
+	{
+		return nullptr;
+	}
+	return SolverParticleHandles[Index];
+}
 
 void FDamageCollector::Reset(int32 NumTransforms)
 {
diff --git a/Engine/Source/Runtime/Experimental/Chaos/Private/PhysicsProxy/SingleParticlePhysicsProxy.cpp b/Engine/Source/Runtime/Experimental/Chaos/Private/PhysicsProxy/SingleParticlePhysicsProxy.cpp
--- a/Engine/Source/Runtime/Experimental/Chaos/Private/PhysicsProxy/SingleParticlePhysicsProxy.cpp
+++ b/Engine/Source/Runtime/Experimental/Chaos/Private/PhysicsProxy/SingleParticlePhysicsProxy.cpp
@@ -10,6 +10,7 @@
 #include "Chaos/Serializable.h"
 #include "Chaos/PerParticleGravity.h"
 #include "Chaos/ParticleHandle.h"
+#include "Chaos/PhysicsObjectInternal.h"
 #include "Chaos/Framework/MultiBufferResource.h"
 #include "PhysicsSolver.h"
 #include "Chaos/ChaosMarshallingManager.h"
@@ -30,6 +31,7 @@ FSingleParticlePhysicsProxy::FSingleParticlePhysicsProxy(TUniquePtr<PARTICLE_TYP
 	, Handle(InHandle)
 {
 	Particle->SetProxy(this);
+	Reference = FPhysicsObjectFactory::CreatePhysicsObject(this);
 }
 
 
diff --git a/Engine/Source/Runtime/Experimental/Chaos/Public/Chaos/Framework/PhysicsProxyBase.h b/Engine/Source/Runtime/Experimental/Chaos/Public/Chaos/Framework/PhysicsProxyBase.h
--- a/Engine/Source/Runtime/Experimental/Chaos/Public/Chaos/Framework/PhysicsProxyBase.h
+++ b/Engine/Source/Runtime/Experimental/Chaos/Public/Chaos/Framework/PhysicsProxyBase.h
@@ -90,6 +90,7 @@ public:
 	void ResetDirtyIdx() { DirtyIdx = INDEX_NONE; }
 
 	void MarkDeleted() { SyncTimestamp->bDeleted = true; }
+	bool GetMarkedDeleted() const { return SyncTimestamp->bDeleted; }
 
 	TSharedPtr<FProxyTimestampBase,ESPMode::ThreadSafe> GetSyncTimestamp() const { return SyncTimestamp; }
 
diff --git a/Engine/Source/Runtime/Experimental/Chaos/Public/Chaos/ParticleHandleFwd.h b/Engine/Source/Runtime/Experimental/Chaos/Public/Chaos/ParticleHandleFwd.h
--- a/Engine/Source/Runtime/Experimental/Chaos/Public/Chaos/ParticleHandleFwd.h
+++ b/Engine/Source/Runtime/Experimental/Chaos/Public/Chaos/ParticleHandleFwd.h
@@ -2,6 +2,8 @@
 #pragma once
 
 #include <Chaos/Real.h>
+#include <Framework/Threading.h>
+#include <type_traits>
 
 // Use to define out code blocks that need to be adapted to use Particle Handles in a searchable way (better than #if 0)
 #define CHAOS_PARTICLEHANDLE_TODO 0
@@ -107,4 +109,10 @@ namespace Chaos
 	using FGeometryParticleHandle = TGeometryParticleHandle<FReal, 3>;
 	using FKinematicGeometryParticle = TKinematicGeometryParticle<FReal, 3>;
 	using FPBDRigidParticle = TPBDRigidParticle<FReal, 3>;
+
+	template<EThreadContext Id>
+	using TThreadParticle = std::conditional_t<Id == EThreadContext::External, FGeometryParticle, FGeometryParticleHandle>;
+
+	template<EThreadContext Id>
+	using TThreadRigidParticle = std::conditional_t<Id == EThreadContext::External, FPBDRigidParticle, FPBDRigidParticleHandle>;
 }
diff --git a/Engine/Source/Runtime/Experimental/Chaos/Public/Chaos/PhysicsObject.h b/Engine/Source/Runtime/Experimental/Chaos/Public/Chaos/PhysicsObject.h
new file mode 100644
--- /dev/null
+++ b/Engine/Source/Runtime/Experimental/Chaos/Public/Chaos/PhysicsObject.h
@@ -0,0 +1,26 @@
+// Copyright Epic Games, Inc. All Rights Reserved.
+
+#pragma once
+
+#include "Templates/UniquePtr.h"
+#include "UObject/NameTypes.h"
+
+class IPhysicsProxyBase;
+class FGeometryCollectionPhysicsProxy;
+
+namespace Chaos
+{
+	/**
+	 * The FPhysicsObject is effectively a reference to a single particle in the solver.
+	 * It maintains this reference indirectly via the physics proxy. This object is meant to be usable on both the game thread and physics thread.
+	 */
+	struct FPhysicsObject;
+	using FPhysicsObjectHandle = FPhysicsObject*;
+
+	struct CHAOS_API FPhysicsObjectDeleter
+	{
+		void operator()(FPhysicsObjectHandle p);
+	};
+	using FPhysicsObjectUniquePtr = TUniquePtr<FPhysicsObject, FPhysicsObjectDeleter>;
+
+} // namespace Chaos
\ No newline at end of file
diff --git a/Engine/Source/Runtime/Experimental/Chaos/Public/Chaos/PhysicsObjectInterface.h b/Engine/Source/Runtime/Experimental/Chaos/Public/Chaos/PhysicsObjectInterface.h
new file mode 100644
--- /dev/null
+++ b/Engine/Source/Runtime/Experimental/Chaos/Public/Chaos/PhysicsObjectInterface.h
@@ -0,0 +1,121 @@
+// Copyright Epic Games, Inc. All Rights Reserved.
+
+#pragma once
+
+#include "Chaos/PhysicsObject.h"
+#include "Framework/Threading.h"
+#include "Containers/ArrayView.h"
+#include "Math/MathFwd.h"
+#include "UObject/ObjectMacros.h"
+
+#include "PhysicsObjectInterface.generated.h"
+
+class FChaosScene;
+class IPhysicsProxyBase;
+
+USTRUCT(BlueprintType)
+struct CHAOS_API FClosestPhysicsObjectResult
+{
+	GENERATED_BODY()
+
+	Chaos::FPhysicsObjectHandle PhysicsObject = nullptr;
+	FVector ClosestLocation;
+	double ClosestDistance;
+	operator bool() const
+	{
+		return IsValid();
+	}
+	bool IsValid() const
+	{
+		return PhysicsObject != nullptr;
+	}
+
+	FName HitName() const;
+};
+
+namespace Chaos
+{
+	class FPBDRigidsSolver;
+
+	/**
+	 * FReadPhysicsObjectInterface will assume that these operations are safe to call (i.e. the relevant scenes have been read locked on the game thread).
+	 */
+	template<EThreadContext Id>
+	class CHAOS_API FReadPhysicsObjectInterface
+	{
+	public:
+		FPhysicsObjectHandle GetRootObject(FPhysicsObjectHandle Object);
+
+		FTransform GetTransform(FPhysicsObjectHandle Object);
+		FVector GetX(FPhysicsObjectHandle Object);
+		FVector GetCoM(FPhysicsObjectHandle Object);
+		FVector GetWorldCoM(FPhysicsObjectHandle Object);
+		FQuat GetR(FPhysicsObjectHandle Object);
+
+		bool AreAllValid(TArrayView<FPhysicsObjectHandle> InObjects);
+		bool AreAllKinematic(TArrayView<FPhysicsObjectHandle> InObjects);
+		bool AreAllSleeping(TArrayView<FPhysicsObjectHandle> InObjects);
+		bool AreAllRigidBody(TArrayView<FPhysicsObjectHandle> InObjects);
+		bool AreAllDynamic(TArrayView<FPhysicsObjectHandle> InObjects);
+		bool AreAllDisabled(TArrayView<FPhysicsObjectHandle> InObjects);
+		float GetMass(TArrayView<FPhysicsObjectHandle> InObjects);
+		FBox GetBounds(TArrayView<FPhysicsObjectHandle> InObjects);
+		FClosestPhysicsObjectResult GetClosestPhysicsBodyFromLocation(TArrayView<FPhysicsObjectHandle> InObjects, const FVector& WorldLocation);
+
+		friend class FPhysicsObjectInterface;
+	protected:
+		FReadPhysicsObjectInterface() = default;
+	};
+
+	using FReadPhysicsObjectInterface_External = FReadPhysicsObjectInterface<EThreadContext::External>;
+	using FReadPhysicsObjectInterface_Internal = FReadPhysicsObjectInterface<EThreadContext::Internal>;
+
+	/**
+	 * FReadPhysicsObjectInterface will assume that these operations are safe to call (i.e. the relevant scenes have been read locked on the physics thread).
+	 */
+	template<EThreadContext Id>
+	class CHAOS_API FWritePhysicsObjectInterface: public FReadPhysicsObjectInterface<Id>
+	{
+	public:
+		void PutToSleep(TArrayView<FPhysicsObjectHandle> InObjects);
+		void WakeUp(TArrayView<FPhysicsObjectHandle> InObjects);
+		void AddForce(TArrayView<FPhysicsObjectHandle> InObjects, const FVector& Force, bool bInvalidate);
+		void AddTorque(TArrayView<FPhysicsObjectHandle> InObjects, const FVector& Torque, bool bInvalidate);
+
+		friend class FPhysicsObjectInterface;
+	protected:
+		FWritePhysicsObjectInterface() = default;
+	};
+
+	using FWritePhysicsObjectInterface_External = FWritePhysicsObjectInterface<EThreadContext::External>;
+	using FWritePhysicsObjectInterface_Internal = FWritePhysicsObjectInterface<EThreadContext::Internal>;
+
+	/**
+	 * The FPhysicsObjectInterface is primarily used to perform maintenance operations on the FPhysicsObject.
+	 * Any operations on the underlying particle/particle handle should use the FReadPhysicsObjectInterface and
+	 * FWritePhysicsObjectInterface.
+	 */
+	class CHAOS_API FPhysicsObjectInterface
+	{
+	public:
+		static bool HasChildren(FPhysicsObjectHandle Object);
+
+		static void SetName(FPhysicsObjectHandle Object, const FName& InName);
+		static FName GetName(FPhysicsObjectHandle Object);
+
+		static void SetId(FPhysicsObjectHandle Object, int32 InId);
+		static int32 GetId(FPhysicsObjectHandle Object);
+
+		static FPBDRigidsSolver* GetSolver(TArrayView<FPhysicsObjectHandle> InObjects);
+		static IPhysicsProxyBase* GetProxy(TArrayView<FPhysicsObjectHandle> InObjects);
+
+	protected:
+		// This function should not be called without an appropriate read-lock on the relevant scene.
+		template<EThreadContext Id>
+		static FReadPhysicsObjectInterface<Id> CreateReadInterface() { return FReadPhysicsObjectInterface<Id>{}; }
+
+		// This function should not be called without an appropriate write-lock on the relevant scene.
+		template<EThreadContext Id>
+		static FWritePhysicsObjectInterface<Id> CreateWriteInterface() { return FWritePhysicsObjectInterface<Id>{}; }
+	};
+}
\ No newline at end of file
diff --git a/Engine/Source/Runtime/Experimental/Chaos/Public/Chaos/PhysicsObjectInternalInterface.h b/Engine/Source/Runtime/Experimental/Chaos/Public/Chaos/PhysicsObjectInternalInterface.h
new file mode 100644
--- /dev/null
+++ b/Engine/Source/Runtime/Experimental/Chaos/Public/Chaos/PhysicsObjectInternalInterface.h
@@ -0,0 +1,15 @@
+// Copyright Epic Games, Inc. All Rights Reserved.
+
+#pragma once
+
+#include "Chaos/PhysicsObjectInterface.h"
+
+namespace Chaos
+{
+	class CHAOS_API FPhysicsObjectInternalInterface : public Chaos::FPhysicsObjectInterface
+	{
+	public:
+		static FReadPhysicsObjectInterface_Internal GetRead();
+		static FWritePhysicsObjectInterface_Internal GetWrite();
+	};
+}
\ No newline at end of file
diff --git a/Engine/Source/Runtime/Experimental/Chaos/Public/Framework/Threading.h b/Engine/Source/Runtime/Experimental/Chaos/Public/Framework/Threading.h
--- a/Engine/Source/Runtime/Experimental/Chaos/Public/Framework/Threading.h
+++ b/Engine/Source/Runtime/Experimental/Chaos/Public/Framework/Threading.h
@@ -136,6 +136,12 @@ namespace Chaos
 /** Signals that we have left a write lock to control the checks above */
 #define CHAOS_RECORD_LEAVE_WRITE_LOCK Chaos::ThreadingPrivate::DecWriteDepth(this);
 
+enum class EThreadContext
+{
+	External,
+	Internal
+};
+
 #if PHYSICS_THREAD_CONTEXT
 /** Debug helper to ensure threading mistakes are caught. Do not use for ship */
 class CHAOS_API FPhysicsThreadContext : public TThreadSingleton<FPhysicsThreadContext>
diff --git a/Engine/Source/Runtime/Experimental/Chaos/Public/GeometryCollection/ManagedArray.h b/Engine/Source/Runtime/Experimental/Chaos/Public/GeometryCollection/ManagedArray.h
--- a/Engine/Source/Runtime/Experimental/Chaos/Public/GeometryCollection/ManagedArray.h
+++ b/Engine/Source/Runtime/Experimental/Chaos/Public/GeometryCollection/ManagedArray.h
@@ -661,6 +661,9 @@ template<> inline void CopyRangeHelper(TArray<TYPE>& Target, const TManagedArray
 typedef TUniquePtr<Chaos::TGeometryParticle<Chaos::FReal, 3>> LOCAL_MA_UniqueTGeometryParticle;
 UNSUPPORTED_UNIQUE_ARRAY_COPIES(LOCAL_MA_UniqueTGeometryParticle, TEXT("Chaos::TGeometryParticle"));
 
+typedef TUniquePtr<Chaos::TPBDRigidParticle<Chaos::FReal, 3>> LOCAL_MA_UniqueTPBDRigidParticle;
+UNSUPPORTED_UNIQUE_ARRAY_COPIES(LOCAL_MA_UniqueTPBDRigidParticle, TEXT("Chaos::TPBDRigidParticle"));
+
 typedef TUniquePtr<Chaos::FBVHParticles, TDefaultDelete<Chaos::FBVHParticles>> LOCAL_MA_UniqueTBVHParticles;
 UNSUPPORTED_UNIQUE_ARRAY_COPIES(LOCAL_MA_UniqueTBVHParticles, TEXT("Chaos::FBVHParticles"));
 
diff --git a/Engine/Source/Runtime/Experimental/Chaos/Public/GeometryCollection/ManagedArrayTypeValues.inl b/Engine/Source/Runtime/Experimental/Chaos/Public/GeometryCollection/ManagedArrayTypeValues.inl
--- a/Engine/Source/Runtime/Experimental/Chaos/Public/GeometryCollection/ManagedArrayTypeValues.inl
+++ b/Engine/Source/Runtime/Experimental/Chaos/Public/GeometryCollection/ManagedArrayTypeValues.inl
@@ -50,6 +50,7 @@ MANAGED_ARRAY_TYPE(TArray<int32>, Int32Array)
 MANAGED_ARRAY_TYPE(TArray<float>, FloatArray)
 MANAGED_ARRAY_TYPE(FVector4f, Vector4f)
 MANAGED_ARRAY_TYPE(TArray<FVector3f>, FVectorArray)
+MANAGED_ARRAY_TYPE(TUniquePtr<Chaos::FPBDRigidParticle>, TPBDRigidParticle3fUniquePtr)
 
 // NOTE: new types must be added at the bottom to keep serialization from breaking
 
diff --git a/Engine/Source/Runtime/Experimental/Chaos/Public/PhysicsProxy/GeometryCollectionPhysicsProxy.h b/Engine/Source/Runtime/Experimental/Chaos/Public/PhysicsProxy/GeometryCollectionPhysicsProxy.h
--- a/Engine/Source/Runtime/Experimental/Chaos/Public/PhysicsProxy/GeometryCollectionPhysicsProxy.h
+++ b/Engine/Source/Runtime/Experimental/Chaos/Public/PhysicsProxy/GeometryCollectionPhysicsProxy.h
@@ -13,6 +13,7 @@
 #include "Chaos/GeometryParticlesfwd.h"
 #include "Chaos/ParticleHandle.h"
 #include "Chaos/ParticleHandleFwd.h"
+#include "Chaos/PhysicsObject.h"
 #include "Containers/Array.h"
 #include "PBDRigidsSolver.h"
 #include "Chaos/Defines.h"
@@ -141,6 +142,7 @@ class CHAOS_API FGeometryCollectionPhysicsProxy : public TPhysicsProxy<FGeometry
 public:
 	typedef TPhysicsProxy<FGeometryCollectionPhysicsProxy, FStubGeometryCollectionData, FGeometryCollectionProxyTimestamp> Base;
 	typedef FCollisionStructureManager::FSimplicial FSimplicial;
+	typedef Chaos::TPBDRigidParticle<Chaos::FReal, 3> FParticle;
 	typedef Chaos::TPBDRigidParticleHandle<Chaos::FReal, 3> FParticleHandle;
 	typedef Chaos::TPBDRigidClusteredParticleHandle<Chaos::FReal, 3> FClusterHandle;
 
@@ -285,11 +287,22 @@ public:
 		return PhysicsThreadCollection;
 	}
 
-	TManagedArray<TUniquePtr<Chaos::FGeometryParticle>>& GetExternalParticles()
+	FGeometryDynamicCollection& GetExternalCollection()
+	{
+		return GameThreadCollection;
+	}
+
+	TManagedArray<TUniquePtr<FParticle>>& GetExternalParticles()
 	{
 		return GTParticles;
 	}
 
+	FParticle* GetParticleByIndex_External(int32 Index);
+	const FParticle* GetParticleByIndex_External(int32 Index) const;
+
+	FParticleHandle* GetParticleByIndex_Internal(int32 Index);
+	const FParticleHandle* GetParticleByIndex_Internal(int32 Index) const;
+
 	/**
 	*  * Get all the geometry collection particle handles based on the processing resolution
 	 */
@@ -337,7 +350,7 @@ public:
 		// first find the GTParticle matching the Transform index
 		if (ChildTransformIndex >= 0 && ChildTransformIndex < GTParticles.Num())
 		{
-			const TUniquePtr<Chaos::FGeometryParticle>& ChildGTParticle = GTParticles[ChildTransformIndex];
+			const TUniquePtr<FParticle>& ChildGTParticle = GTParticles[ChildTransformIndex];
 			if (const int32* InternalClusterUniqueIdx = GTParticlesToInternalClusterUniqueIdx.Find(ChildGTParticle.Get()))
 			{
 				return FGeometryCollectionItemIndex::CreateInternalClusterItemIndex(*InternalClusterUniqueIdx);
@@ -355,7 +368,7 @@ public:
 		return nullptr;
 	}
 	
-	FGeometryCollectionItemIndex GetItemIndexFromGTParticle_External(const Chaos::FGeometryParticle* GTPParticle) const
+	FGeometryCollectionItemIndex GetItemIndexFromGTParticle_External(const FParticle* GTPParticle) const
 	{
 		// internal cluster have  no representation on the GT, so we use the child GT particle to find the matching internal cluster unique index 
 		if (const int32* InternalClusterUniqueIdx = GTParticlesToInternalClusterUniqueIdx.Find(GTPParticle))
@@ -416,6 +429,8 @@ public:
 		PostPhysicsSyncCallback = Callback;
 	}
 	
+	Chaos::FPhysicsObjectHandle GetPhysicsObjectByIndex(int32 Index) const;
+	int32 GetNumParticles() const { return NumParticles; }
 protected:
 	/**
 	* Compute damage threshold for a specific transform
@@ -457,7 +472,6 @@ protected:
 	void CreateNonClusteredParticles(Chaos::FPBDRigidsSolver* RigidsSolver,	const FGeometryCollection& RestCollection, const FGeometryDynamicCollection& DynamicCollection);
 
 	Chaos::FPBDRigidClusteredParticleHandle* FindClusteredParticleHandleByItemIndex_Internal(FGeometryCollectionItemIndex ItemIndex) const;
-	
 private:
 
 	/* set to true once InitializeBodiesPT has been called*/
@@ -469,6 +483,7 @@ private:
 	/** Field Datas stored during evaluation */
 	FFieldExecutionDatas ExecutionDatas;
 
+	TArray<Chaos::FPhysicsObjectUniquePtr> PhysicsObjects;
 	//
 	//  Proxy State Information
 	//
@@ -490,12 +505,12 @@ private:
 
 	EReplicationMode ReplicationMode = EReplicationMode::Unknown;	
 
-	TManagedArray<TUniquePtr<Chaos::FGeometryParticle>> GTParticles;
-	TMap<Chaos::FGeometryParticle*, int32> GTParticlesToTransformGroupIndex;
-	TMap<Chaos::FGeometryParticle*, int32> GTParticlesToInternalClusterUniqueIdx;
+	TManagedArray<TUniquePtr<FParticle>> GTParticles;
+	TMap<FParticle*, int32> GTParticlesToTransformGroupIndex;
+	TMap<FParticle*, int32> GTParticlesToInternalClusterUniqueIdx;
 	TMap<int32, TArray<int32>> InternalClusterUniqueIdxToChildrenTransformIndices;
 
-	TMap<int32, TUniquePtr<Chaos::FGeometryParticle>> GTInternalClustersByUniqueIdx;
+	TMap<int32, TUniquePtr<FParticle>> GTInternalClustersByUniqueIdx;
 
 	// These are read on both threads and should not be changed
 	const FCollisionFilterData SimFilter;
diff --git a/Engine/Source/Runtime/Experimental/Chaos/Public/PhysicsProxy/SingleParticlePhysicsProxy.h b/Engine/Source/Runtime/Experimental/Chaos/Public/PhysicsProxy/SingleParticlePhysicsProxy.h
--- a/Engine/Source/Runtime/Experimental/Chaos/Public/PhysicsProxy/SingleParticlePhysicsProxy.h
+++ b/Engine/Source/Runtime/Experimental/Chaos/Public/PhysicsProxy/SingleParticlePhysicsProxy.h
@@ -8,6 +8,7 @@
 #include "Chaos/Framework/PhysicsProxy.h"
 #include "Chaos/PBDPositionConstraints.h"
 #include "Chaos/ParticleHandle.h"
+#include "Chaos/PhysicsObject.h"
 #include "PhysicsCoreTypes.h"
 #include "Chaos/Defines.h"
 #include "Chaos/PullPhysicsDataImp.h"
@@ -165,9 +166,20 @@ public:
 		return static_cast<const FPBDRigidParticle*>(GetParticle_LowLevel());
 	}
 
+	FPhysicsObjectHandle GetPhysicsObject()
+	{
+		return Reference.Get();
+	}
+
+	const FPhysicsObjectHandle GetPhysicsObject() const
+	{
+		return Reference.Get();
+	}
+
 protected:
 	TUniquePtr<PARTICLE_TYPE> Particle;
 	FParticleHandle* Handle;
+	FPhysicsObjectUniquePtr Reference;
 
 private:
 	FProxyInterpolationData InterpolationData;
diff --git a/Engine/Source/Runtime/PhysicsCore/Private/PhysicsObjectPhysicsCoreInterface.cpp b/Engine/Source/Runtime/PhysicsCore/Private/PhysicsObjectPhysicsCoreInterface.cpp
new file mode 100644
--- /dev/null
+++ b/Engine/Source/Runtime/PhysicsCore/Private/PhysicsObjectPhysicsCoreInterface.cpp
@@ -0,0 +1,16 @@
+// Copyright Epic Games, Inc. All Rights Reserved.
+
+#include "PhysicsObjectPhysicsCoreInterface.h"
+
+#include "Chaos/PhysicsObjectInterface.h"
+#include "PBDRigidsSolver.h"
+
+FChaosScene* PhysicsObjectPhysicsCoreInterface::GetScene(TArrayView<Chaos::FPhysicsObjectHandle> InObjects)
+{
+	Chaos::FPBDRigidsSolver* Solver = Chaos::FPhysicsObjectInterface::GetSolver(InObjects);
+	if (!Solver)
+	{
+		return nullptr;
+	}
+	return reinterpret_cast<FChaosScene*>(Solver->PhysSceneHack);
+}
diff --git a/Engine/Source/Runtime/PhysicsCore/Public/PhysicsObjectPhysicsCoreInterface.h b/Engine/Source/Runtime/PhysicsCore/Public/PhysicsObjectPhysicsCoreInterface.h
new file mode 100644
--- /dev/null
+++ b/Engine/Source/Runtime/PhysicsCore/Public/PhysicsObjectPhysicsCoreInterface.h
@@ -0,0 +1,14 @@
+// Copyright Epic Games, Inc. All Rights Reserved.
+
+#pragma once
+
+#include "Containers/ArrayView.h"
+#include "Chaos/PhysicsObject.h"
+
+class FChaosScene;
+
+class PHYSICSCORE_API PhysicsObjectPhysicsCoreInterface
+{
+public:
+	static FChaosScene* GetScene(TArrayView<Chaos::FPhysicsObjectHandle> InObjects);
+};
\ No newline at end of file
