Chaos Cloth: minor optimization to graph coloring to take the ParticleOffset, which is the start particle index for that cloth. This is an optional parameter. Default behavior is the same as existing behavior.

ParticleOffset should be <= the first node in the graph. This will reduce graph coloring particularly if the ParticleOffset is very large by reducing local allocations and initialization.

Updated the cloth constraints to pass their ParticleOffset to the graph coloring. Did not update the volumetric/flesh constraints because they do not know ParticleOffset already (they don't have weight maps, which is why the ParticleOffset is passed to constraints currently).

Tested with Venus Pantsuit which has two cloths and verified (in debugger + running with Debug checkSlow enabled) that everything is working as expected.

diff --git a/Engine/Source/Runtime/Experimental/Chaos/Private/Chaos/GraphColoring.cpp b/Engine/Source/Runtime/Experimental/Chaos/Private/Chaos/GraphColoring.cpp
--- a/Engine/Source/Runtime/Experimental/Chaos/Private/Chaos/GraphColoring.cpp
+++ b/Engine/Source/Runtime/Experimental/Chaos/Private/Chaos/GraphColoring.cpp
@@ -206,13 +206,17 @@ static bool VerifyGridBasedSubColoring(const TArray<TArray<int32>>& ElementsPerC
 }
 
 template<typename T>
-TArray<TArray<int32>> Chaos::FGraphColoring::ComputeGraphColoring(const TArray<Chaos::TVec2<int32>>& Graph, const Chaos::TDynamicParticles<T, 3>& InParticles)
+TArray<TArray<int32>> Chaos::FGraphColoring::ComputeGraphColoring(const TArray<Chaos::TVec2<int32>>& Graph, const Chaos::TDynamicParticles<T, 3>& InParticles, const int32 GraphParticlesStart, const int32 GraphParticlesEnd)
 {
+	checkSlow(GraphParticlesStart <= GraphParticlesEnd);
+	checkSlow(GraphParticlesEnd <= (int32)InParticles.Size());
 	TArray<TArray<int32>> ColorGraph;
-	TArray<FGraphNode> Nodes;
+	TArray<FGraphNode> NodesSubArray;
 	TArray<FGraphEdge> Edges;
-	Nodes.SetNum(InParticles.Size());
+	NodesSubArray.SetNum(GraphParticlesEnd - GraphParticlesStart);
 	Edges.SetNum(Graph.Num());
+	TArrayView<FGraphNode> Nodes(NodesSubArray.GetData() - GraphParticlesStart, GraphParticlesEnd); // Only nodes starting with GraphParticlesStart are valid to access
+
 	int32 MaxColor = -1;
 
 	for (int32 i = 0; i < Graph.Num(); ++i)
@@ -220,6 +224,8 @@ TArray<TArray<int32>> Chaos::FGraphColoring::ComputeGraphColoring(const TArray<C
 		const TVec2<int32>& Constraint = Graph[i];
 		Edges[i].FirstNode = Constraint[0];
 		Edges[i].SecondNode = Constraint[1];
+		checkSlow(Constraint[0] >= GraphParticlesStart);
+		checkSlow(Constraint[1] >= GraphParticlesStart);
 		Nodes[Constraint[0]].Edges.Add(i);
 		Nodes[Constraint[1]].Edges.Add(i);
 	}
@@ -227,7 +233,7 @@ TArray<TArray<int32>> Chaos::FGraphColoring::ComputeGraphColoring(const TArray<C
 	TSet<int32> ProcessedNodes;
 	TArray<int32> NodesToProcess;
 
-	for (uint32 ParticleNodeIndex = 0; ParticleNodeIndex < InParticles.Size(); ++ParticleNodeIndex)
+	for (int32 ParticleNodeIndex = GraphParticlesStart; ParticleNodeIndex < GraphParticlesEnd; ++ParticleNodeIndex)
 	{
 		const bool bIsParticleDynamic = InParticles.InvM(ParticleNodeIndex) != (T)0.;
 		if (ProcessedNodes.Contains(ParticleNodeIndex) || !bIsParticleDynamic)
@@ -333,15 +339,19 @@ TArray<TArray<int32>> Chaos::FGraphColoring::ComputeGraphColoring(const TArray<C
 }
 
 template<typename T>
-TArray<TArray<int32>> Chaos::FGraphColoring::ComputeGraphColoring(const TArray<TVec3<int32>>& Graph, const Chaos::TDynamicParticles<T, 3>& InParticles)
+TArray<TArray<int32>> Chaos::FGraphColoring::ComputeGraphColoring(const TArray<TVec3<int32>>& Graph, const Chaos::TDynamicParticles<T, 3>& InParticles, const int32 GraphParticlesStart, const int32 GraphParticlesEnd)
 {
 	using namespace Chaos;
 
+	checkSlow(GraphParticlesStart <= GraphParticlesEnd);
+	checkSlow(GraphParticlesEnd <= (int32)InParticles.Size());
 	TArray<TArray<int32>> ColorGraph;
-	TArray<FGraphNode> Nodes;
+	TArray<FGraphNode> NodesSubArray;
 	TArray<FGraph3dEdge> Edges;
-	Nodes.SetNum(InParticles.Size());
+	NodesSubArray.SetNum(GraphParticlesEnd - GraphParticlesStart);
 	Edges.SetNum(Graph.Num());
+	TArrayView<FGraphNode> Nodes(NodesSubArray.GetData() - GraphParticlesStart, GraphParticlesEnd); // Only nodes starting with GraphParticlesStart are valid to access
+
 	int32 MaxColor = -1;
 
 	for (int32 i = 0; i < Graph.Num(); ++i)
@@ -350,6 +360,9 @@ TArray<TArray<int32>> Chaos::FGraphColoring::ComputeGraphColoring(const TArray<T
 		Edges[i].FirstNode = Constraint[0];
 		Edges[i].SecondNode = Constraint[1];
 		Edges[i].ThirdNode = Constraint[2];
+		checkSlow(Constraint[0] >= GraphParticlesStart);
+		checkSlow(Constraint[1] >= GraphParticlesStart);
+		checkSlow(Constraint[2] >= GraphParticlesStart);
 		Nodes[Constraint[0]].Edges.Add(i);
 		Nodes[Constraint[1]].Edges.Add(i);
 		Nodes[Constraint[2]].Edges.Add(i);
@@ -358,7 +371,7 @@ TArray<TArray<int32>> Chaos::FGraphColoring::ComputeGraphColoring(const TArray<T
 	TSet<int32> ProcessedNodes;
 	TArray<int32> NodesToProcess;
 
-	for (uint32 ParticleNodeIndex = 0; ParticleNodeIndex < InParticles.Size(); ++ParticleNodeIndex)
+	for (int32 ParticleNodeIndex = GraphParticlesStart; ParticleNodeIndex < GraphParticlesEnd; ++ParticleNodeIndex)
 	{
 		const bool bIsParticleDynamic = InParticles.InvM(ParticleNodeIndex) != (T)0.;
 		if (ProcessedNodes.Contains(ParticleNodeIndex) || !bIsParticleDynamic)
@@ -515,15 +528,19 @@ TArray<TArray<int32>> Chaos::FGraphColoring::ComputeGraphColoring(const TArray<T
 }
 
 template<typename T>
-TArray<TArray<int32>> Chaos::FGraphColoring::ComputeGraphColoring(const TArray<TVec4<int32>>& Graph, const Chaos::TDynamicParticles<T, 3>& InParticles)
+TArray<TArray<int32>> Chaos::FGraphColoring::ComputeGraphColoring(const TArray<TVec4<int32>>& Graph, const Chaos::TDynamicParticles<T, 3>& InParticles, const int32 GraphParticlesStart, const int32 GraphParticlesEnd)
 {
 	using namespace Chaos;
 
+	checkSlow(GraphParticlesStart <= GraphParticlesEnd);
+	checkSlow(GraphParticlesEnd <= (int32)InParticles.Size());
 	TArray<TArray<int32>> ColorGraph;
-	TArray<FGraphNode> Nodes;
+	TArray<FGraphNode> NodesSubArray;
 	TArray<FGraphTetEdge> Edges;
-	Nodes.SetNum(InParticles.Size());
+	NodesSubArray.SetNum(GraphParticlesEnd - GraphParticlesStart);
 	Edges.SetNum(Graph.Num());
+	TArrayView<FGraphNode> Nodes(NodesSubArray.GetData() - GraphParticlesStart, GraphParticlesEnd); // Only nodes starting with GraphParticlesStart are valid to access
+
 	int32 MaxColor = -1;
 
 	for (int32 i = 0; i < Graph.Num(); ++i)
@@ -533,6 +550,10 @@ TArray<TArray<int32>> Chaos::FGraphColoring::ComputeGraphColoring(const TArray<T
 		Edges[i].SecondNode = Constraint[1];
 		Edges[i].ThirdNode = Constraint[2];
 		Edges[i].FourthNode = Constraint[3];
+		checkSlow(Constraint[0] >= GraphParticlesStart);
+		checkSlow(Constraint[1] >= GraphParticlesStart);
+		checkSlow(Constraint[2] >= GraphParticlesStart);
+		checkSlow(Constraint[3] >= GraphParticlesStart);
 		Nodes[Constraint[0]].Edges.Add(i);
 		Nodes[Constraint[1]].Edges.Add(i);
 		Nodes[Constraint[2]].Edges.Add(i);
@@ -542,7 +563,7 @@ TArray<TArray<int32>> Chaos::FGraphColoring::ComputeGraphColoring(const TArray<T
 	TSet<int32> ProcessedNodes;
 	TArray<int32> NodesToProcess;
 
-	for (uint32 ParticleNodeIndex = 0; ParticleNodeIndex < InParticles.Size(); ++ParticleNodeIndex)
+	for (int32 ParticleNodeIndex = GraphParticlesStart; ParticleNodeIndex < GraphParticlesEnd; ++ParticleNodeIndex)
 	{
 		const bool bIsParticleDynamic = InParticles.InvM(ParticleNodeIndex) != (T)0.;
 		if (ProcessedNodes.Contains(ParticleNodeIndex) || !bIsParticleDynamic)
@@ -770,15 +791,19 @@ TArray<TArray<int32>> Chaos::FGraphColoring::ComputeGraphColoring(const TArray<T
 }
 
 template<typename T>
-TArray<TArray<int32>> Chaos::FGraphColoring::ComputeGraphColoringAllDynamic(const TArray<TVec4<int32>>& Graph, const Chaos::TDynamicParticles<T, 3>& InParticles)
+TArray<TArray<int32>> Chaos::FGraphColoring::ComputeGraphColoringAllDynamic(const TArray<TVec4<int32>>& Graph, const Chaos::TDynamicParticles<T, 3>& InParticles, const int32 GraphParticlesStart, const int32 GraphParticlesEnd)
 {
 	using namespace Chaos;
 
+	checkSlow(GraphParticlesStart <= GraphParticlesEnd);
+	checkSlow(GraphParticlesEnd <= (int32)InParticles.Size());
 	TArray<TArray<int32>> ColorGraph;
-	TArray<FGraphNode> Nodes;
+	TArray<FGraphNode> NodesSubArray;
 	TArray<FGraphTetEdge> Edges;
-	Nodes.SetNum(InParticles.Size());
+	NodesSubArray.SetNum(GraphParticlesEnd - GraphParticlesStart);
 	Edges.SetNum(Graph.Num());
+	TArrayView<FGraphNode> Nodes(NodesSubArray.GetData() - GraphParticlesStart, GraphParticlesEnd); // Only nodes starting with GraphParticlesStart are valid to access
+
 	int32 MaxColor = -1;
 
 	for (int32 i = 0; i < Graph.Num(); ++i)
@@ -788,6 +813,10 @@ TArray<TArray<int32>> Chaos::FGraphColoring::ComputeGraphColoringAllDynamic(cons
 		Edges[i].SecondNode = Constraint[1];
 		Edges[i].ThirdNode = Constraint[2];
 		Edges[i].FourthNode = Constraint[3];
+		checkSlow(Constraint[0] >= GraphParticlesStart);
+		checkSlow(Constraint[1] >= GraphParticlesStart);
+		checkSlow(Constraint[2] >= GraphParticlesStart);
+		checkSlow(Constraint[3] >= GraphParticlesStart);
 		Nodes[Constraint[0]].Edges.Add(i);
 		Nodes[Constraint[1]].Edges.Add(i);
 		Nodes[Constraint[2]].Edges.Add(i);
@@ -797,7 +826,7 @@ TArray<TArray<int32>> Chaos::FGraphColoring::ComputeGraphColoringAllDynamic(cons
 	TSet<int32> ProcessedNodes;
 	TArray<int32> NodesToProcess;
 
-	for (uint32 ParticleNodeIndex = 0; ParticleNodeIndex < InParticles.Size(); ++ParticleNodeIndex)
+	for (int32 ParticleNodeIndex = GraphParticlesStart; ParticleNodeIndex < GraphParticlesEnd; ++ParticleNodeIndex)
 	{
 		//const bool bIsParticleDynamic = InParticles.InvM(ParticleNodeIndex) != (T)0.;
 		if (ProcessedNodes.Contains(ParticleNodeIndex))
@@ -1094,13 +1123,13 @@ void Chaos::ComputeGridBasedGraphSubColoringPointer(const TArray<TArray<int32>>&
 }
 
 
-template CHAOS_API TArray<TArray<int32>> Chaos::FGraphColoring::ComputeGraphColoring<Chaos::FRealSingle>(const TArray<Chaos::TVector<int32, 2>>&, const Chaos::TDynamicParticles<Chaos::FRealSingle, 3>&);
-template CHAOS_API TArray<TArray<int32>> Chaos::FGraphColoring::ComputeGraphColoring<Chaos::FRealDouble>(const TArray<Chaos::TVector<int32, 2>>&, const Chaos::TDynamicParticles<Chaos::FRealDouble, 3>&);
-template CHAOS_API TArray<TArray<int32>> Chaos::FGraphColoring::ComputeGraphColoring<Chaos::FRealSingle>(const TArray<Chaos::TVector<int32, 3>>&, const Chaos::TDynamicParticles<Chaos::FRealSingle, 3>&);
-template CHAOS_API TArray<TArray<int32>> Chaos::FGraphColoring::ComputeGraphColoring<Chaos::FRealDouble>(const TArray<Chaos::TVector<int32, 3>>&, const Chaos::TDynamicParticles<Chaos::FRealDouble, 3>&);
-template CHAOS_API TArray<TArray<int32>> Chaos::FGraphColoring::ComputeGraphColoring<Chaos::FRealSingle>(const TArray<Chaos::TVector<int32, 4>>&, const Chaos::TDynamicParticles<Chaos::FRealSingle, 3>&);
-template CHAOS_API TArray<TArray<int32>> Chaos::FGraphColoring::ComputeGraphColoring<Chaos::FRealDouble>(const TArray<Chaos::TVector<int32, 4>>&, const Chaos::TDynamicParticles<Chaos::FRealDouble, 3>&);
-template CHAOS_API TArray<TArray<int32>> Chaos::FGraphColoring::ComputeGraphColoringAllDynamic<Chaos::FRealSingle>(const TArray<Chaos::TVector<int32, 4>>&, const Chaos::TDynamicParticles<Chaos::FRealSingle, 3>&);
-template CHAOS_API TArray<TArray<int32>> Chaos::FGraphColoring::ComputeGraphColoringAllDynamic<Chaos::FRealDouble>(const TArray<Chaos::TVector<int32, 4>>&, const Chaos::TDynamicParticles<Chaos::FRealDouble, 3>&);
+template CHAOS_API TArray<TArray<int32>> Chaos::FGraphColoring::ComputeGraphColoring<Chaos::FRealSingle>(const TArray<Chaos::TVector<int32, 2>>&, const Chaos::TDynamicParticles<Chaos::FRealSingle, 3>&, const int32 GraphParticlesStart, const int32 GraphParticlesEnd);
+template CHAOS_API TArray<TArray<int32>> Chaos::FGraphColoring::ComputeGraphColoring<Chaos::FRealDouble>(const TArray<Chaos::TVector<int32, 2>>&, const Chaos::TDynamicParticles<Chaos::FRealDouble, 3>&, const int32 GraphParticlesStart, const int32 GraphParticlesEnd);
+template CHAOS_API TArray<TArray<int32>> Chaos::FGraphColoring::ComputeGraphColoring<Chaos::FRealSingle>(const TArray<Chaos::TVector<int32, 3>>&, const Chaos::TDynamicParticles<Chaos::FRealSingle, 3>&, const int32 GraphParticlesStart, const int32 GraphParticlesEnd);
+template CHAOS_API TArray<TArray<int32>> Chaos::FGraphColoring::ComputeGraphColoring<Chaos::FRealDouble>(const TArray<Chaos::TVector<int32, 3>>&, const Chaos::TDynamicParticles<Chaos::FRealDouble, 3>&, const int32 GraphParticlesStart, const int32 GraphParticlesEnd);
+template CHAOS_API TArray<TArray<int32>> Chaos::FGraphColoring::ComputeGraphColoring<Chaos::FRealSingle>(const TArray<Chaos::TVector<int32, 4>>&, const Chaos::TDynamicParticles<Chaos::FRealSingle, 3>&, const int32 GraphParticlesStart, const int32 GraphParticlesEnd);
+template CHAOS_API TArray<TArray<int32>> Chaos::FGraphColoring::ComputeGraphColoring<Chaos::FRealDouble>(const TArray<Chaos::TVector<int32, 4>>&, const Chaos::TDynamicParticles<Chaos::FRealDouble, 3>&, const int32 GraphParticlesStart, const int32 GraphParticlesEnd);
+template CHAOS_API TArray<TArray<int32>> Chaos::FGraphColoring::ComputeGraphColoringAllDynamic<Chaos::FRealSingle>(const TArray<Chaos::TVector<int32, 4>>&, const Chaos::TDynamicParticles<Chaos::FRealSingle, 3>&, const int32 GraphParticlesStart, const int32 GraphParticlesEnd);
+template CHAOS_API TArray<TArray<int32>> Chaos::FGraphColoring::ComputeGraphColoringAllDynamic<Chaos::FRealDouble>(const TArray<Chaos::TVector<int32, 4>>&, const Chaos::TDynamicParticles<Chaos::FRealDouble, 3>&, const int32 GraphParticlesStart, const int32 GraphParticlesEnd);
 template CHAOS_API void Chaos::ComputeGridBasedGraphSubColoringPointer(const TArray<TArray<int32>>& ElementsPerColor, const TMPMGrid<Chaos::FRealSingle>& Grid, const int32 GridSize, TArray<TArray<int32>>*& PreviousColoring, const TArray<TArray<int32>>& ConstraintsNodesSet, TArray<TArray<TArray<int32>>>& ElementsPerSubColors);
 template CHAOS_API void Chaos::ComputeGridBasedGraphSubColoringPointer(const TArray<TArray<int32>>& ElementsPerColor, const TMPMGrid<Chaos::FRealDouble>& Grid, const int32 GridSize, TArray<TArray<int32>>*& PreviousColoring, const TArray<TArray<int32>>& ConstraintsNodesSet, TArray<TArray<TArray<int32>>>& ElementsPerSubColors);
\ No newline at end of file
diff --git a/Engine/Source/Runtime/Experimental/Chaos/Private/Chaos/PBDAxialSpringConstraints.cpp b/Engine/Source/Runtime/Experimental/Chaos/Private/Chaos/PBDAxialSpringConstraints.cpp
--- a/Engine/Source/Runtime/Experimental/Chaos/Private/Chaos/PBDAxialSpringConstraints.cpp
+++ b/Engine/Source/Runtime/Experimental/Chaos/Private/Chaos/PBDAxialSpringConstraints.cpp
@@ -23,14 +23,14 @@ namespace Chaos::Softs {
 int32 Chaos_AxialSpring_ParallelConstraintCount = 100;
 FAutoConsoleVariableRef CVarChaosAxialSpringParallelConstraintCount(TEXT("p.Chaos.AxialSpring.ParallelConstraintCount"), Chaos_AxialSpring_ParallelConstraintCount, TEXT("If we have more constraints than this, use parallel-for in Apply."));
 
-void FPBDAxialSpringConstraints::InitColor(const FSolverParticles& InParticles)
+void FPBDAxialSpringConstraints::InitColor(const FSolverParticles& InParticles, const int32 ParticleOffset, const int32 ParticleCount)
 {
 	// In dev builds we always color so we can tune the system without restarting. See Apply()
 #if UE_BUILD_SHIPPING || UE_BUILD_TEST
 	if (Constraints.Num() > Chaos_AxialSpring_ParallelConstraintCount)
 #endif
 	{
-		const TArray<TArray<int32>> ConstraintsPerColor = FGraphColoring::ComputeGraphColoring(Constraints, InParticles);
+		const TArray<TArray<int32>> ConstraintsPerColor = FGraphColoring::ComputeGraphColoring(Constraints, InParticles, ParticleOffset, ParticleOffset + ParticleCount);
 
 		// Reorder constraints based on color so each array in ConstraintsPerColor contains contiguous elements.
 		TArray<TVec3<int32>> ReorderedConstraints;
diff --git a/Engine/Source/Runtime/Experimental/Chaos/Private/Chaos/PBDBendingConstraints.cpp b/Engine/Source/Runtime/Experimental/Chaos/Private/Chaos/PBDBendingConstraints.cpp
--- a/Engine/Source/Runtime/Experimental/Chaos/Private/Chaos/PBDBendingConstraints.cpp
+++ b/Engine/Source/Runtime/Experimental/Chaos/Private/Chaos/PBDBendingConstraints.cpp
@@ -23,14 +23,14 @@ namespace Chaos::Softs {
 int32 Chaos_Bending_ParallelConstraintCount = 100;
 FAutoConsoleVariableRef CVarChaosBendingParallelConstraintCount(TEXT("p.Chaos.Bending.ParallelConstraintCount"), Chaos_Bending_ParallelConstraintCount, TEXT("If we have more constraints than this, use parallel-for in Apply."));
 
-void FPBDBendingConstraints::InitColor(const FSolverParticles& InParticles)
+void FPBDBendingConstraints::InitColor(const FSolverParticles& InParticles, const int32 ParticleOffset, const int32 ParticleCount)
 {
 	// In dev builds we always color so we can tune the system without restarting. See Apply()
 #if UE_BUILD_SHIPPING || UE_BUILD_TEST
 	if (Constraints.Num() > Chaos_Bending_ParallelConstraintCount)
 #endif
 	{
-		const TArray<TArray<int32>> ConstraintsPerColor = FGraphColoring::ComputeGraphColoring(Constraints, InParticles);
+		const TArray<TArray<int32>> ConstraintsPerColor = FGraphColoring::ComputeGraphColoring(Constraints, InParticles, ParticleOffset, ParticleOffset + ParticleCount);
 
 		// Reorder constraints based on color so each array in ConstraintsPerColor contains contiguous elements.
 		TArray<TVec4<int32>> ReorderedConstraints; 
diff --git a/Engine/Source/Runtime/Experimental/Chaos/Private/Chaos/PBDSpringConstraints.cpp b/Engine/Source/Runtime/Experimental/Chaos/Private/Chaos/PBDSpringConstraints.cpp
--- a/Engine/Source/Runtime/Experimental/Chaos/Private/Chaos/PBDSpringConstraints.cpp
+++ b/Engine/Source/Runtime/Experimental/Chaos/Private/Chaos/PBDSpringConstraints.cpp
@@ -30,14 +30,14 @@ static int32 Chaos_Spring_ParallelConstraintCount = 100;
 FAutoConsoleVariableRef CVarChaosSpringParallelConstraintCount(TEXT("p.Chaos.Spring.ParallelConstraintCount"), Chaos_Spring_ParallelConstraintCount, TEXT("If we have more constraints than this, use parallel-for in Apply."));
 #endif
 
-void FPBDSpringConstraints::InitColor(const FSolverParticles& Particles)
+void FPBDSpringConstraints::InitColor(const FSolverParticles& Particles, const int32 ParticleOffset, const int32 ParticleCount)
 {
 	// In dev builds we always color so we can tune the system without restarting. See Apply()
 #if UE_BUILD_SHIPPING || UE_BUILD_TEST
 	if (Constraints.Num() > Chaos_Spring_ParallelConstraintCount)
 #endif
 	{
-		const TArray<TArray<int32>> ConstraintsPerColor = FGraphColoring::ComputeGraphColoring(Constraints, Particles);
+		const TArray<TArray<int32>> ConstraintsPerColor = FGraphColoring::ComputeGraphColoring(Constraints, Particles, ParticleOffset, ParticleOffset + ParticleCount);
 		
 		// Reorder constraints based on color so each array in ConstraintsPerColor contains contiguous elements.
 		TArray<TVec2<int32>> ReorderedConstraints;
diff --git a/Engine/Source/Runtime/Experimental/Chaos/Private/Chaos/XPBDBendingConstraints.cpp b/Engine/Source/Runtime/Experimental/Chaos/Private/Chaos/XPBDBendingConstraints.cpp
--- a/Engine/Source/Runtime/Experimental/Chaos/Private/Chaos/XPBDBendingConstraints.cpp
+++ b/Engine/Source/Runtime/Experimental/Chaos/Private/Chaos/XPBDBendingConstraints.cpp
@@ -23,14 +23,14 @@ namespace Chaos::Softs {
 int32 Chaos_XPBDBending_ParallelConstraintCount = 100;
 FAutoConsoleVariableRef CVarChaosXPBDBendingParallelConstraintCount(TEXT("p.Chaos.XPBDBending.ParallelConstraintCount"), Chaos_XPBDBending_ParallelConstraintCount, TEXT("If we have more constraints than this, use parallel-for in Apply."));
 
-void FXPBDBendingConstraints::InitColor(const FSolverParticles& InParticles)
+void FXPBDBendingConstraints::InitColor(const FSolverParticles& InParticles, const int32 ParticleOffset, const int32 ParticleCount)
 {
 	// In dev builds we always color so we can tune the system without restarting. See Apply()
 #if UE_BUILD_SHIPPING || UE_BUILD_TEST
 	if (Constraints.Num() > Chaos_XPBDBending_ParallelConstraintCount)
 #endif
 	{
-		const TArray<TArray<int32>> ConstraintsPerColor = FGraphColoring::ComputeGraphColoring(Constraints, InParticles);
+		const TArray<TArray<int32>> ConstraintsPerColor = FGraphColoring::ComputeGraphColoring(Constraints, InParticles, ParticleOffset, ParticleOffset + ParticleCount);
 
 		// Reorder constraints based on color so each array in ConstraintsPerColor contains contiguous elements.
 		TArray<TVec4<int32>> ReorderedConstraints; 
diff --git a/Engine/Source/Runtime/Experimental/Chaos/Private/Chaos/XPBDSpringConstraints.cpp b/Engine/Source/Runtime/Experimental/Chaos/Private/Chaos/XPBDSpringConstraints.cpp
--- a/Engine/Source/Runtime/Experimental/Chaos/Private/Chaos/XPBDSpringConstraints.cpp
+++ b/Engine/Source/Runtime/Experimental/Chaos/Private/Chaos/XPBDSpringConstraints.cpp
@@ -30,14 +30,14 @@ static int32 Chaos_XPBDSpring_ParallelConstraintCount = 100;
 FAutoConsoleVariableRef CVarChaosXPBDSpringParallelConstraintCount(TEXT("p.Chaos.XPBDSpring.ParallelConstraintCount"), Chaos_XPBDSpring_ParallelConstraintCount, TEXT("If we have more constraints than this, use parallel-for in Apply."));
 #endif
 
-void FXPBDSpringConstraints::InitColor(const FSolverParticles& Particles)
+void FXPBDSpringConstraints::InitColor(const FSolverParticles& Particles, const int32 ParticleOffset, const int32 ParticleCount)
 {
 	// In dev builds we always color so we can tune the system without restarting. See Apply()
 #if UE_BUILD_SHIPPING || UE_BUILD_TEST
 	if (Constraints.Num() > Chaos_XPBDSpring_ParallelConstraintCount)
 #endif
 	{
-		const TArray<TArray<int32>> ConstraintsPerColor = FGraphColoring::ComputeGraphColoring(Constraints, Particles);
+		const TArray<TArray<int32>> ConstraintsPerColor = FGraphColoring::ComputeGraphColoring(Constraints, Particles, ParticleOffset, ParticleOffset + ParticleCount);
 		
 		// Reorder constraints based on color so each array in ConstraintsPerColor contains contiguous elements.
 		TArray<TVec2<int32>> ReorderedConstraints;
diff --git a/Engine/Source/Runtime/Experimental/Chaos/Public/Chaos/GraphColoring.h b/Engine/Source/Runtime/Experimental/Chaos/Public/Chaos/GraphColoring.h
--- a/Engine/Source/Runtime/Experimental/Chaos/Public/Chaos/GraphColoring.h
+++ b/Engine/Source/Runtime/Experimental/Chaos/Public/Chaos/GraphColoring.h
@@ -62,13 +62,33 @@ class FGraphColoring
 
   public:
 	template<typename T>
-	CHAOS_API static TArray<TArray<int32>> ComputeGraphColoring(const TArray<TVector<int32, 2>>& Graph, const TDynamicParticles<T, 3>& InParticles);
+	CHAOS_API static TArray<TArray<int32>> ComputeGraphColoring(const TArray<TVector<int32, 2>>& Graph, const TDynamicParticles<T, 3>& InParticles, const int32 GraphParticlesStart, const int32 GraphParticlesEnd);
 	template<typename T>
-	CHAOS_API static TArray<TArray<int32>> ComputeGraphColoring(const TArray<TVector<int32, 3>>& Graph, const TDynamicParticles<T, 3>& InParticles);
+	inline static TArray<TArray<int32>> ComputeGraphColoring(const TArray<TVector<int32, 2>>& Graph, const TDynamicParticles<T, 3>& InParticles)
+	{
+		return ComputeGraphColoring(Graph, InParticles, 0, InParticles.Size());
+	}
+	template<typename T>
+	CHAOS_API static TArray<TArray<int32>> ComputeGraphColoring(const TArray<TVector<int32, 3>>& Graph, const TDynamicParticles<T, 3>& InParticles, const int32 GraphParticlesStart, const int32 GraphParticlesEnd);
+	template<typename T>
+	inline static TArray<TArray<int32>> ComputeGraphColoring(const TArray<TVector<int32, 3>>& Graph, const TDynamicParticles<T, 3>& InParticles)
+	{
+		return ComputeGraphColoring(Graph, InParticles, 0, InParticles.Size());
+	}
 	template<typename T>
-	CHAOS_API static TArray<TArray<int32>> ComputeGraphColoring(const TArray<TVector<int32, 4>>& Graph, const TDynamicParticles<T, 3>& InParticles);
+	CHAOS_API static TArray<TArray<int32>> ComputeGraphColoring(const TArray<TVector<int32, 4>>& Graph, const TDynamicParticles<T, 3>& InParticles, const int32 GraphParticlesStart, const int32 GraphParticlesEnd);
 	template<typename T>
-	CHAOS_API static TArray<TArray<int32>> ComputeGraphColoringAllDynamic(const TArray<TVec4<int32>>& Graph, const Chaos::TDynamicParticles<T, 3>& InParticles);
+	inline static TArray<TArray<int32>> ComputeGraphColoring(const TArray<TVector<int32, 4>>& Graph, const TDynamicParticles<T, 3>& InParticles)
+	{
+		return ComputeGraphColoring(Graph, InParticles, 0, InParticles.Size());
+	}
+	template<typename T>
+	CHAOS_API static TArray<TArray<int32>> ComputeGraphColoringAllDynamic(const TArray<TVec4<int32>>& Graph, const Chaos::TDynamicParticles<T, 3>& InParticles, const int32 GraphParticlesStart, const int32 GraphParticlesEnd);
+	template<typename T>
+	inline static TArray<TArray<int32>> ComputeGraphColoringAllDynamic(const TArray<TVec4<int32>>& Graph, const Chaos::TDynamicParticles<T, 3>& InParticles)
+	{
+		return ComputeGraphColoring(Graph, InParticles, 0, InParticles.Size());
+	}
 };
 
 template<typename T> 
diff --git a/Engine/Source/Runtime/Experimental/Chaos/Public/Chaos/PBDAxialSpringConstraints.h b/Engine/Source/Runtime/Experimental/Chaos/Public/Chaos/PBDAxialSpringConstraints.h
--- a/Engine/Source/Runtime/Experimental/Chaos/Public/Chaos/PBDAxialSpringConstraints.h
+++ b/Engine/Source/Runtime/Experimental/Chaos/Public/Chaos/PBDAxialSpringConstraints.h
@@ -23,7 +23,7 @@ public:
 		bool bTrimKinematicConstraints)
 		: Base(Particles, ParticleOffset, ParticleCount, InConstraints, StiffnessMultipliers, InStiffness, bTrimKinematicConstraints)
 	{
-		InitColor(Particles);
+		InitColor(Particles, ParticleOffset, ParticleCount);
 	}
 
 	virtual ~FPBDAxialSpringConstraints() override {}
@@ -31,7 +31,7 @@ public:
 	void Apply(FSolverParticles& InParticles, const FSolverReal Dt) const;
 
 private:
-	void InitColor(const FSolverParticles& InParticles);
+	void InitColor(const FSolverParticles& InParticles, const int32 ParticleOffset, const int32 ParticleCount);
 	void ApplyHelper(FSolverParticles& InParticles, const FSolverReal Dt, const int32 ConstraintIndex, const FSolverReal ExpStiffnessValue) const;
 
 private:
diff --git a/Engine/Source/Runtime/Experimental/Chaos/Public/Chaos/PBDBendingConstraints.h b/Engine/Source/Runtime/Experimental/Chaos/Public/Chaos/PBDBendingConstraints.h
--- a/Engine/Source/Runtime/Experimental/Chaos/Public/Chaos/PBDBendingConstraints.h
+++ b/Engine/Source/Runtime/Experimental/Chaos/Public/Chaos/PBDBendingConstraints.h
@@ -24,7 +24,7 @@ public:
 		bool bTrimKinematicConstraints = false)
 		:Base(InParticles, ParticleOffset, ParticleCount, MoveTemp(InConstraints), StiffnessMultipliers, BucklingStiffnessMultipliers, InStiffness, InBucklingRatio, InBucklingStiffness, bTrimKinematicConstraints) 
 	{
-		InitColor(InParticles);
+		InitColor(InParticles, ParticleOffset, ParticleCount);
 	}
 
 	FPBDBendingConstraints(const FSolverParticles& InParticles, TArray<TVec4<int32>>&& InConstraints, const FSolverReal InStiffness = (FSolverReal)1.)
@@ -35,7 +35,7 @@ public:
 	void Apply(FSolverParticles& InParticles, const FSolverReal Dt) const;
 
 private:
-	void InitColor(const FSolverParticles& InParticles);
+	void InitColor(const FSolverParticles& InParticles, const int32 ParticleOffset, const int32 ParticleCount);
 	void ApplyHelper(FSolverParticles& InParticles, const FSolverReal Dt, const int32 ConstraintIndex, const FSolverReal ExpStiffnessValue, const FSolverReal ExpBucklingValue) const;
 
 	TArray<int32> ConstraintsPerColorStartIndex; // Constraints are ordered so each batch is contiguous. This is ColorNum + 1 length so it can be used as start and end.
diff --git a/Engine/Source/Runtime/Experimental/Chaos/Public/Chaos/PBDSpringConstraints.h b/Engine/Source/Runtime/Experimental/Chaos/Public/Chaos/PBDSpringConstraints.h
--- a/Engine/Source/Runtime/Experimental/Chaos/Public/Chaos/PBDSpringConstraints.h
+++ b/Engine/Source/Runtime/Experimental/Chaos/Public/Chaos/PBDSpringConstraints.h
@@ -25,7 +25,7 @@ public:
 		typename TEnableIf<Valence >= 2 && Valence <= 4>::Type* = nullptr)
 		: Base(Particles, ParticleOffset, ParticleCount, InConstraints, StiffnessMultipliers, InStiffness, bTrimKinematicConstraints)
 	{
-		InitColor(Particles);
+		InitColor(Particles, ParticleOffset, ParticleCount);
 	}
 
 	virtual ~FPBDSpringConstraints() override {}
@@ -35,7 +35,7 @@ public:
 	const TArray<int32>& GetConstraintsPerColorStartIndex() const { return ConstraintsPerColorStartIndex; }
 
 private:
-	void InitColor(const FSolverParticles& InParticles);
+	void InitColor(const FSolverParticles& InParticles, const int32 ParticleOffset, const int32 ParticleCount);
 	void ApplyHelper(FSolverParticles& Particles, const FSolverReal Dt, const int32 ConstraintIndex, const FSolverReal ExpStiffnessValue) const;
 
 private:
diff --git a/Engine/Source/Runtime/Experimental/Chaos/Public/Chaos/XPBDBendingConstraints.h b/Engine/Source/Runtime/Experimental/Chaos/Public/Chaos/XPBDBendingConstraints.h
--- a/Engine/Source/Runtime/Experimental/Chaos/Public/Chaos/XPBDBendingConstraints.h
+++ b/Engine/Source/Runtime/Experimental/Chaos/Public/Chaos/XPBDBendingConstraints.h
@@ -46,7 +46,7 @@ public:
 		, DampingRatio(FSolverVec2::ZeroVector)
 	{
 		Lambdas.Init((FSolverReal)0., Constraints.Num());
-		InitColor(InParticles);
+		InitColor(InParticles, ParticleOffset, ParticleCount);
 	}
 
 	FXPBDBendingConstraints(const FSolverParticles& InParticles,
@@ -76,7 +76,7 @@ public:
 		, DampingRatio(InDampingRatio, DampingMultipliers, TConstArrayView<TVec2<int32>>(ConstraintSharedEdges), ParticleOffset, ParticleCount)
 	{
 		Lambdas.Init((FSolverReal)0., Constraints.Num());
-		InitColor(InParticles);
+		InitColor(InParticles, ParticleOffset, ParticleCount);
 	}
 
 	virtual ~FXPBDBendingConstraints() override {}
@@ -107,7 +107,7 @@ public:
 	void Apply(FSolverParticles& Particles, const FSolverReal Dt) const;
 
 private:
-	void InitColor(const FSolverParticles& InParticles);
+	void InitColor(const FSolverParticles& InParticles, const int32 ParticleOffset, const int32 ParticleCount);
 	void ApplyHelper(FSolverParticles& Particles, const FSolverReal Dt, const int32 ConstraintIndex, const FSolverReal ExpStiffnessValue, const FSolverReal ExpBucklingValue, const FSolverReal DampingRatioValue) const;
 
 private:
diff --git a/Engine/Source/Runtime/Experimental/Chaos/Public/Chaos/XPBDSpringConstraints.h b/Engine/Source/Runtime/Experimental/Chaos/Public/Chaos/XPBDSpringConstraints.h
--- a/Engine/Source/Runtime/Experimental/Chaos/Public/Chaos/XPBDSpringConstraints.h
+++ b/Engine/Source/Runtime/Experimental/Chaos/Public/Chaos/XPBDSpringConstraints.h
@@ -41,7 +41,7 @@ public:
 		, DampingRatio(FSolverVec2::ZeroVector)
 	{
 		Lambdas.Init((FSolverReal)0., Constraints.Num());
-		InitColor(Particles);
+		InitColor(Particles, ParticleOffset, ParticleCount);
 	}
 	template<int32 Valence, TEMPLATE_REQUIRES(Valence >= 2 && Valence <= 4)>
 	FXPBDSpringConstraints(
@@ -71,7 +71,7 @@ public:
 			ParticleCount)
 	{
 		Lambdas.Init((FSolverReal)0., Constraints.Num());
-		InitColor(Particles);
+		InitColor(Particles, ParticleOffset, ParticleCount);
 	}
 
 	virtual ~FXPBDSpringConstraints() override {}
@@ -91,7 +91,7 @@ public:
 	void Apply(FSolverParticles& Particles, const FSolverReal Dt) const;
 
 private:
-	void InitColor(const FSolverParticles& InParticles);
+	void InitColor(const FSolverParticles& InParticles, const int32 ParticleOffset, const int32 ParticleCount);
 	void ApplyHelper(FSolverParticles& Particles, const FSolverReal Dt, const int32 ConstraintIndex, const FSolverReal StiffnessValue, const FSolverReal DampingRatioValue) const;
 
 	FSolverVec3 GetDelta(const FSolverParticles& Particles, const FSolverReal Dt, const int32 ConstraintIndex, const FSolverReal StiffnessValue, const FSolverReal DampingRatioValue) const
