Chaos - fix CCD collisons

Missing call to FPBDCollisionConstraint::Activate()

diff --git a/Engine/Source/Runtime/Experimental/Chaos/Private/Chaos/Collision/PBDCollisionConstraint.cpp b/Engine/Source/Runtime/Experimental/Chaos/Private/Chaos/Collision/PBDCollisionConstraint.cpp
--- a/Engine/Source/Runtime/Experimental/Chaos/Private/Chaos/Collision/PBDCollisionConstraint.cpp
+++ b/Engine/Source/Runtime/Experimental/Chaos/Private/Chaos/Collision/PBDCollisionConstraint.cpp
@@ -390,7 +390,7 @@ namespace Chaos
 		CCDTargetPenetration = FMath::Min(MinBounds * CVars::CCDAllowedDepthBoundsScale, CCDEnablePenetration);
 	}
 
-	void FPBDCollisionConstraint::Activate(const FReal Dt)
+	void FPBDCollisionConstraint::Activate()
 	{
 		SetDisabled(false);
 
diff --git a/Engine/Source/Runtime/Experimental/Chaos/Private/Chaos/Collision/ParticlePairMidPhase.cpp b/Engine/Source/Runtime/Experimental/Chaos/Private/Chaos/Collision/ParticlePairMidPhase.cpp
--- a/Engine/Source/Runtime/Experimental/Chaos/Private/Chaos/Collision/ParticlePairMidPhase.cpp
+++ b/Engine/Source/Runtime/Experimental/Chaos/Private/Chaos/Collision/ParticlePairMidPhase.cpp
@@ -347,8 +347,6 @@ namespace Chaos
 			{
 				if (Context.GetAllocator()->ActivateConstraint(Constraint.Get()))
 				{
-					Constraint->Activate(Dt);
-
 					LastUsedEpoch = CurrentEpoch;
 					return 1;
 				}
diff --git a/Engine/Source/Runtime/Experimental/Chaos/Public/Chaos/Collision/CollisionConstraintAllocator.h b/Engine/Source/Runtime/Experimental/Chaos/Public/Chaos/Collision/CollisionConstraintAllocator.h
--- a/Engine/Source/Runtime/Experimental/Chaos/Public/Chaos/Collision/CollisionConstraintAllocator.h
+++ b/Engine/Source/Runtime/Experimental/Chaos/Public/Chaos/Collision/CollisionConstraintAllocator.h
@@ -89,6 +89,8 @@ namespace Chaos
 
 				NewActiveConstraints.Push(Constraint);
 
+				Constraint->Activate();
+
 				return true;
 			}
 			return false;
diff --git a/Engine/Source/Runtime/Experimental/Chaos/Public/Chaos/Collision/PBDCollisionConstraint.h b/Engine/Source/Runtime/Experimental/Chaos/Public/Chaos/Collision/PBDCollisionConstraint.h
--- a/Engine/Source/Runtime/Experimental/Chaos/Public/Chaos/Collision/PBDCollisionConstraint.h
+++ b/Engine/Source/Runtime/Experimental/Chaos/Public/Chaos/Collision/PBDCollisionConstraint.h
@@ -299,7 +299,10 @@ namespace Chaos
 		/** 
 		 * Called each frame when the constraint is active after primary collision detection (but not per incremental collision detection call if enabled) 
 		 */
-		void Activate(const FReal Dt);
+		void Activate();
+
+		UE_DEPRECATED(5.2, "Removed parameter")
+		void Activate(const FReal Dt) { Activate(); }
 
 		// When a particle is moved under user control, we need to update some cached state to prevent friction from undoing the move
 		void UpdateParticleTransform(FGeometryParticleHandle* InParticle);
