geometry collection physics object needs thread-specific behavior for querying children and parent object

diff --git a/Engine/Source/Runtime/Engine/Private/PhysicsEngine/PhysicsObjectBlueprintLibrary.cpp b/Engine/Source/Runtime/Engine/Private/PhysicsEngine/PhysicsObjectBlueprintLibrary.cpp
--- a/Engine/Source/Runtime/Engine/Private/PhysicsEngine/PhysicsObjectBlueprintLibrary.cpp
+++ b/Engine/Source/Runtime/Engine/Private/PhysicsEngine/PhysicsObjectBlueprintLibrary.cpp
@@ -17,7 +17,7 @@ FClosestPhysicsObjectResult UPhysicsObjectBlueprintLibrary::GetClosestPhysicsObj
 		// Maybe allow users to customize the predicate?
 		[&Interface](Chaos::FPhysicsObject* Object) {
 			// We need to filter so we only get the leaf bodies. Mainly in consideration of the geometry collection.
-			return !Interface->AreAllDisabled({&Object, 1}) && !Chaos::FPhysicsObjectInterface::HasChildren(Object);
+			return !Interface->AreAllDisabled({&Object, 1}) && !Interface->HasChildren(Object);
 		}
 	);
 	return Interface->GetClosestPhysicsBodyFromLocation(PhysicsObjects, WorldLocation);
diff --git a/Engine/Source/Runtime/Experimental/Chaos/Private/Chaos/PhysicsObject.cpp b/Engine/Source/Runtime/Experimental/Chaos/Private/Chaos/PhysicsObject.cpp
--- a/Engine/Source/Runtime/Experimental/Chaos/Private/Chaos/PhysicsObject.cpp
+++ b/Engine/Source/Runtime/Experimental/Chaos/Private/Chaos/PhysicsObject.cpp
@@ -20,40 +20,4 @@ namespace Chaos
 	{
 		return Proxy != nullptr && !Proxy->GetMarkedDeleted();
 	}
-
-	bool FPhysicsObject::HasChildren() const
-	{
-		EPhysicsProxyType ProxyType = Proxy->GetType();
-		switch (ProxyType)
-		{
-		case EPhysicsProxyType::GeometryCollectionType:
-		{
-			FGeometryDynamicCollection& Collection = static_cast<FGeometryCollectionPhysicsProxy*>(Proxy)->GetExternalCollection();
-			return !Collection.Children[BodyIndex].IsEmpty();
-		}
-		default:
-			break;
-		}
-		return false;
-	}
-
-	FPhysicsObjectHandle FPhysicsObject::GetParentObject() const
-	{
-		EPhysicsProxyType ProxyType = Proxy->GetType();
-		switch (ProxyType)
-		{
-		case EPhysicsProxyType::GeometryCollectionType:
-		{
-			FGeometryCollectionPhysicsProxy* GeometryCollectionProxy = static_cast<FGeometryCollectionPhysicsProxy*>(Proxy);
-			FGeometryDynamicCollection& Collection = GeometryCollectionProxy->GetExternalCollection();
-			if (int32 Index = Collection.Parent[BodyIndex]; Index != INDEX_NONE)
-			{
-				return GeometryCollectionProxy->GetPhysicsObjectByIndex(Index);
-			}
-		}
-		default:
-			break;
-		}
-		return nullptr;
-	}
 }
\ No newline at end of file
diff --git a/Engine/Source/Runtime/Experimental/Chaos/Private/Chaos/PhysicsObjectInterface.cpp b/Engine/Source/Runtime/Experimental/Chaos/Private/Chaos/PhysicsObjectInterface.cpp
--- a/Engine/Source/Runtime/Experimental/Chaos/Private/Chaos/PhysicsObjectInterface.cpp
+++ b/Engine/Source/Runtime/Experimental/Chaos/Private/Chaos/PhysicsObjectInterface.cpp
@@ -86,6 +86,12 @@ namespace Chaos
 		return Object->GetRootObject<Id>();
 	}
 
+	template<EThreadContext Id>
+	bool FReadPhysicsObjectInterface<Id>::HasChildren(FPhysicsObjectHandle Object)
+	{
+		return Object ? Object->HasChildren<Id>() : false;
+	}
+
 	template<EThreadContext Id>
 	FTransform FReadPhysicsObjectInterface<Id>::GetTransform(FPhysicsObjectHandle Object)
 	{
@@ -421,11 +427,6 @@ namespace Chaos
 		}
 	}
 
-	bool FPhysicsObjectInterface::HasChildren(FPhysicsObjectHandle Object)
-	{
-		return Object ? Object->HasChildren() : false;
-	}
-
 	void FPhysicsObjectInterface::SetName(FPhysicsObjectHandle Object, const FName& InName)
 	{
 		if (!Object)
diff --git a/Engine/Source/Runtime/Experimental/Chaos/Private/Chaos/PhysicsObjectInternal.h b/Engine/Source/Runtime/Experimental/Chaos/Private/Chaos/PhysicsObjectInternal.h
--- a/Engine/Source/Runtime/Experimental/Chaos/Private/Chaos/PhysicsObjectInternal.h
+++ b/Engine/Source/Runtime/Experimental/Chaos/Private/Chaos/PhysicsObjectInternal.h
@@ -51,17 +51,36 @@ namespace Chaos
 
 			TThreadParticle<Id>* Particle = GetParticle<Id>();
 			FPhysicsObjectHandle CurrentObject = this;
-			FPhysicsObjectHandle CurrentParent = GetParentObject();
+			FPhysicsObjectHandle CurrentParent = GetParentObject<Id>();
 			while (CurrentParent && IsParticleDisabled<Id>(Particle))
 			{
 				Particle = CurrentParent->GetParticle<Id>();
 				CurrentObject = CurrentParent;
-				CurrentParent = CurrentParent->GetParentObject();
+				CurrentParent = CurrentParent->GetParentObject<Id>();
 			}
 			return CurrentObject;
 		}
 
-		FPhysicsObjectHandle GetParentObject() const;
+		template<EThreadContext Id>
+		FPhysicsObjectHandle GetParentObject() const
+		{
+			EPhysicsProxyType ProxyType = Proxy->GetType();
+			switch (ProxyType)
+			{
+			case EPhysicsProxyType::GeometryCollectionType:
+			{
+				FGeometryCollectionPhysicsProxy* GeometryCollectionProxy = static_cast<FGeometryCollectionPhysicsProxy*>(Proxy);
+				FGeometryDynamicCollection& Collection = GetGeometryCollectionDynamicCollection<Id>(*GeometryCollectionProxy);
+				if (int32 Index = Collection.Parent[BodyIndex]; Index != INDEX_NONE)
+				{
+					return GeometryCollectionProxy->GetPhysicsObjectByIndex(Index);
+				}
+			}
+			default:
+				break;
+			}
+			return nullptr;
+		}
 
 		template <EThreadContext Id>
 		TThreadParticle<Id>* GetRootParticle()
@@ -122,7 +141,22 @@ namespace Chaos
 		IPhysicsProxyBase* PhysicsProxy() { return Proxy; }
 		const IPhysicsProxyBase* PhysicsProxy() const { return Proxy; }
 
-		bool HasChildren() const;
+		template<EThreadContext Id>
+		bool HasChildren() const
+		{
+			EPhysicsProxyType ProxyType = Proxy->GetType();
+			switch (ProxyType)
+			{
+			case EPhysicsProxyType::GeometryCollectionType:
+			{
+				FGeometryDynamicCollection& Collection = GetGeometryCollectionDynamicCollection<Id>(*static_cast<FGeometryCollectionPhysicsProxy*>(Proxy));
+				return !Collection.Children[BodyIndex].IsEmpty();
+			}
+			default:
+				break;
+			}
+			return false;
+		}
 
 		friend class FPhysicsObjectFactory;
 	protected:
@@ -136,6 +170,19 @@ namespace Chaos
 		IPhysicsProxyBase* Proxy = nullptr;
 		int32 BodyIndex = INDEX_NONE;
 		FName BodyName = NAME_None;
+
+		template<EThreadContext Id>
+		static FGeometryDynamicCollection& GetGeometryCollectionDynamicCollection(FGeometryCollectionPhysicsProxy& GeometryCollectionProxy)
+		{
+			if constexpr (Id == EThreadContext::External)
+			{
+				return GeometryCollectionProxy.GetExternalCollection();
+			}
+			else
+			{
+				return GeometryCollectionProxy.GetPhysicsCollection();
+			}
+		}
 	};
 
 	/**
diff --git a/Engine/Source/Runtime/Experimental/Chaos/Public/Chaos/PhysicsObjectInterface.h b/Engine/Source/Runtime/Experimental/Chaos/Public/Chaos/PhysicsObjectInterface.h
--- a/Engine/Source/Runtime/Experimental/Chaos/Public/Chaos/PhysicsObjectInterface.h
+++ b/Engine/Source/Runtime/Experimental/Chaos/Public/Chaos/PhysicsObjectInterface.h
@@ -45,6 +45,7 @@ namespace Chaos
 	{
 	public:
 		FPhysicsObjectHandle GetRootObject(FPhysicsObjectHandle Object);
+		bool HasChildren(FPhysicsObjectHandle Object);
 
 		FTransform GetTransform(FPhysicsObjectHandle Object);
 		FVector GetX(FPhysicsObjectHandle Object);
@@ -98,8 +99,6 @@ namespace Chaos
 	class CHAOS_API FPhysicsObjectInterface
 	{
 	public:
-		static bool HasChildren(FPhysicsObjectHandle Object);
-
 		static void SetName(FPhysicsObjectHandle Object, const FName& InName);
 		static FName GetName(FPhysicsObjectHandle Object);
 
