fix uninitialized variable warning in usage of FClosestPhysicsObjectResult::ClosestDistance

diff --git a/Engine/Source/Runtime/Experimental/Chaos/Public/Chaos/PhysicsObjectInterface.h b/Engine/Source/Runtime/Experimental/Chaos/Public/Chaos/PhysicsObjectInterface.h
--- a/Engine/Source/Runtime/Experimental/Chaos/Public/Chaos/PhysicsObjectInterface.h
+++ b/Engine/Source/Runtime/Experimental/Chaos/Public/Chaos/PhysicsObjectInterface.h
@@ -20,7 +20,7 @@ struct CHAOS_API FClosestPhysicsObjectResult
 
 	Chaos::FPhysicsObjectHandle PhysicsObject = nullptr;
 	FVector ClosestLocation;
-	double ClosestDistance;
+	double ClosestDistance = 0.0;
 	operator bool() const
 	{
 		return IsValid();
