Iris - Renamed Iris file NetHandleManager.h to NetRefHandleManager.h

#rnx

diff --git a/Engine/Plugins/Runtime/ReplicationSystemTestPlugin/Source/Private/Tests/ReplicationSystem/Prioritization/TestNetObjectPrioritizer.cpp b/Engine/Plugins/Runtime/ReplicationSystemTestPlugin/Source/Private/Tests/ReplicationSystem/Prioritization/TestNetObjectPrioritizer.cpp
--- a/Engine/Plugins/Runtime/ReplicationSystemTestPlugin/Source/Private/Tests/ReplicationSystem/Prioritization/TestNetObjectPrioritizer.cpp
+++ b/Engine/Plugins/Runtime/ReplicationSystemTestPlugin/Source/Private/Tests/ReplicationSystem/Prioritization/TestNetObjectPrioritizer.cpp
@@ -2,7 +2,7 @@
 
 #include "NetworkAutomationTest.h"
 #include "NetworkAutomationTestMacros.h"
-#include "Iris/ReplicationSystem/NetHandleManager.h"
+#include "Iris/ReplicationSystem/NetRefHandleManager.h"
 #include "Iris/ReplicationSystem/ReplicationSystemInternal.h"
 #include "Iris/ReplicationSystem/Prioritization/NetObjectPrioritizerDefinitions.h"
 #include "Tests/ReplicationSystem/Prioritization/MockNetObjectPrioritizer.h"
diff --git a/Engine/Source/Runtime/Experimental/Iris/Core/Private/Iris/Core/IrisDebugging.cpp b/Engine/Source/Runtime/Experimental/Iris/Core/Private/Iris/Core/IrisDebugging.cpp
--- a/Engine/Source/Runtime/Experimental/Iris/Core/Private/Iris/Core/IrisDebugging.cpp
+++ b/Engine/Source/Runtime/Experimental/Iris/Core/Private/Iris/Core/IrisDebugging.cpp
@@ -12,7 +12,7 @@
 #include "Iris/ReplicationSystem/ReplicationOperations.h"
 #include "Iris/ReplicationSystem/ReplicationSystem.h"
 #include "Iris/ReplicationSystem/ReplicationSystemInternal.h"
-#include "Iris/ReplicationSystem/NetHandleManager.h"
+#include "Iris/ReplicationSystem/NetRefHandleManager.h"
 #include "Iris/ReplicationSystem/ReplicationProtocolManager.h"
 #include "Iris/ReplicationState/PropertyReplicationState.h"
 #include "Iris/Serialization/NetSerializationContext.h"
diff --git a/Engine/Source/Runtime/Experimental/Iris/Core/Private/Iris/ReplicationState/ReplicationStateStorage.cpp b/Engine/Source/Runtime/Experimental/Iris/Core/Private/Iris/ReplicationState/ReplicationStateStorage.cpp
--- a/Engine/Source/Runtime/Experimental/Iris/Core/Private/Iris/ReplicationState/ReplicationStateStorage.cpp
+++ b/Engine/Source/Runtime/Experimental/Iris/Core/Private/Iris/ReplicationState/ReplicationStateStorage.cpp
@@ -2,7 +2,7 @@
 
 #include "Iris/ReplicationState/ReplicationStateStorage.h"
 #include "Iris/Core/IrisMemoryTracker.h"
-#include "Iris/ReplicationSystem/NetHandleManager.h"
+#include "Iris/ReplicationSystem/NetRefHandleManager.h"
 #include "Iris/ReplicationSystem/ReplicationOperations.h"
 #include "Iris/ReplicationSystem/ReplicationOperationsInternal.h"
 #include "Iris/ReplicationSystem/ReplicationProtocol.h"
diff --git a/Engine/Source/Runtime/Experimental/Iris/Core/Private/Iris/ReplicationSystem/Conditionals/ReplicationConditionals.cpp b/Engine/Source/Runtime/Experimental/Iris/Core/Private/Iris/ReplicationSystem/Conditionals/ReplicationConditionals.cpp
--- a/Engine/Source/Runtime/Experimental/Iris/Core/Private/Iris/ReplicationSystem/Conditionals/ReplicationConditionals.cpp
+++ b/Engine/Source/Runtime/Experimental/Iris/Core/Private/Iris/ReplicationSystem/Conditionals/ReplicationConditionals.cpp
@@ -4,7 +4,7 @@
 #include "Iris/Core/IrisLog.h"
 #include "Iris/Core/IrisProfiler.h"
 #include "Iris/ReplicationState/ReplicationStateUtil.h"
-#include "Iris/ReplicationSystem/NetHandleManager.h"
+#include "Iris/ReplicationSystem/NetRefHandleManager.h"
 #include "Iris/ReplicationSystem/ReplicationConnections.h"
 #include "Iris/ReplicationSystem/ReplicationFragment.h"
 #include "Iris/ReplicationSystem/ReplicationProtocol.h"
diff --git a/Engine/Source/Runtime/Experimental/Iris/Core/Private/Iris/ReplicationSystem/DeltaCompression/DeltaCompressionBaselineManager.cpp b/Engine/Source/Runtime/Experimental/Iris/Core/Private/Iris/ReplicationSystem/DeltaCompression/DeltaCompressionBaselineManager.cpp
--- a/Engine/Source/Runtime/Experimental/Iris/Core/Private/Iris/ReplicationSystem/DeltaCompression/DeltaCompressionBaselineManager.cpp
+++ b/Engine/Source/Runtime/Experimental/Iris/Core/Private/Iris/ReplicationSystem/DeltaCompression/DeltaCompressionBaselineManager.cpp
@@ -6,7 +6,7 @@
 #include "Iris/Core/IrisProfiler.h"
 #include "Iris/Core/IrisMemoryTracker.h"
 #include "Iris/ReplicationSystem/ChangeMaskCache.h"
-#include "Iris/ReplicationSystem/NetHandleManager.h"
+#include "Iris/ReplicationSystem/NetRefHandleManager.h"
 #include "Iris/ReplicationSystem/ReplicationConnections.h"
 #include "Iris/ReplicationSystem/ReplicationProtocol.h"
 #include "Iris/ReplicationSystem/ReplicationSystemTypes.h"
diff --git a/Engine/Source/Runtime/Experimental/Iris/Core/Private/Iris/ReplicationSystem/Filtering/NetObjectGroups.cpp b/Engine/Source/Runtime/Experimental/Iris/Core/Private/Iris/ReplicationSystem/Filtering/NetObjectGroups.cpp
--- a/Engine/Source/Runtime/Experimental/Iris/Core/Private/Iris/ReplicationSystem/Filtering/NetObjectGroups.cpp
+++ b/Engine/Source/Runtime/Experimental/Iris/Core/Private/Iris/ReplicationSystem/Filtering/NetObjectGroups.cpp
@@ -2,7 +2,7 @@
 
 #include "NetObjectGroups.h"
 #include "Iris/ReplicationSystem/Filtering/NetObjectFilter.h"
-#include "Iris/ReplicationSystem/NetHandleManager.h" // for InvalidInternalIndex
+#include "Iris/ReplicationSystem/NetRefHandleManager.h" // for InvalidInternalIndex
 #include "Containers/ArrayView.h"
 
 namespace UE::Net::Private
diff --git a/Engine/Source/Runtime/Experimental/Iris/Core/Private/Iris/ReplicationSystem/Filtering/ReplicationFiltering.cpp b/Engine/Source/Runtime/Experimental/Iris/Core/Private/Iris/ReplicationSystem/Filtering/ReplicationFiltering.cpp
--- a/Engine/Source/Runtime/Experimental/Iris/Core/Private/Iris/ReplicationSystem/Filtering/ReplicationFiltering.cpp
+++ b/Engine/Source/Runtime/Experimental/Iris/Core/Private/Iris/ReplicationSystem/Filtering/ReplicationFiltering.cpp
@@ -6,7 +6,7 @@
 #include "Iris/IrisConstants.h"
 #include "Iris/Core/IrisLog.h"
 #include "Iris/Core/IrisProfiler.h"
-#include "Iris/ReplicationSystem/NetHandleManager.h"
+#include "Iris/ReplicationSystem/NetRefHandleManager.h"
 #include "Iris/ReplicationSystem/ReplicationConnections.h"
 #include "Iris/ReplicationSystem/ReplicationProtocol.h"
 #include "Iris/ReplicationSystem/ReplicationSystem.h"
diff --git a/Engine/Source/Runtime/Experimental/Iris/Core/Private/Iris/ReplicationSystem/LegacyPushModel.cpp b/Engine/Source/Runtime/Experimental/Iris/Core/Private/Iris/ReplicationSystem/LegacyPushModel.cpp
--- a/Engine/Source/Runtime/Experimental/Iris/Core/Private/Iris/ReplicationSystem/LegacyPushModel.cpp
+++ b/Engine/Source/Runtime/Experimental/Iris/Core/Private/Iris/ReplicationSystem/LegacyPushModel.cpp
@@ -6,7 +6,7 @@
 #include "Iris/ReplicationState/PropertyReplicationState.h"
 #include "Iris/ReplicationSystem/ReplicationSystem.h"
 #include "Iris/ReplicationSystem/ReplicationSystemInternal.h"
-#include "Iris/ReplicationSystem/NetHandleManager.h"
+#include "Iris/ReplicationSystem/NetRefHandleManager.h"
 #include "Net/Core/DirtyNetObjectTracker/GlobalDirtyNetObjectTracker.h"
 #include "Containers/ArrayView.h"
 
diff --git a/Engine/Source/Runtime/Experimental/Iris/Core/Private/Iris/ReplicationSystem/NetBlob/NetBlobManager.h b/Engine/Source/Runtime/Experimental/Iris/Core/Private/Iris/ReplicationSystem/NetBlob/NetBlobManager.h
--- a/Engine/Source/Runtime/Experimental/Iris/Core/Private/Iris/ReplicationSystem/NetBlob/NetBlobManager.h
+++ b/Engine/Source/Runtime/Experimental/Iris/Core/Private/Iris/ReplicationSystem/NetBlob/NetBlobManager.h
@@ -4,7 +4,7 @@
 #include "CoreTypes.h"
 #include "Iris/ReplicationSystem/NetBlob/NetBlob.h"
 #include "Iris/ReplicationSystem/NetBlob/NetBlobHandlerManager.h"
-#include "Iris/ReplicationSystem/NetHandleManager.h" // For FInternalNetRefIndex
+#include "Iris/ReplicationSystem/NetRefHandleManager.h" // For FInternalNetRefIndex
 #include "UObject/StrongObjectPtr.h"
 
 class UNetBlobHandler;
diff --git a/Engine/Source/Runtime/Experimental/Iris/Core/Private/Iris/ReplicationSystem/NetHandleManager.cpp b/Engine/Source/Runtime/Experimental/Iris/Core/Private/Iris/ReplicationSystem/NetHandleManager.cpp
--- a/Engine/Source/Runtime/Experimental/Iris/Core/Private/Iris/ReplicationSystem/NetHandleManager.cpp
+++ b/Engine/Source/Runtime/Experimental/Iris/Core/Private/Iris/ReplicationSystem/NetHandleManager.cpp
@@ -1,6 +1,6 @@
 // Copyright Epic Games, Inc. All Rights Reserved.
 
-#include "NetHandleManager.h"
+#include "NetRefHandleManager.h"
 #include "ReplicationOperationsInternal.h"
 #include "ReplicationProtocol.h"
 #include "ReplicationProtocolManager.h"
diff --git a/Engine/Source/Runtime/Experimental/Iris/Core/Private/Iris/ReplicationSystem/NetHandleManager.h b/Engine/Source/Runtime/Experimental/Iris/Core/Private/Iris/ReplicationSystem/NetHandleManager.h
deleted file mode 100644
diff --git a/Engine/Source/Runtime/Experimental/Iris/Core/Private/Iris/ReplicationSystem/NetRefHandleManager.h b/Engine/Source/Runtime/Experimental/Iris/Core/Private/Iris/ReplicationSystem/NetRefHandleManager.h
new file mode 100644
--- /dev/null
+++ b/Engine/Source/Runtime/Experimental/Iris/Core/Private/Iris/ReplicationSystem/NetRefHandleManager.h
@@ -0,0 +1,452 @@
+// Copyright Epic Games, Inc. All Rights Reserved.
+
+#pragma once
+
+#include "CoreTypes.h"
+#include "Containers/Array.h"
+#include "Containers/Map.h"
+#include "Containers/ArrayView.h"
+#include "Net/Core/NetBitArray.h"
+#include "Net/Core/NetHandle/NetHandle.h"
+#include "Iris/ReplicationState/ReplicationStateDescriptor.h"
+#include "Iris/ReplicationSystem/NetHandle.h"
+#include "Iris/ReplicationSystem/NetDependencyData.h"
+
+class FReferenceCollector;
+namespace UE::Net
+{
+	struct FReplicationProtocol;
+	struct FReplicationInstanceProtocol;
+
+	namespace Private
+	{
+		class FReplicationProtocolManager;
+	}
+}
+
+namespace UE::Net::Private
+{
+
+typedef uint32 FInternalNetRefIndex;
+
+enum class EAddSubObjectFlags : uint32
+{
+	None = 0U,
+	WarnIfAlreadySubObject = 1U,
+	SkipIfAlreadySubObject = WarnIfAlreadySubObject << 1U,
+	DestroyWithOwner = SkipIfAlreadySubObject << 1U,
+	ReplicateWithSubObject = DestroyWithOwner << 1U,
+	Default = WarnIfAlreadySubObject | DestroyWithOwner,
+};
+ENUM_CLASS_FLAGS(EAddSubObjectFlags);
+
+enum class EAddDependentObjectFlags : uint32
+{
+	None = 0U,
+	WarnIfAlreadyDependentObject = 1U,
+};
+ENUM_CLASS_FLAGS(EAddDependentObjectFlags);
+
+enum class ERemoveDependentObjectFlags : uint32
+{
+	None = 0U,
+	RemoveFromDependentParentObjects = 1U,
+	RemoveFromParentDependentObjects = RemoveFromDependentParentObjects << 1U,
+	All = RemoveFromDependentParentObjects | RemoveFromParentDependentObjects,
+};
+ENUM_CLASS_FLAGS(ERemoveDependentObjectFlags);
+
+// Internal class to manage NetHandles and their internal data
+class FNetRefHandleManager
+{
+public:
+	// We will never assign InvalidInternalIndex
+	enum { InvalidInternalIndex = 0U };
+
+	// We need to store some internal data for Replicated objects
+	// $TODO: This should be split up into separate array according to usage patterns, it is getting a bit too large
+	struct FReplicatedObjectData
+	{
+		FReplicatedObjectData()
+		: Protocol(nullptr)
+		, InstanceProtocol(nullptr)
+		, ReceiveStateBuffer(nullptr)
+		, SubObjectRootIndex(InvalidInternalIndex)
+		, SubObjectParentIndex(InvalidInternalIndex)
+		, Flags(0U)
+		{
+		}
+
+		FNetRefHandle RefHandle;
+		FNetHandle NetHandle;
+		const FReplicationProtocol* Protocol;
+		const FReplicationInstanceProtocol* InstanceProtocol;
+		uint8* ReceiveStateBuffer;
+		uint32 SubObjectRootIndex;
+		uint32 SubObjectParentIndex;
+
+		union
+		{
+			uint32 Flags : 32U;
+			struct
+			{
+				uint32 bShouldPropagateChangedStates : 1U;
+				uint32 bTearOff : 1U;
+				uint32 bDestroySubObjectWithOwner : 1U;
+				uint32 bIsDependentObject : 1U;
+				uint32 bHasDependentObjects : 1U;
+				uint32 Padding : 27U;
+			};
+		};
+	
+		// Returns true if this is a SubObject
+		bool IsSubObject() const { return SubObjectRootIndex && !!bDestroySubObjectWithOwner; }
+
+		// Returns true if this is a DepedentObject
+		bool IsDependentObject() const { return bIsDependentObject; }
+	};
+
+	// We want to handle both actual Networked and "network addressable" handles in the same way
+	// As we expect the number of FNetRefHandles to be much greater than the actual number of Replicated Objects we
+	// use a map to map from Handle to InternalIndex, only Replicated Objects will ever be assigned an InternalNetHandle
+	typedef TMap<FNetRefHandle, FInternalNetRefIndex> FRefHandleMap;
+	typedef TMap<FNetHandle, FInternalNetRefIndex> FNetHandleMap;
+
+public:
+	FNetRefHandleManager(FReplicationProtocolManager& InReplicationProtocolManager, uint32 InReplicationSystemId, uint32 MaxActiveObjects);
+
+	// Return true if this is a scopable index
+	bool IsScopableIndex(FInternalNetRefIndex InternalIndex) const { return ScopableInternalIndices.GetBit(InternalIndex); }
+
+	static FNetRefHandle MakeNetRefHandle(uint32 Id, uint32 ReplicationSystemId);
+	static FNetRefHandle MakeNetRefHandleFromId(uint32 Id);
+
+	// Returns a valid handle if the wanted handle can be allocated
+	FNetRefHandle AllocateNetRefHandle(bool bIsStatic);
+
+	// Create local Net Object
+	FNetRefHandle CreateNetObject(FNetRefHandle WantedHandle, FNetHandle GlobalHandle, const FReplicationProtocol* ReplicationProtocol);
+
+	// Create NetObject on request from remote
+	FNetRefHandle CreateNetObjectFromRemote(FNetRefHandle WantedHandle, const FReplicationProtocol* ReplicationProtocol);
+
+	// Attach Instance protocol to handle
+	// Instance can be null, we only track the Instance for legacy support
+	void AttachInstanceProtocol(FInternalNetRefIndex InternalIndex, const FReplicationInstanceProtocol* InstanceProtocol, UObject* Instance);
+
+	// Used when we want to detach the instance protocol due to async destroy
+	const FReplicationInstanceProtocol* DetachInstanceProtocol(FInternalNetRefIndex InternalIndex);
+
+	// Creates a DestructionInfo Handle which is used to replicate persistently destroyed static objects when late joining or streaming in levels with static objects that already has been destroyed
+	// Returns a new handle used to replicate the destruction info, and adds a cross-reference to the original handle being destroyed if that still exists 
+	// Handle is a reference to an already replicated Handle that is about to be destroyed
+	FNetRefHandle CreateHandleForDestructionInfo(FNetRefHandle Handle, const FReplicationProtocol* DestroydObjectProtocol);
+	
+	void DestroyNetObject(FNetRefHandle Handle);
+
+	// Mark object as no longer scopable, object will be removed from scope for all connections
+	void RemoveFromScope(FInternalNetRefIndex InternalIndex);
+
+	TArrayView<const FInternalNetRefIndex> GetObjectsPendingDestroy() const { return MakeArrayView(PendingDestroyInternalIndices); }
+	void DestroyObjectsPendingDestroy();
+
+	const FReplicatedObjectData& GetReplicatedObjectDataNoCheck(FInternalNetRefIndex InternalIndex) const { return ReplicatedObjectData[InternalIndex]; }
+	FReplicatedObjectData& GetReplicatedObjectDataNoCheck(FInternalNetRefIndex InternalIndex) { return ReplicatedObjectData[InternalIndex]; }
+
+	inline const FReplicatedObjectData& GetReplicatedObjectData(FInternalNetRefIndex InternalIndex) const;
+
+	inline const uint8* GetReplicatedObjectStateBufferNoCheck(FInternalNetRefIndex InternalObjectIndex) const { return ReplicatedObjectStateBuffers.GetData()[InternalObjectIndex]; }
+	inline uint8* GetReplicatedObjectStateBufferNoCheck(FInternalNetRefIndex InternalObjectIndex) { return ReplicatedObjectStateBuffers.GetData()[InternalObjectIndex]; }
+	inline const TArray<uint8*>& GetReplicatedObjectStateBuffers() const { return ReplicatedObjectStateBuffers; }
+
+	/** Verify a handle against internal handle. A handle is valid if it matches internal storage. */
+	inline bool IsValidNetRefHandle(FNetRefHandle Handle) const;
+
+	/** Returns true if the handle is for a local replicated object. */
+	inline bool IsLocalNetRefHandle(FNetRefHandle Handle) const;
+
+	/** Returns true if the handle is for a remote replicated object. */
+	inline bool IsRemoteNetRefHandle(FNetRefHandle Handle) const;
+
+	/** Extract a full handle from an incomplete one consisting of only an index. */
+	inline FNetRefHandle GetCompleteNetRefHandle(FNetRefHandle IncompleteHandle) const;
+	
+	// Get Handle from internal index
+	inline FNetRefHandle GetNetRefHandleFromInternalIndex(FInternalNetRefIndex InternalIndex) const;
+
+	// Get internal index from handle
+	inline FInternalNetRefIndex GetInternalIndex(FNetRefHandle Handle) const;
+
+	// Get internal index from NetHandle
+	inline FInternalNetRefIndex GetInternalIndexFromNetHandle(FNetHandle Handle) const;
+
+	// Get bitarray for all currently scopable internal indices
+	const FNetBitArray& GetScopableInternalIndices() const { return ScopableInternalIndices; }
+
+	// Get bitarray for all internal indices that was scopable last update
+	const FNetBitArray& GetPrevFrameScopableInternalIndices() const { return PrevFrameScopableInternalIndices; }
+	void SetPrevFrameScopableInternalIndicesToCurrent() { PrevFrameScopableInternalIndices = ScopableInternalIndices; }
+
+	// Get bitarray for all internal indices that currently are assigned
+	const FNetBitArray& GetAssignedInternalIndices() const { return AssignedInternalIndices; }
+
+	// SubObjects
+	const FNetBitArray& GetSubObjectInternalIndices() const { return SubObjectInternalIndices; }	
+
+	bool AddSubObject(FNetRefHandle OwnerHandle, FNetRefHandle SubObjectHandle, FNetRefHandle RelativeOtherSubObjectHandle, EAddSubObjectFlags Flags = EAddSubObjectFlags::Default);
+	bool AddSubObject(FNetRefHandle OwnerHandle, FNetRefHandle SubObjectHandle, EAddSubObjectFlags Flags = EAddSubObjectFlags::Default);
+	void RemoveSubObject(FNetRefHandle SubObjectHandle);
+	FNetRefHandle GetSubObjectOwner(FNetRefHandle SubObjectHandle) const;
+	bool SetSubObjectNetCondition(FInternalNetRefIndex SubObjectInternalIndex, FLifeTimeConditionStorage SubObjectCondition);
+
+	// DependentObjects
+	const FNetBitArray& GetDependentObjectInternalIndices() const { return DependentObjectInternalIndices; }
+	const FNetBitArray& GetObjectsWithDependentObjectsInternalIndices() const { return ObjectsWithDependentObjectsInternalIndices; }
+	bool AddDependentObject(FNetRefHandle ParentHandle, FNetRefHandle DependentHandle, EAddDependentObjectFlags Flags = EAddDependentObjectFlags::WarnIfAlreadyDependentObject);
+	void RemoveDependentObject(FNetRefHandle ParentHandle, FNetRefHandle DependentHandle);
+
+	// Remove DependentHandles from all dependent object tracking
+	// $TODO: This might need a pass ObjectReplicationBridge as well since most likely also must restore filter/polling logic
+	void RemoveDependentObject(FNetRefHandle DependentHandle);
+	
+	void SetShouldPropagateChangedStates(FNetRefHandle Handle, bool bShouldPropagateChangedStates);
+	void SetShouldPropagateChangedStates(FInternalNetRefIndex ObjectInternalIndex, bool bShouldPropagateChangedStates);
+
+	uint32 GetMaxActiveObjectCount() const { return MaxActiveObjectCount; }
+	uint32 GetActiveObjectCount() const { return ActiveObjectCount; }
+
+	// We do refcount objects tracked by each connection in order to know when it is safe to reuse an InternalIndex
+	void AddNetObjectRef(FInternalNetRefIndex InternalIndex) { ++ReplicatedObjectRefCount[InternalIndex]; }
+	void ReleaseNetObjectRef(FInternalNetRefIndex InternalIndex) { check(ReplicatedObjectRefCount[InternalIndex] > 0); --ReplicatedObjectRefCount[InternalIndex]; }
+	uint16 GetNetObjectRefCount(FInternalNetRefIndex ObjectInternalIndex) const { return ReplicatedObjectRefCount[ObjectInternalIndex]; }
+
+	// Get dependent objects for the given ParentIndex
+	inline TArrayView<const FInternalNetRefIndex> GetDependentObjects(FInternalNetRefIndex ParentIndex) const;
+
+	// Get all parents of the given DependentIndex
+	inline TArrayView<const FInternalNetRefIndex> GetDependentObjectParents(FInternalNetRefIndex DependentIndex) const;
+
+	// Get all subobjects for the given OwnerIndex, Note: only valid for the root
+	inline TArrayView<const FInternalNetRefIndex> GetSubObjects(FInternalNetRefIndex OwnerIndex) const;
+
+	// Get child subobjects for a object, used when we do hierarchical operations such as conditional serialization
+	inline TArrayView<const FInternalNetRefIndex> GetChildSubObjects(FInternalNetRefIndex ParentObjectIndex) const;
+
+	// Get child subobjects and condtionals array if one exists, if there are no child subobjects the method returns false
+	inline bool GetChildSubObjects(FInternalNetRefIndex OwnerIndex, FChildSubObjectsInfo& OutInfo) const;
+
+	const FRefHandleMap& GetReplicatedHandles() const { return RefHandleToInternalIndex; }
+
+	const TArray<UObject*>& GetReplicatedInstances() const { return ReplicatedInstances; }
+
+	void AddReferencedObjects(FReferenceCollector& Collector);
+
+	bool GetIsDestroyedStartupObject(FInternalNetRefIndex InternalIndex) const { return DestroyedStartupObjectInternalIndices.GetBit(InternalIndex); }
+	inline uint32 GetOriginalDestroyedStartupObjectIndex(FInternalNetRefIndex InternalIndex) const;
+
+	const FNetBitArray& GetDestroyedStartupObjectInternalIndices() const { return DestroyedStartupObjectInternalIndices; }
+
+	const FNetBitArray& GetWantToBeDormantInternalIndices() const { return WantToBeDormantInternalIndices; }
+	FNetBitArray& GetWantToBeDormantInternalIndices() { return WantToBeDormantInternalIndices; }
+
+	// Iterate over all dependent objects and their dependent objects
+	template <typename T>
+	void ForAllDependentObjectsRecursive(FInternalNetRefIndex ObjectIndex, T&& Functor) const
+	{
+		if (ObjectsWithDependentObjectsInternalIndices.GetBit(ObjectIndex))
+		{
+			for (const FInternalNetRefIndex DependentObjectIndex : GetDependentObjects(ObjectIndex))
+			{
+				Functor(DependentObjectIndex);
+				ForAllDependentObjectsRecursive(DependentObjectIndex, Functor);
+			}
+			return;
+		}
+	};
+
+private:
+	FInternalNetRefIndex InternalCreateNetObject(const FNetRefHandle NetRefHandle, const FNetHandle GlobalHandle, const FReplicationProtocol* ReplicationProtocol);
+	void InternalDestroyNetObject(FInternalNetRefIndex InternalIndex);
+
+	static uint32 MakeNetRefHandleId(uint32 Seed, bool bIsStatic);
+	uint32 GetNextNetRefHandleId(uint32 HandleIndex) const;
+
+	// Get the next free internal index, returns InvalidInternalIndex if a free one cannot be found
+	FInternalNetRefIndex GetNextFreeInternalIndex() const;
+
+	bool InternalAddSubObject(FInternalNetRefIndex OwnerInternalIndex, FInternalNetRefIndex SubObjectInternalIndex, FInternalNetRefIndex RelativeOtherSubObjectInternalIndex, EAddSubObjectFlags Flags);
+
+	inline void SetIsSubObject(FInternalNetRefIndex InternalIndex, bool IsSubObject) { SubObjectInternalIndices.SetBitValue(InternalIndex, IsSubObject); }
+	void InternalRemoveSubObject(FInternalNetRefIndex OwnerInternalIndex, FInternalNetRefIndex SubObjectInternalIndex, bool bRemoveFromSubObjectArray = true);
+
+	void InternalRemoveDependentObject(FInternalNetRefIndex ParentInternalIndex, FInternalNetRefIndex DependentInternalIndex, ERemoveDependentObjectFlags Flags = ERemoveDependentObjectFlags::All);
+	void InternalRemoveDependentObject(FInternalNetRefIndex DependentInternalIndex);
+
+	// The current replicated object count
+	uint32 ActiveObjectCount;
+
+	// Max allowed replicated object count
+	uint32 MaxActiveObjectCount;
+
+	uint32 ReplicationSystemId;
+
+	FRefHandleMap RefHandleToInternalIndex;
+	FNetHandleMap NetHandleToInternalIndex;
+
+	// Bitset used in order to track assigned internal which are scopable
+	FNetBitArray ScopableInternalIndices;
+
+	// Which internal indices were used last net frame. This can be used to find out which ones are new and deleted this frame. 
+	FNetBitArray PrevFrameScopableInternalIndices;
+
+	// Bitset containing all internal indices that are assigned
+	FNetBitArray AssignedInternalIndices;
+
+	// Bitset containing all internal indices that are SubObjects
+	FNetBitArray SubObjectInternalIndices;
+
+	// Bitset containing all internal indices that are dependentObjects
+	FNetBitArray DependentObjectInternalIndices;
+
+	// Bitset containing all internal indices that have dependentObjects
+	FNetBitArray ObjectsWithDependentObjectsInternalIndices;
+
+	// Bitset marking internal indices of static objects that are permanently destroyed
+	// We use this for special scoping
+	FNetBitArray DestroyedStartupObjectInternalIndices;
+
+	// Bitset marking internal indices that wants to be dormant
+	FNetBitArray WantToBeDormantInternalIndices;
+
+	// Both are marked in DestroyedStartupObjectsInternalIndicies
+	// When one is destroyed both mapping will be cleared
+	// Map OriginalIndex -> DestructionInfoIndex
+	// Map DestructionInfoIndex -> Orignal
+	TMap<uint32, uint32> DestroyedStartupObject;
+
+private:
+
+	// Array used in order to track objects pending destroy
+	TArray<FInternalNetRefIndex> PendingDestroyInternalIndices;
+
+	// Just an array containing data about our replicated objects
+	TArray<FReplicatedObjectData> ReplicatedObjectData;
+
+	// Pointers to state buffers for all replicated objects
+	TArray<uint8*> ReplicatedObjectStateBuffers;
+
+	// Refcounts for all tracked objects
+	TArray<uint16> ReplicatedObjectRefCount;
+
+	// Raw pointers to all bound instances
+	TArray<UObject*> ReplicatedInstances;
+
+	// Assign handles
+	uint32 NextStaticHandleIndex;
+	uint32 NextDynamicHandleIndex;
+
+	FNetDependencyData SubObjects;
+
+	FReplicationProtocolManager& ReplicationProtocolManager;
+};
+
+const FNetRefHandleManager::FReplicatedObjectData& FNetRefHandleManager::GetReplicatedObjectData(FInternalNetRefIndex InternalIndex) const
+{
+	check(AssignedInternalIndices.GetBit(InternalIndex));
+	return GetReplicatedObjectDataNoCheck(InternalIndex);
+}
+
+FInternalNetRefIndex FNetRefHandleManager::GetInternalIndex(FNetRefHandle Handle) const
+{
+	const FInternalNetRefIndex* InternalIndex = RefHandleToInternalIndex.Find(Handle);
+	return InternalIndex ? *InternalIndex : FNetRefHandleManager::InvalidInternalIndex;
+}
+
+FInternalNetRefIndex FNetRefHandleManager::GetInternalIndexFromNetHandle(FNetHandle Handle) const
+{
+	const FInternalNetRefIndex* InternalIndex = NetHandleToInternalIndex.Find(Handle);
+	return InternalIndex ? *InternalIndex : FNetRefHandleManager::InvalidInternalIndex;
+}
+
+FNetRefHandle FNetRefHandleManager::GetNetRefHandleFromInternalIndex(FInternalNetRefIndex InternalIndex) const
+{
+	check(AssignedInternalIndices.GetBit(InternalIndex));
+	return GetReplicatedObjectDataNoCheck(InternalIndex).RefHandle;
+}
+
+FNetRefHandle FNetRefHandleManager::GetCompleteNetRefHandle(FNetRefHandle IncompleteHandle) const
+{
+	if (const FInternalNetRefIndex* InternalIndex = RefHandleToInternalIndex.Find(IncompleteHandle))
+	{
+		return GetReplicatedObjectDataNoCheck(*InternalIndex).RefHandle;
+	}
+	else
+	{
+		return FNetRefHandle();
+	}
+}
+
+bool FNetRefHandleManager::IsValidNetRefHandle(FNetRefHandle Handle) const
+{
+	return RefHandleToInternalIndex.Contains(Handle);
+}
+
+bool FNetRefHandleManager::IsLocalNetRefHandle(FNetRefHandle Handle) const
+{
+	if (const FInternalNetRefIndex InternalIndex = GetInternalIndex(Handle))
+	{
+		// For the time being only replicated objects owned by this peer has a state buffer
+		return ReplicatedObjectStateBuffers[InternalIndex] != nullptr;
+	}
+
+	return false;
+}
+
+bool FNetRefHandleManager::IsRemoteNetRefHandle(FNetRefHandle Handle) const
+{
+	if (const uint32 InternalIndex = GetInternalIndex(Handle))
+	{
+		// For the time being only replicated objects owned by this peer has a state buffer
+		return ReplicatedObjectStateBuffers[InternalIndex] == nullptr;
+	}
+
+	return false;
+}
+
+uint32 FNetRefHandleManager::GetOriginalDestroyedStartupObjectIndex(FInternalNetRefIndex InternalIndex) const
+{
+	const uint32* FoundOriginalInternalIndex = DestroyedStartupObject.Find(InternalIndex);
+	return FoundOriginalInternalIndex ? *FoundOriginalInternalIndex : 0U;
+}
+
+inline bool FNetRefHandleManager::AddSubObject(FNetRefHandle OwnerHandle, FNetRefHandle SubObjectHandle, EAddSubObjectFlags Flags)
+{
+	return AddSubObject(OwnerHandle, SubObjectHandle, FNetRefHandle(), Flags);
+}
+
+TArrayView<const FInternalNetRefIndex> FNetRefHandleManager::GetSubObjects(FInternalNetRefIndex OwnerIndex) const
+{
+	return SubObjects.GetInternalIndexArray<FNetDependencyData::EArrayType::SubObjects>(OwnerIndex);
+}
+
+bool FNetRefHandleManager::GetChildSubObjects(FInternalNetRefIndex OwnerIndex, FChildSubObjectsInfo& OutInfo) const
+{
+	return SubObjects.GetChildSubObjects(OwnerIndex, OutInfo);
+}
+
+TArrayView<const FInternalNetRefIndex> FNetRefHandleManager::GetChildSubObjects(FInternalNetRefIndex OwnerIndex) const
+{
+	return SubObjects.GetInternalIndexArray<FNetDependencyData::EArrayType::ChildSubObjects>(OwnerIndex);
+}
+
+TArrayView<const FInternalNetRefIndex> FNetRefHandleManager::GetDependentObjects(FInternalNetRefIndex ParentIndex) const
+{
+	return SubObjects.GetInternalIndexArray<FNetDependencyData::EArrayType::DependentObjects>(ParentIndex);
+}
+
+TArrayView<const FInternalNetRefIndex> FNetRefHandleManager::GetDependentObjectParents(FInternalNetRefIndex DependentIndex) const
+{
+	return SubObjects.GetInternalIndexArray<FNetDependencyData::EArrayType::ParentObjects>(DependentIndex);
+}
+	
+}
diff --git a/Engine/Source/Runtime/Experimental/Iris/Core/Private/Iris/ReplicationSystem/Prioritization/ReplicationPrioritization.cpp b/Engine/Source/Runtime/Experimental/Iris/Core/Private/Iris/ReplicationSystem/Prioritization/ReplicationPrioritization.cpp
--- a/Engine/Source/Runtime/Experimental/Iris/Core/Private/Iris/ReplicationSystem/Prioritization/ReplicationPrioritization.cpp
+++ b/Engine/Source/Runtime/Experimental/Iris/Core/Private/Iris/ReplicationSystem/Prioritization/ReplicationPrioritization.cpp
@@ -5,7 +5,7 @@
 #include "Iris/Core/IrisLog.h"
 #include "Iris/Core/IrisProfiler.h"
 #include "Net/Core/NetBitArray.h"
-#include "Iris/ReplicationSystem/NetHandleManager.h"
+#include "Iris/ReplicationSystem/NetRefHandleManager.h"
 #include "Iris/ReplicationSystem/ReplicationConnections.h"
 #include "Iris/ReplicationSystem/ReplicationWriter.h"
 #include "Iris/ReplicationSystem/Prioritization/NetObjectPrioritizerDefinitions.h"
diff --git a/Engine/Source/Runtime/Experimental/Iris/Core/Private/Iris/ReplicationSystem/ReplicationOperationsInternal.cpp b/Engine/Source/Runtime/Experimental/Iris/Core/Private/Iris/ReplicationSystem/ReplicationOperationsInternal.cpp
--- a/Engine/Source/Runtime/Experimental/Iris/Core/Private/Iris/ReplicationSystem/ReplicationOperationsInternal.cpp
+++ b/Engine/Source/Runtime/Experimental/Iris/Core/Private/Iris/ReplicationSystem/ReplicationOperationsInternal.cpp
@@ -4,7 +4,7 @@
 #include "Iris/Core/IrisLog.h"
 #include "Iris/ReplicationSystem/ChangeMaskCache.h"
 #include "Iris/ReplicationSystem/ChangeMaskUtil.h"
-#include "Iris/ReplicationSystem/NetHandleManager.h"
+#include "Iris/ReplicationSystem/NetRefHandleManager.h"
 #include "Iris/ReplicationSystem/ObjectReferenceCache.h"
 #include "Iris/ReplicationSystem/ReplicationFragment.h"
 #include "Iris/ReplicationSystem/ReplicationProtocol.h"
diff --git a/Engine/Source/Runtime/Experimental/Iris/Core/Private/Iris/ReplicationSystem/ReplicationReader.cpp b/Engine/Source/Runtime/Experimental/Iris/Core/Private/Iris/ReplicationSystem/ReplicationReader.cpp
--- a/Engine/Source/Runtime/Experimental/Iris/Core/Private/Iris/ReplicationSystem/ReplicationReader.cpp
+++ b/Engine/Source/Runtime/Experimental/Iris/Core/Private/Iris/ReplicationSystem/ReplicationReader.cpp
@@ -9,7 +9,7 @@
 #include "Net/Core/Trace/NetDebugName.h"
 #include "Iris/DataStream/DataStream.h"
 #include "Iris/ReplicationSystem/DeltaCompression/DeltaCompressionBaselineManager.h"
-#include "Iris/ReplicationSystem/NetHandleManager.h"
+#include "Iris/ReplicationSystem/NetRefHandleManager.h"
 #include "Iris/Serialization/NetReferenceCollector.h"
 #include "Iris/ReplicationSystem/ReplicationOperations.h"
 #include "Iris/ReplicationSystem/ReplicationOperationsInternal.h"
diff --git a/Engine/Source/Runtime/Experimental/Iris/Core/Private/Iris/ReplicationSystem/ReplicationReader.h b/Engine/Source/Runtime/Experimental/Iris/Core/Private/Iris/ReplicationSystem/ReplicationReader.h
--- a/Engine/Source/Runtime/Experimental/Iris/Core/Private/Iris/ReplicationSystem/ReplicationReader.h
+++ b/Engine/Source/Runtime/Experimental/Iris/Core/Private/Iris/ReplicationSystem/ReplicationReader.h
@@ -5,7 +5,7 @@
 // This class will never be included from public headers
 #include "Iris/ReplicationSystem/AttachmentReplication.h"
 #include "Iris/ReplicationSystem/ChangeMaskUtil.h"
-#include "Iris/ReplicationSystem/NetHandleManager.h"
+#include "Iris/ReplicationSystem/NetRefHandleManager.h"
 #include "Iris/ReplicationSystem/ObjectReferenceCache.h"
 #include "Iris/ReplicationSystem/ReplicationTypes.h"
 #include "Containers/Map.h"
diff --git a/Engine/Source/Runtime/Experimental/Iris/Core/Private/Iris/ReplicationSystem/ReplicationSystemInternal.h b/Engine/Source/Runtime/Experimental/Iris/Core/Private/Iris/ReplicationSystem/ReplicationSystemInternal.h
--- a/Engine/Source/Runtime/Experimental/Iris/Core/Private/Iris/ReplicationSystem/ReplicationSystemInternal.h
+++ b/Engine/Source/Runtime/Experimental/Iris/Core/Private/Iris/ReplicationSystem/ReplicationSystemInternal.h
@@ -6,7 +6,7 @@
 #include "Containers/Array.h"
 #include "Iris/ReplicationState/ReplicationStateStorage.h"
 #include "Iris/ReplicationSystem/DirtyNetObjectTracker.h"
-#include "Iris/ReplicationSystem/NetHandleManager.h"
+#include "Iris/ReplicationSystem/NetRefHandleManager.h"
 #include "Iris/ReplicationSystem/ChangeMaskCache.h"
 #include "Iris/ReplicationSystem/Conditionals/ReplicationConditionals.h"
 #include "Iris/ReplicationSystem/DeltaCompression/DeltaCompressionBaselineManager.h"
diff --git a/Engine/Source/Runtime/Experimental/Iris/Core/Private/Iris/ReplicationSystem/ReplicationWriter.cpp b/Engine/Source/Runtime/Experimental/Iris/Core/Private/Iris/ReplicationSystem/ReplicationWriter.cpp
--- a/Engine/Source/Runtime/Experimental/Iris/Core/Private/Iris/ReplicationSystem/ReplicationWriter.cpp
+++ b/Engine/Source/Runtime/Experimental/Iris/Core/Private/Iris/ReplicationSystem/ReplicationWriter.cpp
@@ -12,7 +12,7 @@
 #include "Iris/ReplicationSystem/ReplicationSystemTypes.h"
 #include "Iris/ReplicationSystem/DeltaCompression/DeltaCompressionBaselineManager.h"
 #include "Iris/ReplicationSystem/Filtering/ReplicationFiltering.h"
-#include "Iris/ReplicationSystem/NetHandleManager.h"
+#include "Iris/ReplicationSystem/NetRefHandleManager.h"
 #include "Iris/ReplicationSystem/ReplicationOperations.h"
 #include "Iris/ReplicationSystem/ReplicationOperationsInternal.h"
 #include "Iris/ReplicationSystem/ReplicationSystem.h"
diff --git a/Engine/Source/Runtime/Experimental/Iris/Core/Private/Iris/ReplicationSystem/ReplicationWriter.h b/Engine/Source/Runtime/Experimental/Iris/Core/Private/Iris/ReplicationSystem/ReplicationWriter.h
--- a/Engine/Source/Runtime/Experimental/Iris/Core/Private/Iris/ReplicationSystem/ReplicationWriter.h
+++ b/Engine/Source/Runtime/Experimental/Iris/Core/Private/Iris/ReplicationSystem/ReplicationWriter.h
@@ -4,7 +4,7 @@
 
 // This class will never be included from public headers
 #include "Net/Core/NetBitArray.h"
-#include "Iris/ReplicationSystem/NetHandleManager.h"
+#include "Iris/ReplicationSystem/NetRefHandleManager.h"
 #include "Iris/ReplicationSystem/ReplicationTypes.h"
 #include "Iris/ReplicationSystem/ReplicationRecord.h"
 #include "Iris/DataStream/DataStream.h"
diff --git a/Engine/Source/Runtime/Experimental/Iris/Core/Private/Iris/Serialization/ObjectNetSerializer.cpp b/Engine/Source/Runtime/Experimental/Iris/Core/Private/Iris/Serialization/ObjectNetSerializer.cpp
--- a/Engine/Source/Runtime/Experimental/Iris/Core/Private/Iris/Serialization/ObjectNetSerializer.cpp
+++ b/Engine/Source/Runtime/Experimental/Iris/Core/Private/Iris/Serialization/ObjectNetSerializer.cpp
@@ -8,7 +8,7 @@
 #include "Iris/Serialization/NetBitStreamReader.h"
 #include "Iris/Serialization/NetBitStreamWriter.h"
 #include "Iris/Serialization/InternalNetSerializationContext.h"
-#include "Iris/ReplicationSystem/NetHandleManager.h"
+#include "Iris/ReplicationSystem/NetRefHandleManager.h"
 #include "Iris/ReplicationSystem/ObjectReferenceCache.h"
 #include "Iris/Serialization/NetBitStreamUtil.h"
 
