Iris - Renamed Iris file NetHandleManager.cpp to NetRefHandleManager.cpp

#rnx

diff --git a/Engine/Source/Runtime/Experimental/Iris/Core/Private/Iris/ReplicationSystem/NetHandleManager.cpp b/Engine/Source/Runtime/Experimental/Iris/Core/Private/Iris/ReplicationSystem/NetHandleManager.cpp
deleted file mode 100644
diff --git a/Engine/Source/Runtime/Experimental/Iris/Core/Private/Iris/ReplicationSystem/NetRefHandleManager.cpp b/Engine/Source/Runtime/Experimental/Iris/Core/Private/Iris/ReplicationSystem/NetRefHandleManager.cpp
new file mode 100644
--- /dev/null
+++ b/Engine/Source/Runtime/Experimental/Iris/Core/Private/Iris/ReplicationSystem/NetRefHandleManager.cpp
@@ -0,0 +1,768 @@
+// Copyright Epic Games, Inc. All Rights Reserved.
+
+#include "NetRefHandleManager.h"
+#include "ReplicationOperationsInternal.h"
+#include "ReplicationProtocol.h"
+#include "ReplicationProtocolManager.h"
+#include "Iris/Core/IrisLog.h"
+#include "Iris/Serialization/NetSerializationContext.h"
+#include "Iris/Serialization/InternalNetSerializationContext.h"
+#include "Net/Core/Trace/NetTrace.h"
+#include "UObject/CoreNetTypes.h"
+#include "UObject/Object.h"
+#include "UObject/UObjectGlobals.h"
+#include "UObject/CoreNetTypes.h"
+
+namespace UE::Net::Private
+{
+
+FNetRefHandleManager::FNetRefHandleManager(FReplicationProtocolManager& InReplicationProtocolManager, uint32 InReplicationSystemId, uint32 InMaxActiveObjectCount)
+: ActiveObjectCount(0)
+, MaxActiveObjectCount(InMaxActiveObjectCount)
+, ReplicationSystemId(InReplicationSystemId)
+, ScopableInternalIndices(MaxActiveObjectCount)
+, PrevFrameScopableInternalIndices(MaxActiveObjectCount)
+, AssignedInternalIndices(MaxActiveObjectCount)
+, SubObjectInternalIndices(MaxActiveObjectCount)
+, DependentObjectInternalIndices(MaxActiveObjectCount)
+, ObjectsWithDependentObjectsInternalIndices(MaxActiveObjectCount)
+, DestroyedStartupObjectInternalIndices(MaxActiveObjectCount)
+, WantToBeDormantInternalIndices(MaxActiveObjectCount)
+, NextStaticHandleIndex(1U) // Index 0 is always reserved, for both static and dynamic handles
+, NextDynamicHandleIndex(1U)
+, ReplicationProtocolManager(InReplicationProtocolManager)
+{
+	static_assert(InvalidInternalIndex == 0, "FNetRefHandleManager::InvalidInternalIndex has an unexpected value");
+	// Mark the invalid index as used
+	AssignedInternalIndices.SetBit(0);
+
+	ReplicatedObjectData.SetNumUninitialized(MaxActiveObjectCount);
+	ReplicatedObjectRefCount.SetNumZeroed(MaxActiveObjectCount);
+	ReplicatedObjectStateBuffers.SetNumZeroed(MaxActiveObjectCount);
+	
+	ReplicatedInstances.SetNumZeroed(MaxActiveObjectCount);
+
+	// For convenience we initialize ReplicatedObjectData for InvalidInternalIndex so that GetReplicatedObjectDataNoCheck returns something useful.
+	ReplicatedObjectData[InvalidInternalIndex] = FReplicatedObjectData();
+}
+
+uint32 FNetRefHandleManager::GetNextNetRefHandleId(uint32 HandleId) const
+{
+	// Since we use the lowest bit in the index to indicate if the handle is static or dynamic we can not use all bits as the index
+	constexpr uint32 NetHandleIdIndexBitMask = (1U << (FNetRefHandle::IdBits - 1U)) - 1U;
+
+	uint32 NextHandleId = (HandleId + 1) & NetHandleIdIndexBitMask;
+	if (NextHandleId == 0)
+	{
+		++NextHandleId;
+	}
+	return NextHandleId;
+}
+
+FInternalNetRefIndex FNetRefHandleManager::GetNextFreeInternalIndex() const
+{
+	const uint32 NextFreeIndex = AssignedInternalIndices.FindFirstZero();
+	return NextFreeIndex != FNetBitArray::InvalidIndex ? NextFreeIndex : InvalidInternalIndex;
+}
+
+FInternalNetRefIndex FNetRefHandleManager::InternalCreateNetObject(const FNetRefHandle NetRefHandle, const FNetHandle GlobalHandle, const FReplicationProtocol* ReplicationProtocol)
+{
+	if (ActiveObjectCount >= MaxActiveObjectCount)
+	{
+		return InvalidInternalIndex;
+	}
+
+	// Verify that the handle is free
+	if (RefHandleToInternalIndex.Contains(NetRefHandle))
+	{
+		ensureAlwaysMsgf(false, TEXT("NetRefHandleManager::InternalCreateNetObject %s already exists"), *NetRefHandle.ToString());
+		return InvalidInternalIndex;
+	}
+
+	const uint32 InternalIndex = GetNextFreeInternalIndex();
+	if (InternalIndex != InvalidInternalIndex)
+	{
+		// Store data;
+		FReplicatedObjectData& Data = ReplicatedObjectData[InternalIndex];
+
+		Data = FReplicatedObjectData();
+
+		Data.RefHandle = NetRefHandle;
+		Data.NetHandle = GlobalHandle;
+		Data.Protocol = ReplicationProtocol;
+		Data.InstanceProtocol = nullptr;
+		ReplicatedObjectStateBuffers.GetData()[InternalIndex] = nullptr;
+		Data.ReceiveStateBuffer = nullptr;
+		Data.bShouldPropagateChangedStates = 1U;
+
+		++ActiveObjectCount;
+
+		// Add map entry so that we can map from NetRefHandle to InternalIndex
+		RefHandleToInternalIndex.Add(NetRefHandle, InternalIndex);
+		
+		// Add mapping from global handle to InternalIndex to speed up lookups for ReplicationSystem public API
+		if (GlobalHandle.IsValid())
+		{
+			NetHandleToInternalIndex.Add(GlobalHandle, InternalIndex);
+		}
+
+		// Mark Handle index as assigned and scopable for now
+		AssignedInternalIndices.SetBit(InternalIndex);
+		ScopableInternalIndices.SetBit(InternalIndex);
+
+		// When a handle is first created, it is not set to be a subobject
+		SubObjectInternalIndices.ClearBit(InternalIndex);
+
+		ReplicatedObjectRefCount[InternalIndex] = 0;
+
+		return InternalIndex;
+	}
+		
+	return InvalidInternalIndex;
+}
+
+void FNetRefHandleManager::AttachInstanceProtocol(FInternalNetRefIndex InternalIndex, const FReplicationInstanceProtocol* InstanceProtocol, UObject* Instance)
+{
+	if (ensureAlways((InternalIndex != InvalidInternalIndex) && InstanceProtocol))
+	{
+		FReplicatedObjectData& Data = ReplicatedObjectData[InternalIndex];
+		Data.InstanceProtocol = InstanceProtocol;
+
+		check(ReplicatedInstances[InternalIndex] == nullptr);
+		ReplicatedInstances[InternalIndex] = Instance;
+	}
+}
+
+const FReplicationInstanceProtocol* FNetRefHandleManager::DetachInstanceProtocol(FInternalNetRefIndex InternalIndex)
+{
+	if (ensure(InternalIndex != InvalidInternalIndex))
+	{
+		FReplicatedObjectData& Data = ReplicatedObjectData[InternalIndex];
+		const FReplicationInstanceProtocol* InstanceProtocol = Data.InstanceProtocol;
+		Data.InstanceProtocol = nullptr;
+		ReplicatedInstances[InternalIndex] = nullptr;
+		
+		return InstanceProtocol;
+	}
+
+	return nullptr;
+}
+
+FNetRefHandle FNetRefHandleManager::AllocateNetRefHandle(bool bIsStatic)
+{
+	uint32& NextHandleId = bIsStatic ? NextStaticHandleIndex : NextDynamicHandleIndex;
+
+	const uint32 NewHandleId = MakeNetRefHandleId(NextHandleId, bIsStatic);
+	FNetRefHandle NewHandle = MakeNetRefHandle(NewHandleId, ReplicationSystemId);
+
+	// Verify that the handle is free
+	if (RefHandleToInternalIndex.Contains(NewHandle))
+	{
+		checkf(false, TEXT("FNetRefHandleManager::AllocateNetHandle - Handle %s already exists!"), *NewHandle.ToString());
+
+		return FNetRefHandle();
+	}
+
+	// Bump NextHandleId
+	NextHandleId = GetNextNetRefHandleId(NextHandleId);
+
+	return NewHandle;
+}
+
+FNetRefHandle FNetRefHandleManager::CreateNetObject(FNetRefHandle WantedHandle, FNetHandle GlobalHandle, const FReplicationProtocol* ReplicationProtocol)
+{
+	FNetRefHandle NetRefHandle = WantedHandle;
+
+	const FInternalNetRefIndex InternalIndex = InternalCreateNetObject(NetRefHandle, GlobalHandle, ReplicationProtocol);
+	if (InternalIndex != InvalidInternalIndex)
+	{
+		FReplicatedObjectData& Data = ReplicatedObjectData[InternalIndex];
+		
+		// Allocate storage for outgoing data. We need a valid pointer even if the size is zero.
+		uint8* StateBuffer = (uint8*)FMemory::MallocZeroed(FPlatformMath::Max(ReplicationProtocol->InternalTotalSize, 1U), ReplicationProtocol->InternalTotalAlignment);
+		if (EnumHasAnyFlags(ReplicationProtocol->ProtocolTraits, EReplicationProtocolTraits::HasConditionalChangeMask))
+		{
+			// Enable all conditions by default. This is to be compatible with the implementation of FRepChangedPropertyTracker where we have hooks into ReplicationSystem.
+			FNetBitArrayView ConditionalChangeMask(reinterpret_cast<uint32*>(StateBuffer + ReplicationProtocol->GetConditionalChangeMaskOffset()), ReplicationProtocol->ChangeMaskBitCount, FNetBitArrayView::NoResetNoValidate);
+			ConditionalChangeMask.SetAllBits();
+		}
+
+		ReplicatedObjectStateBuffers.GetData()[InternalIndex] = StateBuffer;
+
+		// Bump protocol refcount
+		ReplicationProtocol->AddRef();
+
+		return NetRefHandle;
+	}
+	else
+	{
+		return FNetRefHandle();
+	}
+}
+
+// Create NetRefHandle not owned by us
+FNetRefHandle FNetRefHandleManager::CreateNetObjectFromRemote(FNetRefHandle WantedHandle, const FReplicationProtocol* ReplicationProtocol)
+{
+	if (!ensureAlwaysMsgf(WantedHandle.IsValid() && !WantedHandle.IsCompleteHandle(), TEXT("FNetRefHandleManager::CreateNetObjectFromRemote Expected WantedHandle %s to be valid and incomplete"), *WantedHandle.ToString()))
+	{
+		return FNetRefHandle();
+	}
+
+	FNetRefHandle NetRefHandle = MakeNetRefHandle(WantedHandle.GetId(), ReplicationSystemId);
+
+	const FInternalNetRefIndex InternalIndex = InternalCreateNetObject(NetRefHandle, FNetHandle(), ReplicationProtocol);
+	if (InternalIndex != InvalidInternalIndex)
+	{
+		FReplicatedObjectData& Data = ReplicatedObjectData[InternalIndex];
+
+		// Allocate storage for incoming data
+		Data.ReceiveStateBuffer = (uint8*)FMemory::Malloc(ReplicationProtocol->InternalTotalSize, ReplicationProtocol->InternalTotalAlignment);
+
+		// Since we currently do not have any default values, just initialize to zero
+		// N.B. If this Memzero is optimized away for some reason it still needs to be done for protocols with dynamic state.
+		FMemory::Memzero(Data.ReceiveStateBuffer, ReplicationProtocol->InternalTotalSize);
+
+		// Note: We could initialize this from default but at the moment it is part of the contract for all serializers to write the value when we serialize and we will only apply dirty states
+
+		// Don't bother initializing the conditional changemask if present since it's currently unused on the receiving end.
+
+		// Bump protocol refcount
+		ReplicationProtocol->AddRef();
+
+		return NetRefHandle;
+	}
+	else
+	{
+		return FNetRefHandle();
+	}
+}
+
+void FNetRefHandleManager::InternalDestroyNetObject(FInternalNetRefIndex InternalIndex)
+{
+	FReplicatedObjectData& Data = ReplicatedObjectData[InternalIndex];
+
+	UE_LOG(LogIris, Verbose, TEXT("FNetRefHandleManager::InternalDestroyNetObject ( InternalIndex: %u ) %s"), InternalIndex, *Data.RefHandle.ToString());
+
+	uint8* StateBuffer = ReplicatedObjectStateBuffers.GetData()[InternalIndex];
+	// Free any allocated resources
+	if (EnumHasAnyFlags(Data.Protocol->ProtocolTraits, EReplicationProtocolTraits::HasDynamicState))
+	{
+		FNetSerializationContext FreeContext;
+		FInternalNetSerializationContext InternalContext;
+		FreeContext.SetInternalContext(&InternalContext);
+		if (StateBuffer != nullptr)
+		{
+			FReplicationProtocolOperationsInternal::FreeDynamicState(FreeContext, StateBuffer, Data.Protocol);
+		}
+		if (Data.ReceiveStateBuffer != nullptr)
+		{
+			FReplicationProtocolOperationsInternal::FreeDynamicState(FreeContext, Data.ReceiveStateBuffer, Data.Protocol);
+		}
+	}
+	
+	// If we have subobjects remove them from our list
+	if (FNetDependencyData::FInternalNetRefIndexArray* SubObjectArray = SubObjects.GetInternalIndexArray<FNetDependencyData::EArrayType::SubObjects>(InternalIndex))
+	{
+		for (FInternalNetRefIndex SubObjectInternalIndex : *SubObjectArray)
+		{
+			InternalRemoveSubObject(InternalIndex, SubObjectInternalIndex, false);
+		}
+		SubObjectArray->Reset(0U);
+	}
+	// Clear ChildSubObjectArray
+	if (FNetDependencyData::FInternalNetRefIndexArray* ChildSubObjectArray = SubObjects.GetInternalIndexArray<FNetDependencyData::EArrayType::ChildSubObjects>(InternalIndex))
+	{
+		ChildSubObjectArray->Reset(0U);
+	}
+
+	// If we are a subobject remove from owner and hierarchical parents
+	if (Data.SubObjectRootIndex != InvalidInternalIndex)
+	{
+		InternalRemoveSubObject(Data.SubObjectRootIndex, InternalIndex);
+	}
+
+	// Remove from all dependent object relationships and clear data
+	InternalRemoveDependentObject(InternalIndex);
+
+	// Free all stored data for object
+	SubObjects.FreeStoredDependencyDataForObject(InternalIndex);
+
+	// Decrease protocol refcount
+	Data.Protocol->Release();
+	if (Data.Protocol->GetRefCount() == 0)
+	{
+		ReplicationProtocolManager.DestroyReplicationProtocol(Data.Protocol);
+	}
+
+	FMemory::Free(StateBuffer);
+	FMemory::Free(Data.ReceiveStateBuffer);
+
+	// Clear pointer to state buffer
+	ReplicatedObjectStateBuffers.GetData()[InternalIndex] = nullptr;
+
+	UE_NET_TRACE_NETHANDLE_DESTROYED(Data.RefHandle);
+
+	Data = FReplicatedObjectData();
+
+	// tracking
+	AssignedInternalIndices.ClearBit(InternalIndex);
+
+	// Restore internal state
+	SubObjectInternalIndices.ClearBit(InternalIndex);
+	WantToBeDormantInternalIndices.ClearBit(InternalIndex);
+	ObjectsWithDependentObjectsInternalIndices.ClearBit(InternalIndex);
+
+	// Cleanup cross reference to destruction info
+	if (DestroyedStartupObjectInternalIndices.GetBit(InternalIndex))
+	{
+		DestroyedStartupObjectInternalIndices.ClearBit(InternalIndex);
+		uint32 OtherInternalIndex = 0U;
+		if (DestroyedStartupObject.RemoveAndCopyValue(InternalIndex, OtherInternalIndex))
+		{
+			DestroyedStartupObject.Remove(OtherInternalIndex);
+		}
+	}
+
+	--ActiveObjectCount;
+}
+
+FNetRefHandle FNetRefHandleManager::CreateHandleForDestructionInfo(FNetRefHandle Handle, const FReplicationProtocol* DestroyedObjectProtocol)
+{
+	// Create destruction info handle carrying destruction info
+	constexpr bool bIsStaticHandle = false;
+	FNetRefHandle AllocatedHandle = AllocateNetRefHandle(bIsStaticHandle);
+	FNetRefHandle DestructionInfoHandle = CreateNetObject(AllocatedHandle, FNetHandle(), DestroyedObjectProtocol);
+
+	if (DestructionInfoHandle.IsValid())
+	{
+		const uint32 InternalIndex = GetInternalIndex(DestructionInfoHandle);
+		const uint32 DestroyedInternalIndex = GetInternalIndex(Handle);
+		
+		// mark the internal index
+		DestroyedStartupObjectInternalIndices.SetBit(InternalIndex);
+				
+		// Is the object replicated, if so we must make sure that we do not add it to scope by accident
+		if (DestroyedInternalIndex)
+		{			
+			DestroyedStartupObject.FindOrAdd(InternalIndex, DestroyedInternalIndex);
+		
+			// Mark the replicated index as destroyed
+			DestroyedStartupObjectInternalIndices.SetBit(DestroyedInternalIndex);
+			DestroyedStartupObject.FindOrAdd(DestroyedInternalIndex, InternalIndex);
+		}
+	}
+
+	return DestructionInfoHandle;
+}
+
+void FNetRefHandleManager::RemoveFromScope(FInternalNetRefIndex InternalIndex)
+{
+	// Can only remove an object from scope if it is assignable
+	if (ensure(AssignedInternalIndices.GetBit(InternalIndex)))
+	{
+		ScopableInternalIndices.ClearBit(InternalIndex);
+	}
+}
+
+void FNetRefHandleManager::DestroyNetObject(FNetRefHandle RefHandle)
+{
+	FInternalNetRefIndex InternalIndex = RefHandleToInternalIndex.FindAndRemoveChecked(RefHandle);
+	if (ensure(AssignedInternalIndices.GetBit(InternalIndex)))
+	{
+		FReplicatedObjectData& Data = ReplicatedObjectData[InternalIndex];
+		check(Data.RefHandle == RefHandle);
+
+		// Remove mapping from global handle to internal index
+		NetHandleToInternalIndex.Remove(Data.NetHandle);
+
+		// Remove from scopable objects if not already done
+		ScopableInternalIndices.ClearBit(InternalIndex);
+
+		// We always defer the actual destroy
+		PendingDestroyInternalIndices.Add(InternalIndex);
+	}
+}
+
+void FNetRefHandleManager::DestroyObjectsPendingDestroy()
+{
+	// Destroy Objects pending destroy
+	for (int32 It = PendingDestroyInternalIndices.Num() - 1; It >= 0; --It)
+	{
+		const FInternalNetRefIndex InternalIndex = PendingDestroyInternalIndices[It];
+		if (ReplicatedObjectRefCount[InternalIndex] == 0)
+		{
+			InternalDestroyNetObject(InternalIndex);
+			PendingDestroyInternalIndices.RemoveAtSwap(It);
+		}
+	}
+}
+
+bool FNetRefHandleManager::AddSubObject(FNetRefHandle OwnerHandle, FNetRefHandle SubObjectHandle, FNetRefHandle RelativeOtherSubObjectHandle, EAddSubObjectFlags Flags)
+{
+	check(OwnerHandle != SubObjectHandle);
+
+	// validate objects
+	const FInternalNetRefIndex OwnerInternalIndex = GetInternalIndex(OwnerHandle);
+	const FInternalNetRefIndex SubObjectInternalIndex = GetInternalIndex(SubObjectHandle);
+
+	const bool bIsValidOwner = ensure(OwnerInternalIndex != InvalidInternalIndex);
+	const bool bIsValidSubObejct = ensure(SubObjectInternalIndex != InvalidInternalIndex);
+
+	if (!(bIsValidOwner && bIsValidSubObejct))
+	{
+		return false;
+	}
+
+	const FInternalNetRefIndex RelativeOtherSubObjectInternalIndex = EnumHasAnyFlags(Flags, EAddSubObjectFlags::ReplicateWithSubObject) ? GetInternalIndex(RelativeOtherSubObjectHandle) : InvalidInternalIndex;
+	return InternalAddSubObject(OwnerInternalIndex, SubObjectInternalIndex, RelativeOtherSubObjectInternalIndex, Flags);
+}
+
+bool FNetRefHandleManager::InternalAddSubObject(FInternalNetRefIndex OwnerInternalIndex, FInternalNetRefIndex SubObjectInternalIndex, FInternalNetRefIndex RelativeOtherSubObjectInternalIndex, EAddSubObjectFlags Flags)
+{
+	using namespace UE::Net::Private;
+
+	FReplicatedObjectData& SubObjectData = GetReplicatedObjectDataNoCheck(SubObjectInternalIndex);
+	if (!ensureAlwaysMsgf(SubObjectData.SubObjectRootIndex == InvalidInternalIndex, TEXT("FNetRefHandleManager::AddSubObject %s is already marked as a subobject"), ToCStr(SubObjectData.RefHandle.ToString())))
+	{
+		return false;
+	}
+
+	FNetDependencyData::FInternalNetRefIndexArray& SubObjectArray = SubObjects.GetOrCreateInternalIndexArray<FNetDependencyData::SubObjects>(OwnerInternalIndex);
+	SubObjectArray.Add(SubObjectInternalIndex);
+	SubObjectData.SubObjectRootIndex = OwnerInternalIndex;
+	SubObjectData.bDestroySubObjectWithOwner = EnumHasAnyFlags(Flags, EAddSubObjectFlags::DestroyWithOwner);
+	// Mark the object as a subobject
+	SetIsSubObject(SubObjectInternalIndex, true);
+
+	if (RelativeOtherSubObjectInternalIndex != InvalidInternalIndex && ensureAlwaysMsgf(SubObjectArray.Find(RelativeOtherSubObjectInternalIndex) != INDEX_NONE, TEXT("RelativeOtherSubObjectHandle %s Must be a Subobject of %s"), ToCStr(GetNetRefHandleFromInternalIndex(RelativeOtherSubObjectInternalIndex).ToString()), ToCStr(GetNetRefHandleFromInternalIndex(OwnerInternalIndex).ToString())))
+	{
+		// Add to child array of RelativeOtherSubObjectRelativeIndex for hierarchical replication order
+		FNetDependencyData::FSubObjectConditionalsArray* SubObjectConditionalsArray = nullptr;
+		FNetDependencyData::FInternalNetRefIndexArray& ChildSubObjectArray = SubObjects.GetOrCreateInternalChildSubObjectsArray(RelativeOtherSubObjectInternalIndex, SubObjectConditionalsArray);
+		ChildSubObjectArray.Add(SubObjectInternalIndex);
+		if (SubObjectConditionalsArray)
+		{
+			SubObjectConditionalsArray->Add(ELifetimeCondition::COND_None);
+		}
+		SubObjectData.SubObjectParentIndex = RelativeOtherSubObjectInternalIndex;
+	}
+	else
+	{
+		// Add to child array of root
+		FNetDependencyData::FSubObjectConditionalsArray* SubObjectConditionalsArray = nullptr;
+		FNetDependencyData::FInternalNetRefIndexArray& ChildSubObjectArray = SubObjects.GetOrCreateInternalChildSubObjectsArray(OwnerInternalIndex, SubObjectConditionalsArray);
+		ChildSubObjectArray.Add(SubObjectInternalIndex);
+		if (SubObjectConditionalsArray)
+		{
+			SubObjectConditionalsArray->Add(ELifetimeCondition::COND_None);
+		}
+		SubObjectData.SubObjectParentIndex = OwnerInternalIndex;
+	}
+
+	return true;
+}
+
+void FNetRefHandleManager::InternalRemoveSubObject(FInternalNetRefIndex OwnerInternalIndex, FInternalNetRefIndex SubObjectInternalIndex, bool bRemoveFromSubObjectArray)
+{
+	// both must be valid
+	if (OwnerInternalIndex != InvalidInternalIndex && SubObjectInternalIndex != InvalidInternalIndex)
+	{
+		FReplicatedObjectData& SubObjectData = GetReplicatedObjectDataNoCheck(SubObjectInternalIndex);
+		check(SubObjectData.SubObjectRootIndex == OwnerInternalIndex);
+		
+		if (bRemoveFromSubObjectArray)
+		{
+			// Remove from root parent
+			if (FNetDependencyData::FInternalNetRefIndexArray* SubObjectArray = SubObjects.GetInternalIndexArray<FNetDependencyData::EArrayType::SubObjects>(OwnerInternalIndex))
+			{
+				SubObjectArray->Remove(SubObjectInternalIndex);
+			}
+			// Remove from parents ChildSubObjectArray
+			if (SubObjectData.SubObjectParentIndex != InvalidInternalIndex)
+			{
+				FNetDependencyData::FInternalNetRefIndexArray* ChildSubObjectArray;
+				FNetDependencyData::FSubObjectConditionalsArray* SubObjectConditionsArray;
+
+				if (SubObjects.GetInternalChildSubObjectAndConditionalArrays(SubObjectData.SubObjectParentIndex, ChildSubObjectArray, SubObjectConditionsArray))
+				{
+					const int32 ArrayIndex = ChildSubObjectArray->Find(SubObjectInternalIndex);
+					if (ensureAlways(ArrayIndex != INDEX_NONE))
+					{
+						ChildSubObjectArray->RemoveAt(ArrayIndex);
+						if (SubObjectConditionsArray)
+						{
+							SubObjectConditionsArray->RemoveAt(ArrayIndex);
+							check(SubObjectConditionsArray->Num() == ChildSubObjectArray->Num());
+						}
+					}
+				}
+			}
+		}
+
+		SubObjectData.SubObjectRootIndex = InvalidInternalIndex;
+		SubObjectData.bDestroySubObjectWithOwner = false;
+
+		SetIsSubObject(SubObjectInternalIndex, false);
+	}
+}
+
+void FNetRefHandleManager::RemoveSubObject(FNetRefHandle Handle)
+{
+	FInternalNetRefIndex SubObjectInternalIndex = GetInternalIndex(Handle);
+	checkSlow(SubObjectInternalIndex);
+
+	if (SubObjectInternalIndex != InvalidInternalIndex)
+	{
+		const uint32 OwnerInternalIndex = ReplicatedObjectData[SubObjectInternalIndex].SubObjectRootIndex;
+		if (OwnerInternalIndex != InvalidInternalIndex)
+		{
+			InternalRemoveSubObject(OwnerInternalIndex, SubObjectInternalIndex);
+		}
+	}
+}
+
+bool FNetRefHandleManager::SetSubObjectNetCondition(FInternalNetRefIndex SubObjectInternalIndex, FLifeTimeConditionStorage SubObjectCondition)
+{
+	if (ensure(SubObjectInternalIndex != InvalidInternalIndex))
+	{
+		const FInternalNetRefIndex SubObjectParentIndex = ReplicatedObjectData[SubObjectInternalIndex].SubObjectParentIndex;
+		if (ensure(SubObjectParentIndex != InvalidInternalIndex))
+		{
+			FNetDependencyData::FInternalNetRefIndexArray* SubObjectsArray;
+			FNetDependencyData::FSubObjectConditionalsArray* SubObjectConditionals;
+			
+			if (SubObjects.GetInternalChildSubObjectAndConditionalArrays(SubObjectParentIndex, SubObjectsArray, SubObjectConditionals))
+			{
+				const int32 SubObjectArrayIndex = SubObjectsArray->Find(SubObjectInternalIndex);
+				if (ensure(SubObjectArrayIndex != INDEX_NONE))
+				{
+					if (!SubObjectConditionals)
+					{
+						SubObjectConditionals = &SubObjects.GetOrCreateSubObjectConditionalsArray(SubObjectParentIndex);
+					}
+
+					check(SubObjectConditionals->Num() == SubObjectsArray->Num());
+
+					SubObjectConditionals->GetData()[SubObjectArrayIndex] = SubObjectCondition;
+					return true;
+				}
+			}
+
+		}
+	}
+
+	return false;
+}
+
+FNetRefHandle FNetRefHandleManager::GetSubObjectOwner(FNetRefHandle SubObjectRefHandle) const
+{
+	const FInternalNetRefIndex SubObjectInternalIndex = GetInternalIndex(SubObjectRefHandle);
+	const FInternalNetRefIndex OwnerInternalIndex = SubObjectInternalIndex != InvalidInternalIndex ? ReplicatedObjectData[SubObjectInternalIndex].SubObjectRootIndex : InvalidInternalIndex;
+
+	return OwnerInternalIndex != InvalidInternalIndex ? ReplicatedObjectData[OwnerInternalIndex].RefHandle : FNetRefHandle();
+}
+
+bool FNetRefHandleManager::AddDependentObject(FNetRefHandle ParentRefHandle, FNetRefHandle DependentObjectRefHandle, EAddDependentObjectFlags Flags)
+{
+	check(ParentRefHandle != DependentObjectRefHandle);
+
+	// validate objects
+	FInternalNetRefIndex ParentInternalIndex = GetInternalIndex(ParentRefHandle);
+	FInternalNetRefIndex DependentObjectInternalIndex = GetInternalIndex(DependentObjectRefHandle);
+
+	const bool bIsValidOwner = ensure(ParentInternalIndex != InvalidInternalIndex);
+	const bool bIsValidDependentObject = ensure(DependentObjectInternalIndex != InvalidInternalIndex);
+
+	if (!(bIsValidOwner && bIsValidDependentObject))
+	{
+		return false;
+	}
+
+	FReplicatedObjectData& DependentObjectData = GetReplicatedObjectDataNoCheck(DependentObjectInternalIndex);
+	FReplicatedObjectData& ParentObjectData = GetReplicatedObjectDataNoCheck(ParentInternalIndex);
+
+	// SubObjects cannot have dependent objects or be a dependent object (for now)
+	check(!DependentObjectData.IsSubObject() && !ParentObjectData.IsSubObject());
+	check(!SubObjectInternalIndices.GetBit(DependentObjectInternalIndex));
+	check(!SubObjectInternalIndices.GetBit(ParentInternalIndex));
+
+	// Add dependent to parents dependent object list
+	FNetDependencyData::FInternalNetRefIndexArray& ParentDependentObjectArray = SubObjects.GetOrCreateInternalIndexArray<FNetDependencyData::DependentObjects>(ParentInternalIndex);
+	const bool bDependentIsAlreadyDependant = ParentDependentObjectArray.Find(DependentObjectInternalIndex) != INDEX_NONE;
+	if (!bDependentIsAlreadyDependant)
+	{
+		ParentDependentObjectArray.Add(DependentObjectInternalIndex);
+	}
+
+	FNetDependencyData::FInternalNetRefIndexArray& DependentParentObjectArray = SubObjects.GetOrCreateInternalIndexArray<FNetDependencyData::ParentObjects>(DependentObjectInternalIndex);
+	const bool bDependentHadParentAlready = DependentParentObjectArray.Find(ParentInternalIndex) != INDEX_NONE;
+	if (!bDependentHadParentAlready)
+	{
+		DependentParentObjectArray.Add(ParentInternalIndex);
+	}
+
+	// Update cached info to avoid to do map lookups to find out if we have are a dependent object or have dependent objects
+	DependentObjectData.bIsDependentObject = true;
+	ParentObjectData.bHasDependentObjects = true;
+	ObjectsWithDependentObjectsInternalIndices.SetBit(ParentInternalIndex);
+	DependentObjectInternalIndices.SetBit(DependentObjectInternalIndex);
+
+	if (EnumHasAnyFlags(Flags, EAddDependentObjectFlags::WarnIfAlreadyDependentObject) && (bDependentHadParentAlready || bDependentIsAlreadyDependant))
+	{
+		// If this gets out of sync something is messed up
+		check(bDependentHadParentAlready == bDependentIsAlreadyDependant);
+		ensureAlwaysMsgf(false, TEXT("FNetRefHandleManager::AddDependentObject %s already is a child of %s parent %s"), *DependentObjectRefHandle.ToString(), *GetReplicatedObjectDataNoCheck(DependentObjectData.SubObjectRootIndex).RefHandle.ToString(), *ParentRefHandle.ToString());
+	}
+
+	return true;
+}
+
+void FNetRefHandleManager::RemoveDependentObject(FNetRefHandle DependentHandle)
+{
+	FInternalNetRefIndex DependentInternalIndex = GetInternalIndex(DependentHandle);
+
+	if (DependentInternalIndex != InvalidInternalIndex)
+	{
+		InternalRemoveDependentObject(DependentInternalIndex);
+	}
+}
+
+void FNetRefHandleManager::InternalRemoveDependentObject(FInternalNetRefIndex ParentInternalIndex, FInternalNetRefIndex DependentInternalIndex, ERemoveDependentObjectFlags Flags)
+{
+	if (EnumHasAnyFlags(Flags, ERemoveDependentObjectFlags::RemoveFromDependentParentObjects))
+	{
+		if (FNetDependencyData::FInternalNetRefIndexArray* ParentObjectArray = SubObjects.GetInternalIndexArray<FNetDependencyData::ParentObjects>(DependentInternalIndex))
+		{
+			ParentObjectArray->Remove(ParentInternalIndex);
+			if (ParentObjectArray->Num() == 0)
+			{
+				FReplicatedObjectData& DependentObjectData = GetReplicatedObjectDataNoCheck(DependentInternalIndex);
+				DependentObjectData.bIsDependentObject = false;
+				DependentObjectInternalIndices.ClearBit(DependentInternalIndex);
+			}
+		}
+	}
+
+	if (EnumHasAnyFlags(Flags, ERemoveDependentObjectFlags::RemoveFromParentDependentObjects))
+	{
+		if (FNetDependencyData::FInternalNetRefIndexArray* ParentDependentObjectArray = SubObjects.GetInternalIndexArray<FNetDependencyData::DependentObjects>(ParentInternalIndex))
+		{
+			ParentDependentObjectArray->Remove(DependentInternalIndex);
+			if (ParentDependentObjectArray->Num() == 0)
+			{
+				FReplicatedObjectData& ParentObjectData = GetReplicatedObjectDataNoCheck(ParentInternalIndex);
+				ParentObjectData.bHasDependentObjects = false;
+				ObjectsWithDependentObjectsInternalIndices.ClearBit(ParentInternalIndex);
+			}
+		}
+	}
+}
+
+void FNetRefHandleManager::InternalRemoveDependentObject(FInternalNetRefIndex DependentInternalIndex)
+{
+	// Remove from all parents
+	if (FNetDependencyData::FInternalNetRefIndexArray* ParentObjectArray = SubObjects.GetInternalIndexArray<FNetDependencyData::ParentObjects>(DependentInternalIndex))
+	{
+		for (FInternalNetRefIndex ParentInternalIndex : *ParentObjectArray)
+		{
+			// Flag is set to only update data on the parent to avoid modifying the array we iterate over
+			InternalRemoveDependentObject(ParentInternalIndex, DependentInternalIndex, ERemoveDependentObjectFlags::RemoveFromParentDependentObjects);
+		}
+		ParentObjectArray->Reset();
+	}
+
+	// Remove from our dependents
+	if (FNetDependencyData::FInternalNetRefIndexArray* DependentObjectArray = SubObjects.GetInternalIndexArray<FNetDependencyData::DependentObjects>(DependentInternalIndex))
+	{
+		for (FInternalNetRefIndex ChildDepedentInternalIndex : *DependentObjectArray)
+		{
+			// Flag is set to only update data on the childDependentObject to avoid modifying the array we iterate over
+			InternalRemoveDependentObject(DependentInternalIndex, ChildDepedentInternalIndex, ERemoveDependentObjectFlags::RemoveFromDependentParentObjects);
+		}
+		DependentObjectArray->Reset();		
+	}
+
+	FReplicatedObjectData& DependentObjectData = GetReplicatedObjectDataNoCheck(DependentInternalIndex);
+
+	// Clear out flags on this object	
+	DependentObjectData.bIsDependentObject = false;
+	DependentObjectData.bHasDependentObjects = false;
+	ObjectsWithDependentObjectsInternalIndices.ClearBit(DependentInternalIndex);
+	DependentObjectInternalIndices.ClearBit(DependentInternalIndex);
+}
+
+void FNetRefHandleManager::RemoveDependentObject(FNetRefHandle ParentHandle, FNetRefHandle DependentHandle)
+{
+	// Validate objects
+	FInternalNetRefIndex ParentInternalIndex = GetInternalIndex(ParentHandle);
+	FInternalNetRefIndex DependentInternalIndex = GetInternalIndex(DependentHandle);
+
+	if ((ParentInternalIndex == InvalidInternalIndex) || (DependentInternalIndex == InvalidInternalIndex))
+	{
+		return;
+	}
+
+	InternalRemoveDependentObject(ParentInternalIndex, DependentInternalIndex);
+}
+
+void FNetRefHandleManager::SetShouldPropagateChangedStates(FInternalNetRefIndex ObjectInternalIndex, bool bShouldPropagateChangedStates)
+{
+	if (ObjectInternalIndex != InvalidInternalIndex)
+	{
+		if (bShouldPropagateChangedStates)
+		{
+			// Currently we do not support re-enabling state propagation
+			// $IRIS: $TODO: Implement method to force dirty all changes 
+			// https://jira.it.epicgames.com/browse/UE-127368
+
+			checkf(false, TEXT("Re-enabling state change propagation is currently Not implemented."));			
+			return;
+		}
+
+		ReplicatedObjectData[ObjectInternalIndex].bShouldPropagateChangedStates = bShouldPropagateChangedStates ? 1U : 0U;
+	}
+}
+
+void FNetRefHandleManager::SetShouldPropagateChangedStates(FNetRefHandle Handle, bool bShouldPropagateChangedStates)
+{
+	FInternalNetRefIndex ObjectInternalIndex = GetInternalIndex(Handle);
+	return SetShouldPropagateChangedStates(ObjectInternalIndex, bShouldPropagateChangedStates);
+}
+
+void FNetRefHandleManager::AddReferencedObjects(FReferenceCollector& Collector)
+{
+	Collector.AddReferencedObjects(ReplicatedInstances);
+}
+
+uint32 FNetRefHandleManager::MakeNetRefHandleId(uint32 Id, bool bIsStatic)
+{
+	return (Id << 1U) | (bIsStatic ? 1U : 0U);
+}
+
+FNetRefHandle FNetRefHandleManager::MakeNetRefHandle(uint32 Id, uint32 ReplicationSystemId)
+{
+	check((Id & FNetRefHandle::IdMask) == Id);
+	check(ReplicationSystemId < FNetRefHandle::MaxReplicationSystemId);
+
+	FNetRefHandle Handle;
+
+	Handle.Id = Id;
+	Handle.ReplicationSystemId = ReplicationSystemId + 1U;
+
+	return Handle;
+}
+
+FNetRefHandle FNetRefHandleManager::MakeNetRefHandleFromId(uint32 Id)
+{
+	check((Id & FNetRefHandle::IdMask) == Id);
+
+	FNetRefHandle Handle;
+
+	Handle.Id = Id;
+	Handle.ReplicationSystemId = 0U;
+
+	return Handle;
+}
+
+
+}
