Iris - Renamed Iris file NetHandle.h to NetRefHandle.h

#rnx

diff --git a/Engine/Plugins/Runtime/ReplicationSystemTestPlugin/Source/Private/Tests/ReplicationSystem/ReplicatedTestObject.cpp b/Engine/Plugins/Runtime/ReplicationSystemTestPlugin/Source/Private/Tests/ReplicationSystem/ReplicatedTestObject.cpp
--- a/Engine/Plugins/Runtime/ReplicationSystemTestPlugin/Source/Private/Tests/ReplicationSystem/ReplicatedTestObject.cpp
+++ b/Engine/Plugins/Runtime/ReplicationSystemTestPlugin/Source/Private/Tests/ReplicationSystem/ReplicatedTestObject.cpp
@@ -4,7 +4,7 @@
 #include "NetworkAutomationTest.h"
 #include "NetworkAutomationTestMacros.h"
 #include "Iris/ReplicationState/PropertyReplicationState.h"
-#include "Iris/ReplicationSystem/NetHandle.h"
+#include "Iris/ReplicationSystem/NetRefHandle.h"
 #include "Iris/ReplicationSystem/RepTag.h"
 #include "Iris/ReplicationSystem/ReplicationProtocol.h"
 #include "Iris/ReplicationSystem/ReplicationBridge.h"
diff --git a/Engine/Plugins/Runtime/ReplicationSystemTestPlugin/Source/Private/Tests/ReplicationSystem/ReplicatedTestObject.h b/Engine/Plugins/Runtime/ReplicationSystemTestPlugin/Source/Private/Tests/ReplicationSystem/ReplicatedTestObject.h
--- a/Engine/Plugins/Runtime/ReplicationSystemTestPlugin/Source/Private/Tests/ReplicationSystem/ReplicatedTestObject.h
+++ b/Engine/Plugins/Runtime/ReplicationSystemTestPlugin/Source/Private/Tests/ReplicationSystem/ReplicatedTestObject.h
@@ -5,7 +5,7 @@
 #include "NetworkAutomationTest.h"
 #include "UObject/ObjectMacros.h"
 #include "UObject/StrongObjectPtr.h"
-#include "Iris/ReplicationSystem/NetHandle.h"
+#include "Iris/ReplicationSystem/NetRefHandle.h"
 #include "Iris/ReplicationSystem/ReplicationProtocol.h"
 #include "Iris/ReplicationSystem/ObjectReplicationBridge.h"
 #include "Iris/ReplicationSystem/PropertyReplicationFragment.h"
diff --git a/Engine/Source/Runtime/Experimental/Iris/Core/Private/Iris/ReplicationSystem/NetBlob/NetBlobHandlerDefinitions.h b/Engine/Source/Runtime/Experimental/Iris/Core/Private/Iris/ReplicationSystem/NetBlob/NetBlobHandlerDefinitions.h
--- a/Engine/Source/Runtime/Experimental/Iris/Core/Private/Iris/ReplicationSystem/NetBlob/NetBlobHandlerDefinitions.h
+++ b/Engine/Source/Runtime/Experimental/Iris/Core/Private/Iris/ReplicationSystem/NetBlob/NetBlobHandlerDefinitions.h
@@ -2,7 +2,7 @@
 
 #pragma once
 #include "CoreTypes.h"
-#include "Iris/ReplicationSystem/NetHandle.h"
+#include "Iris/ReplicationSystem/NetRefHandle.h"
 #include "NetBlobHandlerDefinitions.generated.h"
 
 namespace UE::Net::Private
diff --git a/Engine/Source/Runtime/Experimental/Iris/Core/Private/Iris/ReplicationSystem/NetHandle.cpp b/Engine/Source/Runtime/Experimental/Iris/Core/Private/Iris/ReplicationSystem/NetHandle.cpp
--- a/Engine/Source/Runtime/Experimental/Iris/Core/Private/Iris/ReplicationSystem/NetHandle.cpp
+++ b/Engine/Source/Runtime/Experimental/Iris/Core/Private/Iris/ReplicationSystem/NetHandle.cpp
@@ -1,6 +1,6 @@
 // Copyright Epic Games, Inc. All Rights Reserved.
 
-#include "Iris/ReplicationSystem/NetHandle.h"
+#include "Iris/ReplicationSystem/NetRefHandle.h"
 #include "Containers/UnrealString.h"
 #include "Misc/StringBuilder.h"
 
diff --git a/Engine/Source/Runtime/Experimental/Iris/Core/Private/Iris/ReplicationSystem/NetRefHandleManager.h b/Engine/Source/Runtime/Experimental/Iris/Core/Private/Iris/ReplicationSystem/NetRefHandleManager.h
--- a/Engine/Source/Runtime/Experimental/Iris/Core/Private/Iris/ReplicationSystem/NetRefHandleManager.h
+++ b/Engine/Source/Runtime/Experimental/Iris/Core/Private/Iris/ReplicationSystem/NetRefHandleManager.h
@@ -9,7 +9,7 @@
 #include "Net/Core/NetBitArray.h"
 #include "Net/Core/NetHandle/NetHandle.h"
 #include "Iris/ReplicationState/ReplicationStateDescriptor.h"
-#include "Iris/ReplicationSystem/NetHandle.h"
+#include "Iris/ReplicationSystem/NetRefHandle.h"
 #include "Iris/ReplicationSystem/NetDependencyData.h"
 
 class FReferenceCollector;
diff --git a/Engine/Source/Runtime/Experimental/Iris/Core/Private/Iris/ReplicationSystem/ObjectReferenceCache.h b/Engine/Source/Runtime/Experimental/Iris/Core/Private/Iris/ReplicationSystem/ObjectReferenceCache.h
--- a/Engine/Source/Runtime/Experimental/Iris/Core/Private/Iris/ReplicationSystem/ObjectReferenceCache.h
+++ b/Engine/Source/Runtime/Experimental/Iris/Core/Private/Iris/ReplicationSystem/ObjectReferenceCache.h
@@ -4,7 +4,7 @@
 
 #include "CoreTypes.h"
 #include "Iris/Core/NetObjectReference.h"
-#include "Iris/ReplicationSystem/NetHandle.h"
+#include "Iris/ReplicationSystem/NetRefHandle.h"
 #include "UObject/WeakObjectPtr.h"
 #include "Containers/Map.h"
 #include "ObjectReferenceCacheFwd.h"
diff --git a/Engine/Source/Runtime/Experimental/Iris/Core/Private/Iris/ReplicationSystem/ReplicationWriter.h b/Engine/Source/Runtime/Experimental/Iris/Core/Private/Iris/ReplicationSystem/ReplicationWriter.h
--- a/Engine/Source/Runtime/Experimental/Iris/Core/Private/Iris/ReplicationSystem/ReplicationWriter.h
+++ b/Engine/Source/Runtime/Experimental/Iris/Core/Private/Iris/ReplicationSystem/ReplicationWriter.h
@@ -10,7 +10,7 @@
 #include "Iris/DataStream/DataStream.h"
 #include "Iris/ReplicationSystem/AttachmentReplication.h"
 #include "Iris/ReplicationSystem/NetBlob/NetBlob.h"
-#include "Iris/ReplicationSystem/NetHandle.h"
+#include "Iris/ReplicationSystem/NetRefHandle.h"
 #include "Iris/ReplicationSystem/NetExports.h"
 #include "Iris/ReplicationSystem/ObjectReferenceCache.h"
 #include "Iris/Stats/NetStats.h"
diff --git a/Engine/Source/Runtime/Experimental/Iris/Core/Private/Iris/Serialization/NetExportContext.h b/Engine/Source/Runtime/Experimental/Iris/Core/Private/Iris/Serialization/NetExportContext.h
--- a/Engine/Source/Runtime/Experimental/Iris/Core/Private/Iris/Serialization/NetExportContext.h
+++ b/Engine/Source/Runtime/Experimental/Iris/Core/Private/Iris/Serialization/NetExportContext.h
@@ -5,7 +5,7 @@
 #include "CoreTypes.h"
 #include "Containers/Array.h"
 #include "Containers/Set.h"
-#include "Iris/ReplicationSystem/NetHandle.h"
+#include "Iris/ReplicationSystem/NetRefHandle.h"
 #include "Iris/ReplicationSystem/NetToken.h"
 #include "Iris/Serialization/NetSerializationContext.h"
 
diff --git a/Engine/Source/Runtime/Experimental/Iris/Core/Public/Iris/Core/IrisDebugging.h b/Engine/Source/Runtime/Experimental/Iris/Core/Public/Iris/Core/IrisDebugging.h
--- a/Engine/Source/Runtime/Experimental/Iris/Core/Public/Iris/Core/IrisDebugging.h
+++ b/Engine/Source/Runtime/Experimental/Iris/Core/Public/Iris/Core/IrisDebugging.h
@@ -3,7 +3,7 @@
 #pragma once
 
 #include "HAL/Platform.h"
-#include "Iris/ReplicationSystem/NetHandle.h"
+#include "Iris/ReplicationSystem/NetRefHandle.h"
 #include "Net/Core/NetHandle/NetHandle.h"
 
 class UReplicationSystem;
diff --git a/Engine/Source/Runtime/Experimental/Iris/Core/Public/Iris/Core/NetObjectReference.h b/Engine/Source/Runtime/Experimental/Iris/Core/Public/Iris/Core/NetObjectReference.h
--- a/Engine/Source/Runtime/Experimental/Iris/Core/Public/Iris/Core/NetObjectReference.h
+++ b/Engine/Source/Runtime/Experimental/Iris/Core/Public/Iris/Core/NetObjectReference.h
@@ -3,7 +3,7 @@
 #pragma once
 
 #include "CoreTypes.h"
-#include "Iris/ReplicationSystem/NetHandle.h"
+#include "Iris/ReplicationSystem/NetRefHandle.h"
 #include "Iris/ReplicationSystem/NetToken.h"
 #include "Misc/EnumClassFlags.h"
 #include "Templates/IsPODType.h"
diff --git a/Engine/Source/Runtime/Experimental/Iris/Core/Public/Iris/ReplicationSystem/NetBlob/ReliableNetBlobQueue.h b/Engine/Source/Runtime/Experimental/Iris/Core/Public/Iris/ReplicationSystem/NetBlob/ReliableNetBlobQueue.h
--- a/Engine/Source/Runtime/Experimental/Iris/Core/Public/Iris/ReplicationSystem/NetBlob/ReliableNetBlobQueue.h
+++ b/Engine/Source/Runtime/Experimental/Iris/Core/Public/Iris/ReplicationSystem/NetBlob/ReliableNetBlobQueue.h
@@ -2,7 +2,7 @@
 
 #pragma once
 #include "CoreTypes.h"
-#include "Iris/ReplicationSystem/NetHandle.h"
+#include "Iris/ReplicationSystem/NetRefHandle.h"
 #include "Iris/PacketControl/PacketNotification.h"
 #include "Templates/RefCounting.h"
 
diff --git a/Engine/Source/Runtime/Experimental/Iris/Core/Public/Iris/ReplicationSystem/NetHandle.h b/Engine/Source/Runtime/Experimental/Iris/Core/Public/Iris/ReplicationSystem/NetHandle.h
deleted file mode 100644
diff --git a/Engine/Source/Runtime/Experimental/Iris/Core/Public/Iris/ReplicationSystem/NetRefHandle.h b/Engine/Source/Runtime/Experimental/Iris/Core/Public/Iris/ReplicationSystem/NetRefHandle.h
new file mode 100644
--- /dev/null
+++ b/Engine/Source/Runtime/Experimental/Iris/Core/Public/Iris/ReplicationSystem/NetRefHandle.h
@@ -0,0 +1,83 @@
+// Copyright Epic Games, Inc. All Rights Reserved.
+
+#pragma once
+
+#include "CoreTypes.h"
+#include "Iris/IrisConfig.h"
+#include "Containers/StringFwd.h"
+#include "Templates/TypeHash.h"
+
+// Forward declarations
+class FString;
+
+namespace UE::Net::Private
+{
+	class FNetRefHandleManager;
+}
+
+namespace UE::Net
+{
+
+/**
+ * FNetRefHandle
+ */
+class FNetRefHandle
+{
+public:
+	enum { Invalid = 0 };
+	enum { IdBits = 32 };
+	enum { ReplicationSystemIdBits = 4 };
+
+public:
+	FNetRefHandle() : Value(Invalid) {}
+
+	uint32 GetId() const { return Id; }
+	uint32 GetReplicationSystemId() const { check(ReplicationSystemId != 0U); return ReplicationSystemId - 1U; }
+	bool IsValid() const { return Value != Invalid; }
+
+	bool IsCompleteHandle() const { return Value != Invalid && ReplicationSystemId != 0U; }
+
+	bool IsStatic() const { return Id & StaticIdMask; }
+	bool IsDynamic() const { return IsValid() && !IsStatic(); }
+
+	bool operator==(const FNetRefHandle& Other)const { return Id == Other.Id; }
+	bool operator<(const FNetRefHandle& Other)const { return Id < Other.Id; }
+	bool operator!=(const FNetRefHandle& Other)const { return Id != Other.Id; }
+
+	IRISCORE_API FString ToString() const;
+
+	static bool FullCompare(FNetRefHandle A, FNetRefHandle B) { return A.Value == B.Value; }
+	
+private:
+	friend uint32 GetTypeHash(const FNetRefHandle& Handle);
+	friend Private::FNetRefHandleManager;
+
+	static constexpr uint32 StaticIdMask = 1U;
+	static constexpr uint32 IdMask = ~0U;
+	static constexpr uint32 MaxReplicationSystemId = (1 << ReplicationSystemIdBits) - 1;
+
+	union 
+	{
+		struct
+		{
+			uint32 Id : IdBits;										// Id, lowest bit indicates if the handle is static or dynamic
+			uint32 ReplicationSystemId : ReplicationSystemIdBits;	// ReplicationSystemId, when running in pie, we track the owning instance
+		};
+		uint64 Value;
+	};
+};
+
+FORCEINLINE uint32 GetTypeHash(const FNetRefHandle& Handle)
+{
+	return ::GetTypeHash(Handle.GetId());
+}
+
+FORCEINLINE uint32 GetObjectIdForNetTrace(const FNetRefHandle& Handle)
+{
+	return Handle.GetId();
+}
+
+}
+
+IRISCORE_API FStringBuilderBase& operator<<(FStringBuilderBase& Builder, const UE::Net::FNetRefHandle& NetHandle);
+IRISCORE_API FAnsiStringBuilderBase& operator<<(FAnsiStringBuilderBase& Builder, const UE::Net::FNetRefHandle& NetHandle);
diff --git a/Engine/Source/Runtime/Experimental/Iris/Core/Public/Iris/ReplicationSystem/ReplicationBridge.h b/Engine/Source/Runtime/Experimental/Iris/Core/Public/Iris/ReplicationSystem/ReplicationBridge.h
--- a/Engine/Source/Runtime/Experimental/Iris/Core/Public/Iris/ReplicationSystem/ReplicationBridge.h
+++ b/Engine/Source/Runtime/Experimental/Iris/Core/Public/Iris/ReplicationSystem/ReplicationBridge.h
@@ -6,7 +6,7 @@
 #include "Containers/ContainersFwd.h"
 #include "Containers/Map.h"
 #include "Iris/Core/NetObjectReference.h"
-#include "Iris/ReplicationSystem/NetHandle.h"
+#include "Iris/ReplicationSystem/NetRefHandle.h"
 #include "Misc/EnumClassFlags.h"
 #include "Net/Core/NetHandle/NetHandle.h"
 #include "UObject/ObjectKey.h"
diff --git a/Engine/Source/Runtime/Experimental/Iris/Core/Public/Iris/ReplicationSystem/ReplicationSystem.h b/Engine/Source/Runtime/Experimental/Iris/Core/Public/Iris/ReplicationSystem/ReplicationSystem.h
--- a/Engine/Source/Runtime/Experimental/Iris/Core/Public/Iris/ReplicationSystem/ReplicationSystem.h
+++ b/Engine/Source/Runtime/Experimental/Iris/Core/Public/Iris/ReplicationSystem/ReplicationSystem.h
@@ -7,7 +7,7 @@
 #include "Templates/PimplPtr.h"
 #include "Templates/RefCounting.h"
 #include "Iris/IrisConfig.h"
-#include "Iris/ReplicationSystem/NetHandle.h"
+#include "Iris/ReplicationSystem/NetRefHandle.h"
 #include "Net/Core/NetHandle/NetHandle.h"
 
 #include "ReplicationSystem.generated.h"
diff --git a/Engine/Source/Runtime/Experimental/Iris/Core/Public/Iris/Serialization/ObjectNetSerializer.h b/Engine/Source/Runtime/Experimental/Iris/Core/Public/Iris/Serialization/ObjectNetSerializer.h
--- a/Engine/Source/Runtime/Experimental/Iris/Core/Public/Iris/Serialization/ObjectNetSerializer.h
+++ b/Engine/Source/Runtime/Experimental/Iris/Core/Public/Iris/Serialization/ObjectNetSerializer.h
@@ -3,7 +3,7 @@
 #pragma once
 
 #include "NetSerializer.h"
-#include "Iris/ReplicationSystem/NetHandle.h"
+#include "Iris/ReplicationSystem/NetRefHandle.h"
 #include "UObject/ObjectMacros.h"
 
 #include "ObjectNetSerializer.generated.h"
