Iris - Renamed Iris file NetHandle.cpp to NetRefHandle.cpp

#rnx

diff --git a/Engine/Source/Runtime/Experimental/Iris/Core/Private/Iris/ReplicationSystem/NetHandle.cpp b/Engine/Source/Runtime/Experimental/Iris/Core/Private/Iris/ReplicationSystem/NetHandle.cpp
deleted file mode 100644
diff --git a/Engine/Source/Runtime/Experimental/Iris/Core/Private/Iris/ReplicationSystem/NetRefHandle.cpp b/Engine/Source/Runtime/Experimental/Iris/Core/Private/Iris/ReplicationSystem/NetRefHandle.cpp
new file mode 100644
--- /dev/null
+++ b/Engine/Source/Runtime/Experimental/Iris/Core/Private/Iris/ReplicationSystem/NetRefHandle.cpp
@@ -0,0 +1,44 @@
+// Copyright Epic Games, Inc. All Rights Reserved.
+
+#include "Iris/ReplicationSystem/NetRefHandle.h"
+#include "Containers/UnrealString.h"
+#include "Misc/StringBuilder.h"
+
+namespace UE::Net
+{
+
+FString FNetRefHandle::ToString() const
+{
+	FString Result;
+#if UE_NET_ALLOW_MULTIPLE_REPLICATION_SYSTEMS
+	const uint32 ReplicationSystemIdToDisplay = ReplicationSystemId - 1U;
+	Result = FString::Printf(TEXT("NetRefHandle (Id=%u):(RepSystemId=%u)"), GetId(), ReplicationSystemIdToDisplay);
+#else
+	Result = FString::Printf(TEXT("NetRefHandle (Id=%u)"), GetId());
+#endif
+	return Result;
+}
+
+}
+
+FStringBuilderBase& operator<<(FStringBuilderBase& Builder, const UE::Net::FNetRefHandle& NetRefHandle) 
+{ 	
+#if UE_NET_ALLOW_MULTIPLE_REPLICATION_SYSTEMS
+	const uint32 ReplicationSystemIdToDisplay = NetRefHandle.GetReplicationSystemId() - 1U;
+	Builder.Appendf(TEXT("NetRefHandle (Id=%u):(RepSystemId=%u)"), NetRefHandle.GetId(), ReplicationSystemIdToDisplay);
+#else
+	Builder.Appendf(TEXT("NetRefHandle (Id=%u)"), NetRefHandle.GetId());
+#endif
+	return Builder;
+}
+
+FAnsiStringBuilderBase& operator<<(FAnsiStringBuilderBase& Builder, const UE::Net::FNetRefHandle& NetRefHandle)
+{
+#if UE_NET_ALLOW_MULTIPLE_REPLICATION_SYSTEMS
+	const uint32 ReplicationSystemIdToDisplay = NetRefHandle.GetReplicationSystemId() - 1U;
+	Builder.Appendf("NetRefHandle (Id=%u):(RepSystemId=%u)", NetRefHandle.GetId(), ReplicationSystemIdToDisplay);
+#else
+	Builder.Appendf("NetRefHandle (Id=%u)", NetRefHandle.GetId());
+#endif
+	return Builder;
+}
