Constraints: avoid possible endless loop when check dependencies

diff --git a/Engine/Source/Runtime/Experimental/Animation/Constraints/Private/TransformConstraint.cpp b/Engine/Source/Runtime/Experimental/Animation/Constraints/Private/TransformConstraint.cpp
--- a/Engine/Source/Runtime/Experimental/Animation/Constraints/Private/TransformConstraint.cpp
+++ b/Engine/Source/Runtime/Experimental/Animation/Constraints/Private/TransformConstraint.cpp
@@ -13,39 +13,6 @@
 namespace
 {
 
-// ensure that InPossiblePrimary is not depending on InPossibleSecondary to avoid creating cycles 
-bool HasPrerequisiteDependencyWith(const FTickFunction* InSecondary, const FTickFunction* InPrimary)
-{
-	if (!InSecondary || !InPrimary)
-	{
-		return false;
-	}
-
-	// is InSecondary a Prereq of InPrimary?
-	const TArray<FTickPrerequisite>& Prerequisites = InPrimary->GetPrerequisites();
-	const bool bIsSecondaryAPrereq = Prerequisites.ContainsByPredicate([InSecondary](const FTickPrerequisite& Prereq)
-	{
-		const FTickFunction* PrereqFunction = Prereq.Get();
-		return PrereqFunction && (PrereqFunction == InSecondary);
-	});
-
-	if (bIsSecondaryAPrereq)
-	{
-		return true;
-	}
-
-	// otherwise, recurse
-	for (const FTickPrerequisite& Prerequisite : Prerequisites)
-	{
-		if (HasPrerequisiteDependencyWith(InSecondary, Prerequisite.Get()))
-		{
-			return true;
-		}
-	}
-
-	return false;
-}
-
 struct FConstraintCycleChecker
 {
 public:
@@ -60,8 +27,9 @@ public:
 			return false;
 		}
 
+		TSet<const FTickFunction*> VisitedFunctions;
 		const FTickFunction* TickFunction = InHandle->GetTickFunction();
-		return HasPrerequisiteDependencyWith(TickFunction, TickFunction);
+		return HasPrerequisiteDependencyWith(TickFunction, TickFunction, VisitedFunctions);
 	}
 
 	/** Checks for cycling constraints and manage tick dependencies if needed to avoid cycles from a tick dependency pov. */
@@ -94,6 +62,12 @@ public:
 				return false;
 			}
 
+			const TObjectPtr<UTransformableHandle>& ChildHandle = TransformConst->ChildTRSHandle;
+			if (!IsValid(ChildHandle) || !ChildHandle->IsValid())
+			{
+				return false;
+			}
+
 			const FTickFunction* ParentTickFunction = TransformConst->GetParentHandleTickFunction();
 			return ParentTickFunction && ParentTickFunction == ChildTickFunction;
 		};
@@ -106,9 +80,10 @@ public:
 		}
 
 		// check if they can cause a cycle and manage dependencies if that's the case
+		TSet<const FTickFunction*> VisitedFunctions;
 		for (const ConstraintPtr& Constraint: CyclingConstraints)
 		{
-			if (HasPrerequisiteDependencyWith(&Constraint->ConstraintTick, &InConstraint->ConstraintTick))
+			if (HasPrerequisiteDependencyWith(&Constraint->ConstraintTick, &InConstraint->ConstraintTick, VisitedFunctions))
 			{
 				UpdateCyclingDependency(World, Cast<UTickableTransformConstraint>(Constraint));
 			}
@@ -117,6 +92,46 @@ public:
 	
 private:
 
+	// ensure that InPossiblePrimary is not depending on InPossibleSecondary to avoid creating cycles 
+	static bool HasPrerequisiteDependencyWith(const FTickFunction* InSecondary, const FTickFunction* InPrimary, TSet<const FTickFunction*>& InOutVisitedFunctions)
+	{
+		if (!InSecondary || !InPrimary)
+		{
+			return false;
+		}
+
+		// is InSecondary a Prereq of InPrimary?
+		const TArray<FTickPrerequisite>& Prerequisites = InPrimary->GetPrerequisites();
+		const bool bIsSecondaryAPrereq = Prerequisites.ContainsByPredicate([InSecondary](const FTickPrerequisite& Prereq)
+		{
+			const FTickFunction* PrereqFunction = Prereq.Get();
+			return PrereqFunction && (PrereqFunction == InSecondary);
+		});
+
+		if (bIsSecondaryAPrereq)
+		{
+			return true;
+		}
+
+		// check if InPrimary has already been visited to avoid endless loop
+		if (InOutVisitedFunctions.Contains(InPrimary))
+		{
+			return false;
+		}
+		InOutVisitedFunctions.Add(InPrimary);
+		
+		// otherwise, recurse
+		for (const FTickPrerequisite& Prerequisite : Prerequisites)
+		{
+			if (HasPrerequisiteDependencyWith(InSecondary, Prerequisite.Get(), InOutVisitedFunctions))
+			{
+				return true;
+			}
+		}
+
+		return false;
+	}
+
 	/**
 	 * Manage tick dependencies if needed to avoid cycles from a tick dependency pov.
 	 * Both InConstraintToUpdate and its parent handle are supposed valid at this point
