Add new 'contact threshold' concept to fracture proximity graph to filter out connections that have insufficient 'contact' + add an initial contact method that uses overlap of projection on major axes.

diff --git a/Engine/Plugins/Experimental/ChaosEditor/Source/FractureEditor/Private/FractureModeSettings.cpp b/Engine/Plugins/Experimental/ChaosEditor/Source/FractureEditor/Private/FractureModeSettings.cpp
--- a/Engine/Plugins/Experimental/ChaosEditor/Source/FractureEditor/Private/FractureModeSettings.cpp
+++ b/Engine/Plugins/Experimental/ChaosEditor/Source/FractureEditor/Private/FractureModeSettings.cpp
@@ -25,6 +25,8 @@ void UFractureModeSettings::ApplyDefaultProximitySettings(FGeometryCollection& G
 	Properties.bUseAsConnectionGraph = bProximityUseAsConnectionGraph;
 	Properties.DistanceThreshold = ProximityDistanceThreshold;
 	Properties.Method = ProximityMethod;
+	Properties.ContactMethod = ProximityContactMethod;
+	Properties.RequireContactAmount = ProximityContactThreshold;
 	GeometryCollection.SetProximityProperties(Properties);
 }
 
diff --git a/Engine/Plugins/Experimental/ChaosEditor/Source/FractureEditor/Private/FractureToolProximity.cpp b/Engine/Plugins/Experimental/ChaosEditor/Source/FractureEditor/Private/FractureToolProximity.cpp
--- a/Engine/Plugins/Experimental/ChaosEditor/Source/FractureEditor/Private/FractureToolProximity.cpp
+++ b/Engine/Plugins/Experimental/ChaosEditor/Source/FractureEditor/Private/FractureToolProximity.cpp
@@ -22,6 +22,8 @@ void UFractureProximityActions::SaveAsDefaults()
 	ModeSettings->ProximityMethod = ProximityTool->ProximitySettings->Method;
 	ModeSettings->ProximityDistanceThreshold = ProximityTool->ProximitySettings->DistanceThreshold;
 	ModeSettings->bProximityUseAsConnectionGraph = ProximityTool->ProximitySettings->bUseAsConnectionGraph;
+	ModeSettings->ProximityContactThreshold = ProximityTool->ProximitySettings->ContactThreshold;
+	ModeSettings->ProximityContactMethod = ProximityTool->ProximitySettings->ContactMethod;
 }
 
 void UFractureProximityActions::SetFromDefaults()
@@ -31,6 +33,8 @@ void UFractureProximityActions::SetFromDefaults()
 	ProximityTool->ProximitySettings->Method = ModeSettings->ProximityMethod;
 	ProximityTool->ProximitySettings->DistanceThreshold = ModeSettings->ProximityDistanceThreshold;
 	ProximityTool->ProximitySettings->bUseAsConnectionGraph = ModeSettings->bProximityUseAsConnectionGraph;
+	ProximityTool->ProximitySettings->ContactThreshold = ModeSettings->ProximityContactThreshold;
+	ProximityTool->ProximitySettings->ContactMethod = ModeSettings->ProximityContactMethod;
 
 	ProximityTool->NotifyOfPropertyChangeByTool(this);
 }
@@ -219,6 +223,8 @@ int32 UFractureToolProximity::ExecuteFracture(const FFractureToolContext& Fractu
 		Properties.Method = ProximitySettings->Method;
 		Properties.DistanceThreshold = ProximitySettings->DistanceThreshold;
 		Properties.bUseAsConnectionGraph = ProximitySettings->bUseAsConnectionGraph;
+		Properties.ContactMethod = ProximitySettings->ContactMethod;
+		Properties.RequireContactAmount = ProximitySettings->ContactThreshold;
 		
 		FractureContext.GetGeometryCollection()->SetProximityProperties(Properties);
 
diff --git a/Engine/Plugins/Experimental/ChaosEditor/Source/FractureEditor/Private/FractureToolProximity.h b/Engine/Plugins/Experimental/ChaosEditor/Source/FractureEditor/Private/FractureToolProximity.h
--- a/Engine/Plugins/Experimental/ChaosEditor/Source/FractureEditor/Private/FractureToolProximity.h
+++ b/Engine/Plugins/Experimental/ChaosEditor/Source/FractureEditor/Private/FractureToolProximity.h
@@ -30,6 +30,14 @@ public:
 	UPROPERTY(EditAnywhere, Category = Automatic, meta = (ClampMin = "0", EditCondition = "Method == EProximityMethod::ConvexHull"))
 	double DistanceThreshold = 1;
 
+	// Method to use to determine the contact between two pieces, if the Contact Threshold is greater than 0
+	UPROPERTY(EditAnywhere, Category = Automatic)
+	EProximityContactMethod ContactMethod = EProximityContactMethod::MinOverlapInProjectionToMajorAxes;
+
+	// If greater than zero, proximity will be additionally filtered by a 'contact' threshold, in cm, to exclude grazing / corner proximity
+	UPROPERTY(EditAnywhere, Category = Automatic, meta = (ClampMin = "0"))
+	double ContactThreshold = 0;
+
 	// Whether to automatically transform the proximity graph into a connection graph to be used for simulation
 	UPROPERTY(EditAnywhere, Category = Automatic)
 	bool bUseAsConnectionGraph = false;
diff --git a/Engine/Plugins/Experimental/ChaosEditor/Source/FractureEditor/Public/FractureModeSettings.h b/Engine/Plugins/Experimental/ChaosEditor/Source/FractureEditor/Public/FractureModeSettings.h
--- a/Engine/Plugins/Experimental/ChaosEditor/Source/FractureEditor/Public/FractureModeSettings.h
+++ b/Engine/Plugins/Experimental/ChaosEditor/Source/FractureEditor/Public/FractureModeSettings.h
@@ -65,6 +65,14 @@ public:
 	/** Whether to automatically transform the proximity graph into a connection graph to be used for simulation */
 	UPROPERTY(config, EditAnywhere, Category = "Fracture Mode|Proximity Detection Defaults", meta = (DisplayName = "Use As Connection Graph"))
 	bool bProximityUseAsConnectionGraph = false;
+
+	// Method to use to determine the contact between two pieces, if the Contact Threshold is greater than 0
+	UPROPERTY(config, EditAnywhere, Category = "Fracture Mode|Proximity Detection Defaults", meta = (DisplayName = "Contact Method"))
+	EProximityContactMethod ProximityContactMethod = EProximityContactMethod::MinOverlapInProjectionToMajorAxes;
+
+	// If greater than zero, proximity will be additionally filtered by a 'contact' threshold, in cm, to exclude grazing / corner proximity
+	UPROPERTY(config, EditAnywhere, Category = "Fracture Mode|Proximity Detection Defaults", meta = (ClampMin = "0", DisplayName = "Contact Threshold"))
+	float ProximityContactThreshold = 0;
 	
 
 	// Apply Convex Generation Defaults to a GeometryCollection
diff --git a/Engine/Plugins/Experimental/GeometryCollectionPlugin/Source/GeometryCollectionNodes/Private/Dataflow/GeometryCollectionNodes.cpp b/Engine/Plugins/Experimental/GeometryCollectionPlugin/Source/GeometryCollectionNodes/Private/Dataflow/GeometryCollectionNodes.cpp
--- a/Engine/Plugins/Experimental/GeometryCollectionPlugin/Source/GeometryCollectionNodes/Private/Dataflow/GeometryCollectionNodes.cpp
+++ b/Engine/Plugins/Experimental/GeometryCollectionPlugin/Source/GeometryCollectionNodes/Private/Dataflow/GeometryCollectionNodes.cpp
@@ -1112,6 +1112,7 @@ void FProximityDataflowNode::Evaluate(Dataflow::FContext& Context, const FDatafl
 			Properties.Method = (EProximityMethod)ProximityMethod;
 			Properties.DistanceThreshold = DistanceThreshold;
 			Properties.bUseAsConnectionGraph = bUseAsConnectionGraph;
+			Properties.RequireContactAmount = ContactThreshold;
 
 			GeomCollection->SetProximityProperties(Properties);
 
diff --git a/Engine/Plugins/Experimental/GeometryCollectionPlugin/Source/GeometryCollectionNodes/Public/Dataflow/GeometryCollectionNodes.h b/Engine/Plugins/Experimental/GeometryCollectionPlugin/Source/GeometryCollectionNodes/Public/Dataflow/GeometryCollectionNodes.h
--- a/Engine/Plugins/Experimental/GeometryCollectionPlugin/Source/GeometryCollectionNodes/Public/Dataflow/GeometryCollectionNodes.h
+++ b/Engine/Plugins/Experimental/GeometryCollectionPlugin/Source/GeometryCollectionNodes/Public/Dataflow/GeometryCollectionNodes.h
@@ -1714,6 +1714,10 @@ public:
 	UPROPERTY(EditAnywhere, Category = "Proximity", meta = (ClampMin = "0", EditCondition = "ProximityMethod == EProximityMethodEnum::Dataflow_ProximityMethod_ConvexHull"))
 	float DistanceThreshold = 1;
 
+	// If greater than zero, proximity will be additionally filtered by a 'contact' threshold, in cm, to exclude grazing / corner proximity
+	UPROPERTY(EditAnywhere, Category = "Proximity", meta = (ClampMin = "0"))
+	float ContactThreshold = 0;
+
 	/** Whether to automatically transform the proximity graph into a connection graph to be used for simulation */
 	UPROPERTY(EditAnywhere, Category = "Proximity")
 	bool bUseAsConnectionGraph = false;
diff --git a/Engine/Source/Runtime/Experimental/Chaos/Private/GeometryCollection/GeometryCollectionProximityPropertiesInterface.cpp b/Engine/Source/Runtime/Experimental/Chaos/Private/GeometryCollection/GeometryCollectionProximityPropertiesInterface.cpp
--- a/Engine/Source/Runtime/Experimental/Chaos/Private/GeometryCollection/GeometryCollectionProximityPropertiesInterface.cpp
+++ b/Engine/Source/Runtime/Experimental/Chaos/Private/GeometryCollection/GeometryCollectionProximityPropertiesInterface.cpp
@@ -7,6 +7,8 @@ const FName FGeometryCollectionProximityPropertiesInterface::ProximityProperties
 const FName FGeometryCollectionProximityPropertiesInterface::ProximityDistanceThreshold = "DistanceThreshold";
 const FName FGeometryCollectionProximityPropertiesInterface::ProximityDetectionMethod = "DetectionMethod";
 const FName FGeometryCollectionProximityPropertiesInterface::ProximityAsConnectionGraph = "AsConnectionGraph";
+const FName FGeometryCollectionProximityPropertiesInterface::ProximityRequireContactAmount = "RequireContactAmount";
+const FName FGeometryCollectionProximityPropertiesInterface::ProximityContactMethod = "ContactMethod";
 
 FGeometryCollectionProximityPropertiesInterface::FGeometryCollectionProximityPropertiesInterface(FGeometryCollection* InGeometryCollection)
 	: FManagedArrayInterface(InGeometryCollection)
@@ -34,6 +36,16 @@ FGeometryCollectionProximityPropertiesInterface::InitializeInterface()
 	{
 		ManagedCollection->AddAttribute<bool>(ProximityAsConnectionGraph, ProximityPropertiesGroup);
 	}
+
+	if (!ManagedCollection->HasAttribute(ProximityContactMethod, ProximityPropertiesGroup))
+	{
+		ManagedCollection->AddAttribute<uint8>(ProximityContactMethod, ProximityPropertiesGroup);
+	}
+
+	if (!ManagedCollection->HasAttribute(ProximityRequireContactAmount, ProximityPropertiesGroup))
+	{
+		ManagedCollection->AddAttribute<float>(ProximityRequireContactAmount, ProximityPropertiesGroup);
+	}
 }
 
 void 
@@ -58,11 +70,15 @@ FGeometryCollectionProximityPropertiesInterface::GetProximityProperties() const
 		const TManagedArray<bool>& AsConnectionGraph = ManagedCollection->GetAttribute<bool>(ProximityAsConnectionGraph, ProximityPropertiesGroup);
 		const TManagedArray<int32>& DetectionMethod = ManagedCollection->GetAttribute<int32>(ProximityDetectionMethod, ProximityPropertiesGroup);
 		const TManagedArray<float>& DistanceThreshold = ManagedCollection->GetAttribute<float>(ProximityDistanceThreshold, ProximityPropertiesGroup);
+		const TManagedArray<uint8>& ContactMethod = ManagedCollection->GetAttribute<uint8>(ProximityContactMethod, ProximityPropertiesGroup);
+		const TManagedArray<float>& ContactAmount = ManagedCollection->GetAttribute<float>(ProximityRequireContactAmount, ProximityPropertiesGroup);
 
 		constexpr int32 DefaultIndex = 0;
 		Properties.bUseAsConnectionGraph = AsConnectionGraph[DefaultIndex];
 		Properties.DistanceThreshold = DistanceThreshold[DefaultIndex];
 		Properties.Method = (EProximityMethod)DetectionMethod[DefaultIndex];
+		Properties.ContactMethod = (EProximityContactMethod)ContactMethod[DefaultIndex];
+		Properties.RequireContactAmount = ContactAmount[DefaultIndex];
 	}
 	return Properties;
 }
@@ -86,8 +102,12 @@ FGeometryCollectionProximityPropertiesInterface::SetProximityProperties(const FP
 	TManagedArray<bool>& AsConnectionGraph = ManagedCollection->ModifyAttribute<bool>(ProximityAsConnectionGraph, ProximityPropertiesGroup);
 	TManagedArray<int32>& Method = ManagedCollection->ModifyAttribute<int32>(ProximityDetectionMethod, ProximityPropertiesGroup);
 	TManagedArray<float>& DistanceThreshold = ManagedCollection->ModifyAttribute<float>(ProximityDistanceThreshold, ProximityPropertiesGroup);
+	TManagedArray<uint8>& ContactMethod = ManagedCollection->ModifyAttribute<uint8>(ProximityContactMethod, ProximityPropertiesGroup);
+	TManagedArray<float>& ContactAmount = ManagedCollection->ModifyAttribute<float>(ProximityRequireContactAmount, ProximityPropertiesGroup);
 
 	AsConnectionGraph[AttributeIndex] = InProximityAttributes.bUseAsConnectionGraph;
 	DistanceThreshold[AttributeIndex] = InProximityAttributes.DistanceThreshold;
 	Method[AttributeIndex] = (int32)InProximityAttributes.Method;
+	ContactMethod[AttributeIndex] = (uint8)InProximityAttributes.ContactMethod;
+	ContactAmount[AttributeIndex] = InProximityAttributes.RequireContactAmount;
 }
diff --git a/Engine/Source/Runtime/Experimental/Chaos/Private/GeometryCollection/GeometryCollectionProximityUtility.cpp b/Engine/Source/Runtime/Experimental/Chaos/Private/GeometryCollection/GeometryCollectionProximityUtility.cpp
--- a/Engine/Source/Runtime/Experimental/Chaos/Private/GeometryCollection/GeometryCollectionProximityUtility.cpp
+++ b/Engine/Source/Runtime/Experimental/Chaos/Private/GeometryCollection/GeometryCollectionProximityUtility.cpp
@@ -626,6 +626,77 @@ void FGeometryCollectionProximityUtility::UpdateProximity(UE::GeometryCollection
 		Spatial.MoveProximityToCollection(Collection);
 	}
 
+	if (Properties.RequireContactAmount > 0.0f)
+	{
+		if (Properties.ContactMethod == EProximityContactMethod::MinOverlapInProjectionToMajorAxes)
+		{
+			TManagedArray<TSet<int32>>& Proximity = Collection->ModifyAttribute<TSet<int32>>("Proximity", FGeometryCollection::GeometryGroup);
+			int32 NumGeometry = Collection->NumElements(FGeometryCollection::GeometryGroup);
+
+			TArray<FBox> GeometryBounds; // Geometry bounding boxes in a shared space
+			GeometryBounds.Init(FBox(EForceInit::ForceInit), NumGeometry);
+
+			TArray<FTransform> GlobalTransformArray;
+			GeometryCollectionAlgo::GlobalMatrices(Collection->Transform, Collection->Parent, GlobalTransformArray);
+
+			ParallelFor(NumGeometry, [&](int32 GeometryIdx)
+			{
+				FBox& Box = GeometryBounds[GeometryIdx];
+				const int32 Start = Collection->VertexStart[GeometryIdx];
+				const int32 End = Start + Collection->VertexCount[GeometryIdx];
+				FTransform GeometryTransform = GlobalTransformArray[Collection->TransformIndex[GeometryIdx]];
+				for (int32 VertIdx = Start; VertIdx < End; ++VertIdx)
+				{
+					Box += GeometryTransform.TransformPosition((FVector)Collection->Vertex[VertIdx]);
+				}
+			});
+
+			TArray<int32> ToRemove;
+			auto ProjectBox = [](const FBox& Box, int32 Axis) -> FBox2D
+			{
+				int32 X = (Axis + 1) % 3;
+				int32 Y = (Axis + 2) % 3;
+				return FBox2D(FVector2D(Box.Min[X], Box.Min[Y]), FVector2D(Box.Max[X], Box.Max[Y]));
+			};
+			for (int32 GeometryIdx = 0; GeometryIdx < NumGeometry; ++GeometryIdx)
+			{
+				FBox Box = GeometryBounds[GeometryIdx];
+				ToRemove.Reset();
+				for (int32 ConnectedGeoIdx : Proximity[GeometryIdx])
+				{
+					FBox OtherBox = GeometryBounds[ConnectedGeoIdx];
+					bool bOverlapAnyAxis = false;
+					for (int ProjAxis = 0; ProjAxis < 3; ++ProjAxis)
+					{
+						FBox2D ProjA = ProjectBox(Box, ProjAxis);
+						FBox2D ProjB = ProjectBox(OtherBox, ProjAxis);
+						FBox2D Overlap = ProjA.Overlap(ProjB);
+						if (Overlap.bIsValid)
+						{
+							float MinBoundsAxis = (float)FMath::Min(ProjA.GetSize().GetMin(), ProjB.GetSize().GetMin());
+							float MinAxisOverlap = (float)Overlap.GetSize().GetMin();
+							// Overlap is accepted if it is greater than the threshold amount OR greater than half the maximum possible amount (to avoid 100% filtering small pieces)
+							if (MinAxisOverlap > FMath::Min(MinBoundsAxis*.5f, Properties.RequireContactAmount))
+							{
+								bOverlapAnyAxis = true;
+								break;
+							}
+						}
+					}
+					if (!bOverlapAnyAxis)
+					{
+						ToRemove.Add(ConnectedGeoIdx);
+					}
+				}
+				for (int32 Nbr : ToRemove)
+				{
+					Proximity[GeometryIdx].Remove(Nbr);
+					Proximity[Nbr].Remove(GeometryIdx);
+				}
+			}
+		}
+	}
+
 	if (Properties.bUseAsConnectionGraph)
 	{
 		CopyProximityToConnectionGraph();
diff --git a/Engine/Source/Runtime/Experimental/Chaos/Public/GeometryCollection/GeometryCollectionProximityPropertiesInterface.h b/Engine/Source/Runtime/Experimental/Chaos/Public/GeometryCollection/GeometryCollectionProximityPropertiesInterface.h
--- a/Engine/Source/Runtime/Experimental/Chaos/Public/GeometryCollection/GeometryCollectionProximityPropertiesInterface.h
+++ b/Engine/Source/Runtime/Experimental/Chaos/Public/GeometryCollection/GeometryCollectionProximityPropertiesInterface.h
@@ -23,12 +23,19 @@ public :
 	static const FName ProximityDistanceThreshold;
 	// Attribute: Whether to use the computed proximity graph as a connection graph
 	static const FName ProximityAsConnectionGraph;
-
+	// Attribute: If greater than zero, filter proximity connections by requiring an amount of 'contact' as computed by the Contact Method
+	// This is a second filter applied after initial proximity connections were determined by any Proximity Detection Method,
+	// and can be used to reduce unsupported or spurious/glancing proximities.
+	static const FName ProximityRequireContactAmount;
+	// Attribute: How to define 'contact' for the above Require Contact Amount.
+	static const FName ProximityContactMethod;
 
 	struct FProximityProperties
 	{
 		EProximityMethod Method;
 		float DistanceThreshold = 1.0f;
+		float RequireContactAmount = 0.0f;
+		EProximityContactMethod ContactMethod = EProximityContactMethod::MinOverlapInProjectionToMajorAxes;
 		bool bUseAsConnectionGraph = false;
 	};
 
diff --git a/Engine/Source/Runtime/Experimental/Chaos/Public/GeometryCollection/GeometryCollectionProximityUtility.h b/Engine/Source/Runtime/Experimental/Chaos/Public/GeometryCollection/GeometryCollectionProximityUtility.h
--- a/Engine/Source/Runtime/Experimental/Chaos/Public/GeometryCollection/GeometryCollectionProximityUtility.h
+++ b/Engine/Source/Runtime/Experimental/Chaos/Public/GeometryCollection/GeometryCollectionProximityUtility.h
@@ -20,6 +20,14 @@ enum class EProximityMethod : int32
 	ConvexHull
 };
 
+UENUM()
+enum class EProximityContactMethod : uint8
+{
+	// Rejects proximity if the bounding boxes do not overlap by more than this many centimeters in any major axis direction. This can filter out corner connections of box-like shapes.
+	MinOverlapInProjectionToMajorAxes
+	//~ TODO: Add other methods for filtering overlaps, e.g. based on approximate surface area of the contact
+};
+
 class CHAOS_API FGeometryCollectionProximityUtility
 {
 public:
