Chaos - fix joint projection disabled

A recent bug woudl skip projection initialization but still run projection when projection is disabled.

diff --git a/Engine/Source/Runtime/Experimental/Chaos/Private/Chaos/Joint/PBDJointContainerSolver.cpp b/Engine/Source/Runtime/Experimental/Chaos/Private/Chaos/Joint/PBDJointContainerSolver.cpp
--- a/Engine/Source/Runtime/Experimental/Chaos/Private/Chaos/Joint/PBDJointContainerSolver.cpp
+++ b/Engine/Source/Runtime/Experimental/Chaos/Private/Chaos/Joint/PBDJointContainerSolver.cpp
@@ -368,10 +368,12 @@ namespace Chaos
 					}
 
 					const FPBDJointSettings& JointSettings = GetConstraintSettings(SolverConstraintIndex);
-					if (JointSettings.bProjectionEnabled)
+					if (!JointSettings.bProjectionEnabled)
 					{
-						Solver.InitProjection(Dt, Settings, JointSettings);
+						continue;
 					}
+
+					Solver.InitProjection(Dt, Settings, JointSettings);
 				}
 			}
 
@@ -384,6 +386,10 @@ namespace Chaos
 				}
 
 				const FPBDJointSettings& JointSettings = GetConstraintSettings(SolverConstraintIndex);
+				if (!JointSettings.bProjectionEnabled)
+				{
+					continue;
+				}
 
 				if (It == 0)
 				{
@@ -408,6 +414,10 @@ namespace Chaos
 				}
 
 				const FPBDJointSettings& JointSettings = GetConstraintSettings(SolverConstraintIndex);
+				if (!JointSettings.bProjectionEnabled)
+				{
+					continue;
+				}
 
 				Solver.Update(Dt, Settings, JointSettings);
 
