Add Tetrahedron class for Chaos.

diff --git a/Engine/Source/Programs/HeadlessChaos/Private/HeadlessChaosApp.cpp b/Engine/Source/Programs/HeadlessChaos/Private/HeadlessChaosApp.cpp
--- a/Engine/Source/Programs/HeadlessChaos/Private/HeadlessChaosApp.cpp
+++ b/Engine/Source/Programs/HeadlessChaos/Private/HeadlessChaosApp.cpp
@@ -72,6 +72,7 @@ DEFINE_LOG_CATEGORY(LogHeadlessChaos);
 
 TEST(ImplicitTests, Implicit) {
 	ChaosTest::ImplicitPlane();
+	ChaosTest::ImplicitTetrahedron();
 	ChaosTest::ImplicitCube();
 	ChaosTest::ImplicitSphere();
 	ChaosTest::ImplicitCylinder();
diff --git a/Engine/Source/Programs/HeadlessChaos/Private/HeadlessChaosTestImplicits.cpp b/Engine/Source/Programs/HeadlessChaos/Private/HeadlessChaosTestImplicits.cpp
--- a/Engine/Source/Programs/HeadlessChaos/Private/HeadlessChaosTestImplicits.cpp
+++ b/Engine/Source/Programs/HeadlessChaos/Private/HeadlessChaosTestImplicits.cpp
@@ -10,6 +10,7 @@
 #include "Chaos/Box.h"
 #include "Chaos/Plane.h"
 #include "Chaos/Sphere.h"
+#include "Chaos/Tetrahedron.h"
 #include "Chaos/Cylinder.h"
 #include "Chaos/TaperedCylinder.h"
 #include "Chaos/TaperedCapsule.h"
@@ -222,6 +223,99 @@ namespace ChaosTest {
 		}
 	}
 
+	void ImplicitTetrahedron()
+	{
+		FString Caller("ImplicitTetrahedron()");
+
+		TTetrahedron<double> Tet(TVec3<double>(0, 0, 0), TVec3<double>(1, 0, 0), TVec3<double>(0, 1, 0), TVec3<double>(0, 0, 1));
+		const TArray<TTriangle<double>> Tris = Tet.GetTriangles();
+
+		double SVol = Tet.GetSignedVolume();
+		EXPECT_NEAR(SVol, 1.0/6, 0.001);
+
+		double Vol = Tet.GetVolume();
+		EXPECT_NEAR(Vol, SVol, 0.001);
+
+		double MinEL = Tet.GetMinEdgeLength();
+		EXPECT_NEAR(MinEL, 1.0, 0.001);
+
+		double MaxEL = Tet.GetMaxEdgeLength();
+		EXPECT_NEAR(MaxEL, 1.4, 0.1);
+
+		TVec3<double> Center = Tet.GetCenter();
+		EXPECT_VECTOR_NEAR(Center, FVector(.25, .25, .25), 0.001);
+
+		// Center
+		TVec3<double> Pt = Center;
+		bool Hit = Tet.Inside(Pt); // generates tris
+		EXPECT_TRUE(Hit);
+		bool RobustHit = Tet.RobustInside(Pt);
+		EXPECT_TRUE(RobustHit);
+		TVec3<double> Bary = Tet.GetFirstThreeBarycentricCoordinates(Pt);
+		EXPECT_VECTOR_NEAR(Bary, FVector(.25, .25, .25), 0.001);
+		TVec3<double> Surf = Tet.ProjectToSurface(Tris, Pt);
+		EXPECT_VECTOR_NEAR(Surf, Tris[3].GetCentroid(), 0.001);
+
+		// Point
+		Pt = Tet[0];
+		Hit = Tet.Inside(Pt, 0.001);
+		EXPECT_TRUE(Hit);
+		RobustHit = Tet.RobustInside(Pt, -0.001);
+		EXPECT_TRUE(RobustHit);
+		Bary = Tet.GetFirstThreeBarycentricCoordinates(Pt);
+		EXPECT_VECTOR_NEAR(Bary, FVector(1, 0, 0), 0.001);
+		Surf = Tet.ProjectToSurface(Tris, Pt);
+		EXPECT_VECTOR_NEAR(Surf, Pt, 0.001);
+
+		Pt[0] -= 0.1;
+		Hit = Tet.Inside(Pt, 0.001);
+		EXPECT_FALSE(Hit);
+		RobustHit = Tet.RobustInside(Pt, -0.001);
+		EXPECT_FALSE(RobustHit);
+		Surf = Tet.ProjectToSurface(Tris, Pt);
+		Pt[0] += 0.1;
+		EXPECT_VECTOR_NEAR(Surf, Pt, 0.001);
+
+		// Edge
+		Pt = TVec3<double>(0.5, 0, 0);
+		Hit = Tet.Inside(Pt, 0.001);
+		EXPECT_TRUE(Hit);
+		RobustHit = Tet.RobustInside(Pt, -0.001);
+		EXPECT_TRUE(RobustHit);
+		Bary = Tet.GetFirstThreeBarycentricCoordinates(Pt);
+		EXPECT_VECTOR_NEAR(Bary, FVector(.5, .5, 0), 0.001);
+		Surf = Tet.ProjectToSurface(Tris, Pt);
+		EXPECT_VECTOR_NEAR(Surf, Pt, 0.001);
+
+		Pt[1] -= 0.1;
+		Hit = Tet.Inside(Pt, 0.001);
+		EXPECT_FALSE(Hit);
+		RobustHit = Tet.RobustInside(Pt, -0.001);
+		EXPECT_FALSE(RobustHit);
+		Surf = Tet.ProjectToSurface(Tris, Pt);
+		Pt[1] += 0.1;
+		EXPECT_VECTOR_NEAR(Surf, Pt, 0.001);
+
+		// Face
+		Pt = Tris[0].GetCentroid();
+		Hit = Tet.Inside(Pt, 0.001);
+		EXPECT_TRUE(Hit);
+		RobustHit = Tet.RobustInside(Pt, -0.001);
+		EXPECT_TRUE(RobustHit);
+		Bary = Tet.GetFirstThreeBarycentricCoordinates(Pt);
+		EXPECT_VECTOR_NEAR(Bary, FVector(.3333, .3333, .3333), 0.001);
+		Surf = Tet.ProjectToSurface(Tris, Pt);
+		EXPECT_VECTOR_NEAR(Surf, Pt, 0.001);
+
+		Pt += Tris[0].GetNormal() * 0.1;
+		Hit = Tet.Inside(Pt, 0.001);
+		EXPECT_FALSE(Hit);
+		RobustHit = Tet.RobustInside(Pt, -0.001);
+		EXPECT_FALSE(RobustHit);
+		Surf = Tet.ProjectToSurface(Tris, Pt);
+		Pt -= Tris[0].GetNormal() * 0.1;
+		EXPECT_VECTOR_NEAR(Surf, Pt, 0.001);
+	}
 
 	void ImplicitCube()
 	{
diff --git a/Engine/Source/Programs/HeadlessChaos/Public/HeadlessChaosTestImplicits.h b/Engine/Source/Programs/HeadlessChaos/Public/HeadlessChaosTestImplicits.h
--- a/Engine/Source/Programs/HeadlessChaos/Public/HeadlessChaosTestImplicits.h
+++ b/Engine/Source/Programs/HeadlessChaos/Public/HeadlessChaosTestImplicits.h
@@ -18,6 +18,9 @@ namespace ChaosTest {
 	/**/
 	void ImplicitPlane();
 
+	/**/
+	void ImplicitTetrahedron();
+
 	/**/
 	void ImplicitSphere();
 
diff --git a/Engine/Source/Runtime/Experimental/Chaos/Public/Chaos/Math/Poisson.h b/Engine/Source/Runtime/Experimental/Chaos/Public/Chaos/Math/Poisson.h
--- a/Engine/Source/Runtime/Experimental/Chaos/Public/Chaos/Math/Poisson.h
+++ b/Engine/Source/Runtime/Experimental/Chaos/Public/Chaos/Math/Poisson.h
@@ -80,7 +80,28 @@ inline TV RowMaj3x3Multiply(const T* A, const TV& x)
 	return Result;
 }
 
-
+template <class T, class TV>
+inline TV RowMaj3x3RobustSolveLinearSystem(const T* A, const TV& b)
+{
+	T Cofactor11 = A[4] * A[8] - A[7] * A[5];
+	T Cofactor12 = A[7] * A[2] - A[1] * A[8];
+	T Cofactor13 = A[1] * A[5] - A[4] * A[2];
+	T Determinant = A[0] * Cofactor11 + A[3] * Cofactor12 + A[6] * Cofactor13;
+
+	T Matrix[9] =
+		{ Cofactor11, Cofactor12, Cofactor13,
+		  A[6] * A[5] - A[3] * A[8], A[0] * A[8] - A[6] * A[2], A[3] * A[2] - A[0] * A[5],
+		  A[3] * A[7] - A[6] * A[4], A[6] * A[1] - A[0] * A[7], A[0] * A[4] - A[3] * A[1] };
+	TV UnscaledResult = RowMaj3x3Multiply(&Matrix[0], b);
+
+	T RelativeTolerance = static_cast<T>(TNumericLimits<float>::Min()) * FMath::Max3(FMath::Abs(UnscaledResult[0]), FMath::Abs(UnscaledResult[1]), FMath::Abs(UnscaledResult[2]));
+	if (FMath::Abs(Determinant) <= RelativeTolerance)
+	{
+		RelativeTolerance = FMath::Max(RelativeTolerance, static_cast<T>(TNumericLimits<float>::Min()));
+		Determinant = Determinant >= 0 ? RelativeTolerance : -RelativeTolerance;
+	}
+	return UnscaledResult / Determinant;
+}
 
 template <class T, class TV=FVector3f, class TV_INT=FIntVector4, int32 d=3>
 void
diff --git a/Engine/Source/Runtime/Experimental/Chaos/Public/Chaos/Tetrahedron.h b/Engine/Source/Runtime/Experimental/Chaos/Public/Chaos/Tetrahedron.h
new file mode 100644
--- /dev/null
+++ b/Engine/Source/Runtime/Experimental/Chaos/Public/Chaos/Tetrahedron.h
@@ -0,0 +1,247 @@
+// Copyright Epic Games, Inc. All Rights Reserved.
+
+#pragma once
+
+#include "Chaos/Core.h"
+#include "Chaos/Triangle.h"
+#include "Chaos/Math/Poisson.h" // 3x3 row major matrix functions
+#include "GenericPlatform/GenericPlatformMath.h"
+
+
+namespace Chaos {
+
+	template <typename T>
+	class TTetrahedron
+	{
+	public:
+		TTetrahedron()
+		{
+		}
+
+		TTetrahedron(const TVec3<T>& In1, const TVec3<T>& In2, const TVec3<T>& In3, const TVec3<T>& In4)
+			: X{ In1, In2, In3, In4 }
+		{
+		}
+
+		FORCEINLINE TVec3<T>& operator[](uint32 InIndex)
+		{
+			checkSlow(InIndex < 4);
+			return X[InIndex];
+		}
+
+		FORCEINLINE const TVec3<T>& operator[](uint32 InIndex) const
+		{
+			checkSlow(InIndex < 4);
+			return X[InIndex];
+		}
+
+		void Invert()
+		{
+			Swap(X[0], X[1]); 
+		}
+
+		TVec3<T> GetCenter() const
+		{
+			return static_cast<T>(.25) * (X[0] + X[1] + X[2] + X[3]);
+		}
+
+		T GetMinEdgeLengthSquared() const
+		{
+			return FMath::Min(
+				FMath::Min3((X[0] - X[1]).SizeSquared(), (X[1] - X[2]).SizeSquared(), (X[2] - X[0]).SizeSquared()), 
+				FMath::Min3((X[0] - X[3]).SizeSquared(), (X[1] - X[3]).SizeSquared(), (X[2] - X[3]).SizeSquared()));
+		}
+
+		T GetMinEdgeLength() const
+		{
+			return FMath::Sqrt(GetMinEdgeLengthSquared());
+		}
+
+		T GetMaxEdgeLengthSquared() const
+		{
+			return FMath::Max(
+				FMath::Max3((X[0] - X[1]).SizeSquared(), (X[1] - X[2]).SizeSquared(), (X[2] - X[0]).SizeSquared()),
+				FMath::Max3((X[0] - X[3]).SizeSquared(), (X[1] - X[3]).SizeSquared(), (X[2] - X[3]).SizeSquared()));
+		}
+
+		T GetMaxEdgeLength() const
+		{
+			return FMath::Sqrt(GetMaxEdgeLengthSquared());
+		}
+
+		T GetVolume() const
+		{
+			return fabs(GetSignedVolume());
+		}
+
+		T GetSignedVolume() const
+		{
+			TVec3<T> U = X[1] - X[0];
+			TVec3<T> V = X[2] - X[0];
+			TVec3<T> W = X[3] - X[0];
+			return TripleProduct(U, V, W) / 6;
+		}
+
+		TVec3<T> GetFirstThreeBarycentricCoordinates(const TVec3<T>& Location) const 
+		{
+			TVec3<T> L1 = X[0] - X[3];
+			TVec3<T> L2 = X[1] - X[3];
+			TVec3<T> L3 = X[2] - X[3];
+			TVec3<T> RHS = Location - X[3];
+			T Matrix[9] = 
+			{ 
+				L1[0], L2[0], L3[0], 
+				L1[1], L2[1], L3[1],
+				L1[2], L2[2], L3[2]
+			};
+			return RowMaj3x3RobustSolveLinearSystem(&Matrix[0], RHS);
+		}
+
+		TVec4<T> GetBarycentricCoordinates(const TVec3<T>& Location) const 
+		{
+			TVec3<T> W = GetFirstThreeBarycentricCoordinates(Location);
+			return TVec4<T>(W[0], W[1], W[2], static_cast<T>(1.0) - W[0] - W[1] - W[2]);
+		}
+
+		TVec3<T> GetPointFromBarycentricCoordinates(const TVec3<T>& Weights) const
+		{
+			return Weights[0] * X[0] + Weights[1] * X[1] + Weights[2] * X[2] + Weights[3] * (static_cast<T>(1.0) - Weights[0] - Weights[1] - Weights[2]);
+		}
+
+		TVec3<T> GetPointFromBarycentricCoordinates(const TVec4<T>& Weights) const
+		{
+			return Weights[0] * X[0] + Weights[1] * X[1] + Weights[2] * X[2] + Weights[3] * X[3];
+		}
+
+		bool BarycentricInside(const TVec3<T>& Location, const T Tolerance=0) const
+		{
+			TVec3<T> Weights = GetFirstThreeBarycentricCoordinates(Location);
+			return Weights[0] >= Tolerance && Weights[1] >= Tolerance && Weights[2] >= Tolerance && Weights[0]+Weights[1]+Weights[2] <= 1+Tolerance;
+		}
+
+		FString ToString() const
+		{
+			return FString::Printf(TEXT("Tetrahedron: A: [%f, %f, %f], B: [%f, %f, %f], C: [%f, %f, %f], D: [%f, %f, %f]"),
+				X[0][0], X[0][1], X[0][2],
+				X[1][0], X[1][1], X[1][2],
+				X[2][0], X[2][1], X[2][2]);
+		}
+
+		//! Initialize outward facing triangles, regardless of the orientation of the tetrahedron.
+		TArray<TTriangle<T>> GetTriangles() const
+		{
+			TArray<TTriangle<T>> Triangles;
+			Triangles.SetNumUninitialized(4);
+			if (GetSignedVolume() <= 0)
+			{
+				Triangles[0] = TTriangle<T>(X[0], X[1], X[2]);
+				Triangles[1] = TTriangle<T>(X[0], X[3], X[1]);
+				Triangles[2] = TTriangle<T>(X[0], X[2], X[3]);
+				Triangles[3] = TTriangle<T>(X[1], X[3], X[2]);
+			}
+			else
+			{
+				Triangles[0] = TTriangle<T>(X[0], X[2], X[1]);
+				Triangles[1] = TTriangle<T>(X[0], X[1], X[3]);
+				Triangles[2] = TTriangle<T>(X[0], X[3], X[2]);
+				Triangles[3] = TTriangle<T>(X[1], X[2], X[3]);
+			}
+			return Triangles;
+		}
+
+		bool Inside(const TVec3<T>& Location, const T HalfThickness=0) const
+		{
+			TArray<TTriangle<T>> Tris = GetTriangles();
+			return Inside(Tris, Location, HalfThickness);
+		}
+
+		static bool Inside(const TArray<TTriangle<T>>& Tris, const TVec3<T>& Location, const T HalfThickness=0)
+		{
+			checkSlow(Tris.Num() == 4);
+			for(int32 i=0; i < Tris.Num(); i++)
+				if (Tris[i].GetPlane().SignedDistance(Location) > HalfThickness)
+				{
+					return false;
+				}
+			return true;
+		}
+
+		//! \p Tolerance should be a small negative number to include boundary.
+		bool RobustInside(const TVec3<T>& Location, const T Tolerance=0) const
+		{
+			T V1 = TripleProduct(Location - X[0], X[1] - X[0], X[2] - X[0]);
+			T V2 = TripleProduct(X[3] - X[0], X[1] - X[0], Location - X[0]);
+			T V3 = TripleProduct(X[3] - X[0], Location - X[0], X[2] - X[0]);
+			T V4 = TripleProduct(X[3] - Location, X[1] - Location, X[2] - Location);
+			return !(fabs(V1) <= Tolerance || fabs(V2) <= Tolerance || fabs(V3) <= Tolerance || fabs(V4) <= Tolerance ||
+				(Sign(V1) != Sign(V2) || Sign(V2) != Sign(V3) || Sign(V3) != Sign(V4)));
+		}
+
+		TVec3<T> ProjectToSurface(const TArray<TTriangle<T>>& Tris, const TVec3<T>& Location) const
+		{
+			if (Inside(Tris, Location)) 
+			{
+				int32 Idx = 0;
+				T Dist = TVec3<T>::DotProduct(Tris[0][0] - Location, Tris[0].GetNormal());
+				T DistTemp = TVec3<T>::DotProduct(Tris[1][0] - Location, Tris[1].GetNormal());
+				if (DistTemp < Dist) 
+				{ 
+					Idx = 1;
+					Dist = DistTemp;
+				}
+				DistTemp = TVec3<T>::DotProduct(Tris[2][0] - Location, Tris[2].GetNormal());
+				if (DistTemp < Dist)
+				{ 
+					Idx = 2;
+					Dist = DistTemp;
+				}
+				DistTemp = TVec3<T>::DotProduct(Tris[3][0] - Location, Tris[3].GetNormal());
+				if (DistTemp < Dist) 
+				{ 
+					Idx = 3;
+					Dist = DistTemp;
+				}
+				return Location + Dist * Tris[Idx].GetNormal();
+			}
+			else 
+			{
+				TVec3<T> SurfacePoint(Location);
+				for (int32 i = 0; i < 4; i++)
+				{
+					if (Tris[i].GetPlane().SignedDistance(Location) > static_cast<T>(0.0))
+					{
+						const TVec3<T> Norm = Tris[i].GetNormal();
+						SurfacePoint -= TVec3<T>::DotProduct(Location - Tris[i][0], Norm) * Norm;
+					}
+				}
+				return SurfacePoint;
+			}
+		}
+
+
+	private:
+		static T TripleProduct(const TVec3<T>& U, const TVec3<T>& V, const TVec3<T>& W)
+		{
+			return TVec3<T>::DotProduct(U, TVec3<T>::CrossProduct(V, W));
+		}
+
+		static int32 Sign(const T Value)
+		{
+			return Value < static_cast<T>(0) ? -1 : 1;
+		}
+
+		friend FChaosArchive& operator<<(FChaosArchive& Ar, TTetrahedron);
+
+		TVec3<T> X[4];
+	};
+
+	using FTetrahedron = TTetrahedron<FReal>;
+
+	template <typename T>
+	inline FChaosArchive& operator<<(FChaosArchive& Ar, TTetrahedron<T>& Value)
+	{
+		Ar << Value.X[0] << Value.X[1] << Value.X[2] << Value.X[3];
+		return Ar;
+	}
+
+} // namespace Chaos
\ No newline at end of file
