Chaos Cloth - Add support for Property Collection in constraint code.

#rnx

diff --git a/Engine/Plugins/Experimental/ChaosCloth/Source/Chaos/Private/ChaosCloth/ChaosClothConstraints.cpp b/Engine/Plugins/Experimental/ChaosCloth/Source/Chaos/Private/ChaosCloth/ChaosClothConstraints.cpp
--- a/Engine/Plugins/Experimental/ChaosCloth/Source/Chaos/Private/ChaosCloth/ChaosClothConstraints.cpp
+++ b/Engine/Plugins/Experimental/ChaosCloth/Source/Chaos/Private/ChaosCloth/ChaosClothConstraints.cpp
@@ -336,11 +336,14 @@ void FClothConstraints::SetEdgeConstraints(const TArray<TVec3<int32>>& SurfaceEl
 
 	if (bUseXPBDConstraints)
 	{
+		const TConstArrayView<FRealSingle> DampingMultipliers;
+
 		XEdgeConstraints = MakeShared<Softs::FXPBDSpringConstraints>(
 			Evolution->Particles(),
 			ParticleOffset, NumParticles,
 			SurfaceElements,
 			StiffnessMultipliers,
+			DampingMultipliers,
 			/*InStiffness =*/ Softs::FSolverVec2::UnitVector,
 			/*bTrimKinematicConstraints =*/ true);
 	}
@@ -381,11 +384,14 @@ void FClothConstraints::SetBendingConstraints(const TArray<TVec2<int32>>& Edges,
 
 	if (bUseXPBDConstraints)
 	{
+		const TConstArrayView<FRealSingle> DampingMultipliers;
+
 		XBendingConstraints = MakeShared<Softs::FXPBDSpringConstraints>(
 			Evolution->Particles(),
 			ParticleOffset, NumParticles,
 			Edges,
 			StiffnessMultipliers,
+			DampingMultipliers,
 			/*InStiffness =*/ Softs::FSolverVec2::UnitVector,
 			/*bTrimKinematicConstraints =*/ true);
 	}
@@ -410,15 +416,19 @@ void FClothConstraints::SetBendingConstraints(TArray<TVec4<int32>>&& BendingElem
 
 	if (bUseXPBDConstraints)
 	{
+		const TConstArrayView<FRealSingle> DampingMultipliers;
+
 		XBendingElementConstraints = MakeShared<Softs::FXPBDBendingConstraints>(
 			Evolution->Particles(),
 			ParticleOffset, NumParticles,
 			MoveTemp(BendingElements),
 			StiffnessMultipliers,
 			BucklingStiffnessMultipliers,
+			DampingMultipliers,
 			/*InStiffness =*/ Softs::FSolverVec2::UnitVector,
 			/*InBucklingRatio=*/ (Softs::FSolverReal)0.f,
 			/*InBucklingStiffness =*/ Softs::FSolverVec2::UnitVector,
+			/*InDampingRatio =*/ Softs::FSolverVec2::ZeroVector,
 			/*bTrimKinematicConstraints =*/ true);
 	}
 	else
@@ -460,9 +470,21 @@ void FClothConstraints::SetXPBDBendingConstraints(TArray<TVec4<int32>>&& Bending
 
 void FClothConstraints::SetBendingConstraints(TArray<TVec4<int32>>&& BendingElements, Softs::FSolverReal BendingStiffness)
 {
+	// Deprecated 5.1
 	check(Evolution);
 
-	BendingElementConstraints = MakeShared<Softs::FPBDBendingConstraints>(Evolution->Particles(), MoveTemp(BendingElements), BendingStiffness);
+	BendingElementConstraints = MakeShared<Softs::FPBDBendingConstraints>(
+		Evolution->Particles(),
+		ParticleOffset,
+		NumParticles,
+		MoveTemp(BendingElements),
+		/*StiffnessMultipliers =*/ TConstArrayView<FRealSingle>(),
+		/*BucklingStiffnessMultipliers =*/ TConstArrayView<FRealSingle>(),
+		/*InStiffness =*/ BendingStiffness,
+		/*InBucklingRatio=*/ (Softs::FSolverReal)0.f,
+		/*InBucklingStiffness =*/ Softs::FSolverVec2::UnitVector,
+		/*bTrimKinematicConstraints =*/ true);
+
 	++NumConstraintInits;  // Uses init to update the property tables
 	++NumConstraintRules;
 }
@@ -709,7 +731,7 @@ void FClothConstraints::SetMaximumDistanceProperties(Softs::FSolverReal MaxDista
 {
 	if (MaximumDistanceConstraints)
 	{
-		MaximumDistanceConstraints->SetSphereRadiiMultiplier(MaxDistancesMultiplier);
+		MaximumDistanceConstraints->SetTransformScale(MaxDistancesMultiplier);
 	}
 }
 
@@ -740,7 +762,7 @@ void FClothConstraints::SetBackstopProperties(bool bEnabled, Softs::FSolverReal
 	if (BackstopConstraints)
 	{
 		BackstopConstraints->SetEnabled(bEnabled);
-		BackstopConstraints->SetSphereRadiiMultiplier(BackstopDistancesMultiplier);
+		BackstopConstraints->SetTransformScale(BackstopDistancesMultiplier);
 	}
 }
 
diff --git a/Engine/Plugins/Experimental/ChaosCloth/Source/Chaos/Private/ChaosCloth/ChaosClothVisualization.cpp b/Engine/Plugins/Experimental/ChaosCloth/Source/Chaos/Private/ChaosCloth/ChaosClothVisualization.cpp
--- a/Engine/Plugins/Experimental/ChaosCloth/Source/Chaos/Private/ChaosCloth/ChaosClothVisualization.cpp
+++ b/Engine/Plugins/Experimental/ChaosCloth/Source/Chaos/Private/ChaosCloth/ChaosClothVisualization.cpp
@@ -965,7 +965,7 @@ namespace Chaos
 					const FLinearColor ColorLight = FLinearColor::MakeFromHSV8(ColorSeed, 160, 128);
 					const FLinearColor ColorDark = FLinearColor::MakeFromHSV8(ColorSeed, 160, 64);
 
-					const FReal BackstopRadius = BackstopRadiuses[Index] * BackstopConstraint->GetSphereRadiiMultiplier();
+					const FReal BackstopRadius = BackstopRadiuses[Index] * BackstopConstraint->GetScale();
 					const FReal BackstopDistance = BackstopDistances[Index];
 
 					const FVector AnimationNormal(AnimationNormals[Index]);
@@ -1023,7 +1023,7 @@ namespace Chaos
 					const FLinearColor ColorLight = FLinearColor::MakeFromHSV8(ColorSeed, 160, 128);
 					const FLinearColor ColorDark = FLinearColor::MakeFromHSV8(ColorSeed, 160, 64);
 
-					const FReal BackstopRadius = BackstopRadiuses[Index] * BackstopConstraint->GetSphereRadiiMultiplier();
+					const FReal BackstopRadius = BackstopRadiuses[Index] * BackstopConstraint->GetScale();
 					const FReal BackstopDistance = BackstopDistances[Index];
 
 					const FVector AnimationNormal(AnimationNormals[Index]);
diff --git a/Engine/Source/Runtime/Experimental/Chaos/Private/Chaos/PBDCollisionSpringConstraintsBase.cpp b/Engine/Source/Runtime/Experimental/Chaos/Private/Chaos/PBDCollisionSpringConstraintsBase.cpp
--- a/Engine/Source/Runtime/Experimental/Chaos/Private/Chaos/PBDCollisionSpringConstraintsBase.cpp
+++ b/Engine/Source/Runtime/Experimental/Chaos/Private/Chaos/PBDCollisionSpringConstraintsBase.cpp
@@ -21,15 +21,15 @@ FPBDCollisionSpringConstraintsBase::FPBDCollisionSpringConstraintsBase(
 	const FSolverReal InThickness,
 	const FSolverReal InStiffness,
 	const FSolverReal InFrictionCoefficient)
-	: TriangleMesh(InTriangleMesh)
+	: Thickness(InThickness)
+	, Stiffness(InStiffness)
+	, FrictionCoefficient(InFrictionCoefficient)
+	, TriangleMesh(InTriangleMesh)
 	, Elements(InTriangleMesh.GetSurfaceElements())
 	, ReferencePositions(InReferencePositions)
 	, DisabledCollisionElements(InDisabledCollisionElements)
 	, Offset(InOffset)
 	, NumParticles(InNumParticles)
-	, Thickness(InThickness)
-	, Stiffness(InStiffness)
-	, FrictionCoefficient(InFrictionCoefficient)
 	, bGlobalIntersectionAnalysis(false)
 {
 }
diff --git a/Engine/Source/Runtime/Experimental/Chaos/Private/Chaos/PBDLongRangeConstraintsBase.cpp b/Engine/Source/Runtime/Experimental/Chaos/Private/Chaos/PBDLongRangeConstraintsBase.cpp
--- a/Engine/Source/Runtime/Experimental/Chaos/Private/Chaos/PBDLongRangeConstraintsBase.cpp
+++ b/Engine/Source/Runtime/Experimental/Chaos/Private/Chaos/PBDLongRangeConstraintsBase.cpp
@@ -42,7 +42,10 @@ FPBDLongRangeConstraintsBase::FPBDLongRangeConstraintsBase(
 		FPBDStiffness::DefaultTableSize,
 		FPBDStiffness::DefaultParameterFitBase,
 		MaxStiffness)
-	, TetherScale(InScale, ScaleMultipliers, InParticleCount)
+	, TetherScale(
+		InScale.ClampAxes(MinTetherScale, MaxTetherScale),
+		ScaleMultipliers,
+		InParticleCount)
 	, ParticleOffset(InParticleOffset)
 {
 	// Apply default properties
diff --git a/Engine/Source/Runtime/Experimental/Chaos/Private/Chaos/PBDSphericalConstraint.cpp b/Engine/Source/Runtime/Experimental/Chaos/Private/Chaos/PBDSphericalConstraint.cpp
--- a/Engine/Source/Runtime/Experimental/Chaos/Private/Chaos/PBDSphericalConstraint.cpp
+++ b/Engine/Source/Runtime/Experimental/Chaos/Private/Chaos/PBDSphericalConstraint.cpp
@@ -29,7 +29,7 @@ void FPBDSphericalConstraint::ApplyHelperISPC(FSolverParticles& Particles, const
 		(ispc::FVector4f*)Particles.GetPAndInvM().GetData(),
 		(const ispc::FVector3f*)AnimationPositions.GetData(),
 		SphereRadii.GetData(),
-		SphereRadiiMultiplier,
+		GetScale(),
 		ParticleOffset,
 		ParticleCount);
 #endif
@@ -47,7 +47,7 @@ void FPBDSphericalBackstopConstraint::ApplyLegacyHelperISPC(FSolverParticles& Pa
 		(const ispc::FVector3f*)AnimationNormals.GetData(),
 		SphereOffsetDistances.GetData(),
 		SphereRadii.GetData(),
-		SphereRadiiMultiplier,
+		GetScale(),
 		ParticleOffset,
 		ParticleCount);
 #endif
@@ -65,7 +65,7 @@ void FPBDSphericalBackstopConstraint::ApplyHelperISPC(FSolverParticles& Particle
 		(const ispc::FVector3f*)AnimationNormals.GetData(),
 		SphereOffsetDistances.GetData(),
 		SphereRadii.GetData(),
-		SphereRadiiMultiplier,
+		GetScale(),
 		ParticleOffset,
 		ParticleCount);
 #endif
diff --git a/Engine/Source/Runtime/Experimental/Chaos/Private/Chaos/XPBDBendingConstraints.cpp b/Engine/Source/Runtime/Experimental/Chaos/Private/Chaos/XPBDBendingConstraints.cpp
--- a/Engine/Source/Runtime/Experimental/Chaos/Private/Chaos/XPBDBendingConstraints.cpp
+++ b/Engine/Source/Runtime/Experimental/Chaos/Private/Chaos/XPBDBendingConstraints.cpp
@@ -77,7 +77,7 @@ void FXPBDBendingConstraints::ApplyHelper(FSolverParticles& Particles, const FSo
 	const int32 i4 = Constraint[3];
 
 	const FSolverReal BiphasicStiffnessValue = IsBuckled[ConstraintIndex] ? BucklingValue : StiffnessValue;
-	if (BiphasicStiffnessValue < XPBDBendMinStiffness || (Particles.InvM(i1) == (FSolverReal)0. && Particles.InvM(i2) == (FSolverReal)0. && Particles.InvM(i3) == (FSolverReal)0. && Particles.InvM(i4) == (FSolverReal)0.))
+	if (BiphasicStiffnessValue < MinStiffness || (Particles.InvM(i1) == (FSolverReal)0. && Particles.InvM(i2) == (FSolverReal)0. && Particles.InvM(i3) == (FSolverReal)0. && Particles.InvM(i4) == (FSolverReal)0.))
 	{
 		return;
 	}
@@ -127,7 +127,7 @@ void FXPBDBendingConstraints::Apply(FSolverParticles& Particles, const FSolverRe
 			const FSolverReal ExpBucklingValue = (FSolverReal)BucklingStiffness;
 			const FSolverReal DampingRatioValue = (FSolverReal)DampingRatio;
 
-			if (ExpStiffnessValue < XPBDBendMinStiffness && ExpBucklingValue < XPBDBendMinStiffness)
+			if (ExpStiffnessValue < MinStiffness && ExpBucklingValue < MinStiffness)
 			{
 				return;
 			}
@@ -149,7 +149,7 @@ void FXPBDBendingConstraints::Apply(FSolverParticles& Particles, const FSolverRe
 							&IsBuckled.GetData()[ColorStart],
 							&Lambdas.GetData()[ColorStart],
 							Dt,
-							XPBDBendMinStiffness,
+							MinStiffness,
 							ExpStiffnessValue,
 							ExpBucklingValue,
 							DampingRatioValue,
@@ -169,7 +169,7 @@ void FXPBDBendingConstraints::Apply(FSolverParticles& Particles, const FSolverRe
 							&IsBuckled.GetData()[ColorStart],
 							&Lambdas.GetData()[ColorStart],
 							Dt,
-							XPBDBendMinStiffness,
+							MinStiffness,
 							ExpStiffnessValue,
 							ExpBucklingValue,
 							ColorSize);
@@ -210,7 +210,7 @@ void FXPBDBendingConstraints::Apply(FSolverParticles& Particles, const FSolverRe
 							&IsBuckled.GetData()[ColorStart],
 							&Lambdas.GetData()[ColorStart],
 							Dt,
-							XPBDBendMinStiffness,
+							MinStiffness,
 							StiffnessHasWeightMap,
 							&Stiffness.GetIndices().GetData()[ColorStart],
 							&Stiffness.GetTable().GetData()[0],
@@ -236,7 +236,7 @@ void FXPBDBendingConstraints::Apply(FSolverParticles& Particles, const FSolverRe
 							&IsBuckled.GetData()[ColorStart],
 							&Lambdas.GetData()[ColorStart],
 							Dt,
-							XPBDBendMinStiffness,
+							MinStiffness,
 							StiffnessHasWeightMap,
 							&Stiffness.GetIndices().GetData()[ColorStart],
 							&Stiffness.GetTable().GetData()[0],
@@ -276,7 +276,7 @@ void FXPBDBendingConstraints::Apply(FSolverParticles& Particles, const FSolverRe
 			const FSolverReal ExpStiffnessValue = (FSolverReal)Stiffness;
 			const FSolverReal ExpBucklingValue = (FSolverReal)BucklingStiffness;
 			const FSolverReal DampingRatioValue = (FSolverReal)DampingRatio;
-			if (ExpStiffnessValue < XPBDBendMinStiffness && ExpBucklingValue < XPBDBendMinStiffness)
+			if (ExpStiffnessValue < MinStiffness && ExpBucklingValue < MinStiffness)
 			{
 				return;
 			}
diff --git a/Engine/Source/Runtime/Experimental/Chaos/Private/Chaos/XPBDSpringConstraints.cpp b/Engine/Source/Runtime/Experimental/Chaos/Private/Chaos/XPBDSpringConstraints.cpp
--- a/Engine/Source/Runtime/Experimental/Chaos/Private/Chaos/XPBDSpringConstraints.cpp
+++ b/Engine/Source/Runtime/Experimental/Chaos/Private/Chaos/XPBDSpringConstraints.cpp
@@ -103,7 +103,7 @@ void FXPBDSpringConstraints::Apply(FSolverParticles& Particles, const FSolverRea
 		{
 			const FSolverReal ExpStiffnessValue = (FSolverReal)Stiffness;
 			const FSolverReal DampingRatioValue = (FSolverReal)DampingRatio;
-			if (ExpStiffnessValue < XPBDSpringMinStiffness)
+			if (ExpStiffnessValue < MinStiffness)
 			{
 				return;
 			}
@@ -179,7 +179,7 @@ void FXPBDSpringConstraints::Apply(FSolverParticles& Particles, const FSolverRea
 							&Dists.GetData()[ColorStart],
 							&Lambdas.GetData()[ColorStart],
 							Dt,
-							XPBDSpringMinStiffness,
+							MinStiffness,
 							StiffnessHasWeightMap,
 							&Stiffness.GetIndices().GetData()[ColorStart],
 							&Stiffness.GetTable().GetData()[0],
@@ -201,7 +201,7 @@ void FXPBDSpringConstraints::Apply(FSolverParticles& Particles, const FSolverRea
 							&Dists.GetData()[ColorStart],
 							&Lambdas.GetData()[ColorStart],
 							Dt,
-							XPBDSpringMinStiffness,
+							MinStiffness,
 							&Stiffness.GetIndices().GetData()[ColorStart],
 							&Stiffness.GetTable().GetData()[0],
 							ColorSize);
@@ -234,7 +234,7 @@ void FXPBDSpringConstraints::Apply(FSolverParticles& Particles, const FSolverRea
 		{
 			const FSolverReal ExpStiffnessValue = (FSolverReal)Stiffness;
 			const FSolverReal DampingRatioValue = (FSolverReal)DampingRatio;
-			if (ExpStiffnessValue < XPBDSpringMinStiffness)
+			if (ExpStiffnessValue < MinStiffness)
 			{
 				return;
 			}
diff --git a/Engine/Source/Runtime/Experimental/Chaos/Public/Chaos/PBDAnimDriveConstraint.h b/Engine/Source/Runtime/Experimental/Chaos/Public/Chaos/PBDAnimDriveConstraint.h
--- a/Engine/Source/Runtime/Experimental/Chaos/Public/Chaos/PBDAnimDriveConstraint.h
+++ b/Engine/Source/Runtime/Experimental/Chaos/Public/Chaos/PBDAnimDriveConstraint.h
@@ -4,6 +4,7 @@
 #include "Chaos/Framework/Parallel.h"
 #include "Chaos/PBDSoftsEvolutionFwd.h"
 #include "Chaos/PBDStiffness.h"
+#include "Chaos/PropertyCollectionAdapter.h"
 #include "ChaosStats.h"
 
 DECLARE_CYCLE_STAT(TEXT("Chaos PBD Anim Drive Constraint"), STAT_PBD_AnimDriveConstraint, STATGROUP_Chaos);
@@ -14,6 +15,11 @@ namespace Chaos::Softs
 	class FPBDAnimDriveConstraint final
 	{
 	public:
+		static bool IsEnabled(const FPropertyCollectionConstAdapter& PropertyCollection)
+		{
+			return IsAnimDriveStiffnessEnabled(PropertyCollection, false);
+		}
+
 		UE_DEPRECATED(5.2, "Use the other constructor supplying AnimationVelocities for correct subframe and damping behavior")
 		FPBDAnimDriveConstraint(
 			const int32 InParticleOffset,
@@ -34,6 +40,32 @@ namespace Chaos::Softs
 		{
 		}
 
+		FPBDAnimDriveConstraint(
+			const int32 InParticleOffset,
+			const int32 InParticleCount,
+			const TArray<FSolverVec3>& InAnimationPositions,  // Use global indexation (will need adding ParticleOffset)
+			const TArray<FSolverVec3>& InAnimationVelocities,  // Use global indexation (will need adding ParticleOffset)
+			const TConstArrayView<FRealSingle>& StiffnessMultipliers,  // Use local indexation
+			const TConstArrayView<FRealSingle>& DampingMultipliers,  // Use local indexation
+			const FPropertyCollectionConstAdapter& PropertyCollection
+		)
+			: AnimationPositions(InAnimationPositions)
+			, OldAnimationPositions_deprecated(InAnimationVelocities) // Unused when not using apply
+			, AnimationVelocities(InAnimationVelocities)
+			, ParticleOffset(InParticleOffset)
+			, ParticleCount(InParticleCount)
+			, UseDeprecatedApply(false)
+			, Stiffness(
+				FSolverVec2(GetWeightedFloatAnimDriveStiffness(PropertyCollection, 1.f)),
+				StiffnessMultipliers,
+				InParticleCount)
+			, Damping(
+				FSolverVec2(GetWeightedFloatAnimDriveDamping(PropertyCollection, 1.f)),
+				DampingMultipliers,
+				InParticleCount)
+		{
+		}
+
 		FPBDAnimDriveConstraint(
 			const int32 InParticleOffset,
 			const int32 InParticleCount,
@@ -62,6 +94,18 @@ namespace Chaos::Softs
 		// Return the damping input values used by the constraint
 		FSolverVec2 GetDamping() const { return Damping.GetWeightedValue(); }
 
+		void SetProperties(const FPropertyCollectionConstAdapter& PropertyCollection)
+		{
+			if (IsAnimDriveStiffnessMutable(PropertyCollection))
+			{
+				Stiffness.SetWeightedValue(FSolverVec2(GetWeightedFloatAnimDriveStiffness(PropertyCollection)));
+			}
+			if (IsAnimDriveDampingMutable(PropertyCollection))
+			{
+				Damping.SetWeightedValue(FSolverVec2(GetWeightedFloatAnimDriveDamping(PropertyCollection)));
+			}
+		}
+
 		inline void SetProperties(const FSolverVec2& InStiffness, const FSolverVec2& InDamping)
 		{
 			Stiffness.SetWeightedValue(InStiffness);
@@ -218,6 +262,9 @@ namespace Chaos::Softs
 
 		FPBDStiffness Stiffness;
 		FPBDStiffness Damping;
+
+		UE_CHAOS_DECLARE_PROPERTYCOLLECTION_NAME(AnimDriveStiffness, float);
+		UE_CHAOS_DECLARE_PROPERTYCOLLECTION_NAME(AnimDriveDamping, float);
 	};
 
 }  // End namespace Chaos::Softs
diff --git a/Engine/Source/Runtime/Experimental/Chaos/Public/Chaos/PBDAxialSpringConstraints.h b/Engine/Source/Runtime/Experimental/Chaos/Public/Chaos/PBDAxialSpringConstraints.h
--- a/Engine/Source/Runtime/Experimental/Chaos/Public/Chaos/PBDAxialSpringConstraints.h
+++ b/Engine/Source/Runtime/Experimental/Chaos/Public/Chaos/PBDAxialSpringConstraints.h
@@ -2,6 +2,7 @@
 #pragma once
 
 #include "Chaos/PBDAxialSpringConstraintsBase.h"
+#include "Chaos/PropertyCollectionAdapter.h"
 
 namespace Chaos::Softs
 {
@@ -30,14 +31,58 @@ public:
 
 	void Apply(FSolverParticles& InParticles, const FSolverReal Dt) const;
 
+protected:
+	using Base::Stiffness;
+
 private:
 	void InitColor(const FSolverParticles& InParticles, const int32 ParticleOffset, const int32 ParticleCount);
 	void ApplyHelper(FSolverParticles& InParticles, const FSolverReal Dt, const int32 ConstraintIndex, const FSolverReal ExpStiffnessValue) const;
 
-private:
 	TArray<int32> ConstraintsPerColorStartIndex; // Constraints are ordered so each batch is contiguous. This is ColorNum + 1 length so it can be used as start and end.
 };
 
+class CHAOS_API FPBDAreaSpringConstraints final : public FPBDAxialSpringConstraints
+{
+public:
+	static bool IsEnabled(const FPropertyCollectionConstAdapter& PropertyCollection)
+	{
+		return IsAreaSpringStiffnessEnabled(PropertyCollection, false);
+	}
+
+	FPBDAreaSpringConstraints(
+		const FSolverParticles& Particles,
+		int32 ParticleOffset,
+		int32 ParticleCount,
+		const TArray<TVec3<int32>>& InConstraints,
+		const TConstArrayView<FRealSingle>& StiffnessMultipliers,
+		const FPropertyCollectionConstAdapter& PropertyCollection,
+		bool bTrimKinematicConstraints)
+		: FPBDAxialSpringConstraints(
+			Particles,
+			ParticleOffset,
+			ParticleCount,
+			InConstraints,
+			StiffnessMultipliers,
+			FSolverVec2(GetWeightedFloatAreaSpringStiffness(PropertyCollection)),
+			bTrimKinematicConstraints)
+	{}
+
+	virtual ~FPBDAreaSpringConstraints() override = default;
+
+	void SetProperties(const FPropertyCollectionConstAdapter& PropertyCollection)
+	{
+		if (IsAreaSpringStiffnessMutable(PropertyCollection))
+		{
+			Stiffness.SetWeightedValue(FSolverVec2(GetWeightedFloatAreaSpringStiffness(PropertyCollection)));
+		}
+	}
+
+private:
+	using FPBDAxialSpringConstraints::Stiffness;
+
+	UE_CHAOS_DECLARE_PROPERTYCOLLECTION_NAME(AreaSpringStiffness, float);
+};
+
 }  // End namespace Chaos::Softs
 
 #if !defined(CHAOS_AXIAL_SPRING_ISPC_ENABLED_DEFAULT)
diff --git a/Engine/Source/Runtime/Experimental/Chaos/Public/Chaos/PBDBendingConstraints.h b/Engine/Source/Runtime/Experimental/Chaos/Public/Chaos/PBDBendingConstraints.h
--- a/Engine/Source/Runtime/Experimental/Chaos/Public/Chaos/PBDBendingConstraints.h
+++ b/Engine/Source/Runtime/Experimental/Chaos/Public/Chaos/PBDBendingConstraints.h
@@ -2,6 +2,7 @@
 #pragma once
 
 #include "Chaos/PBDBendingConstraintsBase.h"
+#include "Chaos/PropertyCollectionAdapter.h"
 
 namespace Chaos::Softs
 {
@@ -9,9 +10,36 @@ namespace Chaos::Softs
 class CHAOS_API FPBDBendingConstraints : public FPBDBendingConstraintsBase
 {
 	typedef FPBDBendingConstraintsBase Base;
-	using Base::Constraints;
 
 public:
+	static bool IsEnabled(const FPropertyCollectionConstAdapter& PropertyCollection)
+	{
+		return IsBendingElementStiffnessEnabled(PropertyCollection, false);
+	}
+
+	FPBDBendingConstraints(const FSolverParticles& InParticles,
+		int32 ParticleOffset,
+		int32 ParticleCount,
+		TArray<TVec4<int32>>&& InConstraints,
+		const TConstArrayView<FRealSingle>& StiffnessMultipliers,
+		const TConstArrayView<FRealSingle>& BucklingStiffnessMultipliers,
+		const FPropertyCollectionConstAdapter& PropertyCollection,
+		bool bTrimKinematicConstraints = false)
+		: Base(
+			InParticles,
+			ParticleOffset,
+			ParticleCount,
+			MoveTemp(InConstraints),
+			StiffnessMultipliers,
+			BucklingStiffnessMultipliers,
+			FSolverVec2(GetWeightedFloatBendingElementStiffness(PropertyCollection, 1.f)),
+			(FSolverReal)GetBucklingRatio(PropertyCollection, 0.f),  // BucklingRatio is clamped in base class
+			FSolverVec2(GetWeightedFloatBucklingStiffness(PropertyCollection, 1.f)),
+			bTrimKinematicConstraints) 
+	{
+		InitColor(InParticles, ParticleOffset, ParticleCount);
+	}
+
 	FPBDBendingConstraints(const FSolverParticles& InParticles,
 		int32 ParticleOffset,
 		int32 ParticleCount,
@@ -22,23 +50,66 @@ public:
 		const FSolverReal InBucklingRatio,
 		const FSolverVec2& InBucklingStiffness,
 		bool bTrimKinematicConstraints = false)
-		:Base(InParticles, ParticleOffset, ParticleCount, MoveTemp(InConstraints), StiffnessMultipliers, BucklingStiffnessMultipliers, InStiffness, InBucklingRatio, InBucklingStiffness, bTrimKinematicConstraints) 
+		: Base(
+			InParticles,
+			ParticleOffset,
+			ParticleCount,
+			MoveTemp(InConstraints),
+			StiffnessMultipliers,
+			BucklingStiffnessMultipliers,
+			InStiffness,
+			InBucklingRatio,
+			InBucklingStiffness,
+			bTrimKinematicConstraints) 
 	{
 		InitColor(InParticles, ParticleOffset, ParticleCount);
 	}
 
-	FPBDBendingConstraints(const FSolverParticles& InParticles, TArray<TVec4<int32>>&& InConstraints, const FSolverReal InStiffness = (FSolverReal)1.)
+	UE_DEPRECATED(5.2, "Use one of the other constructors instead.")
+	FPBDBendingConstraints(
+		const FSolverParticles& InParticles,
+		TArray<TVec4<int32>>&& InConstraints,
+		const FSolverReal InStiffness = (FSolverReal)1.)
+PRAGMA_DISABLE_DEPRECATION_WARNINGS
 		: Base(InParticles, MoveTemp(InConstraints), InStiffness) {}
+PRAGMA_ENABLE_DEPRECATION_WARNINGS
 
 	virtual ~FPBDBendingConstraints() override {}
 
+	using Base::SetProperties;
+
+	void SetProperties(const FPropertyCollectionConstAdapter& PropertyCollection)
+	{
+		if (IsBendingElementStiffnessMutable(PropertyCollection))
+		{
+			Stiffness.SetWeightedValue(FSolverVec2(GetWeightedFloatBendingElementStiffness(PropertyCollection)));
+		}
+		if (IsBucklingRatioMutable(PropertyCollection))
+		{
+			BucklingRatio = (FSolverReal)FMath::Clamp(GetBucklingRatio(PropertyCollection), 0.f, 1.);
+		}
+		if (IsBucklingStiffnessMutable(PropertyCollection))
+		{
+			BucklingStiffness.SetWeightedValue(FSolverVec2(GetWeightedFloatBucklingStiffness(PropertyCollection)));
+		}
+	}
+
 	void Apply(FSolverParticles& InParticles, const FSolverReal Dt) const;
 
 private:
 	void InitColor(const FSolverParticles& InParticles, const int32 ParticleOffset, const int32 ParticleCount);
 	void ApplyHelper(FSolverParticles& InParticles, const FSolverReal Dt, const int32 ConstraintIndex, const FSolverReal ExpStiffnessValue, const FSolverReal ExpBucklingValue) const;
 
+	using Base::Constraints;
+	using Base::Stiffness;
+	using Base::BucklingRatio;
+	using Base::BucklingStiffness;
+
 	TArray<int32> ConstraintsPerColorStartIndex; // Constraints are ordered so each batch is contiguous. This is ColorNum + 1 length so it can be used as start and end.
+
+	UE_CHAOS_DECLARE_PROPERTYCOLLECTION_NAME(BendingElementStiffness, float);
+	UE_CHAOS_DECLARE_PROPERTYCOLLECTION_NAME(BucklingRatio, float);
+	UE_CHAOS_DECLARE_PROPERTYCOLLECTION_NAME(BucklingStiffness, float);
 };
 
 }  // End namespace Chaos::Softs
diff --git a/Engine/Source/Runtime/Experimental/Chaos/Public/Chaos/PBDBendingConstraintsBase.h b/Engine/Source/Runtime/Experimental/Chaos/Public/Chaos/PBDBendingConstraintsBase.h
--- a/Engine/Source/Runtime/Experimental/Chaos/Public/Chaos/PBDBendingConstraintsBase.h
+++ b/Engine/Source/Runtime/Experimental/Chaos/Public/Chaos/PBDBendingConstraintsBase.h
@@ -5,7 +5,6 @@
 #include "Chaos/PBDSoftsSolverParticles.h"
 #include "Chaos/PBDStiffness.h"
 #include "Chaos/ParticleRule.h"
-#include "Chaos/PropertyCollectionAdapter.h"
 #include "Containers/StaticArray.h"
 
 namespace Chaos::Softs
@@ -14,52 +13,6 @@ namespace Chaos::Softs
 class FPBDBendingConstraintsBase
 {
 public:
-	static bool IsEnabled(const FPropertyCollectionConstAdapter& PropertyCollection)
-	{
-		return IsBendingStiffnessEnabled(PropertyCollection, false);
-	}
-
-	FPBDBendingConstraintsBase(
-		const FSolverParticles& InParticles,
-		int32 ParticleOffset,
-		int32 ParticleCount,
-		TArray<TVec4<int32>>&& InConstraints,
-		const TConstArrayView<FRealSingle>& StiffnessMultipliers,
-		const TConstArrayView<FRealSingle>& BucklingStiffnessMultipliers,
-		const FPropertyCollectionConstAdapter& PropertyCollection,
-		bool bTrimKinematicConstraints = false,
-		FSolverReal MaxStiffness = FPBDStiffness::DefaultPBDMaxStiffness)
-		: Constraints(bTrimKinematicConstraints ? TrimKinematicConstraints(InConstraints, InParticles): MoveTemp(InConstraints))
-		, ConstraintSharedEdges(ExtractConstraintSharedEdges(Constraints))
-		, Stiffness(FSolverVec2(GetWeightedFloatBendingStiffness(PropertyCollection, 1.f)),
-			StiffnessMultipliers,
-			TConstArrayView<TVec2<int32>>(ConstraintSharedEdges),
-			ParticleOffset,
-			ParticleCount,
-			FPBDStiffness::DefaultTableSize,
-			FPBDStiffness::DefaultParameterFitBase,
-			MaxStiffness)
-		, BucklingRatio((FSolverReal)FMath::Clamp(GetBucklingRatio(PropertyCollection, 0.f), 0.f, 1.f))
-		, BucklingStiffness(
-			FSolverVec2(GetWeightedFloatBucklingStiffness(PropertyCollection, 1.f)),
-			BucklingStiffnessMultipliers,
-			TConstArrayView<TVec2<int32>>(ConstraintSharedEdges),
-			ParticleOffset,
-			ParticleCount,
-			FPBDStiffness::DefaultTableSize,
-			FPBDStiffness::DefaultParameterFitBase,
-			MaxStiffness)
-	{
-		for (const TVec4<int32>& Constraint : Constraints)
-		{
-			const FSolverVec3& P1 = InParticles.X(Constraint[0]);
-			const FSolverVec3& P2 = InParticles.X(Constraint[1]);
-			const FSolverVec3& P3 = InParticles.X(Constraint[2]);
-			const FSolverVec3& P4 = InParticles.X(Constraint[3]);
-			RestAngles.Add(CalcAngle(P1, P2, P3, P4));
-		}
-	}
-
 	FPBDBendingConstraintsBase(const FSolverParticles& InParticles,
 		int32 ParticleOffset,
 		int32 ParticleCount, 
@@ -82,7 +35,7 @@ public:
 			FPBDStiffness::DefaultTableSize,
 			FPBDStiffness::DefaultParameterFitBase,
 			MaxStiffness)
-		, BucklingRatio(InBucklingRatio)
+		, BucklingRatio(FMath::Clamp(InBucklingRatio, (FSolverReal)0., (FSolverReal)1.))
 		, BucklingStiffness(
 			InBucklingStiffness,
 			BucklingStiffnessMultipliers,
@@ -103,6 +56,7 @@ public:
 		}
 	}
 
+	UE_DEPRECATED(5.2, "Use one of the other constructors instead.")
 	FPBDBendingConstraintsBase(const FSolverParticles& InParticles, TArray<TVec4<int32>>&& InConstraints, const FSolverReal InStiffness = (FSolverReal)1.)
 	    : Constraints(MoveTemp(InConstraints))
 		, ConstraintSharedEdges(ExtractConstraintSharedEdges(Constraints))
@@ -122,23 +76,6 @@ public:
 
 	virtual ~FPBDBendingConstraintsBase() {}
 
-	// Set properties using a property collection
-	void SetProperties(const FPropertyCollectionConstAdapter& PropertyCollection)
-	{
-		if (IsBendingStiffnessMutable(PropertyCollection))
-		{
-			Stiffness.SetWeightedValue(FSolverVec2(GetWeightedFloatBendingStiffness(PropertyCollection)));
-		}
-		if (IsBucklingRatioMutable(PropertyCollection))
-		{
-			BucklingRatio = FMath::Clamp((FSolverReal)GetBucklingRatio(PropertyCollection), (FSolverReal)0., (FSolverReal)1.);
-		}
-		if (IsBucklingStiffnessMutable(PropertyCollection))
-		{
-			BucklingStiffness.SetWeightedValue(FSolverVec2(GetWeightedFloatBucklingStiffness(PropertyCollection)));
-		}
-	}
-
 	// Update stiffness values
 	void SetProperties(const FSolverVec2& InStiffness, const FSolverReal InBucklingRatio, const FSolverVec2& InBucklingStiffness)
 	{ 
@@ -326,11 +263,6 @@ protected:
 
 	TArray<FSolverReal> RestAngles;
 	TArray<bool> IsBuckled;
-
-private:
-	UE_CHAOS_DECLARE_PROPERTYCOLLECTION_NAME(BendingStiffness, float);
-	UE_CHAOS_DECLARE_PROPERTYCOLLECTION_NAME(BucklingRatio, float);
-	UE_CHAOS_DECLARE_PROPERTYCOLLECTION_NAME(BucklingStiffness, float);
 };
 
 }  // End namespace Chaos::Softs
diff --git a/Engine/Source/Runtime/Experimental/Chaos/Public/Chaos/PBDCollisionSpringConstraints.h b/Engine/Source/Runtime/Experimental/Chaos/Public/Chaos/PBDCollisionSpringConstraints.h
--- a/Engine/Source/Runtime/Experimental/Chaos/Public/Chaos/PBDCollisionSpringConstraints.h
+++ b/Engine/Source/Runtime/Experimental/Chaos/Public/Chaos/PBDCollisionSpringConstraints.h
@@ -3,6 +3,7 @@
 
 #if !COMPILE_WITHOUT_UNREAL_SUPPORT
 #include "Chaos/PBDCollisionSpringConstraintsBase.h"
+#include "Chaos/PropertyCollectionAdapter.h"
 
 namespace Chaos::Softs
 {
@@ -14,6 +15,32 @@ class FPBDCollisionSpringConstraints : public FPBDCollisionSpringConstraintsBase
 	using Base::Constraints;
 
 public:
+	static constexpr FSolverReal MinFrictionCoefficient = (FSolverReal)0.;
+	static constexpr FSolverReal MaxFrictionCoefficient = (FSolverReal)10.;
+
+	static bool IsEnabled(const FPropertyCollectionConstAdapter& PropertyCollection)
+	{
+		return IsSelfCollisionStiffnessEnabled(PropertyCollection, false);
+	}
+
+	FPBDCollisionSpringConstraints(
+		const int32 InOffset,
+		const int32 InNumParticles,
+		const FTriangleMesh& InTriangleMesh,
+		const TArray<FSolverVec3>* InRestPositions,
+		TSet<TVec2<int32>>&& InDisabledCollisionElements,
+		const FPropertyCollectionConstAdapter& PropertyCollection)
+		: Base(
+			InOffset,
+			InNumParticles,
+			InTriangleMesh,
+			InRestPositions,
+			MoveTemp(InDisabledCollisionElements),
+			(FSolverReal)FMath::Max(GetSelfCollisionThickness(PropertyCollection, Base::BackCompatThickness), 0.f),
+			(FSolverReal)FMath::Clamp(GetSelfCollisionStiffness(PropertyCollection, Base::BackCompatStiffness), 0.f, 1.f),
+			FMath::Clamp((FSolverReal)GetSelfCollisionFriction(PropertyCollection, Base::BackCompatFrictionCoefficient), MinFrictionCoefficient, MaxFrictionCoefficient))
+	{}
+
 	FPBDCollisionSpringConstraints(
 		const int32 InOffset,
 		const int32 InNumParticles,
@@ -38,6 +65,22 @@ public:
 
 	using Base::Init;
 
+	void SetProperties(const FPropertyCollectionConstAdapter& PropertyCollection)
+	{
+		if (IsSelfCollisionThicknessMutable(PropertyCollection))
+		{
+			Thickness = (FSolverReal)FMath::Max(GetSelfCollisionThickness(PropertyCollection), 0.f);
+		}
+		if (IsSelfCollisionStiffnessMutable(PropertyCollection))
+		{
+			Stiffness = (FSolverReal)FMath::Clamp(GetSelfCollisionStiffness(PropertyCollection), 0.f, 1.f);
+		}
+		if (IsSelfCollisionFrictionMutable(PropertyCollection))
+		{
+			FrictionCoefficient = FMath::Clamp((FSolverReal)GetSelfCollisionFriction(PropertyCollection), MinFrictionCoefficient, MaxFrictionCoefficient);
+		}
+	}
+
 	void Apply(FSolverParticles& Particles, const FSolverReal Dt, const int32 ConstraintIndex) const
 	{
 		const int32 i = ConstraintIndex;
@@ -81,6 +124,14 @@ public:
 		}
 	}
 
+private:
+	using Base::Thickness;
+	using Base::Stiffness;
+	using Base::FrictionCoefficient;
+
+	UE_CHAOS_DECLARE_PROPERTYCOLLECTION_NAME(SelfCollisionThickness, float);
+	UE_CHAOS_DECLARE_PROPERTYCOLLECTION_NAME(SelfCollisionStiffness, float);
+	UE_CHAOS_DECLARE_PROPERTYCOLLECTION_NAME(SelfCollisionFriction, float);
 };
 
 }  // End namespace Chaos::Softs
diff --git a/Engine/Source/Runtime/Experimental/Chaos/Public/Chaos/PBDLongRangeConstraints.h b/Engine/Source/Runtime/Experimental/Chaos/Public/Chaos/PBDLongRangeConstraints.h
--- a/Engine/Source/Runtime/Experimental/Chaos/Public/Chaos/PBDLongRangeConstraints.h
+++ b/Engine/Source/Runtime/Experimental/Chaos/Public/Chaos/PBDLongRangeConstraints.h
@@ -1,8 +1,8 @@
 // Copyright Epic Games, Inc. All Rights Reserved.
 #pragma once
 
-#include "Chaos/PBDSoftsEvolutionFwd.h"
 #include "Chaos/PBDLongRangeConstraintsBase.h"
+#include "Chaos/PropertyCollectionAdapter.h"
 #include "ChaosStats.h"
 
 DECLARE_CYCLE_STAT(TEXT("Chaos PBD Long Range Constraint"), STAT_PBD_LongRange, STATGROUP_Chaos);
@@ -16,6 +16,30 @@ public:
 	typedef FPBDLongRangeConstraintsBase Base;
 	typedef typename Base::FTether FTether;
 
+	static bool IsEnabled(const FPropertyCollectionConstAdapter& PropertyCollection)
+	{
+		return IsTetherStiffnessEnabled(PropertyCollection, false);
+	}
+
+	FPBDLongRangeConstraints(
+		const FSolverParticles& Particles,
+		const int32 InParticleOffset,
+		const int32 InParticleCount,
+		const TArray<TConstArrayView<TTuple<int32, int32, FRealSingle>>>& InTethers,
+		const TConstArrayView<FRealSingle>& StiffnessMultipliers,
+		const TConstArrayView<FRealSingle>& ScaleMultipliers,
+		const FPropertyCollectionConstAdapter& PropertyCollection)
+		: FPBDLongRangeConstraintsBase(
+			Particles,
+			InParticleOffset,
+			InParticleCount,
+			InTethers,
+			StiffnessMultipliers,
+			ScaleMultipliers,
+			FSolverVec2(GetWeightedFloatTetherStiffness(PropertyCollection)),
+			FSolverVec2(GetWeightedFloatTetherScale(PropertyCollection)))  // Scale clamping done in constructor
+	{}
+
 	FPBDLongRangeConstraints(
 		const FSolverParticles& Particles,
 		const int32 InParticleOffset,
@@ -25,9 +49,33 @@ public:
 		const TConstArrayView<FRealSingle>& ScaleMultipliers,
 		const FSolverVec2& InStiffness = FSolverVec2::UnitVector,
 		const FSolverVec2& InScale = FSolverVec2::UnitVector)
-		: FPBDLongRangeConstraintsBase(Particles, InParticleOffset, InParticleCount, InTethers, StiffnessMultipliers, ScaleMultipliers, InStiffness, InScale) {}
+		: FPBDLongRangeConstraintsBase(
+			Particles,
+			InParticleOffset,
+			InParticleCount,
+			InTethers,
+			StiffnessMultipliers,
+			ScaleMultipliers,
+			InStiffness,
+			InScale)
+	{}
+
 	virtual ~FPBDLongRangeConstraints() override {}
 
+	using Base::SetProperties;
+
+	void SetProperties(const FPropertyCollectionConstAdapter& PropertyCollection)
+	{
+		if (IsTetherStiffnessMutable(PropertyCollection))
+		{
+			Stiffness.SetWeightedValue(FSolverVec2(GetWeightedFloatTetherStiffness(PropertyCollection)));
+		}
+		if (IsTetherScaleMutable(PropertyCollection))
+		{
+			TetherScale.SetWeightedValue(FSolverVec2(GetWeightedFloatTetherScale(PropertyCollection)).ClampAxes(MinTetherScale, MaxTetherScale));
+		}
+	}
+
 	void Apply(FSolverParticles& Particles, const FSolverReal Dt) const;
 
 private:
@@ -35,6 +83,9 @@ private:
 	using Base::Stiffness;
 	using Base::TetherScale;
 	using Base::ParticleOffset;
+
+	UE_CHAOS_DECLARE_PROPERTYCOLLECTION_NAME(TetherStiffness, float);
+	UE_CHAOS_DECLARE_PROPERTYCOLLECTION_NAME(TetherScale, float);
 };
 
 }  // End namespace Chaos::Softs
diff --git a/Engine/Source/Runtime/Experimental/Chaos/Public/Chaos/PBDSphericalConstraint.h b/Engine/Source/Runtime/Experimental/Chaos/Public/Chaos/PBDSphericalConstraint.h
--- a/Engine/Source/Runtime/Experimental/Chaos/Public/Chaos/PBDSphericalConstraint.h
+++ b/Engine/Source/Runtime/Experimental/Chaos/Public/Chaos/PBDSphericalConstraint.h
@@ -3,6 +3,7 @@
 
 #include "Chaos/PBDSoftsEvolutionFwd.h"
 #include "Chaos/PBDSoftsSolverParticles.h"
+#include "Chaos/PropertyCollectionAdapter.h"
 #include "ChaosStats.h"
 #include "Containers/ArrayView.h"
 #include "Containers/ContainersFwd.h"
@@ -28,6 +29,27 @@ namespace Chaos::Softs
 class CHAOS_API FPBDSphericalConstraint final
 {
 public:
+	static bool IsEnabled(const FPropertyCollectionConstAdapter& PropertyCollection)
+	{
+		return IsMaxDistanceScaleEnabled(PropertyCollection, false);
+	}
+
+	FPBDSphericalConstraint(
+		const uint32 InParticleOffset,
+		const uint32 InParticleCount,
+		const TArray<FSolverVec3>& InAnimationPositions,  // Use global indexation (will need adding ParticleOffset)
+		const TConstArrayView<FRealSingle>& InSphereRadii,  // Use local indexation
+		const FPropertyCollectionConstAdapter& PropertyCollection
+	)
+		: AnimationPositions(InAnimationPositions)
+		, SphereRadii(InSphereRadii)
+		, ParticleOffset(InParticleOffset)
+		, MaxDistanceScale((FSolverReal)FMath::Max(GetMaxDistanceScale(PropertyCollection, 1.f), 0.f))
+		, TransformScale((FSolverReal)1.)
+	{
+		check(InSphereRadii.Num() == InParticleCount);
+	}
+
 	FPBDSphericalConstraint(
 		const uint32 InParticleOffset,
 		const uint32 InParticleCount,
@@ -37,12 +59,21 @@ public:
 		: AnimationPositions(InAnimationPositions)
 		, SphereRadii(InSphereRadii)
 		, ParticleOffset(InParticleOffset)
-		, SphereRadiiMultiplier((FSolverReal)1.)
+		, MaxDistanceScale((FSolverReal)1.)
+		, TransformScale((FSolverReal)1.)
 	{
 		check(InSphereRadii.Num() == InParticleCount);
 	}
 	~FPBDSphericalConstraint() {}
 
+	void SetProperties(const FPropertyCollectionConstAdapter& PropertyCollection)
+	{
+		if (IsMaxDistanceScaleMutable(PropertyCollection))
+		{
+			MaxDistanceScale = (FSolverReal)FMath::Max(GetMaxDistanceScale(PropertyCollection), 0.f);
+		}
+	}
+
 	void Apply(FSolverParticles& Particles, const FSolverReal Dt) const
 	{
 		SCOPE_CYCLE_COUNTER(STAT_PBD_Spherical);
@@ -57,17 +88,25 @@ public:
 		}
 	}
 
-	void SetSphereRadiiMultiplier(const FSolverReal InSphereRadiiMultiplier)
+	void SetTransformScale(const FSolverReal InTransformScale)
 	{
-		SphereRadiiMultiplier = FMath::Max((FSolverReal)0., InSphereRadiiMultiplier);
+		TransformScale = FMath::Max(InTransformScale, (FSolverReal)0.);
 	}
 
+	FSolverReal GetTransformScale() const { return TransformScale; }
+	FSolverReal GetScale() const { return TransformScale * MaxDistanceScale; }
+
+	UE_DEPRECATED(5.2, "Use SetTransformScale instead.")
+	void SetSphereRadiiMultiplier(const FSolverReal InSphereRadiiMultiplier) { SetTransformScale(InSphereRadiiMultiplier); }
+
 private:
 	void ApplyHelper(FSolverParticles& Particles, const FSolverReal Dt) const
 	{
 		const int32 ParticleCount = SphereRadii.Num();
 
-		PhysicsParallelFor(ParticleCount, [this, &Particles, Dt](int32 Index)  // TODO: profile need for parallel loop based on particle count
+		const FSolverReal Scale = GetScale();
+
+		PhysicsParallelFor(ParticleCount, [this, &Particles, Dt, Scale](int32 Index)  // TODO: profile need for parallel loop based on particle count
 		{
 			const int32 ParticleIndex = ParticleOffset + Index;
 
@@ -76,7 +115,7 @@ private:
 				return;
 			}
 
-			const FSolverReal Radius = SphereRadii[Index] * SphereRadiiMultiplier;
+			const FSolverReal Radius = SphereRadii[Index] * Scale;
 			const FSolverVec3& Center = AnimationPositions[ParticleIndex];
 
 			const FSolverVec3 CenterToParticle = Particles.P(ParticleIndex) - Center;
@@ -98,12 +137,48 @@ protected:
 	const TArray<FSolverVec3>& AnimationPositions;  // Use global indexation (will need adding ParticleOffset)
 	const TConstArrayView<FRealSingle> SphereRadii;  // Use local indexation
 	const int32 ParticleOffset;
-	FSolverReal SphereRadiiMultiplier;
+	UE_DEPRECATED(5.2, "Use TransformScale and MaxDistanceScale instead.")
+	FSolverReal SphereRadiiMultiplier = 1.f;
+
+private:
+	FSolverReal MaxDistanceScale;
+	FSolverReal TransformScale;
+
+	UE_CHAOS_DECLARE_PROPERTYCOLLECTION_NAME(MaxDistanceScale, float);
 };
 
 class CHAOS_API FPBDSphericalBackstopConstraint final
 {
 public:
+	static bool IsEnabled(const FPropertyCollectionConstAdapter& PropertyCollection)
+	{
+		return IsBackstopScaleEnabled(PropertyCollection, false) ||
+			IsBackstopScaleAnimatable(PropertyCollection, false);  // Backstop can be re-enabled if animated
+	}
+
+	FPBDSphericalBackstopConstraint(
+		const int32 InParticleOffset,
+		const int32 InParticleCount,
+		const TArray<FSolverVec3>& InAnimationPositions,  // Use global indexation (will need adding ParticleOffset)
+		const TArray<FSolverVec3>& InAnimationNormals,  // Use global indexation (will need adding ParticleOffset)
+		const TConstArrayView<FRealSingle>& InSphereRadii,  // Use local indexation
+		const TConstArrayView<FRealSingle>& InSphereOffsetDistances,  // Use local indexation
+		const FPropertyCollectionConstAdapter& PropertyCollection
+	)
+		: AnimationPositions(InAnimationPositions)
+		, AnimationNormals(InAnimationNormals)
+		, SphereRadii(InSphereRadii)
+		, SphereOffsetDistances(InSphereOffsetDistances)
+		, ParticleOffset(InParticleOffset)
+		, BackstopScale((FSolverReal)FMath::Max(GetBackstopScale(PropertyCollection), 0.f))
+		, TransformScale((FSolverReal)1.)
+		, bEnabled(true)
+		, bUseLegacyBackstop(false)
+	{
+		check(InSphereRadii.Num() == InParticleCount);
+		check(InSphereOffsetDistances.Num() == InParticleCount);
+	}
+
 	FPBDSphericalBackstopConstraint(
 		const int32 InParticleOffset,
 		const int32 InParticleCount,
@@ -118,7 +193,8 @@ public:
 		, SphereRadii(InSphereRadii)
 		, SphereOffsetDistances(InSphereOffsetDistances)
 		, ParticleOffset(InParticleOffset)
-		, SphereRadiiMultiplier((FSolverReal)1.)
+		, BackstopScale((FSolverReal)1.)
+		, TransformScale((FSolverReal)1.)
 		, bEnabled(true)
 		, bUseLegacyBackstop(bInUseLegacyBackstop)
 	{
@@ -127,6 +203,19 @@ public:
 	}
 	~FPBDSphericalBackstopConstraint() {}
 
+	void SetProperties(const FPropertyCollectionConstAdapter& PropertyCollection)
+	{
+		if (IsBackstopScaleDirty(PropertyCollection) && IsBackstopScaleAnimatable(PropertyCollection))
+		{
+			bEnabled = IsBackstopScaleEnabled(PropertyCollection);
+		}
+
+		if (IsBackstopScaleMutable(PropertyCollection))
+		{
+			BackstopScale = (FSolverReal)FMath::Max(GetBackstopScale(PropertyCollection), 0.f);
+		}
+	}
+
 	void SetEnabled(bool bInEnabled) { bEnabled = bInEnabled; }
 	bool IsEnabled() const { return bEnabled; }
 
@@ -165,15 +254,20 @@ public:
 		}
 	}
 
-	void SetSphereRadiiMultiplier(const FSolverReal InSphereRadiiMultiplier)
+	void SetTransformScale(const FSolverReal InTransformScale)
 	{
-		SphereRadiiMultiplier = FMath::Max((FSolverReal)0., InSphereRadiiMultiplier);
+		TransformScale = FMath::Max(InTransformScale, (FSolverReal)0.);
 	}
+	
+	FSolverReal GetTransformScale() const { return TransformScale; }
 
-	FSolverReal GetSphereRadiiMultiplier() const
-	{
-		return SphereRadiiMultiplier;
-	}
+	FSolverReal GetScale() const { return TransformScale * BackstopScale; }
+
+	UE_DEPRECATED(5.2, "Use SetTransformScale instead.")
+	void SetSphereRadiiMultiplier(const FSolverReal InSphereRadiiMultiplier) { SetTransformScale(InSphereRadiiMultiplier); }
+
+	UE_DEPRECATED(5.2, "Use GetScale() instead.")
+	FSolverReal GetSphereRadiiMultiplier() const { return GetScale(); }
 
 	bool UseLegacyBackstop() const
 	{
@@ -184,8 +278,9 @@ private:
 	void ApplyHelper(FSolverParticles& Particles, const FSolverReal Dt) const
 	{
 		const int32 ParticleCount = SphereRadii.Num();
+		const FSolverReal Scale = GetScale();
 
-		PhysicsParallelFor(ParticleCount, [this, &Particles, Dt](int32 Index)  // TODO: profile need for parallel loop based on particle count
+		PhysicsParallelFor(ParticleCount, [this, &Particles, Dt, Scale](int32 Index)  // TODO: profile need for parallel loop based on particle count
 		{
 			const int32 ParticleIndex = ParticleOffset + Index;
 
@@ -197,8 +292,8 @@ private:
 			const FSolverVec3& AnimationPosition = AnimationPositions[ParticleIndex];
 			const FSolverVec3& AnimationNormal = AnimationNormals[ParticleIndex];
 
-			const FSolverReal SphereOffsetDistance = SphereOffsetDistances[Index] * SphereRadiiMultiplier;
-			const FSolverReal Radius = SphereRadii[Index] * SphereRadiiMultiplier;
+			const FSolverReal SphereOffsetDistance = SphereOffsetDistances[Index] * Scale;
+			const FSolverReal Radius = SphereRadii[Index] * Scale;
 
 			const FSolverVec3 Center = AnimationPosition - (Radius + SphereOffsetDistance) * AnimationNormal;  // Non legacy version adds radius to the distance
 			const FSolverVec3 CenterToParticle = Particles.P(ParticleIndex) - Center;
@@ -221,8 +316,9 @@ private:
 	void ApplyLegacyHelper(FSolverParticles& Particles, const FSolverReal Dt) const
 	{
 		const int32 ParticleCount = SphereRadii.Num();
+		const FSolverReal Scale = GetScale();
 
-		PhysicsParallelFor(ParticleCount, [this, &Particles, Dt](int32 Index)  // TODO: profile need for parallel loop based on particle count
+		PhysicsParallelFor(ParticleCount, [this, &Particles, Dt, Scale](int32 Index)  // TODO: profile need for parallel loop based on particle count
 		{
 			const int32 ParticleIndex = ParticleOffset + Index;
 
@@ -234,8 +330,8 @@ private:
 			const FSolverVec3& AnimationPosition = AnimationPositions[ParticleIndex];
 			const FSolverVec3& AnimationNormal = AnimationNormals[ParticleIndex];
 
-			const FSolverReal SphereOffsetDistance = SphereOffsetDistances[Index] * SphereRadiiMultiplier;
-			const FSolverReal Radius = SphereRadii[Index] * SphereRadiiMultiplier;
+			const FSolverReal SphereOffsetDistance = SphereOffsetDistances[Index] * Scale;
+			const FSolverReal Radius = SphereRadii[Index] * Scale;
 
 			const FSolverVec3 Center = AnimationPosition - SphereOffsetDistance * AnimationNormal;  // Legacy version already includes the radius within the distance
 			const FSolverVec3 CenterToParticle = Particles.P(ParticleIndex) - Center;
@@ -264,9 +360,12 @@ private:
 	const TConstArrayView<FRealSingle> SphereRadii; // Start at index 0, use local indexation
 	const TConstArrayView<FRealSingle> SphereOffsetDistances;  // Sphere position offsets, use local indexation
 	const int32 ParticleOffset;
-	FSolverReal SphereRadiiMultiplier;
+	FSolverReal BackstopScale;
+	FSolverReal TransformScale;
 	bool bEnabled;
 	bool bUseLegacyBackstop;
+
+	UE_CHAOS_DECLARE_PROPERTYCOLLECTION_NAME(BackstopScale, float);
 };
 
 }  // End namespace Chaos::Softs
diff --git a/Engine/Source/Runtime/Experimental/Chaos/Public/Chaos/PBDSpringConstraints.h b/Engine/Source/Runtime/Experimental/Chaos/Public/Chaos/PBDSpringConstraints.h
--- a/Engine/Source/Runtime/Experimental/Chaos/Public/Chaos/PBDSpringConstraints.h
+++ b/Engine/Source/Runtime/Experimental/Chaos/Public/Chaos/PBDSpringConstraints.h
@@ -2,16 +2,13 @@
 #pragma once
 
 #include "Chaos/PBDSpringConstraintsBase.h"
+#include "Chaos/PropertyCollectionAdapter.h"
 
 namespace Chaos::Softs
 {
 
 class CHAOS_API FPBDSpringConstraints : public FPBDSpringConstraintsBase
 {
-	typedef FPBDSpringConstraintsBase Base;
-	using Base::Constraints;
-	using Base::Stiffness;
-
 public:
 	template<int32 Valence>
 	FPBDSpringConstraints(
@@ -34,6 +31,11 @@ public:
 
 	const TArray<int32>& GetConstraintsPerColorStartIndex() const { return ConstraintsPerColorStartIndex; }
 
+protected:
+	typedef FPBDSpringConstraintsBase Base;
+	using Base::Constraints;
+	using Base::Stiffness;
+
 private:
 	void InitColor(const FSolverParticles& InParticles, const int32 ParticleOffset, const int32 ParticleCount);
 	void ApplyHelper(FSolverParticles& Particles, const FSolverReal Dt, const int32 ConstraintIndex, const FSolverReal ExpStiffnessValue) const;
@@ -42,6 +44,90 @@ private:
 	TArray<int32> ConstraintsPerColorStartIndex; // Constraints are ordered so each batch is contiguous. This is ColorNum + 1 length so it can be used as start and end.
 };
 
+class CHAOS_API FPBDEdgeSpringConstraints final : public FPBDSpringConstraints
+{
+public:
+	static constexpr int32 Valence = 2;
+
+	static bool IsEnabled(const FPropertyCollectionConstAdapter& PropertyCollection)
+	{
+		return IsEdgeSpringStiffnessEnabled(PropertyCollection, false);
+	}
+
+	FPBDEdgeSpringConstraints(
+		const FSolverParticles& Particles,
+		int32 ParticleOffset,
+		int32 ParticleCount,
+		const TArray<TVector<int32, Valence>>& InConstraints,
+		const TConstArrayView<FRealSingle>& StiffnessMultipliers,
+		const FPropertyCollectionConstAdapter& PropertyCollection,
+		bool bTrimKinematicConstraints = false)
+		: FPBDSpringConstraints(
+			Particles,
+			ParticleOffset,
+			ParticleCount,
+			InConstraints,
+			StiffnessMultipliers,
+			FSolverVec2(GetWeightedFloatEdgeSpringStiffness(PropertyCollection, 1.f)),
+			bTrimKinematicConstraints)
+	{}
+
+	virtual ~FPBDEdgeSpringConstraints() override = default;
+
+	void SetProperties(const FPropertyCollectionConstAdapter& PropertyCollection)
+	{
+		if (IsEdgeSpringStiffnessMutable(PropertyCollection))
+		{
+			Stiffness.SetWeightedValue(FSolverVec2(GetWeightedFloatEdgeSpringStiffness(PropertyCollection)));
+		}
+	}
+
+private:
+	UE_CHAOS_DECLARE_PROPERTYCOLLECTION_NAME(EdgeSpringStiffness, float);
+};
+
+class CHAOS_API FPBDBendingSpringConstraints final : public FPBDSpringConstraints
+{
+public:
+	static constexpr int32 Valence = 4;
+
+	static bool IsEnabled(const FPropertyCollectionConstAdapter& PropertyCollection)
+	{
+		return IsBendingSpringStiffnessEnabled(PropertyCollection, false);
+	}
+
+	FPBDBendingSpringConstraints(
+		const FSolverParticles& Particles,
+		int32 ParticleOffset,
+		int32 ParticleCount,
+		const TArray<TVector<int32, Valence>>& InConstraints,
+		const TConstArrayView<FRealSingle>& StiffnessMultipliers,
+		const FPropertyCollectionConstAdapter& PropertyCollection,
+		bool bTrimKinematicConstraints = false)
+		: FPBDSpringConstraints(
+			Particles,
+			ParticleOffset,
+			ParticleCount,
+			InConstraints,
+			StiffnessMultipliers,
+			FSolverVec2(GetWeightedFloatBendingSpringStiffness(PropertyCollection, 1.f)),
+			bTrimKinematicConstraints)
+	{}
+
+	virtual ~FPBDBendingSpringConstraints() override = default;
+
+	void SetProperties(const FPropertyCollectionConstAdapter& PropertyCollection)
+	{
+		if (IsBendingSpringStiffnessMutable(PropertyCollection))
+		{
+			Stiffness.SetWeightedValue(FSolverVec2(GetWeightedFloatBendingSpringStiffness(PropertyCollection)));
+		}
+	}
+
+private:
+	UE_CHAOS_DECLARE_PROPERTYCOLLECTION_NAME(BendingSpringStiffness, float);
+};
+
 }  // End namespace Chaos::Softs
 
 #if !defined(CHAOS_SPRING_ISPC_ENABLED_DEFAULT)
diff --git a/Engine/Source/Runtime/Experimental/Chaos/Public/Chaos/PBDTriangleMeshCollisions.h b/Engine/Source/Runtime/Experimental/Chaos/Public/Chaos/PBDTriangleMeshCollisions.h
--- a/Engine/Source/Runtime/Experimental/Chaos/Public/Chaos/PBDTriangleMeshCollisions.h
+++ b/Engine/Source/Runtime/Experimental/Chaos/Public/Chaos/PBDTriangleMeshCollisions.h
@@ -5,7 +5,7 @@
 #include "Chaos/PBDSoftsEvolutionFwd.h"
 #include "Chaos/TriangleMesh.h"
 #include "Chaos/HierarchicalSpatialHash.h"
-
+#include "Chaos/PropertyCollectionAdapter.h"
 
 // This does initialization for PBDCollisionSpringConstraints and PBDTriangleMeshIntersections, 
 // including intersection detection and global intersection analysis
@@ -15,24 +15,6 @@ namespace Chaos::Softs
 class CHAOS_API FPBDTriangleMeshCollisions
 {
 public:
-	FPBDTriangleMeshCollisions(
-		const int32 InOffset,
-		const int32 InNumParticles,
-		const FTriangleMesh& InTriangleMesh,
-		bool bInGlobalIntersectionAnalysis,
-		bool bInContourMinimization
-	)
-		:TriangleMesh(InTriangleMesh)
-		, Offset(InOffset)
-		, NumParticles(InNumParticles)
-		, bGlobalIntersectionAnalysis(bInGlobalIntersectionAnalysis)
-		, bContourMinimization(bInContourMinimization)
-	{}
-
-	~FPBDTriangleMeshCollisions() = default;
-
-	void Init(const FSolverParticles& Particles);
-
 	struct FContourMinimizationIntersection
 	{
 		TVec2<int32> EdgeVertices;
@@ -97,7 +79,6 @@ public:
 		}
 	};
 
-
 	// Debug display of intersection contours
 	struct FBarycentricPoint
 	{
@@ -114,6 +95,66 @@ public:
 		Count
 	};
 
+	static bool IsEnabled(const FPropertyCollectionConstAdapter& PropertyCollection)
+	{
+		return PropertyCollection.IsEnabled(FName(TEXT("SelfCollisionStiffness")).ToString(), false);  // Don't use UE_CHAOS_DECLARE_PROPERTYCOLLECTION_NAME here, SelfCollisionStiffness is only needed for activation
+	}
+
+	FPBDTriangleMeshCollisions(
+		const int32 InOffset,
+		const int32 InNumParticles,
+		const FTriangleMesh& InTriangleMesh,
+		const FPropertyCollectionConstAdapter& PropertyCollection
+	)
+		:TriangleMesh(InTriangleMesh)
+		, Offset(InOffset)
+		, NumParticles(InNumParticles)
+		, bGlobalIntersectionAnalysis(GetUseSelfIntersections(PropertyCollection, false) || GetUseGlobalIntersectionAnalysis(PropertyCollection, false))
+		, bContourMinimization(GetUseSelfIntersections(PropertyCollection, false) || GetUseContourMinimization(PropertyCollection, false))
+	{}
+
+	FPBDTriangleMeshCollisions(
+		const int32 InOffset,
+		const int32 InNumParticles,
+		const FTriangleMesh& InTriangleMesh,
+		bool bInGlobalIntersectionAnalysis,
+		bool bInContourMinimization
+	)
+		:TriangleMesh(InTriangleMesh)
+		, Offset(InOffset)
+		, NumParticles(InNumParticles)
+		, bGlobalIntersectionAnalysis(bInGlobalIntersectionAnalysis)
+		, bContourMinimization(bInContourMinimization)
+	{}
+
+	virtual ~FPBDTriangleMeshCollisions() = default;
+
+	void SetProperties(const FPropertyCollectionConstAdapter& PropertyCollection)
+	{
+		bool bUseSelfIntersections = false;
+		if (IsUseSelfIntersectionsMutable(PropertyCollection))
+		{
+			bUseSelfIntersections = GetUseSelfIntersections(PropertyCollection);
+		}
+		if (bUseSelfIntersections)
+		{
+			bGlobalIntersectionAnalysis = bContourMinimization = true;
+		}
+		else
+		{
+			if (IsUseGlobalIntersectionAnalysisMutable(PropertyCollection))
+			{
+				bGlobalIntersectionAnalysis = GetUseGlobalIntersectionAnalysis(PropertyCollection);
+			}
+			if (IsUseContourMinimizationMutable(PropertyCollection))
+			{
+				bContourMinimization = GetUseContourMinimization(PropertyCollection);
+			}
+		}
+	}
+
+	void Init(const FSolverParticles& Particles);
+
 	void SetGlobalIntersectionAnalysis(bool bInGlobalIntersectionAnalysis) { bGlobalIntersectionAnalysis = bInGlobalIntersectionAnalysis; }
 	void SetContourMinimization(bool bInContourMinimization) { bContourMinimization = bInContourMinimization; }
 
@@ -140,6 +181,9 @@ private:
 	TArray<TArray<FBarycentricPoint>> IntersectionContourPoints;
 	TArray<FContourType> IntersectionContourTypes;
 
+	UE_CHAOS_DECLARE_PROPERTYCOLLECTION_NAME(UseSelfIntersections, bool);
+	UE_CHAOS_DECLARE_PROPERTYCOLLECTION_NAME(UseGlobalIntersectionAnalysis, bool);
+	UE_CHAOS_DECLARE_PROPERTYCOLLECTION_NAME(UseContourMinimization, bool);
 };
 
 }  // End namespace Chaos::Softs
diff --git a/Engine/Source/Runtime/Experimental/Chaos/Public/Chaos/PBDTriangleMeshIntersections.h b/Engine/Source/Runtime/Experimental/Chaos/Public/Chaos/PBDTriangleMeshIntersections.h
--- a/Engine/Source/Runtime/Experimental/Chaos/Public/Chaos/PBDTriangleMeshIntersections.h
+++ b/Engine/Source/Runtime/Experimental/Chaos/Public/Chaos/PBDTriangleMeshIntersections.h
@@ -4,6 +4,7 @@
 #include "Chaos/Core.h"
 #include "Chaos/PBDSoftsEvolutionFwd.h"
 #include "Chaos/PBDTriangleMeshCollisions.h"
+#include "Chaos/PropertyCollectionAdapter.h"
 
 // Resolve intersections via contour minimization. Contour Minimization data is calculated in PBDTriangleMeshCollisions.
 namespace Chaos::Softs
@@ -11,6 +12,11 @@ namespace Chaos::Softs
 class CHAOS_API FPBDTriangleMeshIntersections
 {
 public:
+	static bool IsEnabled(const FPropertyCollectionConstAdapter& PropertyCollection)
+	{
+		return PropertyCollection.IsEnabled(FName(TEXT("SelfCollisionStiffness")).ToString(), false);  // Don't use UE_CHAOS_DECLARE_PROPERTYCOLLECTION_NAME here, SelfCollisionStiffness is only needed for activation
+	}
+
 	FPBDTriangleMeshIntersections(
 		const int32 InOffset,
 		const int32 InNumParticles,
diff --git a/Engine/Source/Runtime/Experimental/Chaos/Public/Chaos/PropertyCollectionAdapter.h b/Engine/Source/Runtime/Experimental/Chaos/Public/Chaos/PropertyCollectionAdapter.h
--- a/Engine/Source/Runtime/Experimental/Chaos/Public/Chaos/PropertyCollectionAdapter.h
+++ b/Engine/Source/Runtime/Experimental/Chaos/Public/Chaos/PropertyCollectionAdapter.h
@@ -318,6 +318,10 @@ namespace Chaos::Softs
 	{ \
 		return PropertyCollection.IsEnabled(PropertyName##String(), bDefault); \
 	} \
+	static bool Is##PropertyName##Animatable(const FPropertyCollectionConstAdapter& PropertyCollection, bool bDefault) \
+	{ \
+		return PropertyCollection.IsAnimatable(PropertyName##String(), bDefault); \
+	} \
 	Type GetLow##PropertyName(const FPropertyCollectionConstAdapter& PropertyCollection, const Type& Default) \
 	{ \
 		return PropertyCollection.GetLowValue<Type>(PropertyName##String(), Default, &PropertyName##Index); \
diff --git a/Engine/Source/Runtime/Experimental/Chaos/Public/Chaos/VelocityField.h b/Engine/Source/Runtime/Experimental/Chaos/Public/Chaos/VelocityField.h
--- a/Engine/Source/Runtime/Experimental/Chaos/Public/Chaos/VelocityField.h
+++ b/Engine/Source/Runtime/Experimental/Chaos/Public/Chaos/VelocityField.h
@@ -4,6 +4,7 @@
 #include "Chaos/PBDSoftsEvolutionFwd.h"
 #include "Chaos/PBDSoftsSolverParticles.h"
 #include "Chaos/TriangleMesh.h"
+#include "Chaos/PropertyCollectionAdapter.h"
 
 namespace Chaos::Softs
 {
@@ -15,6 +16,27 @@ public:
 	static constexpr FSolverReal DefaultDragCoefficient = (FSolverReal)0.5;
 	static constexpr FSolverReal DefaultLiftCoefficient = (FSolverReal)0.1;
 	static constexpr FSolverReal DefaultFluidDensity = (FSolverReal)1.225e-6;
+	static constexpr FSolverReal MinCoefficient = (FSolverReal)0.;   // Applies to both drag and lift
+	static constexpr FSolverReal MaxCoefficient = (FSolverReal)10.;  //
+
+	static bool IsEnabled(const FPropertyCollectionConstAdapter& PropertyCollection)
+	{
+		return
+			IsDragEnabled(PropertyCollection, false) ||
+			IsLiftEnabled(PropertyCollection, false) ||
+			IsPressureEnabled(PropertyCollection, false);
+	}
+
+	explicit FVelocityAndPressureField(const FPropertyCollectionConstAdapter& PropertyCollection)
+		: Offset(INDEX_NONE)
+		, NumParticles(0)
+	{
+		SetProperties(
+			FSolverVec2(GetWeightedFloatDrag(PropertyCollection, DefaultDragCoefficient)),
+			FSolverVec2(GetWeightedFloatLift(PropertyCollection, DefaultLiftCoefficient)),
+			FSolverReal(GetFluidDensity(PropertyCollection, DefaultFluidDensity)),
+			FSolverVec2(GetWeightedFloatPressure(PropertyCollection, 0.f)));
+	}
 
 	// Construct an uninitialized field. Mesh, properties, and velocity will have to be set for this field to be valid.
 	FVelocityAndPressureField()
@@ -39,13 +61,41 @@ public:
 		}
 	}
 
+	void SetProperties(const FPropertyCollectionConstAdapter& PropertyCollection)
+	{
+		if (IsFluidDensityMutable(PropertyCollection))
+		{
+			constexpr FSolverReal OneQuarter = (FSolverReal)0.25;
+			QuarterRho = (FSolverReal)GetFluidDensity(PropertyCollection) * OneQuarter;
+		}
+
+		if (IsDragMutable(PropertyCollection))
+		{
+			const FSolverVec2 Drag(GetWeightedFloatDrag(PropertyCollection));
+			DragBase = FMath::Clamp(Drag[0], MinCoefficient, MaxCoefficient);
+			DragRange = FMath::Clamp(Drag[1], MinCoefficient, MaxCoefficient) - DragBase;
+		}
+
+		if (IsLiftMutable(PropertyCollection))
+		{
+			const FSolverVec2 Lift(GetWeightedFloatLift(PropertyCollection));
+			LiftBase = FMath::Clamp(Lift[0], MinCoefficient, MaxCoefficient);
+			LiftRange = FMath::Clamp(Lift[1], MinCoefficient, MaxCoefficient) - LiftBase;
+		}
+
+		if (IsPressureMutable(PropertyCollection))
+		{
+			const FSolverVec2 Pressure(GetWeightedFloatPressure(PropertyCollection));
+			PressureBase = Pressure[0];
+			PressureRange = Pressure[1] - PressureBase;
+		}
+	}
+
 	void SetProperties(const FSolverVec2& Drag, const FSolverVec2& Lift, const FSolverReal FluidDensity, const FSolverVec2& Pressure = FSolverVec2::ZeroVector)
 	{
 		constexpr FSolverReal OneQuarter = (FSolverReal)0.25;
 		QuarterRho = FluidDensity * OneQuarter;
 
-		constexpr FSolverReal MinCoefficient = (FSolverReal)0.;
-		constexpr FSolverReal MaxCoefficient = (FSolverReal)10.;
 		DragBase = FMath::Clamp(Drag[0], MinCoefficient, MaxCoefficient);
 		DragRange = FMath::Clamp(Drag[1], MinCoefficient, MaxCoefficient) - DragBase;
 		LiftBase = FMath::Clamp(Lift[0], MinCoefficient, MaxCoefficient);
@@ -109,6 +159,11 @@ private:
 	FSolverReal QuarterRho;
 	int32 Offset;
 	int32 NumParticles;
+
+	UE_CHAOS_DECLARE_PROPERTYCOLLECTION_NAME(Drag, float);
+	UE_CHAOS_DECLARE_PROPERTYCOLLECTION_NAME(Lift, float);
+	UE_CHAOS_DECLARE_PROPERTYCOLLECTION_NAME(FluidDensity, float);
+	UE_CHAOS_DECLARE_PROPERTYCOLLECTION_NAME(Pressure, float);
 };
 
 using FVelocityField UE_DEPRECATED(5.1, "Chaos::Softs::FVelocityField has been renamed FVelocityAndPressureField to match its new behavior.") = FVelocityAndPressureField;
diff --git a/Engine/Source/Runtime/Experimental/Chaos/Public/Chaos/XPBDAxialSpringConstraints.h b/Engine/Source/Runtime/Experimental/Chaos/Public/Chaos/XPBDAxialSpringConstraints.h
--- a/Engine/Source/Runtime/Experimental/Chaos/Public/Chaos/XPBDAxialSpringConstraints.h
+++ b/Engine/Source/Runtime/Experimental/Chaos/Public/Chaos/XPBDAxialSpringConstraints.h
@@ -2,6 +2,7 @@
 #pragma once
 
 #include "Chaos/PBDAxialSpringConstraintsBase.h"
+#include "Chaos/PropertyCollectionAdapter.h"
 #include "ChaosStats.h"
 
 DECLARE_CYCLE_STAT(TEXT("Chaos XPBD Axial Spring Constraint"), STAT_XPBD_AxialSpring, STATGROUP_Chaos);
@@ -10,18 +11,20 @@ namespace Chaos::Softs
 {
 
 // Stiffness is in kg/s^2
+UE_DEPRECATED(5.2, "Use FXPBDAxialSpringConstraints::MinStiffness instead.")
 static const FSolverReal XPBDAxialSpringMinStiffness = (FSolverReal)1e-4; // Stiffness below this will be considered 0 since all of our calculations are actually based on 1 / stiffness.
+UE_DEPRECATED(5.2, "Use FXPBDAxialSpringConstraints::MaxStiffness instead.")
 static const FSolverReal XPBDAxialSpringMaxStiffness = (FSolverReal)1e7;
 
-class FXPBDAxialSpringConstraints final : public FPBDAxialSpringConstraintsBase
+class FXPBDAxialSpringConstraints : public FPBDAxialSpringConstraintsBase
 {
 	typedef FPBDAxialSpringConstraintsBase Base;
-	using Base::Barys;
-	using Base::Constraints;
-	using Base::Dists;
-	using Base::Stiffness;
 
 public:
+	// Stiffness is in kg/s^2
+	static constexpr FSolverReal MinStiffness = (FSolverReal)1e-4; // Stiffness below this will be considered 0 since all of our calculations are actually based on 1 / stiffness.
+	static constexpr FSolverReal MaxStiffness = (FSolverReal)1e7;
+
 	FXPBDAxialSpringConstraints(
 		const FSolverParticles& Particles,
 		int32 ParticleOffset,
@@ -38,16 +41,16 @@ public:
 			StiffnessMultipliers,
 			InStiffness,
 			bTrimKinematicConstraints,
-			XPBDAxialSpringMaxStiffness)
+			MaxStiffness)
 	{
 		Lambdas.Init(0.f, Constraints.Num());
 	}
 
 	virtual ~FXPBDAxialSpringConstraints() override {}
 
-	void SetProperties(const FSolverVec2& InStiffness) { Stiffness.SetWeightedValue(InStiffness, XPBDAxialSpringMaxStiffness); }
+	void SetProperties(const FSolverVec2& InStiffness) { Stiffness.SetWeightedValue(InStiffness, MaxStiffness); }
 
-	void ApplyProperties(const FSolverReal /*Dt*/, const int32 /*NumIterations*/) { Stiffness.ApplyXPBDValues(XPBDAxialSpringMaxStiffness); }
+	void ApplyProperties(const FSolverReal /*Dt*/, const int32 /*NumIterations*/) { Stiffness.ApplyXPBDValues(MaxStiffness); }
 
 	void Init() const { for (FSolverReal& Lambda : Lambdas) { Lambda = (FSolverReal)0.; } }
 
@@ -57,7 +60,7 @@ public:
 		if (!Stiffness.HasWeightMap())
 		{
 			const FSolverReal ExpStiffnessValue = (FSolverReal)Stiffness;
-			if (ExpStiffnessValue < XPBDAxialSpringMinStiffness)
+			if (ExpStiffnessValue < MinStiffness)
 			{
 				return;
 			}
@@ -120,7 +123,7 @@ private:
 
 		const FSolverReal Bary = Barys[InConstraintIndex];
 		const FSolverReal PInvMass = Particles.InvM(i3) * ((FSolverReal)1. - Bary) + Particles.InvM(i2) * Bary;
-		if (StiffnessValue < XPBDAxialSpringMinStiffness || ( Particles.InvM(i1) == (FSolverReal)0. && PInvMass == (FSolverReal)0.))
+		if (StiffnessValue < MinStiffness || ( Particles.InvM(i1) == (FSolverReal)0. && PInvMass == (FSolverReal)0.))
 		{
 			return FSolverVec3((FSolverReal)0.);
 		}
@@ -151,8 +154,57 @@ private:
 		return Delta;
 	}
 
+protected:
+	using Base::Stiffness;
+
 private:
+	using Base::Barys;
+	using Base::Constraints;
+	using Base::Dists;
+
 	mutable TArray<FSolverReal> Lambdas;
 };
 
+class FXPBDAreaSpringConstraints final : public FXPBDAxialSpringConstraints
+{
+public:
+	static bool IsEnabled(const FPropertyCollectionConstAdapter& PropertyCollection)
+	{
+		return IsXPBDAreaSpringStiffnessEnabled(PropertyCollection, false);
+	}
+
+	FXPBDAreaSpringConstraints(
+		const FSolverParticles& Particles,
+		int32 ParticleOffset,
+		int32 ParticleCount,
+		const TArray<TVec3<int32>>& InConstraints,
+		const TConstArrayView<FRealSingle>& StiffnessMultipliers,
+		const FPropertyCollectionConstAdapter& PropertyCollection,
+		bool bTrimKinematicConstraints)
+		: FXPBDAxialSpringConstraints(
+			Particles,
+			ParticleOffset,
+			ParticleCount,
+			InConstraints,
+			StiffnessMultipliers,
+			FSolverVec2(GetWeightedFloatXPBDAreaSpringStiffness(PropertyCollection)),
+			bTrimKinematicConstraints)
+	{}
+
+	virtual ~FXPBDAreaSpringConstraints() override = default;
+
+	void SetProperties(const FPropertyCollectionConstAdapter& PropertyCollection)
+	{
+		if (IsXPBDAreaSpringStiffnessMutable(PropertyCollection))
+		{
+			Stiffness.SetWeightedValue(FSolverVec2(GetWeightedFloatXPBDAreaSpringStiffness(PropertyCollection)), MaxStiffness);
+		}
+	}
+
+private:
+	using FXPBDAxialSpringConstraints::Stiffness;
+
+	UE_CHAOS_DECLARE_PROPERTYCOLLECTION_NAME(XPBDAreaSpringStiffness, float);
+};
+
 }  // End namespace Chaos::Softs
diff --git a/Engine/Source/Runtime/Experimental/Chaos/Public/Chaos/XPBDBendingConstraints.h b/Engine/Source/Runtime/Experimental/Chaos/Public/Chaos/XPBDBendingConstraints.h
--- a/Engine/Source/Runtime/Experimental/Chaos/Public/Chaos/XPBDBendingConstraints.h
+++ b/Engine/Source/Runtime/Experimental/Chaos/Public/Chaos/XPBDBendingConstraints.h
@@ -2,25 +2,61 @@
 #pragma once
 
 #include "Chaos/PBDBendingConstraintsBase.h"
+#include "Chaos/PropertyCollectionAdapter.h"
 #include "ChaosStats.h"
 
-
 namespace Chaos::Softs
 {
 
 // Stiffness is in kg cm^2 / rad^2 s^2
+UE_DEPRECATED(5.2, "Use FXPBDBendingConstraints::MinStiffness instead.")
 static const FSolverReal XPBDBendMinStiffness = (FSolverReal)1e-4; // Stiffness below this will be considered 0 since all of our calculations are actually based on 1 / stiffness.
+UE_DEPRECATED(5.2, "Use FXPBDBendingConstraints::MaxStiffness instead.")
 static const FSolverReal XPBDBendMaxStiffness = (FSolverReal)1e7;
 
 class CHAOS_API FXPBDBendingConstraints final : public FPBDBendingConstraintsBase
 {
 	typedef FPBDBendingConstraintsBase Base;
-	using Base::Constraints;
-	using Base::RestAngles;
-	using Base::Stiffness;
-	using Base::BucklingStiffness;
 
 public:
+	// Stiffness is in kg cm^2 / rad^2 s^2
+	static constexpr FSolverReal MinStiffness = (FSolverReal)1e-4; // Stiffness below this will be considered 0 since all of our calculations are actually based on 1 / stiffness.
+	static constexpr FSolverReal MaxStiffness = (FSolverReal)1e7;
+	static constexpr FSolverReal MinDamping = (FSolverReal)0.;
+	static constexpr FSolverReal MaxDamping = (FSolverReal)1000.;
+
+	static bool IsEnabled(const FPropertyCollectionConstAdapter& PropertyCollection)
+	{
+		return IsXPBDBendingElementStiffnessEnabled(PropertyCollection, false);
+	}
+
+	FXPBDBendingConstraints(const FSolverParticles& InParticles,
+		int32 ParticleOffset,
+		int32 ParticleCount,
+		TArray<TVec4<int32>>&& InConstraints,
+		const TConstArrayView<FRealSingle>& StiffnessMultipliers,
+		const TConstArrayView<FRealSingle>& BucklingStiffnessMultipliers,
+		const FPropertyCollectionConstAdapter& PropertyCollection,
+		bool bTrimKinematicConstraints = false)
+		: Base(
+			InParticles,
+			ParticleOffset,
+			ParticleCount,
+			MoveTemp(InConstraints),
+			StiffnessMultipliers,
+			BucklingStiffnessMultipliers,
+			FSolverVec2(GetWeightedFloatXPBDBendingElementStiffness(PropertyCollection, MaxStiffness)),
+			(FSolverReal)GetXPBDBucklingRatio(PropertyCollection, 0.f),
+			FSolverVec2(GetWeightedFloatXPBDBucklingStiffness(PropertyCollection, MaxStiffness)),
+			bTrimKinematicConstraints,
+			MaxStiffness)
+		, DampingRatio(FSolverVec2(GetWeightedFloatXPBDBendingElementDamping(PropertyCollection, 0.f)).ClampAxes(MinDamping, MaxDamping))
+	{
+		Lambdas.Init((FSolverReal)0., Constraints.Num());
+		InitColor(InParticles, ParticleOffset, ParticleCount);
+	}
+
+	UE_DEPRECATED(5.2, "Use one the other constructors instead.")
 	FXPBDBendingConstraints(const FSolverParticles& InParticles,
 		int32 ParticleOffset,
 		int32 ParticleCount,
@@ -42,7 +78,7 @@ public:
 			InBucklingRatio,
 			InBucklingStiffness,
 			bTrimKinematicConstraints,
-			XPBDBendMaxStiffness)
+			MaxStiffness)
 		, DampingRatio(FSolverVec2::ZeroVector)
 	{
 		Lambdas.Init((FSolverReal)0., Constraints.Num());
@@ -72,8 +108,13 @@ public:
 			InBucklingRatio,
 			InBucklingStiffness,
 			bTrimKinematicConstraints,
-			XPBDBendMaxStiffness)
-		, DampingRatio(InDampingRatio, DampingMultipliers, TConstArrayView<TVec2<int32>>(ConstraintSharedEdges), ParticleOffset, ParticleCount)
+			MaxStiffness)
+		, DampingRatio(
+			InDampingRatio.ClampAxes((FSolverReal)0., (FSolverReal)1.),
+			DampingMultipliers,
+			TConstArrayView<TVec2<int32>>(ConstraintSharedEdges),
+			ParticleOffset,
+			ParticleCount)
 	{
 		Lambdas.Init((FSolverReal)0., Constraints.Num());
 		InitColor(InParticles, ParticleOffset, ParticleCount);
@@ -88,21 +129,41 @@ public:
 		FPBDBendingConstraintsBase::Init(InParticles);
 	}
 
-	// Update stiffness values
+	void SetProperties(const FPropertyCollectionConstAdapter& PropertyCollection)
+	{
+		if (IsXPBDBendingElementStiffnessMutable(PropertyCollection))
+		{
+			Stiffness.SetWeightedValue(FSolverVec2(GetWeightedFloatXPBDBendingElementStiffness(PropertyCollection)), MaxStiffness);
+		}
+		if (IsXPBDBucklingRatioMutable(PropertyCollection))
+		{
+			BucklingRatio = FMath::Clamp(GetXPBDBucklingRatio(PropertyCollection), (FSolverReal)0., (FSolverReal)1.);
+		}
+		if (IsXPBDBucklingStiffnessMutable(PropertyCollection))
+		{
+			BucklingStiffness.SetWeightedValue(FSolverVec2(GetWeightedFloatXPBDBucklingStiffness(PropertyCollection)), MaxStiffness);
+		}
+		if (IsXPBDBendingElementDampingMutable(PropertyCollection))
+		{
+			DampingRatio.SetWeightedValue(FSolverVec2(GetWeightedFloatXPBDBendingElementDamping(PropertyCollection)).ClampAxes(MinDamping, MaxDamping));
+		}
+	}
+
 	void SetProperties(const FSolverVec2& InStiffness, const FSolverReal InBucklingRatio, const FSolverVec2& InBucklingStiffness, const FSolverVec2& InDampingRatio = FSolverVec2::ZeroVector)
 	{
-		Stiffness.SetWeightedValue(InStiffness, XPBDBendMaxStiffness);
-		BucklingRatio = InBucklingRatio;
-		BucklingStiffness.SetWeightedValue(InBucklingStiffness, XPBDBendMaxStiffness);
-		DampingRatio.SetWeightedValue(InDampingRatio);
+		Stiffness.SetWeightedValue(InStiffness, MaxStiffness);
+		BucklingRatio = FMath::Clamp(InBucklingRatio, (FSolverReal)0., (FSolverReal)1.);
+		BucklingStiffness.SetWeightedValue(InBucklingStiffness, MaxStiffness);
+		DampingRatio.SetWeightedValue(InDampingRatio.ClampAxes(MinDamping, MaxDamping));
 	}
 
 	// Update stiffness table, as well as the simulation stiffness exponent
 	void ApplyProperties(const FSolverReal /*Dt*/, const int32 /*NumIterations*/)
 	{
-		Stiffness.ApplyXPBDValues(XPBDBendMaxStiffness);
-		BucklingStiffness.ApplyXPBDValues(XPBDBendMaxStiffness);
-		DampingRatio.ApplyValues(); }
+		Stiffness.ApplyXPBDValues(MaxStiffness);
+		BucklingStiffness.ApplyXPBDValues(MaxStiffness);
+		DampingRatio.ApplyValues();
+	}
 
 	void Apply(FSolverParticles& Particles, const FSolverReal Dt) const;
 
@@ -110,11 +171,20 @@ private:
 	void InitColor(const FSolverParticles& InParticles, const int32 ParticleOffset, const int32 ParticleCount);
 	void ApplyHelper(FSolverParticles& Particles, const FSolverReal Dt, const int32 ConstraintIndex, const FSolverReal ExpStiffnessValue, const FSolverReal ExpBucklingValue, const FSolverReal DampingRatioValue) const;
 
-private:
+	using Base::Constraints;
+	using Base::RestAngles;
+	using Base::Stiffness;
+	using Base::BucklingStiffness;
+
 	FPBDWeightMap DampingRatio;
 	TArray<FSolverReal> DampingMultipliers;
 	mutable TArray<FSolverReal> Lambdas;
 	TArray<int32> ConstraintsPerColorStartIndex; // Constraints are ordered so each batch is contiguous. This is ColorNum + 1 length so it can be used as start and end.
+
+	UE_CHAOS_DECLARE_PROPERTYCOLLECTION_NAME(XPBDBendingElementStiffness, float);
+	UE_CHAOS_DECLARE_PROPERTYCOLLECTION_NAME(XPBDBendingElementDamping, float);
+	UE_CHAOS_DECLARE_PROPERTYCOLLECTION_NAME(XPBDBucklingRatio, float);
+	UE_CHAOS_DECLARE_PROPERTYCOLLECTION_NAME(XPBDBucklingStiffness, float);
 };
 
 }  // End namespace Chaos::Softs
diff --git a/Engine/Source/Runtime/Experimental/Chaos/Public/Chaos/XPBDLongRangeConstraints.h b/Engine/Source/Runtime/Experimental/Chaos/Public/Chaos/XPBDLongRangeConstraints.h
--- a/Engine/Source/Runtime/Experimental/Chaos/Public/Chaos/XPBDLongRangeConstraints.h
+++ b/Engine/Source/Runtime/Experimental/Chaos/Public/Chaos/XPBDLongRangeConstraints.h
@@ -3,6 +3,7 @@
 
 #include "Chaos/PBDLongRangeConstraintsBase.h"
 #include "Chaos/PBDParticles.h"
+#include "Chaos/PropertyCollectionAdapter.h"
 #include "ChaosStats.h"
 
 DECLARE_CYCLE_STAT(TEXT("Chaos XPBD Long Range Constraint"), STAT_XPBD_LongRange, STATGROUP_Chaos);
@@ -10,7 +11,9 @@ DECLARE_CYCLE_STAT(TEXT("Chaos XPBD Long Range Constraint"), STAT_XPBD_LongRange
 namespace Chaos::Softs
 {
 
+UE_DEPRECATED(5.2, "Use FXPBDLongRangeConstraints::MinStiffness instead.")
 static const FSolverReal XPBDLongRangeMinStiffness = (FSolverReal)1e-1;
+UE_DEPRECATED(5.2, "Use FXPBDLongRangeConstraints::MaxStiffness instead.")
 static const FSolverReal XPBDLongRangeMaxStiffness = (FSolverReal)1e7;
 
 class FXPBDLongRangeConstraints final : public FPBDLongRangeConstraintsBase
@@ -19,6 +22,36 @@ public:
 	typedef FPBDLongRangeConstraintsBase Base;
 	typedef typename Base::FTether FTether;
 
+	static constexpr FSolverReal MinStiffness = (FSolverReal)1e-1;
+	static constexpr FSolverReal MaxStiffness = (FSolverReal)1e7;
+
+	FXPBDLongRangeConstraints(
+		const FSolverParticles& Particles,
+		const int32 InParticleOffset,
+		const int32 InParticleCount,
+		const TArray<TConstArrayView<TTuple<int32, int32, FRealSingle>>>& InTethers,
+		const TConstArrayView<FRealSingle>& StiffnessMultipliers,
+		const TConstArrayView<FRealSingle>& ScaleMultipliers,
+		const FPropertyCollectionConstAdapter& PropertyCollection)
+	    : FPBDLongRangeConstraintsBase(
+			Particles,
+			InParticleOffset,
+			InParticleCount,
+			InTethers,
+			StiffnessMultipliers,
+			ScaleMultipliers,
+			FSolverVec2(GetWeightedFloatXPBDTetherStiffness(PropertyCollection)),
+			FSolverVec2(GetWeightedFloatXPBDTetherScale(PropertyCollection)),  // Scale clamping done in constructor
+			MaxStiffness)
+	{
+		NumTethers = 0;
+		for (const TConstArrayView<FTether>& TetherBatch : Tethers)
+		{
+			NumTethers += TetherBatch.Num();
+		}
+		Lambdas.Reserve(NumTethers);
+	}
+
 	FXPBDLongRangeConstraints(
 		const FSolverParticles& Particles,
 		const int32 InParticleOffset,
@@ -28,7 +61,7 @@ public:
 		const TConstArrayView<FRealSingle>& ScaleMultipliers,
 		const FSolverVec2& InStiffness = FSolverVec2::UnitVector,
 		const FSolverVec2& InScale = FSolverVec2::UnitVector)
-	    : FPBDLongRangeConstraintsBase(
+		: FPBDLongRangeConstraintsBase(
 			Particles,
 			InParticleOffset,
 			InParticleCount,
@@ -36,8 +69,8 @@ public:
 			StiffnessMultipliers,
 			ScaleMultipliers,
 			InStiffness,
-			InScale,
-			XPBDLongRangeMaxStiffness)
+			InScale,  // Scale clamping done in constructor
+			MaxStiffness)
 	{
 		NumTethers = 0;
 		for (const TConstArrayView<FTether>& TetherBatch : Tethers)
@@ -49,17 +82,29 @@ public:
 
 	virtual ~FXPBDLongRangeConstraints() override {}
 
+	void SetProperties(const FPropertyCollectionConstAdapter& PropertyCollection)
+	{
+		if (IsXPBDTetherStiffnessMutable(PropertyCollection))
+		{
+			Stiffness.SetWeightedValue(FSolverVec2(GetWeightedFloatXPBDTetherStiffness(PropertyCollection)));
+		}
+		if (IsXPBDTetherScaleMutable(PropertyCollection))
+		{
+			TetherScale.SetWeightedValue(FSolverVec2(GetWeightedFloatXPBDTetherScale(PropertyCollection)).ClampAxes(Base::MinTetherScale, Base::MaxTetherScale));
+		}
+	}
+
 	// Set the stiffness and scale values used by the constraint
 	void SetProperties(const FSolverVec2& InStiffness, const FSolverVec2& InScale)
 	{
-		Stiffness.SetWeightedValue(InStiffness, XPBDLongRangeMaxStiffness);
+		Stiffness.SetWeightedValue(InStiffness, MaxStiffness);
 		TetherScale.SetWeightedValue(InScale.ClampAxes(Base::MinTetherScale, Base::MaxTetherScale));
 	}
 
 	// Set stiffness offset and range, as well as the simulation stiffness exponent
 	void ApplyProperties(const FSolverReal /*Dt*/, const int32 /*NumIterations*/)
 	{
-		Stiffness.ApplyXPBDValues(XPBDLongRangeMaxStiffness);
+		Stiffness.ApplyXPBDValues(MaxStiffness);
 		TetherScale.ApplyValues();
 	}
 
@@ -149,7 +194,7 @@ public:
 private:
 	void Apply(FSolverParticles& Particles, const FSolverReal Dt, const FTether& Tether, int32 ConstraintIndex, const FSolverReal InStiffness, const FSolverReal InScale) const
 	{
-		if (InStiffness < XPBDLongRangeMinStiffness)
+		if (InStiffness < MinStiffness)
 		{
 			return;
 		}
@@ -165,13 +210,15 @@ private:
 		Lambda += DLambda;
 	}
 
-private:
 	using Base::Tethers;
 	using Base::Stiffness;
 	using Base::TetherScale;
 
 	mutable TArray<FSolverReal> Lambdas;
 	int32 NumTethers;
+
+	UE_CHAOS_DECLARE_PROPERTYCOLLECTION_NAME(XPBDTetherStiffness, float);
+	UE_CHAOS_DECLARE_PROPERTYCOLLECTION_NAME(XPBDTetherScale, float);
 };
 
 }  // End namespace Chaos::Softs
diff --git a/Engine/Source/Runtime/Experimental/Chaos/Public/Chaos/XPBDSpringConstraints.h b/Engine/Source/Runtime/Experimental/Chaos/Public/Chaos/XPBDSpringConstraints.h
--- a/Engine/Source/Runtime/Experimental/Chaos/Public/Chaos/XPBDSpringConstraints.h
+++ b/Engine/Source/Runtime/Experimental/Chaos/Public/Chaos/XPBDSpringConstraints.h
@@ -4,23 +4,29 @@
 #include "Chaos/PBDSpringConstraintsBase.h"
 #include "ChaosStats.h"
 #include "PBDWeightMap.h"
+#include "Chaos/PropertyCollectionAdapter.h"
 
 namespace Chaos::Softs
 {
 
 // Stiffness is in kg/s^2
+UE_DEPRECATED(5.2, "Use FXPBDSpringConstraints::MinStiffness instead.")
 static const FSolverReal XPBDSpringMinStiffness = (FSolverReal)1e-4; // Stiffness below this will be considered 0 since all of our calculations are actually based on 1 / stiffness.
+UE_DEPRECATED(5.2, "Use FXPBDSpringConstraints::MaxStiffness instead.")
 static const FSolverReal XPBDSpringMaxStiffness = (FSolverReal)1e7;
 
-class CHAOS_API FXPBDSpringConstraints final : public FPBDSpringConstraintsBase
+class CHAOS_API FXPBDSpringConstraints : public FPBDSpringConstraintsBase
 {
 	typedef FPBDSpringConstraintsBase Base;
-	using Base::Constraints;
-	using Base::Dists;
-	using Base::Stiffness;
 
 public:
+	static constexpr FSolverReal MinStiffness = (FSolverReal)1e-4; // Stiffness below this will be considered 0 since all of our calculations are actually based on 1 / stiffness.
+	static constexpr FSolverReal MaxStiffness = (FSolverReal)1e7;
+	static constexpr FSolverReal MinDampingRatio = (FSolverReal)0.;
+	static constexpr FSolverReal MaxDampingRatio = (FSolverReal)1000.;
+
 	template<int32 Valence, TEMPLATE_REQUIRES(Valence >= 2 && Valence <= 4)>
+	UE_DEPRECATED(5.2, "Use the other constructor instead.")
 	FXPBDSpringConstraints(
 		const FSolverParticles& Particles,
 		int32 ParticleOffset,
@@ -37,12 +43,13 @@ public:
 			StiffnessMultipliers,
 			InStiffness,
 			bTrimKinematicConstraints,
-			XPBDSpringMaxStiffness)
+			MaxStiffness)
 		, DampingRatio(FSolverVec2::ZeroVector)
 	{
 		Lambdas.Init((FSolverReal)0., Constraints.Num());
 		InitColor(Particles, ParticleOffset, ParticleCount);
 	}
+
 	template<int32 Valence, TEMPLATE_REQUIRES(Valence >= 2 && Valence <= 4)>
 	FXPBDSpringConstraints(
 		const FSolverParticles& Particles,
@@ -62,9 +69,9 @@ public:
 			StiffnessMultipliers,
 			InStiffness,
 			bTrimKinematicConstraints,
-			XPBDSpringMaxStiffness)
+			MaxStiffness)
 		, DampingRatio(
-			InDampingRatio,
+			InDampingRatio.ClampAxes(MinDampingRatio, MaxDampingRatio),
 			DampingMultipliers,
 			TConstArrayView<TVec2<int32>>(Constraints),
 			ParticleOffset,
@@ -81,12 +88,16 @@ public:
 	// Update stiffness values
 	void SetProperties(const FSolverVec2& InStiffness, const FSolverVec2& InDampingRatio = FSolverVec2::ZeroVector)
 	{ 
-		Stiffness.SetWeightedValue(InStiffness, XPBDSpringMaxStiffness);
-		DampingRatio.SetWeightedValue(InDampingRatio);
+		Stiffness.SetWeightedValue(InStiffness, MaxStiffness);
+		DampingRatio.SetWeightedValue(InDampingRatio.ClampAxes(MinDampingRatio, MaxDampingRatio));
 	}
 	
 	// Update stiffness table, as well as the simulation stiffness exponent
-	inline void ApplyProperties(const FSolverReal Dt, const int32 NumIterations) { Stiffness.ApplyXPBDValues(XPBDSpringMaxStiffness); DampingRatio.ApplyValues(); }
+	void ApplyProperties(const FSolverReal Dt, const int32 NumIterations)
+	{
+		Stiffness.ApplyXPBDValues(MaxStiffness);
+		DampingRatio.ApplyValues();
+	}
 
 	void Apply(FSolverParticles& Particles, const FSolverReal Dt) const;
 
@@ -101,7 +112,7 @@ private:
 		const int32 i1 = Constraint[0];
 		const int32 i2 = Constraint[1];
 
-		if (StiffnessValue < XPBDSpringMinStiffness || (Particles.InvM(i2) == (FSolverReal)0. && Particles.InvM(i1) == (FSolverReal)0.))
+		if (StiffnessValue < MinStiffness || (Particles.InvM(i2) == (FSolverReal)0. && Particles.InvM(i1) == (FSolverReal)0.))
 		{
 			return FSolverVec3((FSolverReal)0.);
 		}
@@ -133,12 +144,127 @@ private:
 		return Delta;
 	}
 
-private:
+protected:
+	using Base::Stiffness;
 	FPBDWeightMap DampingRatio;
+
+private:
+	using Base::Constraints;
+	using Base::Dists;
 	mutable TArray<FSolverReal> Lambdas;
 	TArray<int32> ConstraintsPerColorStartIndex; // Constraints are ordered so each batch is contiguous. This is ColorNum + 1 length so it can be used as start and end.
 };
 
+class CHAOS_API FXPBDEdgeSpringConstraints final : public FXPBDSpringConstraints
+{
+public:
+	static constexpr int32 Valence = 2;
+
+	static bool IsEnabled(const FPropertyCollectionConstAdapter& PropertyCollection)
+	{
+		return IsXPBDEdgeSpringStiffnessEnabled(PropertyCollection, false);
+	}
+
+	FXPBDEdgeSpringConstraints(
+		const FSolverParticles& Particles,
+		int32 ParticleOffset,
+		int32 ParticleCount,
+		const TArray<TVector<int32, 2>>& InConstraints,
+		const TConstArrayView<FRealSingle>& StiffnessMultipliers,
+		const TConstArrayView<FRealSingle>& DampingMultipliers,
+		const FSolverVec2& InStiffness,
+		const FSolverVec2& InDampingRatio,
+		const FPropertyCollectionConstAdapter& PropertyCollection,
+		bool bTrimKinematicConstraints = false)
+		: FXPBDSpringConstraints(
+			Particles,
+			ParticleOffset,
+			ParticleCount,
+			InConstraints,
+			StiffnessMultipliers,
+			DampingMultipliers,
+			FSolverVec2(GetWeightedFloatXPBDEdgeSpringStiffness(PropertyCollection)),
+			FSolverVec2(GetWeightedFloatXPBDEdgeSpringDamping(PropertyCollection)),
+			bTrimKinematicConstraints)
+	{}
+
+	virtual ~FXPBDEdgeSpringConstraints() override = default;
+
+	void SetProperties(const FPropertyCollectionConstAdapter& PropertyCollection)
+	{
+		if (IsXPBDEdgeSpringStiffnessMutable(PropertyCollection))
+		{
+			Stiffness.SetWeightedValue(FSolverVec2(GetWeightedFloatXPBDEdgeSpringStiffness(PropertyCollection)), MaxStiffness);
+		}
+		if (IsXPBDEdgeSpringDampingMutable(PropertyCollection))
+		{
+			DampingRatio.SetWeightedValue(FSolverVec2(GetWeightedFloatXPBDEdgeSpringDamping(PropertyCollection)).ClampAxes(MinDampingRatio, MaxDampingRatio));
+		}
+	}
+
+private:
+	using FXPBDSpringConstraints::Stiffness;
+	using FXPBDSpringConstraints::DampingRatio;
+
+	UE_CHAOS_DECLARE_PROPERTYCOLLECTION_NAME(XPBDEdgeSpringStiffness, float);
+	UE_CHAOS_DECLARE_PROPERTYCOLLECTION_NAME(XPBDEdgeSpringDamping, float);
+};
+
+class CHAOS_API FXPBDBendingSpringConstraints : public FXPBDSpringConstraints
+{
+public:
+	static constexpr int32 Valence = 4;
+
+	static bool IsEnabled(const FPropertyCollectionConstAdapter& PropertyCollection)
+	{
+		return IsXPBDBendingSpringStiffnessEnabled(PropertyCollection, false);
+	}
+
+	FXPBDBendingSpringConstraints(
+		const FSolverParticles& Particles,
+		int32 ParticleOffset,
+		int32 ParticleCount,
+		const TArray<TVector<int32, 4>>& InConstraints,
+		const TConstArrayView<FRealSingle>& StiffnessMultipliers,
+		const TConstArrayView<FRealSingle>& DampingMultipliers,
+		const FSolverVec2& InStiffness,
+		const FSolverVec2& InDampingRatio,
+		const FPropertyCollectionConstAdapter& PropertyCollection,
+		bool bTrimKinematicConstraints = false)
+		: FXPBDSpringConstraints(
+			Particles,
+			ParticleOffset,
+			ParticleCount,
+			InConstraints,
+			StiffnessMultipliers,
+			DampingMultipliers,
+			FSolverVec2(GetWeightedFloatXPBDBendingSpringStiffness(PropertyCollection)),
+			FSolverVec2(GetWeightedFloatXPBDBendingSpringDamping(PropertyCollection)),
+			bTrimKinematicConstraints)
+	{}
+
+	virtual ~FXPBDBendingSpringConstraints() override = default;
+
+	void SetProperties(const FPropertyCollectionConstAdapter& PropertyCollection)
+	{
+		if (IsXPBDBendingSpringStiffnessMutable(PropertyCollection))
+		{
+			Stiffness.SetWeightedValue(FSolverVec2(GetWeightedFloatXPBDBendingSpringStiffness(PropertyCollection)), MaxStiffness);
+		}
+		if (IsXPBDBendingSpringDampingMutable(PropertyCollection))
+		{
+			DampingRatio.SetWeightedValue(FSolverVec2(GetWeightedFloatXPBDBendingSpringDamping(PropertyCollection)).ClampAxes(MinDampingRatio, MaxDampingRatio));
+		}
+	}
+
+private:
+	using FXPBDSpringConstraints::Stiffness;
+	using FXPBDSpringConstraints::DampingRatio;
+
+	UE_CHAOS_DECLARE_PROPERTYCOLLECTION_NAME(XPBDBendingSpringStiffness, float);
+	UE_CHAOS_DECLARE_PROPERTYCOLLECTION_NAME(XPBDBendingSpringDamping, float);
+};
+
 }  // End namespace Chaos::Softs
 
 // Support ISPC enable/disable in non-shipping builds
