Added cvar net.ReplicateCustomDeltaPropertiesInRepIndexOrder to control the replication order of custom delta properties. Defaults to false to keep existing behavior, meaning respect the order from GetLifetimeReplicatedProps.

diff --git a/Engine/Source/Runtime/Engine/Private/RepLayout.cpp b/Engine/Source/Runtime/Engine/Private/RepLayout.cpp
--- a/Engine/Source/Runtime/Engine/Private/RepLayout.cpp
+++ b/Engine/Source/Runtime/Engine/Private/RepLayout.cpp
@@ -103,6 +103,14 @@ namespace UE::Net::Private
 {
 	static bool bDeltaInitialFastArrayElements = false;
 	static FAutoConsoleVariableRef CVarDeltaInitialFastArrayElements(TEXT("net.DeltaInitialFastArrayElements"), bDeltaInitialFastArrayElements, TEXT("If true, send delta struct changelists for initial fast array elements."));
+
+	/* FastArrays and other custom delta properties may have order dependencies due to callbacks being fired during serialization at which time other custom delta properties have not yet received their state.
+	 * This cvar toggles the behavior between using the RepIndex of the property or the order of appearance in the lifetime property array filled during a GetLifetimeReplicatedProps() call.
+	 * Default is false to keep the legacy behavior of using the GetLifetimeReplicatedProps() order for the custom delta properties.
+	 * The cvar is used in ReplicationStateDescriptorBuilder as well. Search for the cvar name in the code base before removing it.
+	 */
+	static bool bReplicateCustomDeltaPropertiesInRepIndexOrder = false;
+	static FAutoConsoleVariableRef CVarReplicateCustomDeltaPropertiesInRepIndexOrder(TEXT("net.ReplicateCustomDeltaPropertiesInRepIndexOrder"), bReplicateCustomDeltaPropertiesInRepIndexOrder, TEXT("If false (default) custom delta properties will replicate in the same order as they're added to the lifetime property array during the call to GetLifetimeReplicatedProps. If true custom delta properties will be replicated in the property RepIndex order, which is typically in increasing property offset order. Note that custom delta properties are always serialized after regular properties."));
 }
 
 extern int32 GNumSharedSerializationHit;
@@ -6180,10 +6188,16 @@ void FRepLayout::InitFromClass(
 	// Initialize lifetime props
 	// Properties that replicate for the lifetime of the channel
 	TArray<FLifetimeProperty> LifetimeProps;
+	LifetimeProps.Reserve(Parents.Num());
 
 	UObject* Object = InObjectClass->GetDefaultObject();
 
 	Object->GetLifetimeReplicatedProps(LifetimeProps);
+	// If there are custom delta properties we may have to change the order we traverse the replicated props.
+	if (UE::Net::Private::bReplicateCustomDeltaPropertiesInRepIndexOrder && (HighestCustomDeltaRepIndex != INDEX_NONE))
+	{
+		Algo::SortBy(LifetimeProps, [](const FLifetimeProperty& Element) { return Element.RepIndex; }, TLess<decltype(FLifetimeProperty::RepIndex)>());
+	}
 
 #if WITH_PUSH_MODEL
 	PushModelProperties.Init(false, Parents.Num());
diff --git a/Engine/Source/Runtime/Experimental/Iris/Core/Private/Iris/ReplicationState/ReplicationStateDescriptorBuilder.cpp b/Engine/Source/Runtime/Experimental/Iris/Core/Private/Iris/ReplicationState/ReplicationStateDescriptorBuilder.cpp
--- a/Engine/Source/Runtime/Experimental/Iris/Core/Private/Iris/ReplicationState/ReplicationStateDescriptorBuilder.cpp
+++ b/Engine/Source/Runtime/Experimental/Iris/Core/Private/Iris/ReplicationState/ReplicationStateDescriptorBuilder.cpp
@@ -2119,6 +2119,8 @@ bool FPropertyReplicationStateDescriptorBuilder::IsFastArraySupported(const UStr
 namespace UE::Net
 {
 
+const IConsoleVariable* FReplicationStateDescriptorBuilder::CVarReplicateCustomDeltaPropertiesInRepIndexOrder = nullptr;
+
 FReplicationStateDescriptorBuilder::FParameters::FParameters()
 : DescriptorRegistry(nullptr),
   DefaultStateSource(nullptr),
@@ -2371,21 +2373,21 @@ SIZE_T FReplicationStateDescriptorBuilder::CreateDescriptorsForClass(FResult& Cr
 	}
 
 	// We need the lifetime properties array to get conditionals
+	TArray<FLifetimeProperty> LifeTimePropertiesOriginalOrder;
 	TArray<FLifetimeProperty> LifeTimeProperties;
 	if (bGetLifeTimeProperties)
 	{
 		IRIS_PROFILER_SCOPE(FReplicationStateDescriptorBuilder_GetLifeTimeProperties);
 
-		TArray<FLifetimeProperty> TempLifetimeProperties;
-		TempLifetimeProperties.Reserve(InObjectClass->ClassReps.Num());
+		LifeTimePropertiesOriginalOrder.Reserve(InObjectClass->ClassReps.Num());
 		UObject* Object = InObjectClass->GetDefaultObject();
-		Object->GetLifetimeReplicatedProps(TempLifetimeProperties);
+		Object->GetLifetimeReplicatedProps(LifeTimePropertiesOriginalOrder);
 
 		// Lifetime properties aren't sorted and searching for a property with a specific RepIndex in the array is linear.
 		// Let's put them where we want them.
 		static_assert(TIsZeroConstructType<FLifetimeProperty>::Value, "Need to use SetNum() instead of SetNumZeroed()");
 		LifeTimeProperties.SetNumZeroed(InObjectClass->ClassReps.Num());
-		for (const FLifetimeProperty& LifetimeProperty : TempLifetimeProperties)
+		for (const FLifetimeProperty& LifetimeProperty : LifeTimePropertiesOriginalOrder)
 		{
 			LifeTimeProperties[LifetimeProperty.RepIndex] = LifetimeProperty;
 		}
@@ -2547,6 +2549,28 @@ SIZE_T FReplicationStateDescriptorBuilder::CreateDescriptorsForClass(FResult& Cr
 	}
 
 	// If we have properties with a CreateAndRegisterReplicationFragmentFunction, we build separate descriptors for them
+	// We may have to re-order custom properties if there are more than one.
+	if (CustomProperties.Num() > 1)
+	{
+		InitCVarReplicateCustomDeltaPropertiesInRepIndexOrder();
+
+		// If CVarReplicateCustomDeltaPropertiesInRepIndexOrder is false then we use the GetLifetimeReplicatedPropsOrder instead, which is the RepLayout legacy behavior.
+		if (CVarReplicateCustomDeltaPropertiesInRepIndexOrder != nullptr && !CVarReplicateCustomDeltaPropertiesInRepIndexOrder->GetBool())
+		{
+			TArray<FPropertyReplicationStateDescriptorBuilder::FMemberProperty, TInlineAllocator<8>> NewCustomProperties;
+			NewCustomProperties.Reserve(CustomProperties.Num());
+			for (const FLifetimeProperty& LifetimeProperty : LifeTimePropertiesOriginalOrder)
+			{
+				if (const FPropertyReplicationStateDescriptorBuilder::FMemberProperty* MemberProperty = CustomProperties.FindByPredicate([RepIndex = LifetimeProperty.RepIndex](const FPropertyReplicationStateDescriptorBuilder::FMemberProperty& MemberProperty) { return RepIndex == MemberProperty.Property->RepIndex; }))
+				{
+					NewCustomProperties.Add(*MemberProperty);
+				}
+			}
+
+			CustomProperties = NewCustomProperties;
+		}
+	}
+
 	for (uint32 CustomPropertyIt = 0, CustomArrayPropertyEndIt = CustomProperties.Num(); CustomPropertyIt != CustomArrayPropertyEndIt; ++CustomPropertyIt)
 	{
 		const FPropertyReplicationStateDescriptorBuilder::FMemberProperty& MemberProperty = CustomProperties[CustomPropertyIt];
@@ -2584,4 +2608,13 @@ SIZE_T FReplicationStateDescriptorBuilder::CreateDescriptorsForClass(FResult& Cr
 	return CreatedDescriptors.Num();
 }
 
+void FReplicationStateDescriptorBuilder::InitCVarReplicateCustomDeltaPropertiesInRepIndexOrder()
+{
+	if (CVarReplicateCustomDeltaPropertiesInRepIndexOrder == nullptr)
+	{
+		CVarReplicateCustomDeltaPropertiesInRepIndexOrder = IConsoleManager::Get().FindConsoleVariable(TEXT("net.ReplicateCustomDeltaPropertiesInRepIndexOrder"));
+		ensureMsgf(CVarReplicateCustomDeltaPropertiesInRepIndexOrder != nullptr, TEXT("Unable to find cvar net.ReplicateCustomDeltaPropertiesInRepIndexOrder"));
+	}
+}
+
 }
diff --git a/Engine/Source/Runtime/Experimental/Iris/Core/Public/Iris/ReplicationState/ReplicationStateDescriptorBuilder.h b/Engine/Source/Runtime/Experimental/Iris/Core/Public/Iris/ReplicationState/ReplicationStateDescriptorBuilder.h
--- a/Engine/Source/Runtime/Experimental/Iris/Core/Public/Iris/ReplicationState/ReplicationStateDescriptorBuilder.h
+++ b/Engine/Source/Runtime/Experimental/Iris/Core/Public/Iris/ReplicationState/ReplicationStateDescriptorBuilder.h
@@ -7,6 +7,7 @@
 #include "Templates/RefCounting.h"
 #include "UObject/CoreNet.h"
 
+class IConsoleVariable;
 namespace UE::Net::Private
 {
 	class FReplicationStateDescriptorRegistry;
@@ -67,6 +68,11 @@ public:
 	 * Any new created ReplicationStateDescriptor will be registered in the provided registry.
 	 */
 	IRISCORE_API static TRefCountPtr<const FReplicationStateDescriptor> CreateDescriptorForFunction(const UFunction* Function, const FParameters& Parameters = FParameters());
+
+private:
+	static void InitCVarReplicateCustomDeltaPropertiesInRepIndexOrder();
+
+	static const IConsoleVariable* CVarReplicateCustomDeltaPropertiesInRepIndexOrder;
 };
 
 }
