DataflowNode: Fix for property tooltip parsing

diff --git a/Engine/Source/Runtime/Experimental/Dataflow/Core/Private/Dataflow/DataflowNode.cpp b/Engine/Source/Runtime/Experimental/Dataflow/Core/Private/Dataflow/DataflowNode.cpp
--- a/Engine/Source/Runtime/Experimental/Dataflow/Core/Private/Dataflow/DataflowNode.cpp
+++ b/Engine/Source/Runtime/Experimental/Dataflow/Core/Private/Dataflow/DataflowNode.cpp
@@ -384,6 +384,12 @@ FString FDataflowNode::GetToolTip()
 					{
 						TArray<FString> OutArr;
 						ToolTipStr.ParseIntoArray(OutArr, TEXT(":\r\n"));
+						
+						if (OutArr.Num() == 0)
+						{
+							break;
+						}
+						const FString& MainTooltipText = (OutArr.Num() > 1) ? OutArr[1] : OutArr[0];
 
 						if (Property->HasMetaData(FDataflowNode::DataflowInput) &&
 							Property->HasMetaData(FDataflowNode::DataflowOutput) &&
@@ -391,34 +397,34 @@ FString FDataflowNode::GetToolTip()
 						{
 							if (Property->HasMetaData(FDataflowNode::DataflowIntrinsic))
 							{
-								InputsStr.Appendf(TEXT("    %s [Intrinsic] - %s\n"), *Property->GetName(), *OutArr[1]);
+								InputsStr.Appendf(TEXT("    %s [Intrinsic] - %s\n"), *Property->GetName(), *MainTooltipText);
 							}
 							else
 							{
-								InputsStr.Appendf(TEXT("    %s - %s\n"), *Property->GetName(), *OutArr[1]);
+								InputsStr.Appendf(TEXT("    %s - %s\n"), *Property->GetName(), *MainTooltipText);
 							}
 
-							OutputsStr.Appendf(TEXT("    %s [Passthrough] - %s\n"), *Property->GetName(), *OutArr[1]);
+							OutputsStr.Appendf(TEXT("    %s [Passthrough] - %s\n"), *Property->GetName(), *MainTooltipText);
 						}					
 						else if (Property->HasMetaData(FDataflowNode::DataflowInput))
 						{
 							if (Property->HasMetaData(FDataflowNode::DataflowIntrinsic))
 							{
-								InputsStr.Appendf(TEXT("    %s [Intrinsic] - %s\n"), *Property->GetName(), *OutArr[1]);
+								InputsStr.Appendf(TEXT("    %s [Intrinsic] - %s\n"), *Property->GetName(), *MainTooltipText);
 							}
 							else
 							{
-								InputsStr.Appendf(TEXT("    %s - %s\n"), *Property->GetName(), *OutArr[1]);
+								InputsStr.Appendf(TEXT("    %s - %s\n"), *Property->GetName(), *MainTooltipText);
 							}
 						}
 						else if (Property->HasMetaData(FDataflowNode::DataflowOutput))
 						{
-							OutputsStr.Appendf(TEXT("    %s - %s\n"), *Property->GetName(), *OutArr[1]);
+							OutputsStr.Appendf(TEXT("    %s - %s\n"), *Property->GetName(), *MainTooltipText);
 						}
 						
 						if (!Property->HasMetaData(FDataflowNode::DataflowInput) && !Property->HasMetaData(FDataflowNode::DataflowOutput))
 						{
-							InputsStr.Appendf(TEXT("    %s - %s\n"), *Property->GetName(), *OutArr[1]);
+							InputsStr.Appendf(TEXT("    %s - %s\n"), *Property->GetName(), *MainTooltipText);
 						}
 					}
 				}
