Fix crash when converting a geometry collectino to another one with missing InstancedMesh attribute on it

diff --git a/Engine/Source/Runtime/Experimental/GeometryCollectionEngine/Private/GeometryCollection/GeometryCollectionEngineConversion.cpp b/Engine/Source/Runtime/Experimental/GeometryCollectionEngine/Private/GeometryCollection/GeometryCollectionEngineConversion.cpp
--- a/Engine/Source/Runtime/Experimental/GeometryCollectionEngine/Private/GeometryCollection/GeometryCollectionEngineConversion.cpp
+++ b/Engine/Source/Runtime/Experimental/GeometryCollectionEngine/Private/GeometryCollection/GeometryCollectionEngineConversion.cpp
@@ -980,23 +980,25 @@ void FGeometryCollectionEngineConversion::AppendGeometryCollectionInstancedMeshe
 		const FCollectionInstancedMeshFacade SourceInstancedMeshFacade(*SourceGeometryCollectionPtr);
 		FCollectionInstancedMeshFacade TargetInstancedMeshFacade(*TargetGeometryCollectionPtr);
 
-		TargetInstancedMeshFacade.DefineSchema();
-
-		const int32 NumSourceIndices = SourceInstancedMeshFacade.GetNumIndices();
-		for (int32 SourceTransformIndex = 0; SourceTransformIndex < NumSourceIndices; SourceTransformIndex++)
+		if (SourceInstancedMeshFacade.IsValid())
 		{
-			int32 TargetInstancedMeshIndex = INDEX_NONE;
-			if (SourceInstancedMeshFacade.IsValid())
+			TargetInstancedMeshFacade.DefineSchema();
+
+			const int32 NumSourceIndices = SourceInstancedMeshFacade.GetNumIndices();
+			for (int32 SourceTransformIndex = 0; SourceTransformIndex < NumSourceIndices; SourceTransformIndex++)
 			{
+				int32 TargetInstancedMeshIndex = INDEX_NONE;
+
 				const int32 SourceAutoInstanceIndex = SourceInstancedMeshFacade.GetIndex(SourceTransformIndex);
 				if (SourceAutoInstanceIndex != INDEX_NONE)
 				{
 					const FGeometryCollectionAutoInstanceMesh& SourceAutoInstanceMesh = SourceGeometryCollectionObject->GetAutoInstanceMesh(SourceAutoInstanceIndex);
 					TargetInstancedMeshIndex = TargetGeometryCollectionObject->FindOrAddAutoInstanceMesh(SourceAutoInstanceMesh);
 				}
+
+				const int32 TargetTransformIndex = TargetTransformStartIndex + SourceTransformIndex;
+				TargetInstancedMeshFacade.SetIndex(TargetTransformIndex, TargetInstancedMeshIndex);
 			}
-			const int32 TargetTransformIndex = TargetTransformStartIndex + SourceTransformIndex;
-			TargetInstancedMeshFacade.SetIndex(TargetTransformIndex, TargetInstancedMeshIndex);
 		}
 	}
 }
