Geometry collection : fix potential crash when geometry collection is deleted at runtime while removal events are still being processed

diff --git a/Engine/Source/Runtime/Experimental/ChaosSolverEngine/Private/Chaos/ChaosGameplayEventDispatcher.cpp b/Engine/Source/Runtime/Experimental/ChaosSolverEngine/Private/Chaos/ChaosGameplayEventDispatcher.cpp
--- a/Engine/Source/Runtime/Experimental/ChaosSolverEngine/Private/Chaos/ChaosGameplayEventDispatcher.cpp
+++ b/Engine/Source/Runtime/Experimental/ChaosSolverEngine/Private/Chaos/ChaosGameplayEventDispatcher.cpp
@@ -501,11 +501,13 @@ void UChaosGameplayEventDispatcher::HandleBreakingEvents(const Chaos::FBreakingE
 		const int32 NumBreaks = BreakingDataArray.Num();
 		if (NumBreaks > 0)
 		{
+			const FPhysScene& Scene = *(GetWorld()->GetPhysicsScene());
+
 			for (const Chaos::FBreakingData& BreakingData : BreakingDataArray)
 			{	
 				if (BreakingData.Proxy)
 				{
-					UPrimitiveComponent* const PrimComp = Cast<UPrimitiveComponent>(BreakingData.Proxy->GetOwner());
+					UPrimitiveComponent* const PrimComp = Scene.GetOwningComponent<UPrimitiveComponent>(BreakingData.Proxy);
 					if (PrimComp && BreakEventRegistrations.Contains(PrimComp))
 					{
 						// queue them up so we can release the physics data before trigging BP events
@@ -523,7 +525,7 @@ void UChaosGameplayEventDispatcher::HandleBreakingEvents(const Chaos::FBreakingE
 
 void UChaosGameplayEventDispatcher::HandleSleepingEvents(const Chaos::FSleepingEventData& SleepingData)
 {
-	FPhysScene& Scene = *(GetWorld()->GetPhysicsScene());
+	const FPhysScene& Scene = *(GetWorld()->GetPhysicsScene());
 
 	const Chaos::FSleepingDataArray& SleepingArray = SleepingData.SleepingData;
 
@@ -566,11 +568,13 @@ void UChaosGameplayEventDispatcher::HandleRemovalEvents(const Chaos::FRemovalEve
 		const int32 NumRemovals = RemovalData.Num();
 		if (NumRemovals > 0)
 		{
+			const FPhysScene& Scene = *(GetWorld()->GetPhysicsScene());
+
 			for (Chaos::FRemovalData const& RemovalDataItem : RemovalData)
 			{
 				if (RemovalDataItem.Proxy)
 				{
-					UPrimitiveComponent* const PrimComp = Cast<UPrimitiveComponent>(RemovalDataItem.Proxy->GetOwner());
+					UPrimitiveComponent* const PrimComp = Scene.GetOwningComponent<UPrimitiveComponent>(RemovalDataItem.Proxy);
 					if (PrimComp && RemovalEventRegistrations.Contains(PrimComp))
 					{
 						// queue them up so we can release the physics data before trigging BP events
@@ -601,13 +605,15 @@ void UChaosGameplayEventDispatcher::HandleCrumblingEvents(const Chaos::FCrumblin
 
 		Chaos::FCrumblingDataArray const& BreakingData = Event.CrumblingData.AllCrumblingsArray;
 
+		const FPhysScene& Scene = *(GetWorld()->GetPhysicsScene());
+
 		// let's assume crumbles are rare, so we will iterate breaks instead of registered components for now
 		TArray<FChaosCrumblingEvent> PendingCrumblingEvent;
 		for (const Chaos::FCrumblingData& CrumblingDataItem : Event.CrumblingData.AllCrumblingsArray)
 		{	
 			if (CrumblingDataItem.Proxy)
 			{
-				if (UPrimitiveComponent* const PrimComp = Cast<UPrimitiveComponent>(CrumblingDataItem.Proxy->GetOwner()))
+				if (UPrimitiveComponent* const PrimComp = Scene.GetOwningComponent<UPrimitiveComponent>(CrumblingDataItem.Proxy))
 				{
 					if (CrumblingEventRegistrations.Contains(PrimComp))
 					{
