fix an issue where a geometry collection would be set kinematically in the distance if FGeometryCollectionPhysicsProxy::SetWorldTransform_External was set prior to the initialization of its particles on the physics thread

diff --git a/Engine/Source/Runtime/Experimental/Chaos/Private/PhysicsProxy/GeometryCollectionPhysicsProxy.cpp b/Engine/Source/Runtime/Experimental/Chaos/Private/PhysicsProxy/GeometryCollectionPhysicsProxy.cpp
--- a/Engine/Source/Runtime/Experimental/Chaos/Private/PhysicsProxy/GeometryCollectionPhysicsProxy.cpp
+++ b/Engine/Source/Runtime/Experimental/Chaos/Private/PhysicsProxy/GeometryCollectionPhysicsProxy.cpp
@@ -896,6 +896,14 @@ void FGeometryCollectionPhysicsProxy::InitializeBodiesPT(Chaos::FPBDRigidsSolver
 		const TManagedArray<TSet<int32>>& Children = DynamicCollection.Children;
 		const TManagedArray<int32>& Parent = DynamicCollection.Parent;
 
+		// In PushToPhysicsState, we're going to compute a relative transform from Parameters.PrevWorldTransform
+		// to a particle's current world transform to get its relative transform. Then, that relative transform is
+		// applied on top of the new Parameters.WorldTransform to get the final world transform. This is problematic
+		// if the particle is initialized with some Parameters.WorldTransform because then the particle's world transform
+		// will have Parameters.WorldTransform baked in but Parameters.PrevWorldTransform will be an identity. Then
+		// when we go to set the kinematic target we will then effectively be applying Parameters.WorldTransform twice.
+		Parameters.PrevWorldTransform = Parameters.WorldTransform;
+
 		TArray<FTransform> Transform;
 		GeometryCollectionAlgo::GlobalMatrices(DynamicCollection.Transform, Parent, Transform);
 
