Iris
- Fixed issue with invalid FastArrays still using custom fragments

diff --git a/Engine/Source/Runtime/Experimental/Iris/Core/Private/Iris/ReplicationState/ReplicationStateDescriptorBuilder.cpp b/Engine/Source/Runtime/Experimental/Iris/Core/Private/Iris/ReplicationState/ReplicationStateDescriptorBuilder.cpp
--- a/Engine/Source/Runtime/Experimental/Iris/Core/Private/Iris/ReplicationState/ReplicationStateDescriptorBuilder.cpp
+++ b/Engine/Source/Runtime/Experimental/Iris/Core/Private/Iris/ReplicationState/ReplicationStateDescriptorBuilder.cpp
@@ -75,7 +75,8 @@ enum class EMemberPropertyTraits : uint32
 	IsFastArray							= IsSourceTriviallyDestructible << 1U,
 	IsNativeFastArray					= IsFastArray << 1U,
 	IsFastArrayItem						= IsNativeFastArray << 1U,
-	HasConnectionSpecificSerialization	= IsFastArrayItem << 1U,
+	IsInvalidFastArray					= IsFastArrayItem << 1U,
+	HasConnectionSpecificSerialization	= IsInvalidFastArray << 1U,
 	HasPushBasedDirtiness				= HasConnectionSpecificSerialization << 1U,
 	UseSerializerIsEqual				= HasPushBasedDirtiness << 1U,
 };
@@ -2035,11 +2036,11 @@ EMemberPropertyTraits FPropertyReplicationStateDescriptorBuilder::GetFastArrayPr
 
 			if (bUseNativeFastArray && Struct->IsChildOf(FIrisFastArraySerializer::StaticStruct()))
 			{
-				return IsFastArraySupported(Struct) ? EMemberPropertyTraits::HasPushBasedDirtiness | EMemberPropertyTraits::IsFastArray | EMemberPropertyTraits::IsNativeFastArray : EMemberPropertyTraits::None;
+				return IsFastArraySupported(Struct) ? EMemberPropertyTraits::HasPushBasedDirtiness | EMemberPropertyTraits::IsFastArray | EMemberPropertyTraits::IsNativeFastArray : EMemberPropertyTraits::IsInvalidFastArray;
 			}
 			else if (Struct->IsChildOf(FFastArraySerializer::StaticStruct()))
 			{
-				return IsFastArraySupported(Struct) ? EMemberPropertyTraits::IsFastArray : EMemberPropertyTraits::None;
+				return IsFastArraySupported(Struct) ? EMemberPropertyTraits::IsFastArray : EMemberPropertyTraits::IsInvalidFastArray;
 			}
 			else if (Struct->IsChildOf(FFastArraySerializerItem::StaticStruct()))
 			{
@@ -2423,9 +2424,14 @@ SIZE_T FReplicationStateDescriptorBuilder::CreateDescriptorsForClass(FResult& Cr
 					continue;
 				}
 
-				// We build separate descriptors for all properties with CustomReplicationFragments
-				if (MemberProperty.CreateAndRegisterReplicationFragmentFunction)
+				if (EnumHasAnyFlags(MemberProperty.Traits, EMemberPropertyTraits::IsInvalidFastArray))
 				{
+					// Invalid fastarrays are treated as a normal struct property
+					MemberProperty.CreateAndRegisterReplicationFragmentFunction = nullptr;
+				}
+				else if (MemberProperty.CreateAndRegisterReplicationFragmentFunction)
+				{
+					// We build separate descriptors for all properties with CustomReplicationFragments					
 					CustomProperties.Add(MemberProperty);
 					continue;
 				}
