geometry collection : only recreate asset simulation data when necessary ( only applies when in editor )
- avoid calling it over and over when components using the same asset are registered
- make creation of physics state faster
- make duplication of components faster

diff --git a/Engine/Source/Runtime/Experimental/GeometryCollectionEngine/Private/GeometryCollection/GeometryCollectionComponent.cpp b/Engine/Source/Runtime/Experimental/GeometryCollectionEngine/Private/GeometryCollection/GeometryCollectionComponent.cpp
--- a/Engine/Source/Runtime/Experimental/GeometryCollectionEngine/Private/GeometryCollection/GeometryCollectionComponent.cpp
+++ b/Engine/Source/Runtime/Experimental/GeometryCollectionEngine/Private/GeometryCollection/GeometryCollectionComponent.cpp
@@ -2454,7 +2454,7 @@ void UGeometryCollectionComponent::OnCreatePhysicsState()
 		{
 			//hack: find a better place for this
 			UGeometryCollection* RestCollectionMutable = const_cast<UGeometryCollection*>(ToRawPtr(RestCollection));
-			RestCollectionMutable->CreateSimulationData();
+			RestCollectionMutable->CreateSimulationDataIfNeeded();
 		}
 #endif
 		const bool bValidWorld = GetWorld() && GetWorld()->IsGameWorld();
diff --git a/Engine/Source/Runtime/Experimental/GeometryCollectionEngine/Private/GeometryCollection/GeometryCollectionObject.cpp b/Engine/Source/Runtime/Experimental/GeometryCollectionEngine/Private/GeometryCollection/GeometryCollectionObject.cpp
--- a/Engine/Source/Runtime/Experimental/GeometryCollectionEngine/Private/GeometryCollection/GeometryCollectionObject.cpp
+++ b/Engine/Source/Runtime/Experimental/GeometryCollectionEngine/Private/GeometryCollection/GeometryCollectionObject.cpp
@@ -53,6 +53,12 @@ FAutoConsoleVariableRef CVarGeometryCollectionEnableForcedConvexGenerationInSeri
 	bGeometryCollectionEnableForcedConvexGenerationInSerialize,
 	TEXT("Enable generation of convex geometry on older destruction files.[def:true]"));
 
+bool bGeometryCollectionAlwaysRecreateSimulationData = false;
+FAutoConsoleVariableRef CVarGeometryCollectionAlwaysRecreateSimulationData(
+	TEXT("p.GeometryCollectionAlwaysRecreateSimulationData"),
+	bGeometryCollectionAlwaysRecreateSimulationData,
+	TEXT("always recreate the simulation data even if the simulation data is not marked as dirty - this has runtime cost in editor - only use as a last resort if default has issues [def:false]"));
+
 
 #if ENABLE_COOK_STATS
 namespace GeometryCollectionCookStats
@@ -953,6 +959,14 @@ void UGeometryCollection::CreateSimulationData()
 	SimulationDataGuid = StateGuid;
 }
 
+void UGeometryCollection::CreateSimulationDataIfNeeded()
+{
+	if (IsSimulationDataDirty() || bGeometryCollectionAlwaysRecreateSimulationData)
+	{
+		CreateSimulationData();
+	}
+}
+
 TUniquePtr<FGeometryCollectionNaniteData> UGeometryCollection::CreateNaniteData(FGeometryCollection* Collection)
 {
 	TUniquePtr<FGeometryCollectionNaniteData> NaniteData;
diff --git a/Engine/Source/Runtime/Experimental/GeometryCollectionEngine/Public/GeometryCollection/GeometryCollectionObject.h b/Engine/Source/Runtime/Experimental/GeometryCollectionEngine/Public/GeometryCollection/GeometryCollectionObject.h
--- a/Engine/Source/Runtime/Experimental/GeometryCollectionEngine/Public/GeometryCollection/GeometryCollectionObject.h
+++ b/Engine/Source/Runtime/Experimental/GeometryCollectionEngine/Public/GeometryCollection/GeometryCollectionObject.h
@@ -437,9 +437,16 @@ public:
 	/** If this flag is set, we only regenerate simulation data when requested via CreateSimulationData() */
 	bool bManualDataCreate;
 	
-	/** Create the simulation data that can be shared among all instances (mass, volume, etc...)*/
+	/** 
+	* Create the simulation data that can be shared among all instances (mass, volume, etc...) 
+	* Note : this does not check if the simulation data is drty or not and will cause a load from the DDC
+	* use CreateSimulationDataIfNeeded() for avoiding extra runtime cost 
+	*/
 	void CreateSimulationData();
 
+	/** Create the simulation data ( calls CreateSimulationData) only if the simulation data is dirty */
+	void CreateSimulationDataIfNeeded();
+
 	/** Create the Nanite rendering data. */
 	static TUniquePtr<FGeometryCollectionNaniteData> CreateNaniteData(FGeometryCollection* Collection);
 #endif
