headlesschaos : fix GeometryCollectionTestProximity
- caused by an undefined value
- make sure Proximity parameter has a default value

diff --git a/Engine/Source/Runtime/Experimental/Chaos/Public/GeometryCollection/GeometryCollectionProximityPropertiesInterface.h b/Engine/Source/Runtime/Experimental/Chaos/Public/GeometryCollection/GeometryCollectionProximityPropertiesInterface.h
--- a/Engine/Source/Runtime/Experimental/Chaos/Public/GeometryCollection/GeometryCollectionProximityPropertiesInterface.h
+++ b/Engine/Source/Runtime/Experimental/Chaos/Public/GeometryCollection/GeometryCollectionProximityPropertiesInterface.h
@@ -32,7 +32,7 @@ public :
 
 	struct FProximityProperties
 	{
-		EProximityMethod Method;
+		EProximityMethod Method = EProximityMethod::Precise;
 		float DistanceThreshold = 1.0f;
 		float RequireContactAmount = 0.0f;
 		EProximityContactMethod ContactMethod = EProximityContactMethod::MinOverlapInProjectionToMajorAxes;
