ChaosFlesh : Add blended XPBD corotated constraint.

diff --git a/Engine/Plugins/Experimental/ChaosFlesh/Source/ChaosFleshEngine/Private/ChaosFlesh/ChaosDeformableSolverComponent.cpp b/Engine/Plugins/Experimental/ChaosFlesh/Source/ChaosFleshEngine/Private/ChaosFlesh/ChaosDeformableSolverComponent.cpp
--- a/Engine/Plugins/Experimental/ChaosFlesh/Source/ChaosFleshEngine/Private/ChaosFlesh/ChaosDeformableSolverComponent.cpp
+++ b/Engine/Plugins/Experimental/ChaosFlesh/Source/ChaosFleshEngine/Private/ChaosFlesh/ChaosDeformableSolverComponent.cpp
@@ -208,6 +208,9 @@ void UDeformableSolverComponent::Reset()
 		, bUseGridBasedConstraints
 		, GridDx
 		, bDoQuasistatics
+		, EMesh
+		, bDoBlended
+		, BlendedZeta
 		}));
 
 	
diff --git a/Engine/Plugins/Experimental/ChaosFlesh/Source/ChaosFleshEngine/Public/ChaosFlesh/ChaosDeformableSolverComponent.h b/Engine/Plugins/Experimental/ChaosFlesh/Source/ChaosFleshEngine/Public/ChaosFlesh/ChaosDeformableSolverComponent.h
--- a/Engine/Plugins/Experimental/ChaosFlesh/Source/ChaosFleshEngine/Public/ChaosFlesh/ChaosDeformableSolverComponent.h
+++ b/Engine/Plugins/Experimental/ChaosFlesh/Source/ChaosFleshEngine/Public/ChaosFlesh/ChaosDeformableSolverComponent.h
@@ -95,6 +95,15 @@ public:
 	UPROPERTY(EditAnywhere, Category = Chaos)
 		bool bDoQuasistatics = false;
 
+	UPROPERTY(EditAnywhere, Category = Chaos)
+		float EMesh = 100000;
+
+	UPROPERTY(EditAnywhere, Category = Chaos)
+		bool bDoBlended = false;
+
+	UPROPERTY(EditAnywhere, Category = Chaos)
+		float BlendedZeta = 0;
+
 	//UPROPERTY(EditAnywhere, Category = Chaos)
 	//	bool bWaitForParallelDeformableTask = true;
 
diff --git a/Engine/Source/Runtime/Experimental/Chaos/Private/Chaos/Deformable/ChaosDeformableSolver.cpp b/Engine/Source/Runtime/Experimental/Chaos/Private/Chaos/Deformable/ChaosDeformableSolver.cpp
--- a/Engine/Source/Runtime/Experimental/Chaos/Private/Chaos/Deformable/ChaosDeformableSolver.cpp
+++ b/Engine/Source/Runtime/Experimental/Chaos/Private/Chaos/Deformable/ChaosDeformableSolver.cpp
@@ -252,25 +252,51 @@ namespace Chaos::Softs
 			}
 
 
-			FXPBDCorotatedConstraints<FSolverReal, FSolverParticles>* CorotatedConstraint =
-				new FXPBDCorotatedConstraints<FSolverReal, FSolverParticles>(
-					Evolution->Particles(), Elements, /*bRecordMetric = */false, (FSolverReal)1000000.0);
-
 			int32 InitIndex = Evolution->AddConstraintInitRange(1, true);
-			Evolution->ConstraintInits()[InitIndex] =
-				[CorotatedConstraint](FSolverParticles& InParticles, const FSolverReal Dt)
+			int32 ConstraintIndex = Evolution->AddConstraintRuleRange(1, true);
+
+			if (Property.bDoBlended) 
 			{
-				CorotatedConstraint->Init();
-			};
+				FBlendedXPBDCorotatedConstraints<FSolverReal, FSolverParticles>* BlendedCorotatedConstraint =
+					new FBlendedXPBDCorotatedConstraints<FSolverReal, FSolverParticles>(
+						Evolution->Particles(), Elements, /*bRecordMetric = */false, Property.EMesh, (FSolverReal).3, Property.BlendedZeta);
+			
+				Evolution->ConstraintInits()[InitIndex] =
+					[BlendedCorotatedConstraint](FSolverParticles& InParticles, const FSolverReal Dt)
+				{
+					BlendedCorotatedConstraint->Init();
+				};
 
+				Evolution->ConstraintRules()[ConstraintIndex] =
+					[BlendedCorotatedConstraint](FSolverParticles& InParticles, const FSolverReal Dt)
+				{
+					BlendedCorotatedConstraint->ApplyInParallel(InParticles, Dt);
+				};
+			
+				BlendedCorotatedConstraints.Add(TUniquePtr<FBlendedXPBDCorotatedConstraints<FSolverReal, FSolverParticles>>(BlendedCorotatedConstraint));
 
-			int32 ConstraintIndex = Evolution->AddConstraintRuleRange(1, true);
-			Evolution->ConstraintRules()[ConstraintIndex] =
-				[CorotatedConstraint](FSolverParticles& InParticles, const FSolverReal Dt)
+			}
+			else
 			{
-				CorotatedConstraint->ApplyInParallel(InParticles, Dt);
-			};
-			CorotatedConstraints.Add(TUniquePtr<FXPBDCorotatedConstraints<FSolverReal, FSolverParticles>>(CorotatedConstraint));
+				FXPBDCorotatedConstraints<FSolverReal, FSolverParticles>* CorotatedConstraint =
+					new FXPBDCorotatedConstraints<FSolverReal, FSolverParticles>(
+						Evolution->Particles(), Elements, /*bRecordMetric = */false, Property.EMesh);
+
+				Evolution->ConstraintInits()[InitIndex] =
+					[CorotatedConstraint](FSolverParticles& InParticles, const FSolverReal Dt)
+				{
+					CorotatedConstraint->Init();
+				};
+
+				Evolution->ConstraintRules()[ConstraintIndex] =
+					[CorotatedConstraint](FSolverParticles& InParticles, const FSolverReal Dt)
+				{
+					CorotatedConstraint->ApplyInParallel(InParticles, Dt);
+				};
+
+				CorotatedConstraints.Add(TUniquePtr<FXPBDCorotatedConstraints<FSolverReal, FSolverParticles>>(CorotatedConstraint));
+			}
+		
 		}
 	}
 
diff --git a/Engine/Source/Runtime/Experimental/Chaos/Public/Chaos/BlendedXPBDCorotatedConstraints.h b/Engine/Source/Runtime/Experimental/Chaos/Public/Chaos/BlendedXPBDCorotatedConstraints.h
new file mode 100644
--- /dev/null
+++ b/Engine/Source/Runtime/Experimental/Chaos/Public/Chaos/BlendedXPBDCorotatedConstraints.h
@@ -0,0 +1,228 @@
+// Copyright Epic Games, Inc. All Rights Reserved.
+#pragma once
+
+#include "Chaos/PBDSpringConstraintsBase.h"
+#include "ChaosStats.h"
+#include "Chaos/ImplicitQRSVD.h"
+#include "Chaos/GraphColoring.h"
+#include "Chaos/Framework/Parallel.h"
+#include "Chaos/XPBDCorotatedConstraints.h"
+
+
+namespace Chaos::Softs
+{
+	template <typename T, typename ParticleType>
+	class FBlendedXPBDCorotatedConstraints : public FXPBDCorotatedConstraints <T, ParticleType>
+	{
+
+		typedef FXPBDCorotatedConstraints<T, ParticleType> Base;
+		using Base::MeshConstraints;
+		using Base::LambdaArray;
+		using Base::Measure;
+		using Base::DmInverse;
+		using Base::Lambda;
+		using Base::bRecordMetric;
+
+	public:
+		//this one only accepts tetmesh input and mesh
+		FBlendedXPBDCorotatedConstraints(
+			const ParticleType& InParticles,
+			const TArray<TVector<int32, 4>>& InMesh,
+			const bool bRecordMetricIn = true,
+			const T& EMesh = (T)10.0,
+			const T& NuMesh = (T).3,
+			const T& InZeta = (T)1.
+		)
+			: Base(InParticles, InMesh, bRecordMetricIn, EMesh, NuMesh), Zeta(InZeta)
+		{
+			if (Zeta < (T)0.) 
+			{
+				Zeta = (T)0.;
+			}
+			if (Zeta > (T)1.) 
+			{
+				Zeta = (T)1.;
+			}
+
+			C1Contribution.Init(TVec4<TVector<T, 3>>(TVector<T, 3>((T)0.)), InMesh.Num());
+			C2Contribution.Init(TVec4<TVector<T, 3>>(TVector<T, 3>((T)0.)), InMesh.Num());
+
+		}
+
+		virtual ~FBlendedXPBDCorotatedConstraints() {}
+
+
+		virtual void ApplyInSerial(ParticleType& Particles, const T Dt, const int32 ElementIndex) const override
+		{
+			TVec4<TVector<T, 3>> PolarDelta = GetPolarDelta(Particles, Dt, ElementIndex);
+
+			for (int i = 0; i < 4; i++)
+			{
+				Particles.P(MeshConstraints[ElementIndex][i]) += PolarDelta[i];
+			}
+
+			TVec4<TVector<T, 3>> DetDelta = GetDeterminantDelta(Particles, Dt, ElementIndex);
+
+			for (int i = 0; i < 4; i++)
+			{
+				Particles.P(MeshConstraints[ElementIndex][i]) += DetDelta[i];
+			}
+		}
+
+		virtual void Init() const override
+		{
+			for (T& Lambdas : LambdaArray) { Lambdas = (T)0.; }
+			C1Contribution.Init(TVec4<TVector<T, 3>>(TVector<T, 3>((T)0.)), MeshConstraints.Num());
+			C2Contribution.Init(TVec4<TVector<T, 3>>(TVector<T, 3>((T)0.)), MeshConstraints.Num());
+		}
+
+	protected:
+
+		virtual TVec4<TVector<T, 3>> GetDeterminantDelta(const ParticleType& Particles, const T Dt, const int32 ElementIndex, const T Tol = (T)1e-3) const override
+		{
+			TRACE_CPUPROFILER_EVENT_SCOPE(TEXT("STAT_ChaosBlendedXPBDCorotatedApplyDet"));
+
+			const PMatrix<T, 3, 3> Fe = Base::F(ElementIndex, Particles);
+			PMatrix<T, 3, 3> DmInvT = Base::ElementDmInv(ElementIndex).GetTransposed();
+
+			T J = Fe.Determinant();
+			if (J - (T)1. < Tol)
+			{
+				return TVec4<TVector<T, 3>>(TVector<T, 3>((T)0.));
+			}
+
+			TVec4<TVector<T, 3>> dC2 = Base::GetDeterminantGradient(Fe, DmInvT);
+
+			T AlphaTilde = (T)2. / (Dt * Dt * Base::Lambda * Measure[ElementIndex]);
+
+			T DLambda = (1 - J) - AlphaTilde * LambdaArray[2 * ElementIndex + 1];
+
+			T Denom = AlphaTilde;
+			for (int i = 0; i < 4; i++)
+			{
+				for (int j = 0; j < 3; j++)
+				{
+					Denom += dC2[i][j] * Particles.InvM(MeshConstraints[ElementIndex][i]) * dC2[i][j];
+				}
+			}
+			DLambda /= Denom;
+			LambdaArray[2 * ElementIndex + 1] += DLambda;
+			TVec4<TVector<T, 3>> Delta2(TVector<T, 3>((T)0.));
+			for (int i = 0; i < 4; i++)
+			{
+				for (int j = 0; j < 3; j++)
+				{
+					Delta2[i][j] = Particles.InvM(MeshConstraints[ElementIndex][i]) * dC2[i][j] * DLambda;
+				}
+			}
+			TVec4<TVector<T, 3>> Delta1(TVector<T, 3>((T)0.));
+			for (int i = 0; i < 4; i++)
+			{
+				for (int j = 0; j < 3; j++)
+				{
+					Delta1[i][j] = Particles.InvM(MeshConstraints[ElementIndex][i]) 
+						* dC2[i][j] * LambdaArray[2 * ElementIndex + 1] - C2Contribution[ElementIndex][i][j];
+				}
+			}
+			TVec4<TVector<T, 3>> Delta(TVector<T, 3>((T)0.));
+			for (int i = 0; i < 4; i++)
+			{
+				Delta[i] = Zeta * Delta1[i] + ((T)1. - Zeta) * Delta2[i];
+			}
+
+			for (int i = 0; i < 4; i++)
+			{
+				C2Contribution[ElementIndex][i] += Delta[i];
+			}
+
+			return Delta;
+		}
+
+
+
+		virtual TVec4<TVector<T, 3>> GetPolarDelta(const ParticleType& Particles, const T Dt, const int32 ElementIndex, const T Tol = (T)1e-3) const override
+		{
+			TRACE_CPUPROFILER_EVENT_SCOPE(TEXT("STAT_ChaosBlendedXPBDCorotatedApplyPolar"));
+			SCOPE_CYCLE_COUNTER(STAT_ChaosXPBDCorotatedPolar);
+			const PMatrix<T, 3, 3> Fe = Base::F(ElementIndex, Particles);
+
+			PMatrix<T, 3, 3> Re((T)0.), Se((T)0.);
+
+			Chaos::PolarDecomposition(Fe, Re, Se);
+
+			T C1 = (T)0.;
+			for (int i = 0; i < 3; i++)
+			{
+				for (int j = 0; j < 3; j++)
+				{
+					C1 += FMath::Square((Fe - Re).GetAt(i, j));
+				}
+			}
+			C1 = FMath::Sqrt(C1);
+
+			if (C1 < Tol)
+			{
+				return TVec4<TVector<T, 3>>(TVector<T, 3>((T)0.));
+			}
+
+			PMatrix<T, 3, 3> DmInvT = Base::ElementDmInv(ElementIndex).GetTransposed();
+
+			TVec4<TVector<T, 3>> dC1 = Base::GetPolarGradient(Fe, Re, DmInvT, C1);
+
+			T AlphaTilde = (T)1. / (Dt * Dt * Base::Mu * Measure[ElementIndex]);
+
+			T DLambda = -C1 - AlphaTilde * LambdaArray[2 * ElementIndex + 0];
+
+			T Denom = AlphaTilde;
+			for (int i = 0; i < 4; i++)
+			{
+				for (int j = 0; j < 3; j++)
+				{
+					Denom += dC1[i][j] * Particles.InvM(MeshConstraints[ElementIndex][i]) * dC1[i][j];
+				}
+			}
+			DLambda /= Denom;
+			LambdaArray[2 * ElementIndex + 0] += DLambda;
+
+			TVec4<TVector<T, 3>> Delta2(TVector<T, 3>((T)0.));
+			for (int i = 0; i < 4; i++)
+			{
+				for (int j = 0; j < 3; j++)
+				{
+					Delta2[i][j] = Particles.InvM(MeshConstraints[ElementIndex][i]) * dC1[i][j] * DLambda;
+				}
+			}
+			TVec4<TVector<T, 3>> Delta1(TVector<T, 3>((T)0.));
+			for (int i = 0; i < 4; i++)
+			{
+				for (int j = 0; j < 3; j++)
+				{
+					Delta1[i][j] = Particles.InvM(MeshConstraints[ElementIndex][i])
+						* dC1[i][j] * LambdaArray[2 * ElementIndex + 0] - C1Contribution[ElementIndex][i][j];
+				}
+			}
+			TVec4<TVector<T, 3>> Delta(TVector<T, 3>((T)0.));
+			for (int i = 0; i < 4; i++)
+			{
+				Delta[i] = Zeta * Delta1[i] + ((T)1. - Zeta) * Delta2[i];
+			}
+
+			for (int i = 0; i < 4; i++)
+			{
+				C1Contribution[ElementIndex][i] += Delta[i];
+			}
+
+			return Delta;
+
+		}
+
+
+	private:
+
+		T Zeta;
+		mutable TArray<TVec4<TVector<T, 3>>> C1Contribution;
+		mutable TArray<TVec4<TVector<T, 3>>> C2Contribution;
+	};
+
+}  // End namespace Chaos::Softs
+
diff --git a/Engine/Source/Runtime/Experimental/Chaos/Public/Chaos/Deformable/ChaosDeformableSolver.h b/Engine/Source/Runtime/Experimental/Chaos/Public/Chaos/Deformable/ChaosDeformableSolver.h
--- a/Engine/Source/Runtime/Experimental/Chaos/Public/Chaos/Deformable/ChaosDeformableSolver.h
+++ b/Engine/Source/Runtime/Experimental/Chaos/Public/Chaos/Deformable/ChaosDeformableSolver.h
@@ -11,6 +11,7 @@
 #include "Chaos/PBDTriangleMeshCollisions.h"
 #include "Chaos/TriangleMesh.h"
 #include "Chaos/XPBDCorotatedConstraints.h"
+#include "Chaos/BlendedXPBDCorotatedConstraints.h"
 #include "Chaos/XPBDGridBasedCorotatedConstraints.h"
 #include "Chaos/Deformable/ChaosDeformableSolverTypes.h"
 #include "Chaos/Deformable/ChaosDeformableSolverProxy.h"
@@ -165,6 +166,7 @@ namespace Chaos::Softs
 		// Simulation Variables
 		TUniquePtr<Softs::FPBDEvolution> Evolution;
 		TArray<TUniquePtr<Softs::FXPBDCorotatedConstraints<Softs::FSolverReal, Softs::FSolverParticles>>> CorotatedConstraints;
+		TArray<TUniquePtr<Softs::FBlendedXPBDCorotatedConstraints<Softs::FSolverReal, Softs::FSolverParticles>>> BlendedCorotatedConstraints;
 		TUniquePtr<Softs::FXPBDGridBasedCorotatedConstraints<Softs::FSolverReal, Softs::FSolverParticles>> GridBasedCorotatedConstraint;
 		TUniquePtr<Softs::FPBDCollisionSpringConstraints> CollisionSpringConstraint;
 		TUniquePtr<Softs::FPBDTriangleMeshCollisions> TriangleMeshCollisions;
diff --git a/Engine/Source/Runtime/Experimental/Chaos/Public/Chaos/Deformable/ChaosDeformableSolverTypes.h b/Engine/Source/Runtime/Experimental/Chaos/Public/Chaos/Deformable/ChaosDeformableSolverTypes.h
--- a/Engine/Source/Runtime/Experimental/Chaos/Public/Chaos/Deformable/ChaosDeformableSolverTypes.h
+++ b/Engine/Source/Runtime/Experimental/Chaos/Public/Chaos/Deformable/ChaosDeformableSolverTypes.h
@@ -24,7 +24,10 @@ namespace Chaos::Softs
 			bool InbDoSelfCollision = false,
 			bool InbUseGridBasedConstraints = false,
 			FSolverReal InGridDx = (FSolverReal)1. ,
-			bool InbDoQuasistatics = false)
+			bool InbDoQuasistatics = false,
+			FSolverReal InEMesh = (FSolverReal)100000.,
+			bool InbDoBlended = false,
+			FSolverReal InBlendedZeta = (FSolverReal).1)
 			: NumSolverSubSteps(InNumSolverSubSteps)
 			, NumSolverIterations(InNumSolverIterations)
 			, FixTimeStep(InFixTimeStep)
@@ -36,6 +39,9 @@ namespace Chaos::Softs
 			, bUseGridBasedConstraints(InbUseGridBasedConstraints)
 			, GridDx(InGridDx)
 			, bDoQuasistatics(InbDoQuasistatics)
+			, EMesh(InEMesh)
+			, bDoBlended(InbDoBlended)
+			, BlendedZeta(InBlendedZeta)
 		{}
 
 		FDeformableSolverProperties(const FDeformableSolverProperties& InProp)
@@ -50,6 +56,9 @@ namespace Chaos::Softs
 			, bUseGridBasedConstraints(InProp.bUseGridBasedConstraints)
 			, GridDx(InProp.GridDx)
 			, bDoQuasistatics(InProp.bDoQuasistatics)
+			, EMesh(InProp.EMesh)
+			, bDoBlended(InProp.bDoBlended)
+			, BlendedZeta(InProp.BlendedZeta)
 		{}
 
 		int32 NumSolverSubSteps = 5;
@@ -63,6 +72,9 @@ namespace Chaos::Softs
 		bool bUseGridBasedConstraints = false;
 		FSolverReal GridDx = (FSolverReal)1.;
 		bool bDoQuasistatics = false;
+		FSolverReal EMesh = (FSolverReal)100000.;
+		bool bDoBlended = false;
+		FSolverReal BlendedZeta = (FSolverReal)0.;
 	};
 
 
diff --git a/Engine/Source/Runtime/Experimental/Chaos/Public/Chaos/XPBDCorotatedConstraints.h b/Engine/Source/Runtime/Experimental/Chaos/Public/Chaos/XPBDCorotatedConstraints.h
--- a/Engine/Source/Runtime/Experimental/Chaos/Public/Chaos/XPBDCorotatedConstraints.h
+++ b/Engine/Source/Runtime/Experimental/Chaos/Public/Chaos/XPBDCorotatedConstraints.h
@@ -17,8 +17,6 @@ namespace Chaos::Softs
 	class FXPBDCorotatedConstraints 
 	{
 
-		//TODO(Yizhou Chen): COuld be optimized . The SVD is using Chaos::Softs::Tdouble so float -> double -> float all the time
-		//should change data type in accordance 
 
 	public:
 		//this one only accepts tetmesh input and mesh
@@ -133,7 +131,7 @@ namespace Chaos::Softs
 			return DmInv;
 		}
 		
-		void Init() const 
+		virtual void Init() const 
 		{
 			for (T& Lambdas : LambdaArray) { Lambdas = (T)0.; }
 		}
@@ -321,7 +319,7 @@ namespace Chaos::Softs
 		
 
 
-		TVec4<TVector<T, 3>> GetDeterminantDelta(const ParticleType& Particles, const T Dt, const int32 ElementIndex, const T Tol = (T)1e-3) const
+		virtual TVec4<TVector<T, 3>> GetDeterminantDelta(const ParticleType& Particles, const T Dt, const int32 ElementIndex, const T Tol = (T)1e-3) const
 		{
 			TRACE_CPUPROFILER_EVENT_SCOPE(TEXT("STAT_ChaosXPBDCorotatedApplyDet"));
 			//SCOPE_CYCLE_COUNTER(STAT_ChaosXPBDCorotatedDet);
@@ -379,7 +377,7 @@ namespace Chaos::Softs
 
 
 
-		TVec4<TVector<T, 3>> GetPolarDelta(const ParticleType& Particles, const T Dt, const int32 ElementIndex, const T Tol = (T)1e-3) const
+		virtual TVec4<TVector<T, 3>> GetPolarDelta(const ParticleType& Particles, const T Dt, const int32 ElementIndex, const T Tol = (T)1e-3) const
 		{	
 			TRACE_CPUPROFILER_EVENT_SCOPE(TEXT("STAT_ChaosXPBDCorotatedApplyPolar"));
 			SCOPE_CYCLE_COUNTER(STAT_ChaosXPBDCorotatedPolar);
diff --git a/Engine/Source/Runtime/Experimental/Chaos/Public/Chaos/XPBDGridBasedCorotatedConstraints.h b/Engine/Source/Runtime/Experimental/Chaos/Public/Chaos/XPBDGridBasedCorotatedConstraints.h
--- a/Engine/Source/Runtime/Experimental/Chaos/Public/Chaos/XPBDGridBasedCorotatedConstraints.h
+++ b/Engine/Source/Runtime/Experimental/Chaos/Public/Chaos/XPBDGridBasedCorotatedConstraints.h
@@ -25,6 +25,7 @@ namespace Chaos::Softs
 		using Base::Measure;
 		using Base::DmInverse;
 		using Base::Mu;
+		using Base::Init;
 		using Base::Lambda;
 
 	public:
