Add method for smooth projection onto a triangulated surface.

diff --git a/Engine/Source/Programs/HeadlessChaos/Private/HeadlessChaosApp.cpp b/Engine/Source/Programs/HeadlessChaos/Private/HeadlessChaosApp.cpp
--- a/Engine/Source/Programs/HeadlessChaos/Private/HeadlessChaosApp.cpp
+++ b/Engine/Source/Programs/HeadlessChaos/Private/HeadlessChaosApp.cpp
@@ -14,6 +14,7 @@
 #include "HeadlessChaosTestRaycast.h"
 #include "HeadlessChaosTestSerialization.h"
 #include "HeadlessChaosTestSpatialHashing.h"
+#include "HeadlessChaosTestTriangleMesh.h"
 #include "Modules/ModuleManager.h"
 #include "RequiredProgramMainCPPInclude.h"
 #include "Chaos/PBDRigidsEvolution.h"
@@ -289,6 +290,11 @@ TEST(Handles, FrameworkTests)
 	ChaosTest::Handles::HandleSerializeTest();
 }
 
+TEST(TriangleMesh, TriangleMeshTests) {
+	ChaosTest::TriangleMeshProjectTest();
+	SUCCEED();
+}
+
 //TEST(Vehicle, VehicleTests) {
 //
 //	ChaosTest::SystemTemplateTest<float>();
diff --git a/Engine/Source/Programs/HeadlessChaos/Private/HeadlessChaosTestTriangleMesh.cpp b/Engine/Source/Programs/HeadlessChaos/Private/HeadlessChaosTestTriangleMesh.cpp
new file mode 100644
--- /dev/null
+++ b/Engine/Source/Programs/HeadlessChaos/Private/HeadlessChaosTestTriangleMesh.cpp
@@ -0,0 +1,126 @@
+// Copyright Epic Games, Inc. All Rights Reserved.
+
+#include "HeadlessChaosTestTriangleMesh.h"
+#include "HeadlessChaos.h"
+#include "HeadlessChaosTestUtility.h"
+#include "Chaos/Triangle.h"
+#include "Chaos/TriangleMesh.h"
+
+namespace ChaosTest
+{
+    using namespace Chaos;
+
+    void TriangleMeshProjectTest()
+    {
+		FTriangleMesh TriMesh;
+		TParticles<double, 3> Particles;
+		Particles.AddParticles(3);
+		Particles.X(0) = FVec3(0, 0, 0);
+		Particles.X(1) = FVec3(1, 0, 0);
+		Particles.X(2) = FVec3(0, 1, 0);
+		TArray<TVec3<int32>> Elements;
+		Elements.Add(TVec3<int32>(0, 1, 2));
+		TriMesh.Init(MoveTemp(Elements));
+
+		TTriangle<double> Tri(Particles.X(0), Particles.X(1), Particles.X(2));
+		Chaos::FVec3 Center = Tri.GetCentroid();
+		Chaos::FVec3 Norm = Tri.GetNormal();
+
+		TArray<TVec3<double>> FaceNormals;
+		FaceNormals.Add(Norm);
+		TArray<TVec3<double>> PointNormals;
+		PointNormals.Add(Norm);
+		PointNormals.Add(Norm);
+		PointNormals.Add(Norm);
+		EXPECT_TRUE(PointNormals.Num() > 0);
+
+		FTriangleMesh::TBVHType<double> BVH;
+		TriMesh.BuildBVH(TConstArrayView<TVec3<double>>(Particles.XArray()), BVH);
+
+		int32 HitTriIdx = INDEX_NONE;
+		FVec3 Weights;
+
+		// Center
+		
+		Chaos::FVec3 Pt = Center;
+		bool Hit = TriMesh.SmoothProject(BVH, Particles.XArray(), PointNormals, Pt, HitTriIdx, Weights);
+		EXPECT_TRUE(Hit);
+		EXPECT_VECTOR_NEAR(Weights, FVec3(1.0 / 3), 0.0001);
+
+		Pt = Center + Norm;
+		Hit = TriMesh.SmoothProject(BVH, Particles.XArray(), PointNormals, Pt, HitTriIdx, Weights);
+		EXPECT_TRUE(Hit);
+		EXPECT_VECTOR_NEAR(Weights, FVec3(1.0 / 3), 0.0001);
+
+		Pt = Center - Norm;
+		Hit = TriMesh.SmoothProject(BVH, Particles.XArray(), PointNormals, Pt, HitTriIdx, Weights);
+		EXPECT_TRUE(Hit);
+		EXPECT_VECTOR_NEAR(Weights, FVec3(1.0 / 3), 0.0001);
+
+		// Vertices
+
+		Pt = Particles.X(0);
+		Hit = TriMesh.SmoothProject(BVH, Particles.XArray(), PointNormals, Pt, HitTriIdx, Weights);
+		EXPECT_TRUE(Hit);
+		EXPECT_VECTOR_NEAR(Weights, FVec3(1,0,0), 0.01);
+
+		Pt = Particles.X(1);
+		Hit = TriMesh.SmoothProject(BVH, Particles.XArray(), PointNormals, Pt, HitTriIdx, Weights);
+		EXPECT_TRUE(Hit);
+		EXPECT_VECTOR_NEAR(Weights, FVec3(0, 1, 0), 0.01);
+
+		Pt = Particles.X(2);
+		Hit = TriMesh.SmoothProject(BVH, Particles.XArray(), PointNormals, Pt, HitTriIdx, Weights);
+		EXPECT_TRUE(Hit);
+		EXPECT_VECTOR_NEAR(Weights, FVec3(0, 0, 1), 0.01);
+
+		// Edges
+
+		Pt = (Particles.X(0) + Particles.X(1)) * 0.5;
+		Hit = TriMesh.SmoothProject(BVH, Particles.XArray(), PointNormals, Pt, HitTriIdx, Weights);
+		EXPECT_TRUE(Hit);
+		EXPECT_VECTOR_NEAR(Weights, FVec3(.5, .5, 0), 0.01);
+
+		Pt = (Particles.X(1) + Particles.X(2)) * 0.5;
+		Hit = TriMesh.SmoothProject(BVH, Particles.XArray(), PointNormals, Pt, HitTriIdx, Weights);
+		EXPECT_TRUE(Hit);
+		EXPECT_VECTOR_NEAR(Weights, FVec3(0, .5, .5), 0.01);
+
+		Pt = (Particles.X(0) + Particles.X(2)) * 0.5;
+		Hit = TriMesh.SmoothProject(BVH, Particles.XArray(), PointNormals, Pt, HitTriIdx, Weights);
+		EXPECT_TRUE(Hit);
+		EXPECT_VECTOR_NEAR(Weights, FVec3(.5, 0, .5), 0.01);
+
+		// Off vertices - Expect false
+
+		Pt = Center + (Particles.X(0) - Center) * 1.1;
+		Hit = TriMesh.SmoothProject(BVH, Particles.XArray(), PointNormals, Pt, HitTriIdx, Weights);
+		EXPECT_FALSE(Hit);
+
+		Pt = Center + (Particles.X(1) - Center) * 1.1;
+		Hit = TriMesh.SmoothProject(BVH, Particles.XArray(), PointNormals, Pt, HitTriIdx, Weights);
+		EXPECT_FALSE(Hit);
+
+		Pt = Center + (Particles.X(2) - Center) * 1.1;
+		Hit = TriMesh.SmoothProject(BVH, Particles.XArray(), PointNormals, Pt, HitTriIdx, Weights);
+		EXPECT_FALSE(Hit);
+
+		// Off edges - Expect false
+
+		Pt = (Particles.X(0) + Particles.X(1)) * 0.5;
+		Pt = Center + (Pt - Center) * 1.1;
+		Hit = TriMesh.SmoothProject(BVH, Particles.XArray(), PointNormals, Pt, HitTriIdx, Weights);
+		EXPECT_FALSE(Hit);
+
+		Pt = (Particles.X(1) + Particles.X(2)) * 0.5;
+		Pt = Center + (Pt - Center) * 1.1;
+		Hit = TriMesh.SmoothProject(BVH, Particles.XArray(), PointNormals, Pt, HitTriIdx, Weights);
+		EXPECT_FALSE(Hit);
+
+		Pt = (Particles.X(0) + Particles.X(2)) * 0.5;
+		Pt = Center + (Pt - Center) * 1.1;
+		Hit = TriMesh.SmoothProject(BVH, Particles.XArray(), PointNormals, Pt, HitTriIdx, Weights);
+		EXPECT_FALSE(Hit);
+	}
+
+} // namespace ChaosTest
diff --git a/Engine/Source/Programs/HeadlessChaos/Public/HeadlessChaosTestTriangleMesh.h b/Engine/Source/Programs/HeadlessChaos/Public/HeadlessChaosTestTriangleMesh.h
new file mode 100644
--- /dev/null
+++ b/Engine/Source/Programs/HeadlessChaos/Public/HeadlessChaosTestTriangleMesh.h
@@ -0,0 +1,12 @@
+// Copyright Epic Games, Inc. All Rights Reserved.
+
+#pragma once
+
+#include "HeadlessChaosTestUtility.h"
+#include "Modules/ModuleManager.h"
+
+namespace ChaosTest {
+
+	void TriangleMeshProjectTest();
+
+} // namespace ChaosTest
diff --git a/Engine/Source/Runtime/Experimental/Chaos/Private/Chaos/TriangleMesh.cpp b/Engine/Source/Runtime/Experimental/Chaos/Private/Chaos/TriangleMesh.cpp
--- a/Engine/Source/Runtime/Experimental/Chaos/Private/Chaos/TriangleMesh.cpp
+++ b/Engine/Source/Runtime/Experimental/Chaos/Private/Chaos/TriangleMesh.cpp
@@ -3,6 +3,7 @@
 
 #include "Chaos/Box.h"
 #include "Chaos/Defines.h"
+#include "Chaos/Math/Poisson.h"
 #include "Chaos/Plane.h"
 #include "Chaos/Triangle.h"
 #include "Chaos/TriangleCollisionPoint.h"
@@ -1596,6 +1597,207 @@ template bool FTriangleMesh::EdgeIntersectionQuery<FRealSingle>(const TBVHType<F
 template bool FTriangleMesh::EdgeIntersectionQuery<FRealDouble>(const TBVHType<FRealDouble>& BVH, const TConstArrayView<TVec3<FRealDouble>>& Points, const int32 EdgeIndex, const TVec3<FRealDouble>& EdgePosition1, const TVec3<FRealDouble>& EdgePosition2,
 	TFunctionRef<bool(const int32 EdgeIndex, const int32 TriangleIndex)> BroadphaseTest, TArray<TTriangleCollisionPoint<FRealDouble>>& Result) const;
 
+template <class T>
+bool SmoothProject(
+	const TConstArrayView<FVec3>& Points,
+	const TArray<TVec3<int32>>& Tris,
+	const TArray<FVec3>& PointNormals,
+	const FVec3& Pos,
+	const int32 TriIdx,
+	FVec3& Weights,
+	TArray<TArray<FVec3>>& TangentBases,
+	T* B,
+	T* C,
+	TVec3<T>& C0,
+	TVec3<T>& C1,
+	TVec3<T>& C2,
+	TVec3<T>& W0,
+	TVec3<T>& W1,
+	TVec3<T>& W2)
+{
+	// Construct orthogonal tangents from point normal
+	for (int32 i = 0; i < 3; i++)
+	{
+		const int32 PtIdx = Tris[TriIdx][i];
+		const FVec3& PointNormal = PointNormals[PtIdx];
+		TangentBases[i][0] = PointNormal.GetOrthogonalVector().GetSafeNormal();
+		TangentBases[i][1] = FVec3::CrossProduct(TangentBases[i][0], PointNormal);
+	}
+	// Solve
+	for (int32 i = 0; i < 3; i++)
+	{
+		for (int j = 0; j < 3; j++)
+		{
+			const int32 PtIdx = Tris[TriIdx][i];
+			const FVec3& Pt = Points[PtIdx];
+			const FVec3& TangentBase0 = TangentBases[i][0];
+			const FVec3& TangentBase1 = TangentBases[i][1];
+			FReal B0j = FVec3::DotProduct(TangentBase0, Pt);
+			FReal B1j = FVec3::DotProduct(TangentBase1, Pt);
+			RowMaj3x3Set(B, j, 0, static_cast<T>(B0j));
+			RowMaj3x3Set(B, j, 1, static_cast<T>(B1j));
+			RowMaj3x3Set(B, j, 2, static_cast<T>(1.0));
+		}
+		RowMaj3x3Inverse(B, B); // Sets B
+		FVec3 PV(
+			FVec3::DotProduct(TangentBases[i][0], Pos),
+			FVec3::DotProduct(TangentBases[i][1], Pos),
+			static_cast<T>(1.0));
+		RowMaj3x3SetRow(C, i, RowMaj3x3Multiply(B, PV));
+	}
+
+	RowMaj3x3Get(C, 0, 0) -= static_cast<T>(1.0);
+	RowMaj3x3Get(C, 1, 1) -= static_cast<T>(1.0);
+	RowMaj3x3Get(C, 2, 2) -= static_cast<T>(1.0);
+
+	RowMaj3x3GetRow(C, 0, C0);
+	RowMaj3x3GetRow(C, 1, C1);
+	RowMaj3x3GetRow(C, 2, C2);
+
+	W0 = FVec3::CrossProduct(C0, C1);
+	W1 = FVec3::CrossProduct(C0, C2);
+	W2 = FVec3::CrossProduct(C1, C2);
+
+	if (W0[0] >= 0.0 && W0[1] >= 0.0 && W0[2] >= 0.0)
+	{
+		Weights = W0;
+		Swap(Weights[0], Weights[2]);
+	}
+	else if (W1[0] >= 0.0 && W1[1] >= 0.0 && W1[2] >= 0.0)
+	{
+		Weights = W1;
+	}
+	else if (W2[0] >= 0.0 && W2[1] >= 0.0 && W2[2] >= 0.0)
+	{
+		Weights = W2;
+		Swap(Weights[0], Weights[2]);
+	}
+	else
+	{
+		// We're outside of the cone implied by the normals,
+		// so the projection doesn't exist.
+		return false;
+	}
+
+	FReal WeightsSum = Weights[0] + Weights[1] + Weights[2];
+	if (fabs(WeightsSum) < 1.0e-6)
+	{
+		// The point is at the vertex of the normal cone, so the
+		// projection isn't unique.  Any face location will do.
+		static const T OneThird = static_cast<T>(1.0 / 3);
+		Weights = FVec3(OneThird, OneThird, OneThird);
+		return true;
+	}
+	Weights /= WeightsSum;
+	return true;
+}
+
+template <class T>
+bool SmoothProject(
+	const TConstArrayView<FVec3>& Points,
+	const TArray<TVec3<int32>>& Tris,
+	const TArray<FVec3>& PointNormals,
+	const FVec3& Pos,
+	const int32 TriIdx,
+	FVec3& Weights)
+{
+	TArray<TArray<FVec3>> TangentBases;
+	TangentBases.SetNum(3);
+	for (int32 i = 0; i < 3; i++)
+	{
+		TangentBases[i].SetNum(2);
+	}
+	T B[9];
+	T C[9];
+	TVec3<T> C0;
+	TVec3<T> C1;
+	TVec3<T> C2;
+	TVec3<T> W0;
+	TVec3<T> W1;
+	TVec3<T> W2;
+
+	return SmoothProject(Points, Tris, PointNormals, Pos, TriIdx, Weights,
+		TangentBases, &B[0], &C[0], C0, C1, C2, W0, W1, W2);
+}
+
+template <class T>
+TArray<bool> SmoothProject(
+	const TConstArrayView<FVec3>& Points,
+	const TArray<TVec3<int32>>& Tris,
+	const TArray<FVec3>& PointNormals,
+	const FVec3& Pos,
+	const TArray<int32>& TriIdx,
+	TArray<FVec3>& Weights,
+	const bool UseFirstFound=false)
+{
+	TArray<TArray<FVec3>> TangentBases;
+	TangentBases.SetNum(3);
+	for (int32 i = 0; i < 3; i++)
+	{
+		TangentBases[i].SetNum(2);
+	}
+	T B[9];
+	T C[9];
+	TVec3<T> C0;
+	TVec3<T> C1;
+	TVec3<T> C2;
+	TVec3<T> W0;
+	TVec3<T> W1;
+	TVec3<T> W2;
+
+	TArray<bool> RetVal;
+	RetVal.SetNumZeroed(TriIdx.Num());
+	Weights.SetNum(TriIdx.Num());
+	for (int i = 0; i < TriIdx.Num(); i++)
+	{
+		RetVal[i] = SmoothProject(Points, Tris, PointNormals, Pos, TriIdx[i], Weights[i],
+			TangentBases, &B[0], &C[0], C0, C1, C2, W0, W1, W2);
+		if (UseFirstFound && RetVal[i])
+		{
+			return RetVal;
+		}
+	}
+	return RetVal;
+}
+
+template<typename T>
+bool FTriangleMesh::SmoothProject(
+	const TBVHType<T>& BVH, 
+	const TConstArrayView<FVec3>& Points,
+	const TArray<FVec3>& PointNormals,
+	const FVec3& Pos,
+	int32& TriangleIndex,
+	FVec3& Weights, 
+	const int32 MaxIters) const
+{
+	TArray<int32> CandidateTris;
+	int32 Iter = 0;
+	do {
+		FAABB3 Box(Pos-FVec3(.5*Iter), Pos+FVec3(.5*Iter));
+		CandidateTris = BVH.FindAllIntersections(Box);
+	} while (!CandidateTris.Num() && ++Iter < MaxIters);
+
+	if (!CandidateTris.Num())
+	{
+		return false;
+	}
+
+	TArray<FVec3> CandidateWeights;
+	TArray<bool> Success = Chaos::SmoothProject<T>(Points, MElements, PointNormals, Pos, CandidateTris, CandidateWeights, true);
+	for (int32 i = 0; i < Success.Num(); i++)
+	{
+		if (Success[i])
+		{
+			TriangleIndex = CandidateTris[i];
+			Weights = CandidateWeights[i];
+			return true;
+		}
+	}
+	return false;
+}
+template bool FTriangleMesh::SmoothProject<FRealSingle>(const TBVHType<FRealSingle>& BVH, const TConstArrayView<FVec3>& Points, const TArray<FVec3>& PointNormals, const FVec3& Point, int32& TriangleIndex, FVec3& Weights, const int32 MaxIters) const;
+template bool FTriangleMesh::SmoothProject<FRealDouble>(const TBVHType<FRealDouble>& BVH, const TConstArrayView<FVec3>& Points, const TArray<FVec3>& PointNormals, const FVec3& Point, int32& TriangleIndex, FVec3& Weights, const int32 MaxIters) const;
+
 template<typename T>
 void FTriangleMesh::BuildSpatialHash(const TConstArrayView<TVec3<T>>& Points, TSpatialHashType<T>& SpatialHash) const
 {
diff --git a/Engine/Source/Runtime/Experimental/Chaos/Public/Chaos/Math/Poisson.h b/Engine/Source/Runtime/Experimental/Chaos/Public/Chaos/Math/Poisson.h
--- a/Engine/Source/Runtime/Experimental/Chaos/Public/Chaos/Math/Poisson.h
+++ b/Engine/Source/Runtime/Experimental/Chaos/Public/Chaos/Math/Poisson.h
@@ -12,6 +12,84 @@
 
 namespace Chaos {
 
+// Set/Get i, j
+
+template <class T>
+inline void RowMaj3x3Set(T* A, const int32 i, const int32 j, const T Value)
+{
+	A[3 * i + j] = Value;
+}
+
+template <class T>
+inline const T& RowMaj3x3Get(const T* A, const int32 i, const int32 j)
+{
+	return A[3 * i + j];
+}
+
+template <class T>
+inline T& RowMaj3x3Get(T* A, const int32 i, const int32 j)
+{
+	return A[3 * i + j];
+}
+
+// Set/Get row
+
+template <class T>
+inline void RowMaj3x3SetRow(T* A, const int32 i, const T* Values)
+{
+	for (int32 j = 0; j < 3; j++)
+	{
+		RowMaj3x3Set(A, i, j, Values[i]);
+	}
+}
+
+template <class T, class TV>
+inline void RowMaj3x3SetRow(T* A, const int32 i, const TV Value)
+{
+	const T Tmp[3]{ static_cast<T>(Value[0]), static_cast<T>(Value[1]), static_cast<T>(Value[2]) };
+	RowMaj3x3SetRow(A, i, &Tmp[0]);
+}
+
+template <class T, class TV>
+inline void RowMaj3x3GetRow(T* A, const int32 i, TV& Row)
+{
+	for (int j = 0; j < 3; j++)
+	{
+		Row[j] = RowMaj3x3Get(A, i, j);
+	}
+}
+
+// Set/Get col
+
+template <class T>
+inline void RowMaj3x3SetCol(T* A, const int32 j, const T* Values)
+{
+	for (int32 i = 0; i < 3; i++)
+	{
+		RowMaj3x3Set(A, i, j, Values[i]);
+	}
+}
+
+template <class T, class TV>
+inline void RowMaj3x3SetCol(T* A, const int32 j, const TV Value)
+{
+	for (int32 i = 0; i < 3; i++)
+	{
+		RowMaj3x3Set(A, i, j, static_cast<T>(Value[i]));
+	}
+}
+
+template <class T, class TV>
+inline void RowMaj3x3GetCol(T* A, const int32 j, TV& Col)
+{
+	for (int i = 0; i < 3; i++)
+	{
+		Col[i] = RowMaj3x3Get(A, i, j);
+	}
+}
+
+// Determinant
+
 template <class T>
 inline T RowMaj3x3Determinant(const T A0, const T A1, const T A2, const T A3, const T A4, const T A5, const T A6, const T A7, const T A8)
 {
@@ -24,6 +102,8 @@ inline T RowMaj3x3Determinant(const T* A)
 	return RowMaj3x3Determinant(A[0], A[1], A[2], A[3], A[4], A[5], A[6], A[7], A[8]);
 }
 
+// Inverse
+
 template <class T>
 inline void RowMaj3x3Inverse(const T Det, const T A0, const T A1, const T A2, const T A3, const T A4, const T A5, const T A6, const T A7, const T A8, T* Inv)
 {
@@ -54,6 +134,14 @@ inline void RowMaj3x3Inverse(const T Det, const T* A, T* Inv)
 	RowMaj3x3Inverse(Det, A[0], A[1], A[2], A[3], A[4], A[5], A[6], A[7], A[8], Inv);
 }
 
+template <class T>
+inline void RowMaj3x3Inverse(const T* A, T* Inv)
+{
+	RowMaj3x3Inverse(RowMaj3x3Determinant(A), A, Inv);
+}
+
+// Transpose
+
 template <class T>
 inline void RowMaj3x3Transpose(const T* A, T* Transpose)
 {
@@ -480,4 +568,4 @@ ComputeFiberField(
 }
 
 
-} // namespace ChaosFlesh
\ No newline at end of file
+} // namespace ChaosFlesh
diff --git a/Engine/Source/Runtime/Experimental/Chaos/Public/Chaos/TriangleMesh.h b/Engine/Source/Runtime/Experimental/Chaos/Public/Chaos/TriangleMesh.h
--- a/Engine/Source/Runtime/Experimental/Chaos/Public/Chaos/TriangleMesh.h
+++ b/Engine/Source/Runtime/Experimental/Chaos/Public/Chaos/TriangleMesh.h
@@ -298,6 +298,11 @@ namespace Chaos
 		bool EdgeIntersectionQuery(const TBVHType<T>& BVH, const TConstArrayView<TVec3<T>>& Points, const int32 EdgeIndex, const TVec3<T>& EdgePosition1, const TVec3<T>& EdgePosition2,
 			TFunctionRef<bool(const int32 EdgeIndex, const int32 TriangleIndex)> BroadphaseTest, TArray<TTriangleCollisionPoint<T>>& Result) const;
 
+		//! Returns \c false if \p Point is outside of the smooth normal cone, where a smooth projection doesn't exist.
+		template<typename T>
+		bool SmoothProject(const TBVHType<T>& BVH, const TConstArrayView<FVec3>& Points, const TArray<FVec3>& PointNormals,
+			const FVec3& Point, int32& TriangleIndex, FVec3& Weights, const int32 MaxIters=10) const;
+
 		template<typename T>
 		using TSpatialHashType = THierarchicalSpatialHash<int32, T>;
 
