Remove dataflow from cook builds

diff --git a/Engine/Source/Runtime/Experimental/GeometryCollectionEngine/Private/GeometryCollection/GeometryCollectionObject.cpp b/Engine/Source/Runtime/Experimental/GeometryCollectionEngine/Private/GeometryCollection/GeometryCollectionObject.cpp
--- a/Engine/Source/Runtime/Experimental/GeometryCollectionEngine/Private/GeometryCollection/GeometryCollectionObject.cpp
+++ b/Engine/Source/Runtime/Experimental/GeometryCollectionEngine/Private/GeometryCollection/GeometryCollectionObject.cpp
@@ -709,6 +709,8 @@ void UGeometryCollection::Serialize(FArchive& Ar)
 	// The Geometry Collection we will be archiving. This may be replaced with a transient, stripped back Geometry Collection if we are cooking.
 	TSharedPtr<FGeometryCollection, ESPMode::ThreadSafe> ArchiveGeometryCollection = GeometryCollection;
 
+	TObjectPtr<UDataflow> StrippedDataflowAsset = nullptr;
+
 	bool bIsCookedOrCooking = Ar.IsCooking();
 	if (bIsCookedOrCooking && Ar.IsSaving())
 	{
@@ -734,6 +736,10 @@ void UGeometryCollection::Serialize(FArchive& Ar)
 				ArchiveGeometryCollection = CopyCollectionAndRemoveGeometry(GeometryCollection);
 			}
 		}
+
+		// The dataflow asset is only needed for the editor, so we just remove it when cooking 
+		StrippedDataflowAsset = DataflowAsset;
+		DataflowAsset = nullptr;
 #endif
 	}
 
@@ -776,6 +782,11 @@ void UGeometryCollection::Serialize(FArchive& Ar)
 		Super::Serialize(Ar);
 	}
 
+	// Important : this needs to remain after the call to Super::Serialize
+	if (StrippedDataflowAsset)
+	{
+		DataflowAsset = StrippedDataflowAsset;
+	}
 
 	if (!SizeSpecificData.Num())
 	{
