client geometry collections need to wait for its physics proxy to initialize prior to processing replication data

diff --git a/Engine/Source/Runtime/Experimental/GeometryCollectionEngine/Private/GeometryCollection/GeometryCollectionComponent.cpp b/Engine/Source/Runtime/Experimental/GeometryCollectionEngine/Private/GeometryCollection/GeometryCollectionComponent.cpp
--- a/Engine/Source/Runtime/Experimental/GeometryCollectionEngine/Private/GeometryCollection/GeometryCollectionComponent.cpp
+++ b/Engine/Source/Runtime/Experimental/GeometryCollectionEngine/Private/GeometryCollection/GeometryCollectionComponent.cpp
@@ -1668,7 +1668,7 @@ FAutoConsoleVariableRef CVarGeometryCollectionHardsnapThresholdMs(TEXT("p.Geomet
 void UGeometryCollectionComponent::ProcessRepData()
 {
 	using namespace Chaos;
-	if(VersionProcessed == RepData.Version || PhysicsProxy->GetReplicationMode() != FGeometryCollectionPhysicsProxy::EReplicationMode::Client)
+	if(!PhysicsProxy || !PhysicsProxy->IsInitializedOnPhysicsThread() || VersionProcessed == RepData.Version || PhysicsProxy->GetReplicationMode() != FGeometryCollectionPhysicsProxy::EReplicationMode::Client)
 	{
 		return;
 	}
@@ -1700,6 +1700,7 @@ void UGeometryCollectionComponent::ProcessRepData()
 	{
 		const FGeometryCollectionActivatedCluster& ActivatedCluster = RepData.OneOffActivated[OneOffActivatedProcessed];
 		FPBDRigidParticleHandle* OneOff = PhysicsProxy->GetParticles()[ActivatedCluster.ActivatedIndex];
+		check(OneOff);
 
 		// Set initial velocities if not hard snapping
 		if(!bHardSnap)
