group the "Notify" booleans in the geometry collection component into it a separate "Events" category

diff --git a/Engine/Source/Runtime/Experimental/GeometryCollectionEngine/Public/GeometryCollection/GeometryCollectionComponent.h b/Engine/Source/Runtime/Experimental/GeometryCollectionEngine/Public/GeometryCollection/GeometryCollectionComponent.h
--- a/Engine/Source/Runtime/Experimental/GeometryCollectionEngine/Public/GeometryCollection/GeometryCollectionComponent.h
+++ b/Engine/Source/Runtime/Experimental/GeometryCollectionEngine/Public/GeometryCollection/GeometryCollectionComponent.h
@@ -984,27 +984,27 @@ public:
 	virtual void DispatchChaosPhysicsCollisionBlueprintEvents(const FChaosPhysicsCollisionInfo& CollisionInfo) override;
 	
 	/** If true, this component will generate breaking events that other systems may subscribe to. */
-	UPROPERTY(EditAnywhere, BlueprintReadOnly, Category = "ChaosPhysics|General")
+	UPROPERTY(EditAnywhere, BlueprintReadOnly, Category = "ChaosPhysics|Events")
 	bool bNotifyBreaks;
 
 	/** If true, this component will generate collision events that other systems may subscribe to. */
-	UPROPERTY(EditAnywhere, BlueprintReadOnly, Category = "ChaosPhysics|General")
+	UPROPERTY(EditAnywhere, BlueprintReadOnly, Category = "ChaosPhysics|Events")
 	bool bNotifyCollisions;
 
 	/** If true, this component will generate trailing events that other systems may subscribe to. */
-	UPROPERTY(EditAnywhere, BlueprintReadOnly, Category = "ChaosPhysics|General")
+	UPROPERTY(EditAnywhere, BlueprintReadOnly, Category = "ChaosPhysics|Events")
 	bool bNotifyTrailing;
 
 	/** If true, this component will generate removal events that other systems may subscribe to. */
-	UPROPERTY(EditAnywhere, BlueprintReadOnly, Category = "ChaosPhysics|General")
+	UPROPERTY(EditAnywhere, BlueprintReadOnly, Category = "ChaosPhysics|Events")
 	bool bNotifyRemovals;
 
 	/** If true, this component will generate crumbling events that other systems may subscribe to. */
-	UPROPERTY(EditAnywhere, BlueprintReadOnly, Category = "ChaosPhysics|General")
+	UPROPERTY(EditAnywhere, BlueprintReadOnly, Category = "ChaosPhysics|Events")
 	bool bNotifyCrumblings;
 
 	/** If this and bNotifyCrumblings are true, the crumbling events will contain released children indices. */
-	UPROPERTY(EditAnywhere, BlueprintReadOnly, Category = "ChaosPhysics|General", meta = (EditCondition = "bNotifyCrumblings"))
+	UPROPERTY(EditAnywhere, BlueprintReadOnly, Category = "ChaosPhysics|Events", meta = (EditCondition = "bNotifyCrumblings"))
 	bool bCrumblingEventIncludesChildren;
 	
 	/** If true, this component will save linear and angular velocities on its DynamicCollection. */
