geometry collection should return the root bone if bone name is specified as none.

diff --git a/Engine/Source/Runtime/Experimental/GeometryCollectionEngine/Private/GeometryCollection/GeometryCollectionComponent.cpp b/Engine/Source/Runtime/Experimental/GeometryCollectionEngine/Private/GeometryCollection/GeometryCollectionComponent.cpp
--- a/Engine/Source/Runtime/Experimental/GeometryCollectionEngine/Private/GeometryCollection/GeometryCollectionComponent.cpp
+++ b/Engine/Source/Runtime/Experimental/GeometryCollectionEngine/Private/GeometryCollection/GeometryCollectionComponent.cpp
@@ -4651,6 +4651,22 @@ Chaos::FPhysicsObject* UGeometryCollectionComponent::GetPhysicsObjectByName(cons
 	{
 		return nullptr;
 	}
+
+	if (Name == NAME_None)
+	{
+		// Special case where it's more convenient for us to return the root bone instead.
+		TArray<int32> Roots;
+		FGeometryCollectionClusteringUtility::GetRootBones(RestCollection->GetGeometryCollection().Get(), Roots);
+
+		if (Roots.IsEmpty())
+		{
+			return nullptr;
+		}
+
+		// More convenient just to assume there's one root for this special case here.
+		return GetPhysicsObjectById(Roots[0]);
+	}
+
 	const int32 Index = RestCollection->GetGeometryCollection()->BoneName.Find(Name.ToString());
 	return GetPhysicsObjectById(Index);
 }
