geometry collection : fix cached geometry collection being set to dynamic too early
- Reverted back set simulate physics to not remove all anchors and instead updated higher level code to do so at the cost of a conditional cast

diff --git a/Engine/Source/Runtime/Experimental/GeometryCollectionEngine/Private/GeometryCollection/GeometryCollectionComponent.cpp b/Engine/Source/Runtime/Experimental/GeometryCollectionEngine/Private/GeometryCollection/GeometryCollectionComponent.cpp
--- a/Engine/Source/Runtime/Experimental/GeometryCollectionEngine/Private/GeometryCollection/GeometryCollectionComponent.cpp
+++ b/Engine/Source/Runtime/Experimental/GeometryCollectionEngine/Private/GeometryCollection/GeometryCollectionComponent.cpp
@@ -785,11 +785,6 @@ void UGeometryCollectionComponent::SetSimulatePhysics(bool bEnabled)
 	{
 		RegisterAndInitializePhysicsProxy();
 	}
-
-	if (bEnabled)
-	{
-		RemoveAllAnchors();
-	}
 }
 
 void UGeometryCollectionComponent::AddForce(FVector Force, FName BoneName, bool bAccelChange)
