Fracture editor : crash fix : make sure to check the validtity of the bone index in GetParentOfBoneAtSpecifiedLevel

diff --git a/Engine/Source/Runtime/Experimental/Chaos/Private/GeometryCollection/GeometryCollectionClusteringUtility.cpp b/Engine/Source/Runtime/Experimental/Chaos/Private/GeometryCollection/GeometryCollectionClusteringUtility.cpp
--- a/Engine/Source/Runtime/Experimental/Chaos/Private/GeometryCollection/GeometryCollectionClusteringUtility.cpp
+++ b/Engine/Source/Runtime/Experimental/Chaos/Private/GeometryCollection/GeometryCollectionClusteringUtility.cpp
@@ -587,7 +587,7 @@ int32 FGeometryCollectionClusteringUtility::GetParentOfBoneAtSpecifiedLevel(cons
 	const TManagedArray<int32>& Levels = GeometryCollection->GetAttribute<int32>("Level", FGeometryCollection::TransformGroup);
 	const TManagedArray<int32>& SimTypes = GeometryCollection->SimulationType;
 
-	if (SourceBone >= 0)
+	if (SourceBone >= 0 && SourceBone < Parents.Num())
 	{
 		int32 SourceParent = SourceBone;
 		while (Levels[SourceParent] > Level || 
