Iris - Fixed destruction infos not destroying the object instance.

diff --git a/Engine/Source/Runtime/Experimental/Iris/Core/Private/Iris/ReplicationSystem/ReplicationBridge.cpp b/Engine/Source/Runtime/Experimental/Iris/Core/Private/Iris/ReplicationSystem/ReplicationBridge.cpp
--- a/Engine/Source/Runtime/Experimental/Iris/Core/Private/Iris/ReplicationSystem/ReplicationBridge.cpp
+++ b/Engine/Source/Runtime/Experimental/Iris/Core/Private/Iris/ReplicationSystem/ReplicationBridge.cpp
@@ -247,7 +247,7 @@ void UReplicationBridge::ReadAndExecuteDestructionInfoFromRemote(FReplicationBri
 	{
 		// Need to forward this as we do not really know how to destroy the actual instance
 		constexpr bool bTearOff = false;
-		constexpr bool bShouldDestroyInstance = false;
+		constexpr bool bShouldDestroyInstance = true;
 		CallDetachInstanceFromRemote(ReferenceToDestroy.GetRefHandle(), bTearOff, bShouldDestroyInstance);
 	}
 }
