CHaos : fix potential crash where cluster can be null when processing cluster groups

diff --git a/Engine/Source/Runtime/Experimental/Chaos/Private/Chaos/PBDRigidClustering.cpp b/Engine/Source/Runtime/Experimental/Chaos/Private/Chaos/PBDRigidClustering.cpp
--- a/Engine/Source/Runtime/Experimental/Chaos/Private/Chaos/PBDRigidClustering.cpp
+++ b/Engine/Source/Runtime/Experimental/Chaos/Private/Chaos/PBDRigidClustering.cpp
@@ -323,7 +323,7 @@ namespace Chaos
 					TArray<FPBDRigidParticleHandle*> ClusterBodies;
 					for(FPBDRigidClusteredParticleHandle* ActiveCluster : Handles)
 					{
-						if(!ActiveCluster->Disabled())
+						if(ActiveCluster && !ActiveCluster->Disabled())
 						{
 							// If this is an external cluster (from the rest collection) we release its children and append them to the current group
 							TSet<FPBDRigidParticleHandle*> Children;
