re-add a function to compute a non-world bounding box to the physics object interface

diff --git a/Engine/Source/Runtime/Experimental/Chaos/Private/Chaos/PhysicsObjectInterface.cpp b/Engine/Source/Runtime/Experimental/Chaos/Private/Chaos/PhysicsObjectInterface.cpp
--- a/Engine/Source/Runtime/Experimental/Chaos/Private/Chaos/PhysicsObjectInterface.cpp
+++ b/Engine/Source/Runtime/Experimental/Chaos/Private/Chaos/PhysicsObjectInterface.cpp
@@ -357,6 +357,38 @@ namespace Chaos
 		return Mass;
 	}
 
+	template<EThreadContext Id>
+	FBox FReadPhysicsObjectInterface<Id>::GetBounds(TArrayView<FPhysicsObjectHandle> InObjects)
+	{
+		FBox RetBox(ForceInit);
+		for (FPhysicsObjectHandle Object : InObjects)
+		{
+			if (!Object)
+			{
+				continue;
+			}
+
+			const TThreadParticle<Id>* Particle = Object->GetParticle<Id>();
+			if (!Particle)
+			{
+				continue;
+			}
+
+			FBox ParticleBox(ForceInit);
+			if (const FImplicitObject* Geometry = Particle->Geometry().Get(); Geometry && Geometry->HasBoundingBox())
+			{
+				const Chaos::FAABB3 Box = Geometry->BoundingBox();
+				ParticleBox = FBox{ Box.Min(), Box.Max() };
+			}
+
+			if (ParticleBox.IsValid)
+			{
+				RetBox += ParticleBox;
+			}
+		}
+		return RetBox;
+	}
+
 	template<EThreadContext Id>
 	FBox FReadPhysicsObjectInterface<Id>::GetWorldBounds(TArrayView<FPhysicsObjectHandle> InObjects)
 	{
diff --git a/Engine/Source/Runtime/Experimental/Chaos/Public/Chaos/PhysicsObjectInterface.h b/Engine/Source/Runtime/Experimental/Chaos/Public/Chaos/PhysicsObjectInterface.h
--- a/Engine/Source/Runtime/Experimental/Chaos/Public/Chaos/PhysicsObjectInterface.h
+++ b/Engine/Source/Runtime/Experimental/Chaos/Public/Chaos/PhysicsObjectInterface.h
@@ -65,6 +65,7 @@ namespace Chaos
 		bool AreAllDynamic(TArrayView<FPhysicsObjectHandle> InObjects);
 		bool AreAllDisabled(TArrayView<FPhysicsObjectHandle> InObjects);
 		float GetMass(TArrayView<FPhysicsObjectHandle> InObjects);
+		FBox GetBounds(TArrayView<FPhysicsObjectHandle> InObjects);
 		FBox GetWorldBounds(TArrayView<FPhysicsObjectHandle> InObjects);
 		FClosestPhysicsObjectResult GetClosestPhysicsBodyFromLocation(TArrayView<FPhysicsObjectHandle> InObjects, const FVector& WorldLocation);
 
