Chaos - broadphase optimizations

Break the collision detection loop into 3 phases: broadphase; midphase assignment; midphase/narrowphase.

diff --git a/Engine/Source/Runtime/Experimental/Chaos/Public/Chaos/Collision/CollisionContext.h b/Engine/Source/Runtime/Experimental/Chaos/Public/Chaos/Collision/CollisionContext.h
--- a/Engine/Source/Runtime/Experimental/Chaos/Public/Chaos/Collision/CollisionContext.h
+++ b/Engine/Source/Runtime/Experimental/Chaos/Public/Chaos/Collision/CollisionContext.h
@@ -60,6 +60,13 @@ namespace Chaos
 		{
 		}
 
+		void Reset()
+		{
+			Settings = nullptr;
+			Allocator = nullptr;
+			MultiShapeCollisionDetector = nullptr;
+		}
+
 		const FCollisionDetectorSettings& GetSettings() const { return *Settings; }
 		void SetSettings(const FCollisionDetectorSettings& InSettings) { Settings = &InSettings; }
 
diff --git a/Engine/Source/Runtime/Experimental/Chaos/Public/Chaos/Collision/SpatialAccelerationBroadPhase.h b/Engine/Source/Runtime/Experimental/Chaos/Public/Chaos/Collision/SpatialAccelerationBroadPhase.h
--- a/Engine/Source/Runtime/Experimental/Chaos/Public/Chaos/Collision/SpatialAccelerationBroadPhase.h
+++ b/Engine/Source/Runtime/Experimental/Chaos/Public/Chaos/Collision/SpatialAccelerationBroadPhase.h
@@ -21,81 +21,289 @@ namespace Chaos
 	class ISpatialAcceleration;
 	class IResimCacheBase;
 
-	// Data collected from the spatial partitioning system
-	class FSimOverlapVisitorEntry
+	namespace Private
 	{
-	public:
-		FSimOverlapVisitorEntry()
-			: Particle(nullptr)
-		{}
-
-		FSimOverlapVisitorEntry(const FAccelerationStructureHandle& Handle)
-			: Particle(Handle.GetGeometryParticleHandle_PhysicsThread())
+		/**
+		 * Output element from the broadphase
+		*/
+		class FBroadphaseOverlap
 		{
-		}
+		public:
+			FBroadphaseOverlap()
+			{
+			}
 
-		FGeometryParticleHandle* Particle;
-	};
+			FBroadphaseOverlap(FGeometryParticleHandle* Particle0, FGeometryParticleHandle* Particle1)
+				: Particles{ Particle0, Particle1 }
+			{
+			}
 
-	/**
-	 * A visitor for the spatial partitioning system used to build the set of objects overlapping a bounding box.
-	 */
-	struct FSimOverlapVisitor
-	{
-		TArray<FSimOverlapVisitorEntry>& Intersections;
+			void SwapOrder()
+			{
+				Swap(Particles[0], Particles[1]);
+			}
 
-		FSimOverlapVisitor(FGeometryParticleHandle* ParticleHandle, const FCollisionFilterData& InSimFilterData, TArray<FSimOverlapVisitorEntry>& InIntersections)
-			: Intersections(InIntersections)
-			, SimFilterData(InSimFilterData)
-			, ParticleUniqueIdx(ParticleHandle ? ParticleHandle->UniqueIdx() : FUniqueIdx(0))
-			, AccelerationHandle(ParticleHandle)
-		{
-		}
+			FGeometryParticleHandle* Particles[2];
+		};
 
-		bool VisitOverlap(const TSpatialVisitorData<FAccelerationStructureHandle>& Instance)
+		/**
+		 * Per-thread context for the broadphase.
+		*/
+		class FBroadphaseContext
 		{
-			Intersections.Emplace(FSimOverlapVisitorEntry(Instance.Payload));
-			return true;
-		}
+		public:
+			FBroadphaseContext()
+			{
+			}
 
-		bool VisitSweep(TSpatialVisitorData<FAccelerationStructureHandle>, FQueryFastData& CurData)
-		{
-			check(false);
-			return false;
-		}
+			void Reset()
+			{
+				Overlaps.Reset();
+				MidPhases.Reset();
+				CollisionContext.Reset();
+			}
+
+			TArray<FBroadphaseOverlap> Overlaps;
+			TArray<FParticlePairMidPhase*> MidPhases;
+			FCollisionContext CollisionContext;
+		};
 
-		bool VisitRaycast(TSpatialVisitorData<FAccelerationStructureHandle>, FQueryFastData& CurData)
+		/**
+		 * Check whether the two particles are allowed to collide, and also whether we should flip their order.
+		 */
+		inline bool ParticlePairCollisionAllowed(const FGeometryParticleHandle* Particle1, const FGeometryParticleHandle* Particle2, const FIgnoreCollisionManager& IgnoreCollisionManager, const bool bIsResimming,  bool& bOutSwapOrder)
 		{
-			check(false);
-			return false;
-		}
+			bOutSwapOrder = false;
 
-		const void* GetQueryData() const { return nullptr; }
+			if (Particle1 == Particle2)
+			{
+				return false;
+			}
 
-		const void* GetSimData() const { return &SimFilterData; }
+			bool bIsKinematic1 = true;
+			bool bIsDynamicAwake1 = false;
+			bool bIsDynamicAsleep1 = false;
+			bool bUseIgnoreCollisionManager1 = false;
+			bool bDisabled1 = false;
+			int32 CollisionGroup1 = 0;
+			const FPBDRigidParticleHandle* Rigid1 = Particle1->CastToRigidParticle();
+			if (Rigid1 != nullptr)
+			{
+				bIsKinematic1 = Rigid1->IsKinematic();
+				bIsDynamicAsleep1 = !bIsKinematic1 && Rigid1->IsSleeping();
+				bIsDynamicAwake1 = !bIsKinematic1 && !bIsDynamicAsleep1;
+				bUseIgnoreCollisionManager1 = Rigid1->HasCollisionConstraintFlag(ECollisionConstraintFlags::CCF_BroadPhaseIgnoreCollisions);
+				bDisabled1 = Rigid1->Disabled();
+				CollisionGroup1 = Rigid1->CollisionGroup();
+			}
 
-		bool ShouldIgnore(const TSpatialVisitorData<FAccelerationStructureHandle>& Instance) const
-		{
-			return Instance.Payload.UniqueIdx() == ParticleUniqueIdx;
-		}
-		/** Return a pointer to the payload on which we are querying the acceleration structure */
-		const void* GetQueryPayload() const
-		{
-			return &AccelerationHandle;
+			bool bIsKinematic2 = true;
+			bool bIsDynamicAwake2 = false;
+			bool bIsDynamicAsleep2 = false;
+			bool bUseIgnoreCollisionManager2 = false;
+			bool bDisabled2 = false;
+			int32 CollisionGroup2 = 0;
+			const FPBDRigidParticleHandle* Rigid2 = Particle2->CastToRigidParticle();
+			if (Rigid2 != nullptr)
+			{
+				bIsKinematic2 = Rigid2->IsKinematic();
+				bIsDynamicAsleep2 = !bIsKinematic2 && Rigid2->IsSleeping();
+				bIsDynamicAwake2 = !bIsKinematic2 && !bIsDynamicAsleep2;
+				bUseIgnoreCollisionManager2 = Rigid2->HasCollisionConstraintFlag(ECollisionConstraintFlags::CCF_BroadPhaseIgnoreCollisions);
+				bDisabled2 = Rigid2->Disabled();
+				CollisionGroup2 = Rigid2->CollisionGroup();
+			}
+
+			// @todo(chaos): This should not be happening if the disabled particles are removed from the active particles list, but GeometryCollection may leave them there
+			//check(!bDisabled2);
+			if (bDisabled1 || bDisabled2)
+			{
+				return false;
+			}
+
+			// At least one particle needs to be dynamic to generate a collision response
+			if (bIsKinematic1 && bIsKinematic2)
+			{
+				return false;
+			}
+			check((Rigid1 != nullptr) || (Rigid2 != nullptr));
+
+			// Is this particle interaction governed by the IgnoreCollisionManager? If so, check to see if interaction is allowed
+			if (bUseIgnoreCollisionManager1 || bUseIgnoreCollisionManager2)
+			{
+				if (IgnoreCollisionManager.IgnoresCollision(Particle1->UniqueIdx(), Particle2->UniqueIdx()))
+				{
+					return false;
+				}
+			}
+
+			// CollisionGroups are used by geometry collections for high-level collision filtering
+			// CollisionGroup == 0 : Collide_With_Everything
+			// CollisionGroup == INDEX_NONE : Disabled collisions
+			// CollisionGroup1 != CollisionGroup2 : Disabled collisions (if other conditions not met)
+			if (CollisionGroup1 == INDEX_NONE || CollisionGroup2 == INDEX_NONE)
+			{
+				return false;
+			}
+			if ((CollisionGroup1 != 0) && (CollisionGroup2 != 0) && (CollisionGroup1 != CollisionGroup2))
+			{
+				return false;
+			}
+
+			// In both cases (resim or not) we will generate (1) dynamic-(sleeping,kinematic(steady+moving),static) pairs + (2) sleeping-moving kinematic ones
+			// Sleeping particles could collide with dynamic ones but these collisions are already handled in case 1
+			// Sleeping particles won't collide with static or steady kinematic particles since neither are moving
+			// Sleeping particles will collide against moving kinematic particles
+			bool bAcceptParticlePair = false;
+			if ((bIsDynamicAwake1 && !bIsDynamicAwake2) || (bIsDynamicAsleep1 && bIsKinematic2))
+			{
+				bAcceptParticlePair = true;
+			}
+
+			// Used to determine a winner in cases where we will visit particle pairs in both orders
+			const bool bIsParticle1Preferred = AreParticlesInPreferredOrder(Particle1, Particle2);
+
+			if (!bIsResimming)
+			{
+				// Normally (not resimming) we iterate over dynamic and asleep|kinematic particles, so:
+				// - Particle1 is dynamic, asleep OR kinematic
+				// - Particle2 may be static, kinematic, dynamic, asleep
+
+				// If Particle1 is non dynamic but particle 2 is dynamic, the case should already be handled by (1)
+				if (!bIsDynamicAwake1 && bIsDynamicAwake2)
+				{
+					bAcceptParticlePair = false;
+				}
+				// If Particle1 and Particle2 are dynamic we validate the pair if particle1 has higher ID to discard duplicates since we will visit twice the pairs
+				else if (bIsDynamicAwake1 && bIsDynamicAwake2)
+				{
+					bAcceptParticlePair = bIsParticle1Preferred;
+				}
+				// If Particle1 is kinematic we should in theory discard the pairs against sleeping particles
+				// since the sleeping-kinematic case has been validated in (2). But Particle1 is asleep OR kinematic so
+				// when entering this condition we are sure that we enver entered (2). It is why we validate the kinematic-sleeping pairs as well
+				else if (bIsKinematic1 && bIsDynamicAsleep2)
+				{
+					bAcceptParticlePair = true;
+				}
+			}
+			else
+			{
+				// When resimming we iterate over "desynced" particles which may be kinematic so:
+				// - Particle1 is always desynced
+				// - Particle2 may also be desynced, in which case we will also visit the opposite ordering regardless of dynamic/kinematic status
+				// - Particle1 may be static, kinematic, dynamic, asleep
+				// - Particle2 may be static, kinematic, dynamic, asleep
+				// 
+				// Even though Particle1 may be kinematic when resimming, we want to create the contacts in the original order (i.e., dynamic first)
+				// 
+				const bool bIsParticle2Desynced = bIsResimming && (Particle2->SyncState() == ESyncState::HardDesync);
+
+				// If Particle1 is non dynamic but particle 2 is dynamic, the case should already be handled by (1) for
+				// the desynced dynamic - synced/desynced (static,kinematic,asleep) pairs. But we still need to process
+				// the desynced (static,kinematic,asleep) against the synced dynamic since this case have not been handled by (1)
+				if (!bIsDynamicAwake1 && bIsDynamicAwake2)
+				{
+					bAcceptParticlePair = !bIsParticle2Desynced;
+				}
+				// If Particle1 and Particle2 are dynamic we validate the pair if particle1 has higher ID to discard duplicates since we will visit twice the pairs
+				// We validate the pairs as well if the particle 2 is synced since  we will never visit the opposite order (we only iterate over desynced particles)
+				else if (bIsDynamicAwake1 && bIsDynamicAwake2)
+				{
+					bAcceptParticlePair = bIsParticle1Preferred || !bIsParticle2Desynced;
+				}
+				// If Particle1 is kinematic we are discarding the pairs against sleeping desynced particles since the case has been handled by (2)
+				else if (bIsKinematic1 && bIsDynamicAsleep2)
+				{
+					bAcceptParticlePair = !bIsParticle2Desynced;
+				}
+			}
+
+			// Since particles can change type (between Kinematic and Dynamic) we may visit them in different orders at different times, but if we allow
+			// that it would break Resim and constraint reuse. Also, if only one particle is Dynamic, we want it in first position. This isn't a strtct 
+			// requirement but some downstream systems assume this is true (e.g., CCD, TriMesh collision).
+			bOutSwapOrder = ShouldSwapParticleOrder((bIsDynamicAwake1 || bIsDynamicAsleep1), (bIsDynamicAwake2 || bIsDynamicAsleep2), bIsParticle1Preferred);
+
+			return bAcceptParticlePair;
 		}
 
-		bool HasBlockingHit() const
+		/**
+		 * A visitor for the spatial partitioning system used to build the set of objects overlapping a bounding box.
+		 */
+		struct FSimOverlapVisitor
 		{
-			return false;
-		}
+			FSimOverlapVisitor(FGeometryParticleHandle* ParticleHandle, const FCollisionFilterData& InSimFilterData, const FIgnoreCollisionManager& InIgnoreCollisionManager, const bool bInIsResimming, Private::FBroadphaseContext& InContext)
+				: Context(InContext)
+				, IgnoreCollisionManager(InIgnoreCollisionManager)
+				, SimFilterData(InSimFilterData)
+				, ParticleUniqueIdx(ParticleHandle ? ParticleHandle->UniqueIdx() : FUniqueIdx(0))
+				, bIsResimming(bInIsResimming)
+				, AccelerationHandle(ParticleHandle)
+			{
+			}
 
-	private:
-		FCollisionFilterData SimFilterData;
-		FUniqueIdx ParticleUniqueIdx; // unique id of the particle visiting, used to skip self intersection as early as possible
+			bool VisitOverlap(const TSpatialVisitorData<FAccelerationStructureHandle>& Instance)
+			{
+				FGeometryParticleHandle* Particle1 = AccelerationHandle.GetGeometryParticleHandle_PhysicsThread();
+				FGeometryParticleHandle* Particle2 = Instance.Payload.GetGeometryParticleHandle_PhysicsThread();
+			
+				bool bSwapOrder = false;
+				const bool bAccept = Private::ParticlePairCollisionAllowed(Particle1, Particle2, IgnoreCollisionManager, bIsResimming, bSwapOrder);
+				
+				if (bAccept)
+				{
+					if (bSwapOrder)
+					{
+						Swap(Particle1, Particle2);
+					}
+					Context.Overlaps.Emplace(Particle1, Particle2);
+				}
 
-		/** Handle to be stored to retrieve the payload on which we are querying the acceleration structure*/
-		FAccelerationStructureHandle AccelerationHandle;
-	};
+				return true;
+			}
+
+			bool VisitSweep(TSpatialVisitorData<FAccelerationStructureHandle>, FQueryFastData& CurData)
+			{
+				check(false);
+				return false;
+			}
+
+			bool VisitRaycast(TSpatialVisitorData<FAccelerationStructureHandle>, FQueryFastData& CurData)
+			{
+				check(false);
+				return false;
+			}
+
+			const void* GetQueryData() const { return nullptr; }
+
+			const void* GetSimData() const { return &SimFilterData; }
+
+			bool ShouldIgnore(const TSpatialVisitorData<FAccelerationStructureHandle>& Instance) const
+			{
+				return Instance.Payload.UniqueIdx() == ParticleUniqueIdx;
+			}
+			/** Return a pointer to the payload on which we are querying the acceleration structure */
+			const void* GetQueryPayload() const
+			{
+				return &AccelerationHandle;
+			}
+
+			bool HasBlockingHit() const
+			{
+				return false;
+			}
+
+		private:
+			Private::FBroadphaseContext& Context;
+			const FIgnoreCollisionManager& IgnoreCollisionManager;
+			FCollisionFilterData SimFilterData;
+			FUniqueIdx ParticleUniqueIdx; // unique id of the particle visiting, used to skip self intersection as early as possible
+			bool bIsResimming;
+
+			/** Handle to be stored to retrieve the payload on which we are querying the acceleration structure*/
+			FAccelerationStructureHandle AccelerationHandle;
+		};
+	}
 
 	/**
 	 * A broad phase that iterates over particle and uses a spatial acceleration structure to output
@@ -164,7 +372,6 @@ namespace Chaos
 		 * @param InSpatialAcceleration Spatial acceleration (AABB, bounding volumes...) to be used for broadphase collision detection
 		 * @param NarrowPhase Narrowphase collision detection that will be executed on each potential pairs coming from the broadphase detection
 		 * */
-
 		template<bool bNeedsResim, bool bOnlyRigid, typename ViewType, typename SpatialAccelerationType>
 		void ComputeParticlesOverlaps(
 			ViewType& OverlapView, 
@@ -173,37 +380,75 @@ namespace Chaos
 			FCollisionConstraintAllocator* Allocator,
 			const FCollisionDetectorSettings& Settings)
 		{
+			// Reset all the contexts (we don't always use all of them, but don't reduce the array size so that the element arrays don't need reallocating)
+			for (Private::FBroadphaseContext& BroadphaseContext : BroadphaseContexts)
+			{
+				BroadphaseContext.Reset();
+			}
+			int32 NumActiveContexts = 0;
+
 			// A set of contexts, one for each worker thread, containing the allocation buffers etc.
-			// NOTE: Ideally contexts would not resize after the initial creation, but actually it might because
+			// NOTE: Ideally contexts arrays would not resize after the initial creation, but actually it might because
 			// of the way the ParticleView parallel-for is implemented. If the view is a set of sub-views, we
 			// will get a parallel-for per sub-view which will have different numbers of particles and may
 			// therefore require a different number of workers.
 			// Fortunately, we only get calls to ContextCreator for a sub-view (with a possibly different NumWorkers)
 			// after the previous sub-view's parallel-for has completed, so we don't need to worry about context addresses 
 			// changing (we could use TChunkedArray, or allocate contexts on the heap if we ever need to handle that).
-			TArray<FCollisionContext> Contexts;
-			auto ContextCreator = [Allocator, &Settings, &Contexts](const int32 WorkerIndex, const int32 NumWorkers) -> int32
+			const auto& ContextCreator = [this, Allocator, &Settings, &OverlapView, &NumActiveContexts](const int32 WorkerIndex, const int32 NumWorkers) -> int32
 			{
-				Contexts.SetNum(NumWorkers);
+				// Make sure we have enough contexts 
+				// (NOTE: ContextCreator gets called for every worker, but in serial and always with the same NumWorkers)
+				NumActiveContexts = NumWorkers;
+				if (NumWorkers > BroadphaseContexts.Num())
+				{
+					check(WorkerIndex == 0);	// Should be the same for all subsequent workers
+					BroadphaseContexts.SetNum(NumWorkers);
+				}
+
+				// Initialize the broadphase context for this worker
+				BroadphaseContexts[WorkerIndex].Overlaps.Reset(OverlapView.Num());
 
 				// Retrieve the midphase/collision allocator for this context
 				Allocator->SetMaxContexts(NumWorkers);
 				FCollisionContextAllocator* ContextAllocator = Allocator->GetContextAllocator(WorkerIndex);
 
-				// Build the context for the worker
-				Contexts[WorkerIndex].SetSettings(Settings);
-				Contexts[WorkerIndex].SetAllocator(ContextAllocator);
+				// Build the collision context for the worker
+				BroadphaseContexts[WorkerIndex].CollisionContext.SetSettings(Settings);
+				BroadphaseContexts[WorkerIndex].CollisionContext.SetAllocator(ContextAllocator);
 
 				return WorkerIndex;
 			};
 
-			// Detect collisions for a single particle versus all others
-			auto ParticleCollisionGenerator = [this, Dt, &InSpatialAcceleration, &Contexts](auto& Particle1, const int32 ActiveIdxIdx, const int32 ContextIndex)
+			// Generate the set of AABB overlapping particles that are allowed to collide with each other
 			{
-				ProduceParticleOverlaps<bNeedsResim, bOnlyRigid>(Dt, Particle1.Handle(), InSpatialAcceleration, Contexts[ContextIndex]);
-			};
+				CSV_SCOPED_TIMING_STAT(PhysicsVerbose, DetectCollisions_AABBTree);
+				const auto& ProduceOverlapsWorker = [this, Dt, &InSpatialAcceleration](auto& Particle1, const int32 ActiveIdxIdx, const int32 ContextIndex)
+				{
+					ProduceParticleOverlaps<bNeedsResim, bOnlyRigid>(Dt, Particle1.Handle(), InSpatialAcceleration, BroadphaseContexts[ContextIndex]);
+				};
+				OverlapView.ParallelFor(ContextCreator, ProduceOverlapsWorker, bDisableCollisionParallelFor);
+			}
+
+			// Find or assign a midphase to each overlapping particle pair
+			{
+				CSV_SCOPED_TIMING_STAT(PhysicsVerbose, DetectCollisions_AssignMidPhases);
+				const auto& AssignMidphasesWorker = [this, Dt](const int32 ContextIndex)
+				{
+					AssignMidphases(BroadphaseContexts[ContextIndex]);
+				};
+				PhysicsParallelFor(NumActiveContexts, AssignMidphasesWorker, bDisableCollisionParallelFor);
+			}
 
-			OverlapView.ParallelFor(ContextCreator, ParticleCollisionGenerator, bDisableCollisionParallelFor);
+			// Use the midphase to create contacts between the shapes on the two partiucles
+			{
+				CSV_SCOPED_TIMING_STAT(PhysicsVerbose, DetectCollisions_MidPhase);
+				const auto& ProcessMidphasesWorker = [this, Dt](const int32 ContextIndex)
+				{
+					ProcessMidphases(Dt, BroadphaseContexts[ContextIndex]);
+				};
+				PhysicsParallelFor(NumActiveContexts, ProcessMidphasesWorker, bDisableCollisionParallelFor);
+			}
 		}
 
 		template<typename T_SPATIALACCELERATION>
@@ -251,40 +496,35 @@ namespace Chaos
 
 	private:
 
+		// Generate the set of particles that overlap the specified particle and are allowed to collide with it
 		template<bool bIsResimming, bool bOnlyRigid, typename T_SPATIALACCELERATION>
 		void ProduceParticleOverlaps(
 		    FReal Dt,
 		    FGeometryParticleHandle* Particle1,
 		    const T_SPATIALACCELERATION& InSpatialAcceleration,
-			const FCollisionContext& Context)
+			Private::FBroadphaseContext& Context)
 		{
-			TArray<FSimOverlapVisitorEntry> PotentialIntersections;
-
+			// @todo(chaos):We shouldn't need this data here (see uses below)
 			bool bIsKinematic1 = true;
 			bool bIsDynamicAwake1 = false;
 			bool bIsDynamicAsleep1 = false;
-			bool bUseIgnoreCollisionManager1 = false;
 			bool bDisabled1 = false;
-			int32 CollisionGroup1 = 0;
 			const FPBDRigidParticleHandle* Rigid1 = Particle1->CastToRigidParticle();
 			if (Rigid1 != nullptr)
 			{
 				bIsKinematic1 = Rigid1->IsKinematic();
 				bIsDynamicAsleep1 = !bIsKinematic1 && Rigid1->IsSleeping();
 				bIsDynamicAwake1 = !bIsKinematic1 && !bIsDynamicAsleep1;
-				bUseIgnoreCollisionManager1 = Rigid1->HasCollisionConstraintFlag(ECollisionConstraintFlags::CCF_BroadPhaseIgnoreCollisions);
 				bDisabled1 = Rigid1->Disabled();
-				CollisionGroup1 = Rigid1->CollisionGroup();
 			}
 
-			// @todo(chaos): We should not be getting disabled particles, but we do from GeometryCollections
-			//check(!bDisabled1);
+			// Skip particles we are not interested in
+			// @todo(chaos) Ideally we would we should not be getting disabled particles, but we currently do from GeometryCollections.
 			if (bDisabled1)
 			{
 				return;
 			}
 
-			// Skip particles we are not interested in
 			// @todo(chaos): This should already be handled by selecting the appropriate particle view in the calling function. GeometryCollections?
 			const bool bHasValidState = bOnlyRigid ? (bIsDynamicAwake1 || bIsDynamicAsleep1) : (bIsDynamicAwake1 || bIsKinematic1);
 			if (!bHasValidState && !bIsResimming)
@@ -304,7 +544,7 @@ namespace Chaos
 
 				{
 					PHYSICS_CSV_SCOPED_EXPENSIVE(PhysicsVerbose, DetectCollisions_BroadPhase);
-					FSimOverlapVisitor OverlapVisitor(Particle1, ParticleSimData, PotentialIntersections);
+					Private::FSimOverlapVisitor OverlapVisitor(Particle1, ParticleSimData, IgnoreCollisionManager, bIsResimming, Context);
 					InSpatialAcceleration.Overlap(Box1, OverlapVisitor);
 				}
 			}
@@ -312,176 +552,36 @@ namespace Chaos
 			{
 				//SCOPE_CYCLE_COUNTER(STAT_Collisions_AABBTree);
 				const auto& GlobalElems = InSpatialAcceleration.GlobalObjects();
-				PotentialIntersections.Reserve(GlobalElems.Num());
+				Context.Overlaps.Reserve(GlobalElems.Num());
 
 				for (auto& Elem : GlobalElems)
 				{
-					PotentialIntersections.Add(Elem.Payload);
+					Context.Overlaps.Emplace(Particle1, Elem.Payload.GetGeometryParticleHandle_PhysicsThread());
 				}
 			}
+		}
 
-			//SCOPE_CYCLE_COUNTER(STAT_Collisions_Filtering);
-			const int32 NumPotentials = PotentialIntersections.Num();
-			TArray<FParticlePairMidPhase*> MidPhases;
-			MidPhases.Reserve(NumPotentials);
-			for (int32 PotentialIndex = 0; PotentialIndex < NumPotentials; ++PotentialIndex)
+		// Find or assign midphases to each of the overlapping particle pairs
+		// @todo(chaos): optimize
+		void AssignMidphases(Private::FBroadphaseContext& BroadphaseContext)
+		{
+			BroadphaseContext.MidPhases.Reset(BroadphaseContext.Overlaps.Num());
+			for (Private::FBroadphaseOverlap& Overlap : BroadphaseContext.Overlaps)
 			{
-				FGeometryParticleHandle* Particle2 = PotentialIntersections[PotentialIndex].Particle;
-				if (Particle1 == Particle2)
-				{
-					continue;
-				}
-
-				bool bIsKinematic2 = true;
-				bool bIsDynamicAwake2 = false;
-				bool bIsDynamicAsleep2 = false;
-				bool bUseIgnoreCollisionManager2 = false;
-				bool bDisabled2 = false;
-				int32 CollisionGroup2 = 0;
-				const FPBDRigidParticleHandle* Rigid2 = Particle2->CastToRigidParticle();
-				if (Rigid2 != nullptr)
-				{
-					bIsKinematic2 = Rigid2->IsKinematic();
-					bIsDynamicAsleep2 = !bIsKinematic2 && Rigid2->IsSleeping();
-					bIsDynamicAwake2 = !bIsKinematic2 && !bIsDynamicAsleep2;
-					bUseIgnoreCollisionManager2 = Rigid2->HasCollisionConstraintFlag(ECollisionConstraintFlags::CCF_BroadPhaseIgnoreCollisions);
-					bDisabled2 = Rigid2->Disabled();
-					CollisionGroup2 = Rigid2->CollisionGroup();
-				}
-
-				// @todo(chaos): This should not be happening if the disabled particles are removed from the active particles list, but GeometryCollection particles leave there there
-				//check(!bDisabled2);
-				if (bDisabled2)
-				{
-					continue;
-				}
-
-				// At least one particle needs to be dynamic to generate a collision response
-				if (bIsKinematic1 && bIsKinematic2)
-				{
-					continue;
-				}
-				check((Rigid1 != nullptr) || (Rigid2 != nullptr));
-
-				// Is this particle interaction governed by the IgnoreCollisionManager? If so, check to see if interaction is allowed
-				if (bUseIgnoreCollisionManager1 || bUseIgnoreCollisionManager2)
-				{
-					if (IgnoreCollisionManager.IgnoresCollision(Particle1->UniqueIdx(), Particle2->UniqueIdx()))
-					{
-						continue;
-					}
-				}
-
-				// CollisionGroups are used by geometry collections for high-level collision filtering
-				// CollisionGroup == 0 : Collide_With_Everything
-				// CollisionGroup == INDEX_NONE : Disabled collisions
-				// CollisionGroup1 != CollisionGroup2 : Disabled collisions (if other conditions not met)
-				if (CollisionGroup1 == INDEX_NONE || CollisionGroup2 == INDEX_NONE)
-				{
-					continue;
-				}
-				if ((CollisionGroup1 != 0) && (CollisionGroup2 != 0) && (CollisionGroup1 != CollisionGroup2))
-				{
-					continue;
-				}
-
-				// In both cases (resim or not) we will generate (1) dynamic-(sleeping,kinematic(steady+moving),static) pairs + (2) sleeping-moving kinematic ones
-				// Sleeping particles could collide with dynamic ones but these collisions are already handled in case 1
-				// Sleeping particles won't collide with static or steady kinematic particles since neither are moving
-				// Sleeping particles will collide against moving kinematic particles
-				bool bAcceptParticlePair = false;
-				if ((bIsDynamicAwake1 && !bIsDynamicAwake2) || (bIsDynamicAsleep1 && bIsKinematic2))
-				{
-					bAcceptParticlePair = true;
-				}
-					
-				// Used to determine a winner in cases where we will visit particle pairs in both orders
-				const bool bIsParticle1Preferred = AreParticlesInPreferredOrder(Particle1, Particle2);
-					
-				if(!bIsResimming)
-				{
-					// Normally (not resimming) we iterate over dynamic and asleep|kinematic particles, so:
-					// - Particle1 is dynamic, asleep OR kinematic
-					// - Particle2 may be static, kinematic, dynamic, asleep
-
-					// If Particle1 is non dynamic but particle 2 is dynamic, the case should already be handled by (1)
-					if(!bIsDynamicAwake1 && bIsDynamicAwake2)
-					{
-						bAcceptParticlePair = false;
-					}
-					// If Particle1 and Particle2 are dynamic we validate the pair if particle1 has higher ID to discard duplicates since we will visit twice the pairs
-					else if(bIsDynamicAwake1 && bIsDynamicAwake2)
-					{
-						bAcceptParticlePair = bIsParticle1Preferred;
-					}
-					// If Particle1 is kinematic we should in theory discard the pairs against sleeping particles
-					// since the sleeping-kinematic case has been validated in (2). But Particle1 is asleep OR kinematic so
-					// when entering this condition we are sure that we enver entered (2). It is why we validate the kinematic-sleeping pairs as well
-					else if(bIsKinematic1 && bIsDynamicAsleep2)
-					{
-						bAcceptParticlePair = true;
-					}
-				}
-				else
-				{
-					// When resimming we iterate over "desynced" particles which may be kinematic so:
-					// - Particle1 is always desynced
-					// - Particle2 may also be desynced, in which case we will also visit the opposite ordering regardless of dynamic/kinematic status
-					// - Particle1 may be static, kinematic, dynamic, asleep
-					// - Particle2 may be static, kinematic, dynamic, asleep
-					// 
-					// Even though Particle1 may be kinematic when resimming, we want to create the contacts in the original order (i.e., dynamic first)
-					// 
-					const bool bIsParticle2Desynced = bIsResimming && (Particle2->SyncState() == ESyncState::HardDesync);
-
-					// If Particle1 is non dynamic but particle 2 is dynamic, the case should already be handled by (1) for
-					// the desynced dynamic - synced/desynced (static,kinematic,asleep) pairs. But we still need to process
-					// the desynced (static,kinematic,asleep) against the synced dynamic since this case have not been handled by (1)
-					if(!bIsDynamicAwake1 && bIsDynamicAwake2)
-					{
-						bAcceptParticlePair = !bIsParticle2Desynced;
-					}
-					// If Particle1 and Particle2 are dynamic we validate the pair if particle1 has higher ID to discard duplicates since we will visit twice the pairs
-					// We validate the pairs as well if the particle 2 is synced since  we will never visit the opposite order (we only iterate over desynced particles)
-					else if(bIsDynamicAwake1 && bIsDynamicAwake2)
-					{
-						bAcceptParticlePair = bIsParticle1Preferred || !bIsParticle2Desynced;
-					}
-					// If Particle1 is kinematic we are discarding the pairs against sleeping desynced particles since the case has been handled by (2)
-					else if(bIsKinematic1 && bIsDynamicAsleep2)
-					{
-						bAcceptParticlePair = !bIsParticle2Desynced;
-					}
-				}
-					
-				if (!bAcceptParticlePair)
-				{
-					continue;
-				}
-
-				// Since particles can change type (between Kinematic and Dynamic) we may visit them in different orders at different times, but if we allow
-				// that it would break Resim and constraint re-use. Also, if only one particle is Dynamic, we want it in first position. This isn't a strtct 
-				// requirement but some downstream systems assume this is true (e.g., CCD, TriMesh collision).
-				const bool bSwapOrder = ShouldSwapParticleOrder((bIsDynamicAwake1 || bIsDynamicAsleep1), (bIsDynamicAwake2 || bIsDynamicAsleep2), bIsParticle1Preferred);
-				FGeometryParticleHandle* ParticleA = bSwapOrder ? Particle2 : Particle1;
-				FGeometryParticleHandle* ParticleB = bSwapOrder ? Particle1 : Particle2;
-
-				// NOTE: This searches the array of existing midphases on the particle specified in the last parameter. The assumption
-				// here is that Particle1 is more likely to have less collision on it because it is dynamic. Kinematics can often be
-				// large and may collide with many objects. Maybe we should consider using size instead.
-				FParticlePairMidPhase* MidPhase = Context.GetAllocator()->GetMidPhase(ParticleA, ParticleB, Particle1, Context);
-				if (MidPhase != nullptr)
-				{
-					MidPhases.Add(MidPhase);
-				}
+				FParticlePairMidPhase* MidPhase = BroadphaseContext.CollisionContext.GetAllocator()->GetMidPhase(Overlap.Particles[0], Overlap.Particles[1], Overlap.Particles[0], BroadphaseContext.CollisionContext);
+				BroadphaseContext.MidPhases.Add(MidPhase);
 			}
+		}
 
+		// Process all the midphases: generate constraints and execute the narrowphase
+		void ProcessMidphases(const FReal Dt, const Private::FBroadphaseContext& BroadphaseContext)
+		{
 			// Prefetch initial set of MidPhases
 			const int32 PrefetchLookahead = 4;
-			const int32 NumMidPhases = MidPhases.Num();
+			const int32 NumMidPhases = BroadphaseContext.MidPhases.Num();
 			for (int32 Index = 0; Index < NumMidPhases && Index < PrefetchLookahead; Index++)
 			{
-				MidPhases[Index]->CachePrefetch();
+				BroadphaseContext.MidPhases[Index]->CachePrefetch();
 			}
 
 			for (int32 Index = 0; Index < NumMidPhases; Index++)
@@ -489,11 +589,11 @@ namespace Chaos
 				// Prefetch next MidPhase
 				if (Index + PrefetchLookahead < NumMidPhases)
 				{
-					MidPhases[Index + PrefetchLookahead]->CachePrefetch();
+					BroadphaseContext.MidPhases[Index + PrefetchLookahead]->CachePrefetch();
 				}
 
 				// Run MidPhase + NarrowPhase
-				MidPhases[Index]->GenerateCollisions(Context.GetSettings().BoundsExpansion, Dt, Context);
+				BroadphaseContext.MidPhases[Index]->GenerateCollisions(BroadphaseContext.CollisionContext.GetSettings().BoundsExpansion, Dt, BroadphaseContext.CollisionContext);
 			}
 
 			PHYSICS_CSV_CUSTOM_EXPENSIVE(PhysicsCounters, NumFromBroadphase, NumPotentials, ECsvCustomStatOp::Accumulate);
@@ -501,7 +601,8 @@ namespace Chaos
 
 		const FPBDRigidsSOAs& Particles;
 		const FAccelerationStructure* SpatialAcceleration;
-
+		TArray<Private::FBroadphaseContext> BroadphaseContexts;
 		FIgnoreCollisionManager IgnoreCollisionManager;
 	};
+
 }
