Chaos : fix geometry collection replication failing because the client side component has a authority local role when physics state is being created

diff --git a/Engine/Source/Runtime/Experimental/GeometryCollectionEngine/Private/GeometryCollection/GeometryCollectionComponent.cpp b/Engine/Source/Runtime/Experimental/GeometryCollectionEngine/Private/GeometryCollection/GeometryCollectionComponent.cpp
--- a/Engine/Source/Runtime/Experimental/GeometryCollectionEngine/Private/GeometryCollection/GeometryCollectionComponent.cpp
+++ b/Engine/Source/Runtime/Experimental/GeometryCollectionEngine/Private/GeometryCollection/GeometryCollectionComponent.cpp
@@ -2372,7 +2372,7 @@ void UGeometryCollectionComponent::ResetDynamicCollection()
 	bCreateDynamicCollection = false;
 	if (UWorld* World = GetWorld())
 	{
-		if(World->IsGameWorld())
+		if(World->IsGameWorld() || World->IsPreviewWorld())
 		{
 			bCreateDynamicCollection = true;
 		}
@@ -2453,7 +2453,7 @@ void UGeometryCollectionComponent::OnCreatePhysicsState()
 			RestCollectionMutable->CreateSimulationDataIfNeeded();
 		}
 #endif
-		const bool bValidWorld = GetWorld() && GetWorld()->IsGameWorld();
+		const bool bValidWorld = GetWorld() && (GetWorld()->IsGameWorld() || GetWorld()->IsPreviewWorld());
 		const bool bValidCollection = DynamicCollection && DynamicCollection->Transform.Num() > 0;
 		if (bValidWorld && bValidCollection)
 		{
