Chaos - fix Headless crash (mishandling of infinite bounds objects in broadphase)
- re-add midphase search particle hint
- move collision allocator classes to Private namespace

diff --git a/Engine/Source/Programs/HeadlessChaos/Private/HeadlessChaosTestShapePadding.cpp b/Engine/Source/Programs/HeadlessChaos/Private/HeadlessChaosTestShapePadding.cpp
--- a/Engine/Source/Programs/HeadlessChaos/Private/HeadlessChaosTestShapePadding.cpp
+++ b/Engine/Source/Programs/HeadlessChaos/Private/HeadlessChaosTestShapePadding.cpp
@@ -85,7 +85,7 @@ namespace ChaosTest {
 		EXPECT_NEAR(BoxBounds1.Extents().Y, Size.Y, Tolerance);
 		EXPECT_NEAR(BoxBounds1.Extents().Z, Size.Z, Tolerance);
 
-		FCollisionConstraintAllocator CollisionAllocator;
+		Private::FCollisionConstraintAllocator CollisionAllocator;
 		CollisionAllocator.SetMaxContexts(1);
 
 		FPBDCollisionConstraintPtr Constraint = CollisionAllocator.GetContextAllocator(0)->CreateConstraint(
@@ -198,7 +198,7 @@ namespace ChaosTest {
 		EXPECT_NEAR(BoxBounds1.Extents().Y, Size.Y, Tolerance);
 		EXPECT_NEAR(BoxBounds1.Extents().Z, Size.Z, Tolerance);
 
-		FCollisionConstraintAllocator CollisionAllocator;
+		Private::FCollisionConstraintAllocator CollisionAllocator;
 		CollisionAllocator.SetMaxContexts(1);
 
 		FPBDCollisionConstraintPtr Constraint = CollisionAllocator.GetContextAllocator(0)->CreateConstraint(
diff --git a/Engine/Source/Runtime/Experimental/Chaos/Private/Chaos/CCDUtilities.cpp b/Engine/Source/Runtime/Experimental/Chaos/Private/Chaos/CCDUtilities.cpp
--- a/Engine/Source/Runtime/Experimental/Chaos/Private/Chaos/CCDUtilities.cpp
+++ b/Engine/Source/Runtime/Experimental/Chaos/Private/Chaos/CCDUtilities.cpp
@@ -93,7 +93,7 @@ namespace Chaos
 		return INDEX_NONE;
 	}
 
-	void FCCDManager::ApplyConstraintsPhaseCCD(const FReal Dt, FCollisionConstraintAllocator *CollisionAllocator, const int32 NumDynamicParticles)
+	void FCCDManager::ApplyConstraintsPhaseCCD(const FReal Dt, Private::FCollisionConstraintAllocator *CollisionAllocator, const int32 NumDynamicParticles)
 	{
 		SweptConstraints = CollisionAllocator->GetCCDConstraints();
 		if (SweptConstraints.Num() > 0)
@@ -896,7 +896,7 @@ namespace Chaos
 		}
 	}
 
-	void FCCDManager::UpdateSweptConstraints(const FReal Dt, FCollisionConstraintAllocator *CollisionAllocator)
+	void FCCDManager::UpdateSweptConstraints(const FReal Dt, Private::FCollisionConstraintAllocator *CollisionAllocator)
 	{
 		// Buld the set of collision whose contact data will be out of date because we moved one or both of its particles. 
 		// This is all collision constraints, including non-swept ones, for any particle that was relocated by the CCD sweep 
diff --git a/Engine/Source/Runtime/Experimental/Chaos/Private/Chaos/ChaosDebugDraw.cpp b/Engine/Source/Runtime/Experimental/Chaos/Private/Chaos/ChaosDebugDraw.cpp
--- a/Engine/Source/Runtime/Experimental/Chaos/Private/Chaos/ChaosDebugDraw.cpp
+++ b/Engine/Source/Runtime/Experimental/Chaos/Private/Chaos/ChaosDebugDraw.cpp
@@ -1702,7 +1702,7 @@ namespace Chaos
 			}
 		}
 
-		void DrawCollisions(const FRigidTransform3& SpaceTransform, const FCollisionConstraintAllocator& CollisionAllocator, FRealSingle ColorScale, const FChaosDebugDrawSettings* Settings)
+		void DrawCollisions(const FRigidTransform3& SpaceTransform, const Private::FCollisionConstraintAllocator& CollisionAllocator, FRealSingle ColorScale, const FChaosDebugDrawSettings* Settings)
 		{
 			if (FDebugDrawQueue::IsDebugDrawingEnabled())
 			{
diff --git a/Engine/Source/Runtime/Experimental/Chaos/Private/Chaos/Collision/CollisionConstraintAllocator.cpp b/Engine/Source/Runtime/Experimental/Chaos/Private/Chaos/Collision/CollisionConstraintAllocator.cpp
--- a/Engine/Source/Runtime/Experimental/Chaos/Private/Chaos/Collision/CollisionConstraintAllocator.cpp
+++ b/Engine/Source/Runtime/Experimental/Chaos/Private/Chaos/Collision/CollisionConstraintAllocator.cpp
@@ -4,99 +4,102 @@
 
 namespace Chaos
 {
-	FPBDCollisionConstraintPtr FCollisionContextAllocator::CreateConstraint(
-		FGeometryParticleHandle* Particle0,
-		const FImplicitObject* Implicit0,
-		const FPerShapeData* Shape0,
-		const FBVHParticles* Simplicial0,
-		const FRigidTransform3& ShapeRelativeTransform0,
-		FGeometryParticleHandle* Particle1,
-		const FImplicitObject* Implicit1,
-		const FPerShapeData* Shape1,
-		const FBVHParticles* Simplicial1,
-		const FRigidTransform3& ShapeRelativeTransform1,
-		const FReal CullDistance,
-		const bool bUseManifold,
-		const EContactShapesType ShapePairType)
+	namespace Private
 	{
-		FPBDCollisionConstraintPtr Constraint = CreateConstraint();
-
-		if (Constraint.IsValid())
+		FPBDCollisionConstraintPtr FCollisionContextAllocator::CreateConstraint(
+			FGeometryParticleHandle* Particle0,
+			const FImplicitObject* Implicit0,
+			const FPerShapeData* Shape0,
+			const FBVHParticles* Simplicial0,
+			const FRigidTransform3& ShapeRelativeTransform0,
+			FGeometryParticleHandle* Particle1,
+			const FImplicitObject* Implicit1,
+			const FPerShapeData* Shape1,
+			const FBVHParticles* Simplicial1,
+			const FRigidTransform3& ShapeRelativeTransform1,
+			const FReal CullDistance,
+			const bool bUseManifold,
+			const EContactShapesType ShapePairType)
 		{
-			Constraint->SetContainer(CollisionContainer);
-			FPBDCollisionConstraint::Make(Particle0, Implicit0, Shape0, Simplicial0, ShapeRelativeTransform0, Particle1, Implicit1, Shape1, Simplicial1, ShapeRelativeTransform1, CullDistance, bUseManifold, ShapePairType, *Constraint);
-		}
+			FPBDCollisionConstraintPtr Constraint = CreateConstraint();
 
-		return Constraint;
-	}
+			if (Constraint.IsValid())
+			{
+				Constraint->SetContainer(CollisionContainer);
+				FPBDCollisionConstraint::Make(Particle0, Implicit0, Shape0, Simplicial0, ShapeRelativeTransform0, Particle1, Implicit1, Shape1, Simplicial1, ShapeRelativeTransform1, CullDistance, bUseManifold, ShapePairType, *Constraint);
+			}
 
-	void FCollisionContextAllocator::ProcessNewMidPhases(FCollisionConstraintAllocator* Allocator)
-	{
-		// NOTE: Called from the physics thread, and never from a physics task/parallel-for. No need for locks.
+			return Constraint;
+		}
 
-		// Register the midphases with each of their particles and add the midphase to the central list
-		for (FParticlePairMidPhase* MidPhase : NewMidPhases)
+		void FCollisionContextAllocator::ProcessNewMidPhases(FCollisionConstraintAllocator* Allocator)
 		{
-			FGeometryParticleHandle* Particle0 = MidPhase->GetParticle0();
-			FGeometryParticleHandle* Particle1 = MidPhase->GetParticle1();
+			// NOTE: Called from the physics thread, and never from a physics task/parallel-for. No need for locks.
 
-			// NOTE: It may look weird that we are passing the particle's self to the AddMidPhase method. 
-			// The midphase represents two particles and it needs to know which one we're currently adding it to.
-			// We could make this API nicer but providing a helper on the Particle, but it's a bit awkard to do so
-			// and since AddMidPahse is only called from here we'll leave it like this for now.
-			Particle0->ParticleCollisions().AddMidPhase(Particle0, MidPhase);
-			Particle1->ParticleCollisions().AddMidPhase(Particle1, MidPhase);
+			// Register the midphases with each of their particles and add the midphase to the central list
+			for (FParticlePairMidPhase* MidPhase : NewMidPhases)
+			{
+				FGeometryParticleHandle* Particle0 = MidPhase->GetParticle0();
+				FGeometryParticleHandle* Particle1 = MidPhase->GetParticle1();
 
-#if CHAOS_COLLISION_OBJECTPOOL_ENABLED 
-			Allocator->AddMidPhase(FParticlePairMidPhasePtr(MidPhase, FParticlePairMidPhaseDeleter(MidPhasePool)));
-#else
-			Allocator->AddMidPhase(FParticlePairMidPhasePtr(MidPhase));
-#endif
-		}
-		NewMidPhases.Reset();
-	}
+				// NOTE: It may look weird that we are passing the particle's self to the AddMidPhase method. 
+				// The midphase represents two particles and it needs to know which one we're currently adding it to.
+				// We could make this API nicer but providing a helper on the Particle, but it's a bit awkard to do so
+				// and since AddMidPahse is only called from here we'll leave it like this for now.
+				Particle0->ParticleCollisions().AddMidPhase(Particle0, MidPhase);
+				Particle1->ParticleCollisions().AddMidPhase(Particle1, MidPhase);
 
-	void FCollisionContextAllocator::ProcessNewConstraints(FCollisionConstraintAllocator* Allocator)
-	{
-		// NOTE: Called from the physics thread, and never from a physics task/parallel-for. No need for locks.
+	#if CHAOS_COLLISION_OBJECTPOOL_ENABLED 
+				Allocator->AddMidPhase(FParticlePairMidPhasePtr(MidPhase, FParticlePairMidPhaseDeleter(MidPhasePool)));
+	#else
+				Allocator->AddMidPhase(FParticlePairMidPhasePtr(MidPhase));
+	#endif
+			}
+			NewMidPhases.Reset();
+		}
 
-		for (FPBDCollisionConstraint* NewConstraint : NewActiveConstraints)
+		void FCollisionContextAllocator::ProcessNewConstraints(FCollisionConstraintAllocator* Allocator)
 		{
-			Allocator->ActivateConstraintImp(NewConstraint);
-		}
-		NewActiveConstraints.Reset();
-	}
+			// NOTE: Called from the physics thread, and never from a physics task/parallel-for. No need for locks.
 
-	void FCollisionConstraintAllocator::PruneExpiredMidPhases()
-	{
-		// NOTE: Called from the physics thread, and never from a physics task/parallel-for. No need for locks.
+			for (FPBDCollisionConstraint* NewConstraint : NewActiveConstraints)
+			{
+				Allocator->ActivateConstraintImp(NewConstraint);
+			}
+			NewActiveConstraints.Reset();
+		}
 
-		// Determine which particle pairs are no longer overlapping
-		// Prune all pairs which were not updated this tick as part of the collision
-		// detection loop and are not asleep
-		for (int32 Index = ParticlePairMidPhases.Num() - 1; Index >= 0; --Index)
+		void FCollisionConstraintAllocator::PruneExpiredMidPhases()
 		{
-			FParticlePairMidPhasePtr& MidPhase = ParticlePairMidPhases[Index];
+			// NOTE: Called from the physics thread, and never from a physics task/parallel-for. No need for locks.
 
-			// We could also add !MidPhase->IsInConstraintGraph() here, but we know that we will not be in the graph if we were
-			// not active this tick and were not asleep. The constraint graph ejects all non-sleeping constraints each tick.
-			// (There is a check in the collision destructor that verified this).
-			if (!MidPhase->IsUsedSince(CurrentEpoch) && !MidPhase->IsSleeping())
+			// Determine which particle pairs are no longer overlapping
+			// Prune all pairs which were not updated this tick as part of the collision
+			// detection loop and are not asleep
+			for (int32 Index = ParticlePairMidPhases.Num() - 1; Index >= 0; --Index)
 			{
-				// Remove from the particles' lists of contacts
-				DetachParticlePairMidPhase(MidPhase.Get());
+				FParticlePairMidPhasePtr& MidPhase = ParticlePairMidPhases[Index];
+
+				// We could also add !MidPhase->IsInConstraintGraph() here, but we know that we will not be in the graph if we were
+				// not active this tick and were not asleep. The constraint graph ejects all non-sleeping constraints each tick.
+				// (There is a check in the collision destructor that verified this).
+				if (!MidPhase->IsUsedSince(CurrentEpoch) && !MidPhase->IsSleeping())
+				{
+					// Remove from the particles' lists of contacts
+					DetachParticlePairMidPhase(MidPhase.Get());
 
-				// Destroy the midphase and its collisions
-				MidPhase.Reset();
+					// Destroy the midphase and its collisions
+					MidPhase.Reset();
 
-				// Remove the midphase from the list. 
-				// ParticlePairMidPhases can get large so we allow it to shrink from time to time
-				const int32 MaxSlack = 1000;
-				const int32 Slack = ParticlePairMidPhases.Max() - ParticlePairMidPhases.Num();
-				const bool bAllowShrink = (Slack > MaxSlack);
-				ParticlePairMidPhases.RemoveAtSwap(Index, 1, bAllowShrink);
+					// Remove the midphase from the list. 
+					// ParticlePairMidPhases can get large so we allow it to shrink from time to time
+					const int32 MaxSlack = 1000;
+					const int32 Slack = ParticlePairMidPhases.Max() - ParticlePairMidPhases.Num();
+					const bool bAllowShrink = (Slack > MaxSlack);
+					ParticlePairMidPhases.RemoveAtSwap(Index, 1, bAllowShrink);
+				}
 			}
 		}
-	}
 
-}
\ No newline at end of file
+	}	// namespace Private
+}	// namespace Chaos
\ No newline at end of file
diff --git a/Engine/Source/Runtime/Experimental/Chaos/Private/Chaos/PBDRigidsEvolutionGBF.cpp b/Engine/Source/Runtime/Experimental/Chaos/Private/Chaos/PBDRigidsEvolutionGBF.cpp
--- a/Engine/Source/Runtime/Experimental/Chaos/Private/Chaos/PBDRigidsEvolutionGBF.cpp
+++ b/Engine/Source/Runtime/Experimental/Chaos/Private/Chaos/PBDRigidsEvolutionGBF.cpp
@@ -712,7 +712,7 @@ void FPBDRigidsEvolutionGBF::TransferJointConstraintCollisions()
 	// which is quite difficult for large mass ratios and would require many iterations.
 	if (DoTransferJointConstraintCollisions)
 	{
-		FCollisionConstraintAllocator& CollisionAllocator = CollisionConstraints.GetConstraintAllocator();
+		Private::FCollisionConstraintAllocator& CollisionAllocator = CollisionConstraints.GetConstraintAllocator();
 
 		// @todo(chaos): we should only visit the joints that have ContactTransferScale > 0 
 		for (int32 JointConstraintIndex = 0; JointConstraintIndex < GetJointConstraints().NumConstraints(); ++JointConstraintIndex)
diff --git a/Engine/Source/Runtime/Experimental/Chaos/Public/Chaos/CCDUtilities.h b/Engine/Source/Runtime/Experimental/Chaos/Public/Chaos/CCDUtilities.h
--- a/Engine/Source/Runtime/Experimental/Chaos/Public/Chaos/CCDUtilities.h
+++ b/Engine/Source/Runtime/Experimental/Chaos/Public/Chaos/CCDUtilities.h
@@ -84,7 +84,7 @@ namespace Chaos
 	{
 	public:
 		FCCDManager(){}
-		void ApplyConstraintsPhaseCCD(const FReal Dt, FCollisionConstraintAllocator *CollisionAllocator, const int32 NumDynamicParticles = TNumericLimits<int32>::Max());
+		void ApplyConstraintsPhaseCCD(const FReal Dt, Private::FCollisionConstraintAllocator *CollisionAllocator, const int32 NumDynamicParticles = TNumericLimits<int32>::Max());
 
 		// A post process on CCD contacts to ensure that CCD objects never pass through non-dynamic objects
 		void ApplyCorrections(const FReal Dt);
@@ -93,7 +93,7 @@ namespace Chaos
 		void ApplySweptConstraints(const FReal Dt, TArrayView<FPBDCollisionConstraint* const> InSweptConstraints, const int32 NumDynamicParticles = TNumericLimits<int32>::Max());
 		bool UpdateParticleSweptConstraints(FCCDParticle* CCDParticle, const FReal IslandTOI, const FReal Dt);
 		// This is called after ApplySweptConstraints. This function updates manifold data which will be used in normal solve.
-		void UpdateSweptConstraints(const FReal Dt, FCollisionConstraintAllocator *CollisionAllocator);
+		void UpdateSweptConstraints(const FReal Dt, Private::FCollisionConstraintAllocator *CollisionAllocator);
 		void OverwriteXUsingV(const FReal Dt);
 		bool Init(const FReal Dt, const int32 NumDynamicParticles = TNumericLimits<int32>::Max());
 		void AssignParticleIslandsAndGroupParticles();
diff --git a/Engine/Source/Runtime/Experimental/Chaos/Public/Chaos/ChaosDebugDraw.h b/Engine/Source/Runtime/Experimental/Chaos/Public/Chaos/ChaosDebugDraw.h
--- a/Engine/Source/Runtime/Experimental/Chaos/Public/Chaos/ChaosDebugDraw.h
+++ b/Engine/Source/Runtime/Experimental/Chaos/Public/Chaos/ChaosDebugDraw.h
@@ -9,9 +9,13 @@
 
 namespace Chaos
 {
+	namespace Private
+	{
+		class FCollisionConstraintAllocator;
+	}
+
 	class FAccelerationStructureHandle;
 	struct FCCDConstraint;
-	class FCollisionConstraintAllocator;
 	template <typename PayloadType, typename T, int d> class ISpatialAcceleration;
 	class FPBDSuspensionConstraints;
 	class FRigidClustering;
@@ -210,7 +214,7 @@ namespace Chaos
 		CHAOS_API void DrawParticleTransforms(const FRigidTransform3& SpaceTransform, const TParticleView<FPBDRigidParticles>& ParticlesView, const FChaosDebugDrawSettings* Settings = nullptr);
 		CHAOS_API void DrawParticleCollisions(const FRigidTransform3& SpaceTransform, const FGeometryParticleHandle* Particle, const FPBDCollisionConstraints& Collisions, const FChaosDebugDrawSettings* Settings = nullptr);
 		CHAOS_API void DrawCollisions(const FRigidTransform3& SpaceTransform, const FPBDCollisionConstraints& Collisions, FRealSingle ColorScale, const FChaosDebugDrawSettings* Settings = nullptr);
-		CHAOS_API void DrawCollisions(const FRigidTransform3& SpaceTransform, const FCollisionConstraintAllocator& CollisionAllocator, FRealSingle ColorScale, const FChaosDebugDrawSettings* Settings = nullptr);
+		CHAOS_API void DrawCollisions(const FRigidTransform3& SpaceTransform, const Private::FCollisionConstraintAllocator& CollisionAllocator, FRealSingle ColorScale, const FChaosDebugDrawSettings* Settings = nullptr);
 		CHAOS_API void DrawJointConstraints(const FRigidTransform3& SpaceTransform, const TArray<FPBDJointConstraintHandle*>& ConstraintHandles, FRealSingle ColorScale, const FChaosDebugDrawJointFeatures& FeatureMask = FChaosDebugDrawJointFeatures::MakeDefault(), const FChaosDebugDrawSettings* Settings = nullptr);
 		CHAOS_API void DrawJointConstraints(const FRigidTransform3& SpaceTransform, const FPBDJointConstraints& Constraints, FRealSingle ColorScale, const FChaosDebugDrawJointFeatures& FeatureMask = FChaosDebugDrawJointFeatures::MakeDefault(), const FChaosDebugDrawSettings* Settings = nullptr);
 		CHAOS_API void DrawSimulationSpace(const FSimulationSpace& SimSpace, const FChaosDebugDrawSettings* Settings = nullptr);
diff --git a/Engine/Source/Runtime/Experimental/Chaos/Public/Chaos/Collision/CollisionConstraintAllocator.h b/Engine/Source/Runtime/Experimental/Chaos/Public/Chaos/Collision/CollisionConstraintAllocator.h
--- a/Engine/Source/Runtime/Experimental/Chaos/Public/Chaos/Collision/CollisionConstraintAllocator.h
+++ b/Engine/Source/Runtime/Experimental/Chaos/Public/Chaos/Collision/CollisionConstraintAllocator.h
@@ -17,569 +17,572 @@ namespace Chaos
 {
 	class FPBDCollisionConstraints;
 
-	/**
-	 * Container the storage for the FCollisionConstraintAllocator, as well as the API to create new midphases and collision constraints.
-	 * We have one of these objects per thread on which collisions detection is performed to get lock-free allocations and lists.
-	 */
-	class CHAOS_API FCollisionContextAllocator
+	namespace Private
 	{
-	public:
-		FCollisionContextAllocator(FPBDCollisionConstraints* InCollisionContainer, const int32 InNumCollisionsPerBlock, const int32 InCurrentEpoch)
-			: CollisionContainer(InCollisionContainer)
-			, CurrentEpoch(InCurrentEpoch)
-#if CHAOS_COLLISION_OBJECTPOOL_ENABLED
-			, ConstraintPool(InNumCollisionsPerBlock, 0)
-			, MidPhasePool(InNumCollisionsPerBlock, 0)
-#endif
-		{
-		}
-
 		/**
-		 * The current epoch used to determine if a collision is up to date
+		 * Container the storage for the FCollisionConstraintAllocator, as well as the API to create new midphases and collision constraints.
+		 * We have one of these objects per thread on which collisions detection is performed to get lock-free allocations and lists.
 		 */
-		int32 GetCurrentEpoch() const
-		{
-			return CurrentEpoch;
-		}
+		class CHAOS_API FCollisionContextAllocator
+		{
+		public:
+			FCollisionContextAllocator(FPBDCollisionConstraints* InCollisionContainer, const int32 InNumCollisionsPerBlock, const int32 InCurrentEpoch)
+				: CollisionContainer(InCollisionContainer)
+				, CurrentEpoch(InCurrentEpoch)
+	#if CHAOS_COLLISION_OBJECTPOOL_ENABLED
+				, ConstraintPool(InNumCollisionsPerBlock, 0)
+				, MidPhasePool(InNumCollisionsPerBlock, 0)
+	#endif
+			{
+			}
 
-		/**
-		 * Create a constraint (called by the MidPhase)
-		 */
-		FPBDCollisionConstraintPtr CreateConstraint(
-			FGeometryParticleHandle* Particle0,
-			const FImplicitObject* Implicit0,
-			const FPerShapeData* Shape0,
-			const FBVHParticles* Simplicial0,
-			const FRigidTransform3& ShapeRelativeTransform0,
-			FGeometryParticleHandle* Particle1,
-			const FImplicitObject* Implicit1,
-			const FPerShapeData* Shape1,
-			const FBVHParticles* Simplicial1,
-			const FRigidTransform3& ShapeRelativeTransform1,
-			const FReal CullDistance,
-			const bool bUseManifold,
-			const EContactShapesType ShapePairType);
+			/**
+			 * The current epoch used to determine if a collision is up to date
+			 */
+			int32 GetCurrentEpoch() const
+			{
+				return CurrentEpoch;
+			}
 
-		/**
-		 * Create an uninitialized collision constraint (public only for use by Resim which overwrites it with a saved constraint)
-		 */
-		FPBDCollisionConstraintPtr CreateConstraint()
-		{
-#if CHAOS_COLLISION_OBJECTPOOL_ENABLED 
-			FPBDCollisionConstraint* Constraint = ConstraintPool.Alloc();
-			return FPBDCollisionConstraintPtr(Constraint, FPBDCollisionConstraintDeleter(ConstraintPool));
-#else
-			return MakeUnique<FPBDCollisionConstraint>();
-#endif
-		}
+			/**
+			 * Create a constraint (called by the MidPhase)
+			 */
+			FPBDCollisionConstraintPtr CreateConstraint(
+				FGeometryParticleHandle* Particle0,
+				const FImplicitObject* Implicit0,
+				const FPerShapeData* Shape0,
+				const FBVHParticles* Simplicial0,
+				const FRigidTransform3& ShapeRelativeTransform0,
+				FGeometryParticleHandle* Particle1,
+				const FImplicitObject* Implicit1,
+				const FPerShapeData* Shape1,
+				const FBVHParticles* Simplicial1,
+				const FRigidTransform3& ShapeRelativeTransform1,
+				const FReal CullDistance,
+				const bool bUseManifold,
+				const EContactShapesType ShapePairType);
+
+			/**
+			 * Create an uninitialized collision constraint (public only for use by Resim which overwrites it with a saved constraint)
+			 */
+			FPBDCollisionConstraintPtr CreateConstraint()
+			{
+	#if CHAOS_COLLISION_OBJECTPOOL_ENABLED 
+				FPBDCollisionConstraint* Constraint = ConstraintPool.Alloc();
+				return FPBDCollisionConstraintPtr(Constraint, FPBDCollisionConstraintDeleter(ConstraintPool));
+	#else
+				return MakeUnique<FPBDCollisionConstraint>();
+	#endif
+			}
 
-		/**
-		 * Set the constraint as Active - it will be added to the graph and solved this tick.
-		 * NOTE: This actually just enqueues the constraint - actual activation happens after the parallel phase is done.
-		 */
-		bool ActivateConstraint(FPBDCollisionConstraint* Constraint)
-		{
-			// When we wake an Island, we reactivate all constraints for all dynamic particles in the island. This
-			// results in duplicate calls to active for constraints involving two dynamic particles, hence the check for CurrentEpoch.
-			// @todo(chaos): fix duplicate calls from island wake. See UpdateSleepState in IslandManager.cpp
-			FPBDCollisionConstraintContainerCookie& Cookie = Constraint->GetContainerCookie();
-			if (Cookie.LastUsedEpoch != CurrentEpoch)
+			/**
+			 * Set the constraint as Active - it will be added to the graph and solved this tick.
+			 * NOTE: This actually just enqueues the constraint - actual activation happens after the parallel phase is done.
+			 */
+			bool ActivateConstraint(FPBDCollisionConstraint* Constraint)
 			{
-				Cookie.LastUsedEpoch = CurrentEpoch;
+				// When we wake an Island, we reactivate all constraints for all dynamic particles in the island. This
+				// results in duplicate calls to active for constraints involving two dynamic particles, hence the check for CurrentEpoch.
+				// @todo(chaos): fix duplicate calls from island wake. See UpdateSleepState in IslandManager.cpp
+				FPBDCollisionConstraintContainerCookie& Cookie = Constraint->GetContainerCookie();
+				if (Cookie.LastUsedEpoch != CurrentEpoch)
+				{
+					Cookie.LastUsedEpoch = CurrentEpoch;
 
-				NewActiveConstraints.Push(Constraint);
+					NewActiveConstraints.Push(Constraint);
 
-				Constraint->Activate();
+					Constraint->Activate();
 
-				return true;
+					return true;
+				}
+				return false;
 			}
-			return false;
-		}
 
-		/**
-		 * Return a midphase for a particle pair.
-		 * This wil create a new midphase if the particle pairs were not recently overlapping, or return an
-		 * existing one if they were.
-		 * @note Nothing outside of thie allocator should hold a pointer to the midphase, or any constraints
-		 * it creates for more than the duration of the tick. Except the IslandManager :|
-		*/
-		FParticlePairMidPhase* GetMidPhase(FGeometryParticleHandle* Particle0, FGeometryParticleHandle* Particle1, FGeometryParticleHandle* SearchParticlePerformanceHint, const FCollisionContext& Context)
-		{
-			FParticlePairMidPhase* MidPhase = FindMidPhaseImpl(Particle0, Particle1, SearchParticlePerformanceHint);
-			if (MidPhase == nullptr)
+			/**
+			 * Return a midphase for a particle pair.
+			 * This wil create a new midphase if the particle pairs were not recently overlapping, or return an
+			 * existing one if they were.
+			 * @note Nothing outside of thie allocator should hold a pointer to the midphase, or any constraints
+			 * it creates for more than the duration of the tick. Except the IslandManager :|
+			*/
+			FParticlePairMidPhase* GetMidPhase(FGeometryParticleHandle* Particle0, FGeometryParticleHandle* Particle1, FGeometryParticleHandle* SearchParticlePerformanceHint, const FCollisionContext& Context)
 			{
-				MidPhase = CreateMidPhase(Particle0, Particle1, Context);
-			}
+				FParticlePairMidPhase* MidPhase = FindMidPhaseImpl(Particle0, Particle1, SearchParticlePerformanceHint);
+				if (MidPhase == nullptr)
+				{
+					MidPhase = CreateMidPhase(Particle0, Particle1, Context);
+				}
 
-			return MidPhase;
-		}
+				return MidPhase;
+			}
 
-		/**
-		 * Return a midphase for a particle pair if it already exists, otherwise return null
-		*/
-		FParticlePairMidPhase* FindMidPhase(FGeometryParticleHandle* Particle0, FGeometryParticleHandle* Particle1, FGeometryParticleHandle* SearchParticlePerformanceHint)
-		{
-			return FindMidPhaseImpl(Particle0, Particle1, SearchParticlePerformanceHint);
-		}
+			/**
+			 * Return a midphase for a particle pair if it already exists, otherwise return null
+			*/
+			FParticlePairMidPhase* FindMidPhase(FGeometryParticleHandle* Particle0, FGeometryParticleHandle* Particle1, FGeometryParticleHandle* SearchParticlePerformanceHint)
+			{
+				return FindMidPhaseImpl(Particle0, Particle1, SearchParticlePerformanceHint);
+			}
 
-	private:
-		friend class FCollisionConstraintAllocator;
+		private:
+			friend class Private::FCollisionConstraintAllocator;
 
-		void Reset()
-		{
-#if CHAOS_COLLISION_OBJECTPOOL_ENABLED
-			ConstraintPool.Reset();
-			MidPhasePool.Reset();
-#endif
-		}
+			void Reset()
+			{
+	#if CHAOS_COLLISION_OBJECTPOOL_ENABLED
+				ConstraintPool.Reset();
+				MidPhasePool.Reset();
+	#endif
+			}
 
-		void BeginDetectCollisions(const int32 InEpoch)
-		{
-			check(NewActiveConstraints.IsEmpty());
-			check(NewMidPhases.IsEmpty());
+			void BeginDetectCollisions(const int32 InEpoch)
+			{
+				check(NewActiveConstraints.IsEmpty());
+				check(NewMidPhases.IsEmpty());
 
-			CurrentEpoch = InEpoch;
-		}
+				CurrentEpoch = InEpoch;
+			}
 
-		// Process all the new midphases from the parallel collision detection
-		void ProcessNewMidPhases(FCollisionConstraintAllocator* Allocator);
+			// Process all the new midphases from the parallel collision detection
+			void ProcessNewMidPhases(FCollisionConstraintAllocator* Allocator);
 
-		// Process all the activated collisions from the parallel collision detection
-		void ProcessNewConstraints(FCollisionConstraintAllocator* Allocator);
+			// Process all the activated collisions from the parallel collision detection
+			void ProcessNewConstraints(FCollisionConstraintAllocator* Allocator);
 
-		// Find the midphase for the particle pair if it exists. Every particle holds a list of its midphases. We search "SearchParticle" which should be
-		// ideally be the one with fewer midphases on it.
-		FParticlePairMidPhase* FindMidPhaseImpl(FGeometryParticleHandle* Particle0, FGeometryParticleHandle* Particle1, FGeometryParticleHandle* SearchParticle)
-		{
-			check((SearchParticle == Particle0) || (SearchParticle == Particle1));
+			// Find the midphase for the particle pair if it exists. Every particle holds a list of its midphases. We search "SearchParticle" which should be
+			// ideally be the one with fewer midphases on it.
+			FParticlePairMidPhase* FindMidPhaseImpl(FGeometryParticleHandle* Particle0, FGeometryParticleHandle* Particle1, FGeometryParticleHandle* SearchParticle)
+			{
+				check((SearchParticle == Particle0) || (SearchParticle == Particle1));
 
-			const FCollisionParticlePairKey Key = FCollisionParticlePairKey(Particle0, Particle1);
-			return SearchParticle->ParticleCollisions().FindMidPhase(Key.GetKey());
-		}
+				const FCollisionParticlePairKey Key = FCollisionParticlePairKey(Particle0, Particle1);
+				return SearchParticle->ParticleCollisions().FindMidPhase(Key.GetKey());
+			}
 
-		// Create and initialize a midphase for the particle pair. Adds it to the particles' lists of midphases.
-		FParticlePairMidPhase* CreateMidPhase(FGeometryParticleHandle* Particle0, FGeometryParticleHandle* Particle1, const FCollisionContext& Context)
-		{
-			const FCollisionParticlePairKey Key = FCollisionParticlePairKey(Particle0, Particle1);
+			// Create and initialize a midphase for the particle pair. Adds it to the particles' lists of midphases.
+			FParticlePairMidPhase* CreateMidPhase(FGeometryParticleHandle* Particle0, FGeometryParticleHandle* Particle1, const FCollisionContext& Context)
+			{
+				const FCollisionParticlePairKey Key = FCollisionParticlePairKey(Particle0, Particle1);
 
-			// We temporarily hold new midphases as raw pointers and wrap them in a unique ptr later in ProcessNewMidPhases
-#if CHAOS_COLLISION_OBJECTPOOL_ENABLED 
-			FParticlePairMidPhase* MidPhase = MidPhasePool.Alloc();
-#else
-			FParticlePairMidPhase* MidPhase = new FParticlePairMidPhase();
-#endif
+				// We temporarily hold new midphases as raw pointers and wrap them in a unique ptr later in ProcessNewMidPhases
+	#if CHAOS_COLLISION_OBJECTPOOL_ENABLED 
+				FParticlePairMidPhase* MidPhase = MidPhasePool.Alloc();
+	#else
+				FParticlePairMidPhase* MidPhase = new FParticlePairMidPhase();
+	#endif
 
-			NewMidPhases.Add(MidPhase);
+				NewMidPhases.Add(MidPhase);
 
-			MidPhase->Init(Particle0, Particle1, Key, Context);
-			return MidPhase;
-		}
+				MidPhase->Init(Particle0, Particle1, Key, Context);
+				return MidPhase;
+			}
 
-		FPBDCollisionConstraints* CollisionContainer;
-		int32 CurrentEpoch;
-
-		TArray<FPBDCollisionConstraint*> NewActiveConstraints;
-		TArray<FParticlePairMidPhase*> NewMidPhases;
-
-#if CHAOS_COLLISION_OBJECTPOOL_ENABLED
-		FPBDCollisionConstraintPool ConstraintPool;
-		FParticlePairMidPhasePool MidPhasePool;
-#endif
-	};
-
-
-	/**
-	 * @brief An allocator and container of collision constraints that supports reuse of constraints from the previous tick
-	 * 
-	 * All constraint pointers are persistent in memory until Destroy() is called, or until they are pruned.
-
-	 * This allocator maintains the set of all overlapping particle pairs, with each overlapping particle pair managed
-	 * by a FParticlePairMidPhase object. the MidPhase object is what actually calls the Narrow Phase and maintains
-	 * the set of collision constraints for all the shape pairs on the particles.
-	 * 
-	 * Constraints are allocated during the collision detection phase and retained between ticks. An attempt to create
-	 * a constraint for the same shape pair as seen on the previous tick will return the existing collision constraint with
-	 * all of its data intact.
-	 *
-	 * The allocator also keeps a list of Standard and Swept collision constraints that are active for the current tick.
-	 * This list gets reset and rebuilt every frame during collision detection. It may get added to by the IslandManager
-	 * if some islands are woken following collision detection.
-	 * 
-	 * The allocators Epoch counter is used to determine whether a constraint (or midphase object) generated any
-	 * contacts for the current frame. When a midphase creates or updates a constraint, it copies the current Epoch
-	 * counter.
-	 * 
-	 * The Midphase list is pruned at the end of each tick so if particles are destroyed or a particle pair is no longer 
-	 * overlapping, the collisions will be destroyed.
-	 * 
-	 * When particles are destroyed, we do not immediately destroy the MidPhases (or Collisions) that are associoated with
-	 * the particle. Instead, we clear the particle pointer from them, but leave their destruction to the pruning process.
-	 * This avoids the need to parse collision and midphase lists whenever a particle is disabled.
-	 * 
-	 * NOTE: To reduce RBAN memory use, we do not create any collision blocks until the first call to CreateCollisionConstraint
-	 * (see ConstraintPool initialization)
-	*/
-	class CHAOS_API FCollisionConstraintAllocator
-	{
-	public:
-		UE_NONCOPYABLE(FCollisionConstraintAllocator);
-
-		FCollisionConstraintAllocator(const int32 InNumCollisionsPerBlock = 1000)
-			: CollisionContainer(nullptr)
-			, ContextAllocators()
-			, NumCollisionsPerBlock(InNumCollisionsPerBlock)
-			, ParticlePairMidPhases()
-			, ActiveConstraints()
-			, ActiveCCDConstraints()
-			, CurrentEpoch(0)
-			, bInCollisionDetectionPhase(false)
-		{
-		}
+			FPBDCollisionConstraints* CollisionContainer;
+			int32 CurrentEpoch;
 
-		~FCollisionConstraintAllocator()
-		{
-			// Explicit clear so we don't have to worry about destruction order of pool and midphases
-			Reset();
-		}
+			TArray<FPBDCollisionConstraint*> NewActiveConstraints;
+			TArray<FParticlePairMidPhase*> NewMidPhases;
 
-		void SetCollisionContainer(FPBDCollisionConstraints* InCollisionContainer)
-		{
-			CollisionContainer = InCollisionContainer;
-		}
-
-		void SetMaxContexts(const int32 MaxContexts)
-		{
-			// Forgot to call SetCollisionContainer after construction
-			// @todo(chaos): this check makes the allocator difficult to unit test - fix this
-			//check(CollisionContainer != nullptr);
+	#if CHAOS_COLLISION_OBJECTPOOL_ENABLED
+			FPBDCollisionConstraintPool ConstraintPool;
+			FParticlePairMidPhasePool MidPhasePool;
+	#endif
+		};
 
-			const int32 NumNewContexts = MaxContexts - ContextAllocators.Num();
-			if (NumNewContexts > 0)
-			{
-				ContextAllocators.Reserve(MaxContexts);
-				for (int32 NewIndex = 0; NewIndex < NumNewContexts; ++NewIndex)
-				{
-					TUniquePtr<FCollisionContextAllocator> ContextAllocator = MakeUnique<FCollisionContextAllocator>(CollisionContainer, NumCollisionsPerBlock, CurrentEpoch);
-					ContextAllocators.Emplace(MoveTemp(ContextAllocator));
-				}
-			}
-		}
 
-		FCollisionContextAllocator* GetContextAllocator(const int32 Index)
-		{
-			return ContextAllocators[Index].Get();
-		}
-		
 		/**
-		 * @brief The set of collision constraints for the current tick (created or reinstated)
+		 * @brief An allocator and container of collision constraints that supports reuse of constraints from the previous tick
 		 * 
-		 * @note Some elements may be null (constraints that have been deleted)
-		 * @note This is not thread-safe and should not be used during the collision detection phase (i.e., when the list is being built)
-		*/
-		TArrayView<FPBDCollisionConstraint* const> GetConstraints() const
-		{ 
-			return MakeArrayView(ActiveConstraints);
-		}
+		 * All constraint pointers are persistent in memory until Destroy() is called, or until they are pruned.
 
-		/**
-		 * @brief The set of sweep collision constraints for the current tick (created or reinstated)
+		 * This allocator maintains the set of all overlapping particle pairs, with each overlapping particle pair managed
+		 * by a FParticlePairMidPhase object. the MidPhase object is what actually calls the Narrow Phase and maintains
+		 * the set of collision constraints for all the shape pairs on the particles.
 		 * 
-		 * @note Some elements may be null (constraints that have been explicitly deleted)
-		 * @note This is not thread-safe and should not be used during the collision detection phase (i.e., when the list is being built)
-		*/
-		TArrayView<FPBDCollisionConstraint* const> GetCCDConstraints() const
-		{ 
-			return MakeArrayView(ActiveCCDConstraints);
-		}
-
-		/**
-		 * @brief The set of collision constraints for the current tick (created or reinstated)
+		 * Constraints are allocated during the collision detection phase and retained between ticks. An attempt to create
+		 * a constraint for the same shape pair as seen on the previous tick will return the existing collision constraint with
+		 * all of its data intact.
+		 *
+		 * The allocator also keeps a list of Standard and Swept collision constraints that are active for the current tick.
+		 * This list gets reset and rebuilt every frame during collision detection. It may get added to by the IslandManager
+		 * if some islands are woken following collision detection.
 		 * 
-		 * @note Some elements may be null (constraints that have been explicitly deleted)
-		 * @note This is not thread-safe and should not be used during the collision detection phase (i.e., when the list is being built)
-		 */
-		TArrayView<const FPBDCollisionConstraint* const> GetConstConstraints() const
-		{
-			return MakeArrayView(ActiveConstraints);
-		}
-
-		int32 GetCurrentEpoch() const
-		{
-			return CurrentEpoch;
-		}
-
-		/**
-		* Has the constraint expired. An expired constraint is one that was not refreshed this tick.
-		* 
-		* @note Sleeping constrainst will report they are expired, but they should not be deleted until awoken.
-		* We deliberately don't check the sleeping flag here because it adds a cache miss. This function
-		* should be called only when you already know the sleep state (or you must also check IsSleeping())
+		 * The allocators Epoch counter is used to determine whether a constraint (or midphase object) generated any
+		 * contacts for the current frame. When a midphase creates or updates a constraint, it copies the current Epoch
+		 * counter.
+		 * 
+		 * The Midphase list is pruned at the end of each tick so if particles are destroyed or a particle pair is no longer 
+		 * overlapping, the collisions will be destroyed.
+		 * 
+		 * When particles are destroyed, we do not immediately destroy the MidPhases (or Collisions) that are associoated with
+		 * the particle. Instead, we clear the particle pointer from them, but leave their destruction to the pruning process.
+		 * This avoids the need to parse collision and midphase lists whenever a particle is disabled.
+		 * 
+		 * NOTE: To reduce RBAN memory use, we do not create any collision blocks until the first call to CreateCollisionConstraint
+		 * (see ConstraintPool initialization)
 		*/
-		bool IsConstraintExpired(const FPBDCollisionConstraint& Constraint)
-		{
-			const bool bIsExpired = Constraint.GetContainerCookie().LastUsedEpoch < CurrentEpoch;
-			return bIsExpired;
-		}
+		class CHAOS_API FCollisionConstraintAllocator
+		{
+		public:
+			UE_NONCOPYABLE(FCollisionConstraintAllocator);
+
+			FCollisionConstraintAllocator(const int32 InNumCollisionsPerBlock = 1000)
+				: CollisionContainer(nullptr)
+				, ContextAllocators()
+				, NumCollisionsPerBlock(InNumCollisionsPerBlock)
+				, ParticlePairMidPhases()
+				, ActiveConstraints()
+				, ActiveCCDConstraints()
+				, CurrentEpoch(0)
+				, bInCollisionDetectionPhase(false)
+			{
+			}
 
-		/**
-		 * @brief Destroy all constraints
-		*/
-		void Reset()
-		{
-			ActiveConstraints.Reset();
-			ActiveCCDConstraints.Reset();
-			ParticlePairMidPhases.Reset();
+			~FCollisionConstraintAllocator()
+			{
+				// Explicit clear so we don't have to worry about destruction order of pool and midphases
+				Reset();
+			}
 
-			// NOTE: this deletes all the storage for constraints and midphases - must be last
-			ContextAllocators.Reset();
-		}
+			void SetCollisionContainer(FPBDCollisionConstraints* InCollisionContainer)
+			{
+				CollisionContainer = InCollisionContainer;
+			}
 
-		/**
-		 * @brief Called at the start of the frame to clear the frame's active collision list.
-		 * @todo(chaos): This is only required because of the way events work (see AdvanceOneTimeStepTask::DoWork)
-		*/
-		void BeginFrame()
-		{
-			ActiveConstraints.Reset();
-			ActiveCCDConstraints.Reset();
-		}
+			void SetMaxContexts(const int32 MaxContexts)
+			{
+				// Forgot to call SetCollisionContainer after construction
+				// @todo(chaos): this check makes the allocator difficult to unit test - fix this
+				//check(CollisionContainer != nullptr);
 
-		/**
-		 * @brief Called at the start of the tick to prepare for collision detection.
-		 * Resets the list of active contacts.
-		*/
-		void BeginDetectCollisions()
-		{
-			check(!bInCollisionDetectionPhase);
-			bInCollisionDetectionPhase = true;
+				const int32 NumNewContexts = MaxContexts - ContextAllocators.Num();
+				if (NumNewContexts > 0)
+				{
+					ContextAllocators.Reserve(MaxContexts);
+					for (int32 NewIndex = 0; NewIndex < NumNewContexts; ++NewIndex)
+					{
+						TUniquePtr<FCollisionContextAllocator> ContextAllocator = MakeUnique<FCollisionContextAllocator>(CollisionContainer, NumCollisionsPerBlock, CurrentEpoch);
+						ContextAllocators.Emplace(MoveTemp(ContextAllocator));
+					}
+				}
+			}
 
-			// Update the tick counter
-			// NOTE: We do this here rather than in EndDetectionCollisions so that any contacts injected
-			// before collision detection count as the previous frame's collisions, e.g., from Islands
-			// that are manually awoken by modifying a particle on the game thread. This also needs to be
-			// done where we reset the Constraints array so that we can tell we have a valid index from
-			// the Epoch.
-			++CurrentEpoch;
+			FCollisionContextAllocator* GetContextAllocator(const int32 Index)
+			{
+				return ContextAllocators[Index].Get();
+			}
+		
+			/**
+			 * @brief The set of collision constraints for the current tick (created or reinstated)
+			 * 
+			 * @note Some elements may be null (constraints that have been deleted)
+			 * @note This is not thread-safe and should not be used during the collision detection phase (i.e., when the list is being built)
+			*/
+			TArrayView<FPBDCollisionConstraint* const> GetConstraints() const
+			{ 
+				return MakeArrayView(ActiveConstraints);
+			}
+
+			/**
+			 * @brief The set of sweep collision constraints for the current tick (created or reinstated)
+			 * 
+			 * @note Some elements may be null (constraints that have been explicitly deleted)
+			 * @note This is not thread-safe and should not be used during the collision detection phase (i.e., when the list is being built)
+			*/
+			TArrayView<FPBDCollisionConstraint* const> GetCCDConstraints() const
+			{ 
+				return MakeArrayView(ActiveCCDConstraints);
+			}
 
-			for (const TUniquePtr<FCollisionContextAllocator>& ContextAllocator : ContextAllocators)
+			/**
+			 * @brief The set of collision constraints for the current tick (created or reinstated)
+			 * 
+			 * @note Some elements may be null (constraints that have been explicitly deleted)
+			 * @note This is not thread-safe and should not be used during the collision detection phase (i.e., when the list is being built)
+			 */
+			TArrayView<const FPBDCollisionConstraint* const> GetConstConstraints() const
 			{
-				ContextAllocator->BeginDetectCollisions(CurrentEpoch);
+				return MakeArrayView(ActiveConstraints);
 			}
 
-			// Clear the collision list for this tick - we are about to rebuild them
-			ActiveConstraints.Reset();
-			ActiveCCDConstraints.Reset();
-		}
+			int32 GetCurrentEpoch() const
+			{
+				return CurrentEpoch;
+			}
 
-		/**
-		 * @brief Called after collision detection to clean up
-		 * Prunes unused contacts
-		*/
-		void EndDetectCollisions()
-		{
-			check(bInCollisionDetectionPhase);
-			bInCollisionDetectionPhase = false;
+			/**
+			* Has the constraint expired. An expired constraint is one that was not refreshed this tick.
+			* 
+			* @note Sleeping constrainst will report they are expired, but they should not be deleted until awoken.
+			* We deliberately don't check the sleeping flag here because it adds a cache miss. This function
+			* should be called only when you already know the sleep state (or you must also check IsSleeping())
+			*/
+			bool IsConstraintExpired(const FPBDCollisionConstraint& Constraint)
+			{
+				const bool bIsExpired = Constraint.GetContainerCookie().LastUsedEpoch < CurrentEpoch;
+				return bIsExpired;
+			}
 
-			ProcessNewItems();
-		}
+			/**
+			 * @brief Destroy all constraints
+			*/
+			void Reset()
+			{
+				ActiveConstraints.Reset();
+				ActiveCCDConstraints.Reset();
+				ParticlePairMidPhases.Reset();
 
-		/**
-		 * @brief Called each tick after the graph is updated to remove unused collisions
-		*/
-		void PruneExpiredItems()
-		{
-			PruneExpiredMidPhases();
-		}
+				// NOTE: this deletes all the storage for constraints and midphases - must be last
+				ContextAllocators.Reset();
+			}
 
-		/**
-		 * @brief If we add new constraints after collision detection, do what needs to be done to add them to the system
-		*/
-		void ProcessInjectedConstraints()
-		{
-			ProcessNewItems();
-		}
+			/**
+			 * @brief Called at the start of the frame to clear the frame's active collision list.
+			 * @todo(chaos): This is only required because of the way events work (see AdvanceOneTimeStepTask::DoWork)
+			*/
+			void BeginFrame()
+			{
+				ActiveConstraints.Reset();
+				ActiveCCDConstraints.Reset();
+			}
 
-		/**
-		 * @brief Add a set of pre-built constraints and build required internal mapping data
-		 * This is used by the resim cache when restoring constraints after a desync
-		*/
-		void AddResimConstraints(const TArray<FPBDCollisionConstraint>& InConstraints, const FCollisionContext Context)
-		{
-			for (const FPBDCollisionConstraint& SourceConstraint : InConstraints)
-			{
-				// We must keep the particles in the same order that the broadphase would generate when
-				// finding or creating the midphase. This is because Collision Constraints may have the
-				// particles in the opposite order to the midphase tha owns them.
-				FGeometryParticleHandle* Particle0 = SourceConstraint.Particle[0];
-				FGeometryParticleHandle* Particle1 = SourceConstraint.Particle[1];
-				if (ShouldSwapParticleOrder(Particle0, Particle1))
+			/**
+			 * @brief Called at the start of the tick to prepare for collision detection.
+			 * Resets the list of active contacts.
+			*/
+			void BeginDetectCollisions()
+			{
+				check(!bInCollisionDetectionPhase);
+				bInCollisionDetectionPhase = true;
+
+				// Update the tick counter
+				// NOTE: We do this here rather than in EndDetectionCollisions so that any contacts injected
+				// before collision detection count as the previous frame's collisions, e.g., from Islands
+				// that are manually awoken by modifying a particle on the game thread. This also needs to be
+				// done where we reset the Constraints array so that we can tell we have a valid index from
+				// the Epoch.
+				++CurrentEpoch;
+
+				for (const TUniquePtr<FCollisionContextAllocator>& ContextAllocator : ContextAllocators)
 				{
-					Swap(Particle0, Particle1);
+					ContextAllocator->BeginDetectCollisions(CurrentEpoch);
 				}
 
-				FParticlePairMidPhase* MidPhase = Context.GetAllocator()->GetMidPhase(Particle0, Particle1, SourceConstraint.Particle[0], Context);
+				// Clear the collision list for this tick - we are about to rebuild them
+				ActiveConstraints.Reset();
+				ActiveCCDConstraints.Reset();
+			}
+
+			/**
+			 * @brief Called after collision detection to clean up
+			 * Prunes unused contacts
+			*/
+			void EndDetectCollisions()
+			{
+				check(bInCollisionDetectionPhase);
+				bInCollisionDetectionPhase = false;
 
-				// We may be adding multiple constraints for the same particle pair, so we need
-				// to make sure the map is up to date in the case where we just created a new MidPhase
-				ProcessNewMidPhases();
-				
-				if (MidPhase != nullptr)
-				{
-					MidPhase->InjectCollision(SourceConstraint, Context);
-				}
+				ProcessNewItems();
 			}
 
-			ProcessNewConstraints();
-		}
+			/**
+			 * @brief Called each tick after the graph is updated to remove unused collisions
+			*/
+			void PruneExpiredItems()
+			{
+				PruneExpiredMidPhases();
+			}
 
-		/**
-		* @brief Sort all the constraints for better solver stability
-		*/
-		void SortConstraintsHandles()
-		{
-			if(ActiveConstraints.Num())
+			/**
+			 * @brief If we add new constraints after collision detection, do what needs to be done to add them to the system
+			*/
+			void ProcessInjectedConstraints()
 			{
-				// We need to sort constraints for solver stability
-				// We have to use StableSort so that constraints of the same pair stay in the same order
-				// Otherwise the order within each pair can change due to where they start out in the array
-				// @todo(chaos): we should label each contact (and shape) for things like warm starting GJK
-				// and so we could use that label as part of the key
-				// and then we could use regular Sort (which is faster)			
-				// @todo(chaos): this can be moved to the island and therefoe done in parallel
-				ActiveConstraints.StableSort(ContactConstraintSortPredicate);
+				ProcessNewItems();
 			}
-		}
 
-		/**
-		 * @brief Destroy all collision and caches involving the particle
-		 * Called when a particle is destroyed or disabled (not sleeping).
-		*/
-		void RemoveParticle(FGeometryParticleHandle* Particle)
-		{
-			// We will be removing collisions, and don't want to have to prune the queues
-			check(!bInCollisionDetectionPhase);
+			/**
+			 * @brief Add a set of pre-built constraints and build required internal mapping data
+			 * This is used by the resim cache when restoring constraints after a desync
+			*/
+			void AddResimConstraints(const TArray<FPBDCollisionConstraint>& InConstraints, const FCollisionContext Context)
+			{
+				for (const FPBDCollisionConstraint& SourceConstraint : InConstraints)
+				{
+					// We must keep the particles in the same order that the broadphase would generate when
+					// finding or creating the midphase. This is because Collision Constraints may have the
+					// particles in the opposite order to the midphase tha owns them.
+					FGeometryParticleHandle* Particle0 = SourceConstraint.Particle[0];
+					FGeometryParticleHandle* Particle1 = SourceConstraint.Particle[1];
+					if (ShouldSwapParticleOrder(Particle0, Particle1))
+					{
+						Swap(Particle0, Particle1);
+					}
+
+					FParticlePairMidPhase* MidPhase = Context.GetAllocator()->GetMidPhase(Particle0, Particle1, SourceConstraint.Particle[0], Context);
+
+					// We may be adding multiple constraints for the same particle pair, so we need
+					// to make sure the map is up to date in the case where we just created a new MidPhase
+					ProcessNewMidPhases();
+				
+					if (MidPhase != nullptr)
+					{
+						MidPhase->InjectCollision(SourceConstraint, Context);
+					}
+				}
 
-			// Loop over all particle pairs involving this particle.
-			// Tell each Particle Pair MidPhase that one of its particles is gone. 
-			// It will get pruned at the next collision detection phase.
-			Particle->ParticleCollisions().VisitMidPhases([Particle](FParticlePairMidPhase& MidPhase)
+				ProcessNewConstraints();
+			}
+
+			/**
+			* @brief Sort all the constraints for better solver stability
+			*/
+			void SortConstraintsHandles()
+			{
+				if(ActiveConstraints.Num())
 				{
-					MidPhase.DetachParticle(Particle);
-					return ECollisionVisitorResult::Continue;
-				});
-		}
+					// We need to sort constraints for solver stability
+					// We have to use StableSort so that constraints of the same pair stay in the same order
+					// Otherwise the order within each pair can change due to where they start out in the array
+					// @todo(chaos): we should label each contact (and shape) for things like warm starting GJK
+					// and so we could use that label as part of the key
+					// and then we could use regular Sort (which is faster)			
+					// @todo(chaos): this can be moved to the island and therefoe done in parallel
+					ActiveConstraints.StableSort(ContactConstraintSortPredicate);
+				}
+			}
 
-		/**
-		 * @brief Iterate over all collisions, including sleeping ones
-		*/
-		template<typename TLambda>
-		void VisitConstCollisions(const TLambda& Visitor) const
-		{
-			for (const FParticlePairMidPhasePtr& MidPhase : ParticlePairMidPhases)
+			/**
+			 * @brief Destroy all collision and caches involving the particle
+			 * Called when a particle is destroyed or disabled (not sleeping).
+			*/
+			void RemoveParticle(FGeometryParticleHandle* Particle)
 			{
-				if (MidPhase->VisitConstCollisions(Visitor) == ECollisionVisitorResult::Stop)
+				// We will be removing collisions, and don't want to have to prune the queues
+				check(!bInCollisionDetectionPhase);
+
+				// Loop over all particle pairs involving this particle.
+				// Tell each Particle Pair MidPhase that one of its particles is gone. 
+				// It will get pruned at the next collision detection phase.
+				Particle->ParticleCollisions().VisitMidPhases([Particle](FParticlePairMidPhase& MidPhase)
+					{
+						MidPhase.DetachParticle(Particle);
+						return ECollisionVisitorResult::Continue;
+					});
+			}
+
+			/**
+			 * @brief Iterate over all collisions, including sleeping ones
+			*/
+			template<typename TLambda>
+			void VisitConstCollisions(const TLambda& Visitor) const
+			{
+				for (const FParticlePairMidPhasePtr& MidPhase : ParticlePairMidPhases)
 				{
-					return;
+					if (MidPhase->VisitConstCollisions(Visitor) == ECollisionVisitorResult::Stop)
+					{
+						return;
+					}
 				}
 			}
-		}
 		
-	private:
-		friend class FCollisionContextAllocator;
+		private:
+			friend class FCollisionContextAllocator;
 
-		// Gather all the newly created items and finalize their setup
-		void ProcessNewItems()
-		{
-			ProcessNewMidPhases();
+			// Gather all the newly created items and finalize their setup
+			void ProcessNewItems()
+			{
+				ProcessNewMidPhases();
 
-			ProcessNewConstraints();
-		}
+				ProcessNewConstraints();
+			}
 
-		// Collect all the midphases created on the context allocators (probably on multiple threads) and register them
-		void ProcessNewMidPhases()
-		{
-			for (TUniquePtr<FCollisionContextAllocator>& ContextAllocator : ContextAllocators)
+			// Collect all the midphases created on the context allocators (probably on multiple threads) and register them
+			void ProcessNewMidPhases()
 			{
-				ContextAllocator->ProcessNewMidPhases(this);
+				for (TUniquePtr<FCollisionContextAllocator>& ContextAllocator : ContextAllocators)
+				{
+					ContextAllocator->ProcessNewMidPhases(this);
+				}
 			}
-		}
 
-		void AddMidPhase(FParticlePairMidPhasePtr&& MidPhase)
-		{
-			ParticlePairMidPhases.Emplace(MoveTemp(MidPhase));
-		}
-
-		// Remove ParticlePairMidPhase from each Particles list of collisions
-		// NOTE: One or both particles may have been destroyed in which case it will
-		// have been set to null on the midphase.
-		void DetachParticlePairMidPhase(FParticlePairMidPhase* MidPhase)
-		{
-			if (FGeometryParticleHandle* Particle0 = MidPhase->GetParticle0())
+			void AddMidPhase(FParticlePairMidPhasePtr&& MidPhase)
 			{
-				Particle0->ParticleCollisions().RemoveMidPhase(Particle0, MidPhase);
+				ParticlePairMidPhases.Emplace(MoveTemp(MidPhase));
 			}
 
-			if (FGeometryParticleHandle* Particle1 = MidPhase->GetParticle1())
+			// Remove ParticlePairMidPhase from each Particles list of collisions
+			// NOTE: One or both particles may have been destroyed in which case it will
+			// have been set to null on the midphase.
+			void DetachParticlePairMidPhase(FParticlePairMidPhase* MidPhase)
 			{
-				Particle1->ParticleCollisions().RemoveMidPhase(Particle1, MidPhase);
+				if (FGeometryParticleHandle* Particle0 = MidPhase->GetParticle0())
+				{
+					Particle0->ParticleCollisions().RemoveMidPhase(Particle0, MidPhase);
+				}
+
+				if (FGeometryParticleHandle* Particle1 = MidPhase->GetParticle1())
+				{
+					Particle1->ParticleCollisions().RemoveMidPhase(Particle1, MidPhase);
+				}
 			}
-		}
 
-		// Find all midphases that have not been updated this tick (i.e., bounds no longer overlap) and destroy them
-		void PruneExpiredMidPhases();
+			// Find all midphases that have not been updated this tick (i.e., bounds no longer overlap) and destroy them
+			void PruneExpiredMidPhases();
 
-		// Collect all the constraints activated in the collision tasks and register them (calls ActivateConstraintImp on each)
-		void ProcessNewConstraints()
-		{
-			for (TUniquePtr<FCollisionContextAllocator>& ContextAllocator : ContextAllocators)
+			// Collect all the constraints activated in the collision tasks and register them (calls ActivateConstraintImp on each)
+			void ProcessNewConstraints()
 			{
-				ContextAllocator->ProcessNewConstraints(this);
+				for (TUniquePtr<FCollisionContextAllocator>& ContextAllocator : ContextAllocators)
+				{
+					ContextAllocator->ProcessNewConstraints(this);
+				}
 			}
-		}
 
-		// Register an activated constraint
-		void ActivateConstraintImp(FPBDCollisionConstraint* CollisionConstraint)
-		{
-			FPBDCollisionConstraintContainerCookie& Cookie = CollisionConstraint->GetContainerCookie();
+			// Register an activated constraint
+			void ActivateConstraintImp(FPBDCollisionConstraint* CollisionConstraint)
+			{
+				FPBDCollisionConstraintContainerCookie& Cookie = CollisionConstraint->GetContainerCookie();
 
-			// Add the constraint to the active list and update its epoch
-			checkSlow(ActiveConstraints.Find(CollisionConstraint) == INDEX_NONE);
-			Cookie.ConstraintIndex = ActiveConstraints.Add(CollisionConstraint);
+				// Add the constraint to the active list and update its epoch
+				checkSlow(ActiveConstraints.Find(CollisionConstraint) == INDEX_NONE);
+				Cookie.ConstraintIndex = ActiveConstraints.Add(CollisionConstraint);
 
-			// If the constraint uses CCD, keep it in another list so we don't have to search the full list for CCD contacts
-			if (CollisionConstraint->GetCCDEnabled())
-			{
-				checkSlow(ActiveCCDConstraints.Find(CollisionConstraint) == INDEX_NONE);
-				Cookie.CCDConstraintIndex = ActiveCCDConstraints.Add(CollisionConstraint);
-			}
+				// If the constraint uses CCD, keep it in another list so we don't have to search the full list for CCD contacts
+				if (CollisionConstraint->GetCCDEnabled())
+				{
+					checkSlow(ActiveCCDConstraints.Find(CollisionConstraint) == INDEX_NONE);
+					Cookie.CCDConstraintIndex = ActiveCCDConstraints.Add(CollisionConstraint);
+				}
 
-			Cookie.LastUsedEpoch = CurrentEpoch;
-		}
+				Cookie.LastUsedEpoch = CurrentEpoch;
+			}
 
-		// The container that owns the allocator (only needed because new constraints need to know)
-		FPBDCollisionConstraints* CollisionContainer;
+			// The container that owns the allocator (only needed because new constraints need to know)
+			FPBDCollisionConstraints* CollisionContainer;
 
-		// Storage for collisins and midphases, one for each thread on which we detect collisions
-		TArray<TUniquePtr<FCollisionContextAllocator>> ContextAllocators;
-		int32 NumCollisionsPerBlock;
+			// Storage for collisins and midphases, one for each thread on which we detect collisions
+			TArray<TUniquePtr<FCollisionContextAllocator>> ContextAllocators;
+			int32 NumCollisionsPerBlock;
 
-		// All of the overlapping particle pairs in the scene
-		TArray<FParticlePairMidPhasePtr> ParticlePairMidPhases;
+			// All of the overlapping particle pairs in the scene
+			TArray<FParticlePairMidPhasePtr> ParticlePairMidPhases;
 
-		// The active constraints (added or recovered this tick)
-		TArray<FPBDCollisionConstraint*> ActiveConstraints;
+			// The active constraints (added or recovered this tick)
+			TArray<FPBDCollisionConstraint*> ActiveConstraints;
 
-		// The active sweep constraints (added or recovered this tick)
-		TArray<FPBDCollisionConstraint*> ActiveCCDConstraints;
+			// The active sweep constraints (added or recovered this tick)
+			TArray<FPBDCollisionConstraint*> ActiveCCDConstraints;
 
-		// The current epoch used to track out-of-date contacts. A constraint whose Epoch is
-		// older than the current Epoch at the end of the tick was not refreshed this tick.
-		int32 CurrentEpoch;
+			// The current epoch used to track out-of-date contacts. A constraint whose Epoch is
+			// older than the current Epoch at the end of the tick was not refreshed this tick.
+			int32 CurrentEpoch;
 
-		// For assertions
-		bool bInCollisionDetectionPhase;
-	};
+			// For assertions
+			bool bInCollisionDetectionPhase;
+		};
 
+	}	// namespace Private
 
 	///////////////////////////////////////////////////////////////////////////////////////////////
 	///////////////////////////////////////////////////////////////////////////////////////////////
diff --git a/Engine/Source/Runtime/Experimental/Chaos/Public/Chaos/Collision/CollisionContext.h b/Engine/Source/Runtime/Experimental/Chaos/Public/Chaos/Collision/CollisionContext.h
--- a/Engine/Source/Runtime/Experimental/Chaos/Public/Chaos/Collision/CollisionContext.h
+++ b/Engine/Source/Runtime/Experimental/Chaos/Public/Chaos/Collision/CollisionContext.h
@@ -5,7 +5,10 @@
 
 namespace Chaos
 {
-	class FCollisionContextAllocator;
+	namespace Private
+	{
+		class FCollisionContextAllocator;
+	}
 	class FMultiShapePairCollisionDetector;
 
 	class FCollisionDetectorSettings
@@ -70,8 +73,8 @@ namespace Chaos
 		const FCollisionDetectorSettings& GetSettings() const { return *Settings; }
 		void SetSettings(const FCollisionDetectorSettings& InSettings) { Settings = &InSettings; }
 
-		FCollisionContextAllocator* GetAllocator() const { return Allocator; }
-		void SetAllocator(FCollisionContextAllocator* InAllocator) { Allocator = InAllocator; }
+		Private::FCollisionContextAllocator* GetAllocator() const { return Allocator; }
+		void SetAllocator(Private::FCollisionContextAllocator* InAllocator) { Allocator = InAllocator; }
 
 		FMultiShapePairCollisionDetector* GetMultiShapePairCollisionDetector() const { return MultiShapeCollisionDetector; }
 		void SetMultiShapePairCollisionDetector(FMultiShapePairCollisionDetector* InDetector) { MultiShapeCollisionDetector = InDetector; }
@@ -79,7 +82,7 @@ namespace Chaos
 
 		const FCollisionDetectorSettings* Settings;
 
-		FCollisionContextAllocator* Allocator;
+		Private::FCollisionContextAllocator* Allocator;
 
 		// This is used in the older collision detection path which is still used for particles that do not flatten their implicit hierrarchies
 		// into the Particle's ShapesArray. Currently this is only Clusters.
diff --git a/Engine/Source/Runtime/Experimental/Chaos/Public/Chaos/Collision/CollisionKeys.h b/Engine/Source/Runtime/Experimental/Chaos/Public/Chaos/Collision/CollisionKeys.h
--- a/Engine/Source/Runtime/Experimental/Chaos/Public/Chaos/Collision/CollisionKeys.h
+++ b/Engine/Source/Runtime/Experimental/Chaos/Public/Chaos/Collision/CollisionKeys.h
@@ -72,6 +72,11 @@ namespace Chaos
 			return Key.Key64;
 		}
 
+		friend bool operator<(const FCollisionParticlePairKey& L, const FCollisionParticlePairKey& R)
+		{
+			return L.Key.Key64 < R.Key.Key64;
+		}
+
 	private:
 		void GenerateKey(const FGeometryParticleHandle* Particle0, const FGeometryParticleHandle* Particle1)
 		{
diff --git a/Engine/Source/Runtime/Experimental/Chaos/Public/Chaos/Collision/PBDCollisionConstraint.h b/Engine/Source/Runtime/Experimental/Chaos/Public/Chaos/Collision/PBDCollisionConstraint.h
--- a/Engine/Source/Runtime/Experimental/Chaos/Public/Chaos/Collision/PBDCollisionConstraint.h
+++ b/Engine/Source/Runtime/Experimental/Chaos/Public/Chaos/Collision/PBDCollisionConstraint.h
@@ -17,8 +17,11 @@
 
 namespace Chaos
 {
-	class FCollisionConstraintAllocator;
-	class FCollisionContextAllocator;
+	namespace Private
+	{
+		class FCollisionConstraintAllocator;
+		class FCollisionContextAllocator;
+	}
 	class FConstGenericParticleHandle;
 	class FImplicitObject;
 	class FParticlePairMidPhase;
@@ -156,8 +159,8 @@ namespace Chaos
 	*/
 	class CHAOS_API FPBDCollisionConstraint final : public FPBDCollisionConstraintHandle
 	{
-		friend class FCollisionConstraintAllocator;
-		friend class FCollisionContextAllocator;
+		friend class Private::FCollisionConstraintAllocator;
+		friend class Private::FCollisionContextAllocator;
 		friend class FMultiShapePairCollisionDetector;
 		friend class FParticlePairMidPhase;
 		friend class FPBDCollisionConstraints;
diff --git a/Engine/Source/Runtime/Experimental/Chaos/Public/Chaos/Collision/ParticlePairBroadPhase.h b/Engine/Source/Runtime/Experimental/Chaos/Public/Chaos/Collision/ParticlePairBroadPhase.h
--- a/Engine/Source/Runtime/Experimental/Chaos/Public/Chaos/Collision/ParticlePairBroadPhase.h
+++ b/Engine/Source/Runtime/Experimental/Chaos/Public/Chaos/Collision/ParticlePairBroadPhase.h
@@ -34,7 +34,7 @@ namespace Chaos
 		 */
 		void ProduceOverlaps(
 			FReal Dt, 
-			FCollisionConstraintAllocator* Allocator,
+			Private::FCollisionConstraintAllocator* Allocator,
 			const FCollisionDetectorSettings& Settings)
 		{
 			SCOPE_CYCLE_COUNTER(STAT_Collisions_ParticlePairBroadPhase);
diff --git a/Engine/Source/Runtime/Experimental/Chaos/Public/Chaos/Collision/ParticlePairMidPhase.h b/Engine/Source/Runtime/Experimental/Chaos/Public/Chaos/Collision/ParticlePairMidPhase.h
--- a/Engine/Source/Runtime/Experimental/Chaos/Public/Chaos/Collision/ParticlePairMidPhase.h
+++ b/Engine/Source/Runtime/Experimental/Chaos/Public/Chaos/Collision/ParticlePairMidPhase.h
@@ -17,8 +17,11 @@
 
 namespace Chaos
 {
+	namespace Private
+	{
+		class FCollisionConstraintAllocator;
+	}
 	class FCollisionContext;
-	class FCollisionConstraintAllocator;
 	class FParticlePairMidPhase;
 	class FPBDCollisionConstraints;
 	class FPerShapeData;
diff --git a/Engine/Source/Runtime/Experimental/Chaos/Public/Chaos/Collision/SpatialAccelerationBroadPhase.h b/Engine/Source/Runtime/Experimental/Chaos/Public/Chaos/Collision/SpatialAccelerationBroadPhase.h
--- a/Engine/Source/Runtime/Experimental/Chaos/Public/Chaos/Collision/SpatialAccelerationBroadPhase.h
+++ b/Engine/Source/Runtime/Experimental/Chaos/Public/Chaos/Collision/SpatialAccelerationBroadPhase.h
@@ -33,8 +33,9 @@ namespace Chaos
 			{
 			}
 
-			FBroadphaseOverlap(FGeometryParticleHandle* Particle0, FGeometryParticleHandle* Particle1)
+			FBroadphaseOverlap(FGeometryParticleHandle* Particle0, FGeometryParticleHandle* Particle1, const int32 InSearchParticleIndex)
 				: Particles{ Particle0, Particle1 }
+				, SearchParticleIndex(InSearchParticleIndex)
 			{
 			}
 
@@ -44,6 +45,7 @@ namespace Chaos
 			}
 
 			FGeometryParticleHandle* Particles[2];
+			int32 SearchParticleIndex;
 		};
 
 		/**
@@ -252,11 +254,13 @@ namespace Chaos
 				
 				if (bAccept)
 				{
+					int32 SearchParticleIndex = 0;
 					if (bSwapOrder)
 					{
 						Swap(Particle1, Particle2);
+						SearchParticleIndex = 1;
 					}
-					Context.Overlaps.Emplace(Particle1, Particle2);
+					Context.Overlaps.Emplace(Particle1, Particle2, SearchParticleIndex);
 				}
 
 				return true;
@@ -330,7 +334,7 @@ namespace Chaos
 		 */
 		void ProduceOverlaps(
 			FReal Dt, 
-			FCollisionConstraintAllocator* Allocator,
+			Private::FCollisionConstraintAllocator* Allocator,
 			const FCollisionDetectorSettings& Settings,
 			IResimCacheBase* ResimCache
 			)
@@ -377,7 +381,7 @@ namespace Chaos
 			ViewType& OverlapView, 
 			FReal Dt,
 			const SpatialAccelerationType& InSpatialAcceleration, 
-			FCollisionConstraintAllocator* Allocator,
+			Private::FCollisionConstraintAllocator* Allocator,
 			const FCollisionDetectorSettings& Settings)
 		{
 			// Reset all the contexts (we don't always use all of them, but don't reduce the array size so that the element arrays don't need reallocating)
@@ -411,7 +415,7 @@ namespace Chaos
 
 				// Retrieve the midphase/collision allocator for this context
 				Allocator->SetMaxContexts(NumWorkers);
-				FCollisionContextAllocator* ContextAllocator = Allocator->GetContextAllocator(WorkerIndex);
+				Private::FCollisionContextAllocator* ContextAllocator = Allocator->GetContextAllocator(WorkerIndex);
 
 				// Build the collision context for the worker
 				BroadphaseContexts[WorkerIndex].CollisionContext.SetSettings(Settings);
@@ -455,7 +459,7 @@ namespace Chaos
 		void ProduceOverlaps(
 			FReal Dt, 
 			const T_SPATIALACCELERATION& InSpatialAcceleration, 
-			FCollisionConstraintAllocator* Allocator,
+			Private::FCollisionConstraintAllocator* Allocator,
 			const FCollisionDetectorSettings& Settings,
 			IResimCacheBase* ResimCache
 			)
@@ -556,7 +560,10 @@ namespace Chaos
 
 				for (auto& Elem : GlobalElems)
 				{
-					Context.Overlaps.Emplace(Particle1, Elem.Payload.GetGeometryParticleHandle_PhysicsThread());
+					if (Particle1 != Elem.Payload.GetGeometryParticleHandle_PhysicsThread())
+					{
+						Context.Overlaps.Emplace(Particle1, Elem.Payload.GetGeometryParticleHandle_PhysicsThread(), 1);
+					}
 				}
 			}
 		}
@@ -565,11 +572,15 @@ namespace Chaos
 		// @todo(chaos): optimize
 		void AssignMidphases(Private::FBroadphaseContext& BroadphaseContext)
 		{
-			BroadphaseContext.MidPhases.Reset(BroadphaseContext.Overlaps.Num());
-			for (Private::FBroadphaseOverlap& Overlap : BroadphaseContext.Overlaps)
+			Private::FCollisionContextAllocator* ContextAllocator = BroadphaseContext.CollisionContext.GetAllocator();
+
+			BroadphaseContext.MidPhases.SetNum(BroadphaseContext.Overlaps.Num());
+
+			for (int32 OverlapIndex = 0, OverlapEnd = BroadphaseContext.Overlaps.Num(); OverlapIndex < OverlapEnd; ++OverlapIndex)
 			{
-				FParticlePairMidPhase* MidPhase = BroadphaseContext.CollisionContext.GetAllocator()->GetMidPhase(Overlap.Particles[0], Overlap.Particles[1], Overlap.Particles[0], BroadphaseContext.CollisionContext);
-				BroadphaseContext.MidPhases.Add(MidPhase);
+				Private::FBroadphaseOverlap& Overlap = BroadphaseContext.Overlaps[OverlapIndex];
+
+				BroadphaseContext.MidPhases[OverlapIndex] = ContextAllocator->GetMidPhase(Overlap.Particles[0], Overlap.Particles[1], Overlap.Particles[Overlap.SearchParticleIndex], BroadphaseContext.CollisionContext);
 			}
 		}
 
diff --git a/Engine/Source/Runtime/Experimental/Chaos/Public/Chaos/ISpatialAccelerationCollection.h b/Engine/Source/Runtime/Experimental/Chaos/Public/Chaos/ISpatialAccelerationCollection.h
--- a/Engine/Source/Runtime/Experimental/Chaos/Public/Chaos/ISpatialAccelerationCollection.h
+++ b/Engine/Source/Runtime/Experimental/Chaos/Public/Chaos/ISpatialAccelerationCollection.h
@@ -9,8 +9,11 @@
 
 namespace Chaos
 {
+	namespace Private
+	{
+		class FCollisionConstraintAllocator;
+	}
 	class FAsyncCollisionReceiver;
-	class FCollisionConstraintAllocator;
 	class FCollisionDetectorSettings;
 	class FSpatialAccelerationBroadPhase;
 	class IResimCacheBase;
@@ -31,7 +34,7 @@ public:
 	virtual void SwapSubstructure(ISpatialAccelerationCollection<TPayloadType, T, d>& Other, FSpatialAccelerationIdx Idx) = 0;	
 
 	/** This is kind of a hack to avoid virtuals. We simply route calls into templated functions */
-	virtual void PBDComputeConstraintsLowLevel(T Dt, FSpatialAccelerationBroadPhase& BroadPhase, FCollisionConstraintAllocator* Allocator, const FCollisionDetectorSettings& Settings, IResimCacheBase* ResimCache) const = 0;
+	virtual void PBDComputeConstraintsLowLevel(T Dt, FSpatialAccelerationBroadPhase& BroadPhase, Private::FCollisionConstraintAllocator* Allocator, const FCollisionDetectorSettings& Settings, IResimCacheBase* ResimCache) const = 0;
 	virtual TArray<FSpatialAccelerationIdx> GetAllSpatialIndices() const = 0;
 
 	bool IsBucketActive(uint8 BucketIdx) const
diff --git a/Engine/Source/Runtime/Experimental/Chaos/Public/Chaos/PBDCollisionConstraints.h b/Engine/Source/Runtime/Experimental/Chaos/Public/Chaos/PBDCollisionConstraints.h
--- a/Engine/Source/Runtime/Experimental/Chaos/Public/Chaos/PBDCollisionConstraints.h
+++ b/Engine/Source/Runtime/Experimental/Chaos/Public/Chaos/PBDCollisionConstraints.h
@@ -247,7 +247,7 @@ public:
 
 	const FPBDCollisionConstraint& GetConstraint(int32 Index) const;
 
-	FCollisionConstraintAllocator& GetConstraintAllocator() { return ConstraintAllocator; }
+	Private::FCollisionConstraintAllocator& GetConstraintAllocator() { return ConstraintAllocator; }
 
 	void UpdateConstraintMaterialProperties(FPBDCollisionConstraint& Contact);
 
@@ -275,7 +275,7 @@ private:
 
 	const FPBDRigidsSOAs& Particles;
 
-	FCollisionConstraintAllocator ConstraintAllocator;
+	Private::FCollisionConstraintAllocator ConstraintAllocator;
 	int32 NumActivePointConstraints;
 	TArray<FPBDCollisionConstraintHandle*> TempCollisions;	// Reused from tick to tick to build contact lists
 
diff --git a/Engine/Source/Runtime/Experimental/Chaos/Public/Chaos/PBDRigidsEvolutionGBF.h b/Engine/Source/Runtime/Experimental/Chaos/Public/Chaos/PBDRigidsEvolutionGBF.h
--- a/Engine/Source/Runtime/Experimental/Chaos/Public/Chaos/PBDRigidsEvolutionGBF.h
+++ b/Engine/Source/Runtime/Experimental/Chaos/Public/Chaos/PBDRigidsEvolutionGBF.h
@@ -22,7 +22,6 @@
 
 namespace Chaos
 {
-	class FCollisionConstraintAllocator;
 	class FChaosArchive;
 	class IResimCacheBase;
 	class FEvolutionResimCache;
diff --git a/Engine/Source/Runtime/Experimental/Chaos/Public/Chaos/SpatialAccelerationCollection.h b/Engine/Source/Runtime/Experimental/Chaos/Public/Chaos/SpatialAccelerationCollection.h
--- a/Engine/Source/Runtime/Experimental/Chaos/Public/Chaos/SpatialAccelerationCollection.h
+++ b/Engine/Source/Runtime/Experimental/Chaos/Public/Chaos/SpatialAccelerationCollection.h
@@ -354,13 +354,13 @@ struct TSpatialAccelerationCollectionHelper
 };
 
 template <typename SpatialAccelerationCollection>
-typename TEnableIf<TIsSame<typename SpatialAccelerationCollection::TPayloadType, FAccelerationStructureHandle>::Value, void>::Type PBDComputeConstraintsLowLevel_Helper(FReal Dt, const SpatialAccelerationCollection& Accel, FSpatialAccelerationBroadPhase& BroadPhase, FCollisionConstraintAllocator* Allocator, const FCollisionDetectorSettings& Settings, IResimCacheBase* ResimCache)
+typename TEnableIf<TIsSame<typename SpatialAccelerationCollection::TPayloadType, FAccelerationStructureHandle>::Value, void>::Type PBDComputeConstraintsLowLevel_Helper(FReal Dt, const SpatialAccelerationCollection& Accel, FSpatialAccelerationBroadPhase& BroadPhase, Private::FCollisionConstraintAllocator* Allocator, const FCollisionDetectorSettings& Settings, IResimCacheBase* ResimCache)
 {
 	BroadPhase.ProduceOverlaps(Dt, Accel, Allocator, Settings, ResimCache);
 }
 
 template <typename SpatialAccelerationCollection>
-typename TEnableIf<!TIsSame<typename SpatialAccelerationCollection::TPayloadType, FAccelerationStructureHandle>::Value, void>::Type PBDComputeConstraintsLowLevel_Helper(FReal Dt, const SpatialAccelerationCollection& Accel, FSpatialAccelerationBroadPhase& BroadPhase, FCollisionConstraintAllocator* Allocator, const FCollisionDetectorSettings& Settings, IResimCacheBase* ResimCache)
+typename TEnableIf<!TIsSame<typename SpatialAccelerationCollection::TPayloadType, FAccelerationStructureHandle>::Value, void>::Type PBDComputeConstraintsLowLevel_Helper(FReal Dt, const SpatialAccelerationCollection& Accel, FSpatialAccelerationBroadPhase& BroadPhase, Private::FCollisionConstraintAllocator* Allocator, const FCollisionDetectorSettings& Settings, IResimCacheBase* ResimCache)
 {
 }
 
@@ -562,7 +562,7 @@ public:
 		return TUniquePtr<ISpatialAcceleration<TPayloadType, T, d>>(new TSpatialAccelerationCollection<TSpatialAccelerationTypes...>(*this));
 	}
 
-	virtual void PBDComputeConstraintsLowLevel(T Dt, FSpatialAccelerationBroadPhase& BroadPhase, FCollisionConstraintAllocator* Allocator, const FCollisionDetectorSettings& Settings, IResimCacheBase* ResimCache) const override
+	virtual void PBDComputeConstraintsLowLevel(T Dt, FSpatialAccelerationBroadPhase& BroadPhase, Private::FCollisionConstraintAllocator* Allocator, const FCollisionDetectorSettings& Settings, IResimCacheBase* ResimCache) const override
 	{
 		PBDComputeConstraintsLowLevel_Helper(Dt, *this, BroadPhase, Allocator, Settings, ResimCache);
 	}
