Chaos - split broadphase and narrow phase so we can add a broadphase modifier

diff --git a/Engine/Source/Runtime/Experimental/Chaos/Private/Chaos/PBDRigidsEvolutionGBF.cpp b/Engine/Source/Runtime/Experimental/Chaos/Private/Chaos/PBDRigidsEvolutionGBF.cpp
--- a/Engine/Source/Runtime/Experimental/Chaos/Private/Chaos/PBDRigidsEvolutionGBF.cpp
+++ b/Engine/Source/Runtime/Experimental/Chaos/Private/Chaos/PBDRigidsEvolutionGBF.cpp
@@ -365,7 +365,9 @@ void FPBDRigidsEvolutionGBF::AdvanceOneTimeStepImpl(const FReal Dt, const FSubSt
 		CSV_SCOPED_TIMING_STAT(PhysicsVerbose, StepSolver_DetectCollisions);
 		CollisionDetector.GetBroadPhase().SetSpatialAcceleration(InternalAcceleration);
 
-		CollisionDetector.DetectCollisions(Dt, GetCurrentStepResimCache());
+		CollisionDetector.RunBroadPhase(Dt, GetCurrentStepResimCache());
+
+		CollisionDetector.RunNarrowPhase(Dt, GetCurrentStepResimCache());
 	}
 
 	if (PostDetectCollisionsCallback != nullptr)
diff --git a/Engine/Source/Runtime/Experimental/Chaos/Public/Chaos/Collision/CollisionConstraintAllocator.h b/Engine/Source/Runtime/Experimental/Chaos/Public/Chaos/Collision/CollisionConstraintAllocator.h
--- a/Engine/Source/Runtime/Experimental/Chaos/Public/Chaos/Collision/CollisionConstraintAllocator.h
+++ b/Engine/Source/Runtime/Experimental/Chaos/Public/Chaos/Collision/CollisionConstraintAllocator.h
@@ -393,6 +393,17 @@ namespace Chaos
 				PruneExpiredMidPhases();
 			}
 
+			/**
+			 * Collect all the midphases created on the context allocators (probably on multiple threads) and register them
+			*/
+			void ProcessNewMidPhases()
+			{
+				for (TUniquePtr<FCollisionContextAllocator>& ContextAllocator : ContextAllocators)
+				{
+					ContextAllocator->ProcessNewMidPhases(this);
+				}
+			}
+
 			/**
 			 * @brief If we add new constraints after collision detection, do what needs to be done to add them to the system
 			*/
@@ -473,6 +484,7 @@ namespace Chaos
 
 			/**
 			 * @brief Iterate over all collisions, including sleeping ones
+			 * Visitor signature: ECollisionVisitorResult(const FPBDCollisionConstraint&)
 			*/
 			template<typename TLambda>
 			void VisitConstCollisions(const TLambda& Visitor) const
@@ -485,7 +497,23 @@ namespace Chaos
 					}
 				}
 			}
-		
+
+			/**
+			 * @brief Iterate over all midphases
+			 * Visitor signature: ECollisionVisitorResult(FParticlePairMidPhase&)
+			*/
+			template<typename TLambda>
+			void VisitMidPhases(const TLambda& Visitor)
+			{
+				for (FParticlePairMidPhasePtr& MidPhase : ParticlePairMidPhases)
+				{
+					if (Visitor(*MidPhase) == ECollisionVisitorResult::Stop)
+					{
+						return;
+					}
+				}
+			}
+
 		private:
 			friend class FCollisionContextAllocator;
 
@@ -497,15 +525,6 @@ namespace Chaos
 				ProcessNewConstraints();
 			}
 
-			// Collect all the midphases created on the context allocators (probably on multiple threads) and register them
-			void ProcessNewMidPhases()
-			{
-				for (TUniquePtr<FCollisionContextAllocator>& ContextAllocator : ContextAllocators)
-				{
-					ContextAllocator->ProcessNewMidPhases(this);
-				}
-			}
-
 			void AddMidPhase(FParticlePairMidPhasePtr&& MidPhase)
 			{
 				ParticlePairMidPhases.Emplace(MoveTemp(MidPhase));
diff --git a/Engine/Source/Runtime/Experimental/Chaos/Public/Chaos/Collision/SpatialAccelerationBroadPhase.h b/Engine/Source/Runtime/Experimental/Chaos/Public/Chaos/Collision/SpatialAccelerationBroadPhase.h
--- a/Engine/Source/Runtime/Experimental/Chaos/Public/Chaos/Collision/SpatialAccelerationBroadPhase.h
+++ b/Engine/Source/Runtime/Experimental/Chaos/Public/Chaos/Collision/SpatialAccelerationBroadPhase.h
@@ -26,14 +26,14 @@ namespace Chaos
 		/**
 		 * Output element from the broadphase
 		*/
-		class FBroadphaseOverlap
+		class FBroadPhaseOverlap
 		{
 		public:
-			FBroadphaseOverlap()
+			FBroadPhaseOverlap()
 			{
 			}
 
-			FBroadphaseOverlap(FGeometryParticleHandle* Particle0, FGeometryParticleHandle* Particle1, const int32 InSearchParticleIndex)
+			FBroadPhaseOverlap(FGeometryParticleHandle* Particle0, FGeometryParticleHandle* Particle1, const int32 InSearchParticleIndex)
 				: Particles{ Particle0, Particle1 }
 				, SearchParticleIndex(InSearchParticleIndex)
 			{
@@ -51,10 +51,10 @@ namespace Chaos
 		/**
 		 * Per-thread context for the broadphase.
 		*/
-		class FBroadphaseContext
+		class FBroadPhaseContext
 		{
 		public:
-			FBroadphaseContext()
+			FBroadPhaseContext()
 			{
 			}
 
@@ -65,7 +65,7 @@ namespace Chaos
 				CollisionContext.Reset();
 			}
 
-			TArray<FBroadphaseOverlap> Overlaps;
+			TArray<FBroadPhaseOverlap> Overlaps;
 			TArray<FParticlePairMidPhase*> MidPhases;
 			FCollisionContext CollisionContext;
 		};
@@ -234,7 +234,7 @@ namespace Chaos
 		 */
 		struct FSimOverlapVisitor
 		{
-			FSimOverlapVisitor(FGeometryParticleHandle* ParticleHandle, const FCollisionFilterData& InSimFilterData, const FIgnoreCollisionManager& InIgnoreCollisionManager, const bool bInIsResimming, Private::FBroadphaseContext& InContext)
+			FSimOverlapVisitor(FGeometryParticleHandle* ParticleHandle, const FCollisionFilterData& InSimFilterData, const FIgnoreCollisionManager& InIgnoreCollisionManager, const bool bInIsResimming, Private::FBroadPhaseContext& InContext)
 				: Context(InContext)
 				, IgnoreCollisionManager(InIgnoreCollisionManager)
 				, SimFilterData(InSimFilterData)
@@ -298,7 +298,7 @@ namespace Chaos
 			}
 
 		private:
-			Private::FBroadphaseContext& Context;
+			Private::FBroadPhaseContext& Context;
 			const FIgnoreCollisionManager& IgnoreCollisionManager;
 			FCollisionFilterData SimFilterData;
 			FUniqueIdx ParticleUniqueIdx; // unique id of the particle visiting, used to skip self intersection as early as possible
@@ -330,7 +330,7 @@ namespace Chaos
 		}
 
 		/**
-		 * Generate all overlapping pairs and pass them to the narrow phase.
+		 * Generate all overlapping pairs and spawn a misphase object to handle collisions for each of them
 		 */
 		void ProduceOverlaps(
 			FReal Dt, 
@@ -367,7 +367,35 @@ namespace Chaos
 			{
 				check(false);  //question: do we want to support a dynamic dispatch version?
 			}
+
+			{
+				CSV_SCOPED_TIMING_STAT(PhysicsVerbose, DetectCollisions_AssignMidPhases);
+
+				// Find or assign a midphase to each overlapping particle pair
+				const auto& AssignMidphasesWorker = [this, Dt](const int32 ContextIndex)
+				{
+					AssignMidphases(BroadphaseContexts[ContextIndex]);
+				};
+				PhysicsParallelFor(NumActiveBroadphaseContexts, AssignMidphasesWorker, bDisableCollisionParallelFor);
+
+				// Merge all the midphases from each worker into the primary allocator
+				Allocator->ProcessNewMidPhases();
+			}
+		}
+
+		/**
+		 * Generate all the collision constraints for the set of overlapping objects produced byt he broad phase
+		*/
+		void ProduceCollisions(FReal Dt)
+		{
+			CSV_SCOPED_TIMING_STAT(PhysicsVerbose, DetectCollisions_MidPhase);
+			const auto& ProcessMidphasesWorker = [this, Dt](const int32 ContextIndex)
+			{
+				ProcessMidphases(Dt, BroadphaseContexts[ContextIndex]);
+			};
+			PhysicsParallelFor(NumActiveBroadphaseContexts, ProcessMidphasesWorker, bDisableCollisionParallelFor);
 		}
+	
 		
 		/** @brief This function is the outer loop of collision detection. It loops over the
 		 * particles view and do the broadphase + narrowphase collision detection
@@ -385,11 +413,11 @@ namespace Chaos
 			const FCollisionDetectorSettings& Settings)
 		{
 			// Reset all the contexts (we don't always use all of them, but don't reduce the array size so that the element arrays don't need reallocating)
-			for (Private::FBroadphaseContext& BroadphaseContext : BroadphaseContexts)
+			for (Private::FBroadPhaseContext& BroadphaseContext : BroadphaseContexts)
 			{
 				BroadphaseContext.Reset();
 			}
-			int32 NumActiveContexts = 0;
+			NumActiveBroadphaseContexts = 0;
 
 			// A set of contexts, one for each worker thread, containing the allocation buffers etc.
 			// NOTE: Ideally contexts arrays would not resize after the initial creation, but actually it might because
@@ -399,11 +427,11 @@ namespace Chaos
 			// Fortunately, we only get calls to ContextCreator for a sub-view (with a possibly different NumWorkers)
 			// after the previous sub-view's parallel-for has completed, so we don't need to worry about context addresses 
 			// changing (we could use TChunkedArray, or allocate contexts on the heap if we ever need to handle that).
-			const auto& ContextCreator = [this, Allocator, &Settings, &OverlapView, &NumActiveContexts](const int32 WorkerIndex, const int32 NumWorkers) -> int32
+			const auto& ContextCreator = [this, Allocator, &Settings, &OverlapView](const int32 WorkerIndex, const int32 NumWorkers) -> int32
 			{
 				// Make sure we have enough contexts 
 				// (NOTE: ContextCreator gets called for every worker, but in serial and always with the same NumWorkers)
-				NumActiveContexts = NumWorkers;
+				NumActiveBroadphaseContexts = NumWorkers;
 				if (NumWorkers > BroadphaseContexts.Num())
 				{
 					check(WorkerIndex == 0);	// Should be the same for all subsequent workers
@@ -433,26 +461,6 @@ namespace Chaos
 				};
 				OverlapView.ParallelFor(ContextCreator, ProduceOverlapsWorker, bDisableCollisionParallelFor);
 			}
-
-			// Find or assign a midphase to each overlapping particle pair
-			{
-				CSV_SCOPED_TIMING_STAT(PhysicsVerbose, DetectCollisions_AssignMidPhases);
-				const auto& AssignMidphasesWorker = [this, Dt](const int32 ContextIndex)
-				{
-					AssignMidphases(BroadphaseContexts[ContextIndex]);
-				};
-				PhysicsParallelFor(NumActiveContexts, AssignMidphasesWorker, bDisableCollisionParallelFor);
-			}
-
-			// Use the midphase to create contacts between the shapes on the two partiucles
-			{
-				CSV_SCOPED_TIMING_STAT(PhysicsVerbose, DetectCollisions_MidPhase);
-				const auto& ProcessMidphasesWorker = [this, Dt](const int32 ContextIndex)
-				{
-					ProcessMidphases(Dt, BroadphaseContexts[ContextIndex]);
-				};
-				PhysicsParallelFor(NumActiveContexts, ProcessMidphasesWorker, bDisableCollisionParallelFor);
-			}
 		}
 
 		template<typename T_SPATIALACCELERATION>
@@ -506,7 +514,7 @@ namespace Chaos
 		    FReal Dt,
 		    FGeometryParticleHandle* Particle1,
 		    const T_SPATIALACCELERATION& InSpatialAcceleration,
-			Private::FBroadphaseContext& Context)
+			Private::FBroadPhaseContext& Context)
 		{
 			// @todo(chaos):We shouldn't need this data here (see uses below)
 			bool bIsKinematic1 = true;
@@ -570,7 +578,7 @@ namespace Chaos
 
 		// Find or assign midphases to each of the overlapping particle pairs
 		// @todo(chaos): optimize
-		void AssignMidphases(Private::FBroadphaseContext& BroadphaseContext)
+		void AssignMidphases(Private::FBroadPhaseContext& BroadphaseContext)
 		{
 			Private::FCollisionContextAllocator* ContextAllocator = BroadphaseContext.CollisionContext.GetAllocator();
 
@@ -578,14 +586,14 @@ namespace Chaos
 
 			for (int32 OverlapIndex = 0, OverlapEnd = BroadphaseContext.Overlaps.Num(); OverlapIndex < OverlapEnd; ++OverlapIndex)
 			{
-				Private::FBroadphaseOverlap& Overlap = BroadphaseContext.Overlaps[OverlapIndex];
+				Private::FBroadPhaseOverlap& Overlap = BroadphaseContext.Overlaps[OverlapIndex];
 
 				BroadphaseContext.MidPhases[OverlapIndex] = ContextAllocator->GetMidPhase(Overlap.Particles[0], Overlap.Particles[1], Overlap.Particles[Overlap.SearchParticleIndex], BroadphaseContext.CollisionContext);
 			}
 		}
 
 		// Process all the midphases: generate constraints and execute the narrowphase
-		void ProcessMidphases(const FReal Dt, const Private::FBroadphaseContext& BroadphaseContext)
+		void ProcessMidphases(const FReal Dt, const Private::FBroadPhaseContext& BroadphaseContext)
 		{
 			// Prefetch initial set of MidPhases
 			const int32 PrefetchLookahead = 4;
@@ -612,7 +620,8 @@ namespace Chaos
 
 		const FPBDRigidsSOAs& Particles;
 		const FAccelerationStructure* SpatialAcceleration;
-		TArray<Private::FBroadphaseContext> BroadphaseContexts;
+		TArray<Private::FBroadPhaseContext> BroadphaseContexts;
+		int32 NumActiveBroadphaseContexts;
 		FIgnoreCollisionManager IgnoreCollisionManager;
 	};
 
diff --git a/Engine/Source/Runtime/Experimental/Chaos/Public/Chaos/Collision/SpatialAccelerationCollisionDetector.h b/Engine/Source/Runtime/Experimental/Chaos/Public/Chaos/Collision/SpatialAccelerationCollisionDetector.h
--- a/Engine/Source/Runtime/Experimental/Chaos/Public/Chaos/Collision/SpatialAccelerationCollisionDetector.h
+++ b/Engine/Source/Runtime/Experimental/Chaos/Public/Chaos/Collision/SpatialAccelerationCollisionDetector.h
@@ -24,12 +24,8 @@ namespace Chaos
 			return BroadPhase;
 		}
 
-		virtual void DetectCollisions(const FReal Dt, FEvolutionResimCache* ResimCache) override
+		void RunBroadPhase(const FReal Dt, FEvolutionResimCache* ResimCache)
 		{
-			SCOPE_CYCLE_COUNTER(STAT_Collisions_Detect);
-			CHAOS_SCOPED_TIMER(DetectCollisions);
-			CSV_SCOPED_TIMING_STAT(Chaos, DetectCollisions);
-
 			if (!GetCollisionContainer().GetCollisionsEnabled())
 			{
 				return;
@@ -37,13 +33,24 @@ namespace Chaos
 
 			GetCollisionContainer().BeginDetectCollisions();
 
-			// Collision detection pipeline: BroadPhase -[parallel]-> NarrowPhase -[parallel]-> CollisionAllocator -[serial]-> Container
+			// Run the broadphase and generate a midphase object for every overlapping particle pair
 			BroadPhase.ProduceOverlaps(Dt, &GetCollisionContainer().GetConstraintAllocator(), GetCollisionContainer().GetDetectorSettings(), ResimCache);
+		}
+
+		void RunNarrowPhase(const FReal Dt, FEvolutionResimCache* ResimCache)
+		{
+			if (!GetCollisionContainer().GetCollisionsEnabled())
+			{
+				return;
+			}
+
+			// Run collision detection on the output of the broadphase
+			BroadPhase.ProduceCollisions(Dt);
 
 			GetCollisionContainer().EndDetectCollisions();
 
 			// If we have a resim cache restore and save contacts
-			if(ResimCache)
+			if (ResimCache)
 			{
 				// Ensure we have at least one allocator
 				GetCollisionContainer().GetConstraintAllocator().SetMaxContexts(1);
@@ -58,6 +65,16 @@ namespace Chaos
 			}
 		}
 
+		virtual void DetectCollisions(const FReal Dt, FEvolutionResimCache* ResimCache) override
+		{
+			SCOPE_CYCLE_COUNTER(STAT_Collisions_Detect);
+			CHAOS_SCOPED_TIMER(DetectCollisions);
+			CSV_SCOPED_TIMING_STAT(Chaos, DetectCollisions);
+
+			RunBroadPhase(Dt, ResimCache);
+			RunNarrowPhase(Dt, ResimCache);
+		}
+
 	private:
 		FSpatialAccelerationBroadPhase& BroadPhase;
 	};
