Heightfield sweep fixes:
- Consider bound rotation
- Don't skip cells when bounds are large

diff --git a/Engine/Source/Runtime/Experimental/Chaos/Private/Chaos/HeightField.cpp b/Engine/Source/Runtime/Experimental/Chaos/Private/Chaos/HeightField.cpp
--- a/Engine/Source/Runtime/Experimental/Chaos/Private/Chaos/HeightField.cpp
+++ b/Engine/Source/Runtime/Experimental/Chaos/Private/Chaos/HeightField.cpp
@@ -297,7 +297,10 @@ namespace Chaos
 		{
 			const FAABB3 QueryBounds = InQueryGeom.BoundingBox();
 			StartPoint = StartTM.TransformPositionNoScale(QueryBounds.Center());
-			Inflation3D = QueryBounds.Extents() * 0.5 + FVec3(Thickness);
+			//Need inflation to consider rotation, there's probably a cheaper way to do this
+			FTransform RotationTM(StartTM.GetRotation(), FVector(0));
+			const FAABB3 RotatedBounds = QueryBounds.TransformedAABB(RotationTM);
+			Inflation3D = RotatedBounds.Extents() * 0.5 + FVec3(Thickness);
 			for (int Axis = 0; Axis < 3; ++Axis)
 			{
 				bParallel[Axis] = FMath::IsNearlyZero(Dir[Axis], (FReal)1.e-8);
@@ -1406,16 +1409,38 @@ namespace Chaos
 		FVec2 ClippedStart;
 		FVec2 ClippedEnd;
 
-		if(InflatedBounds.ClipLine(StartPoint, StartPoint + Dir * Length, ClippedStart, ClippedEnd))
+		if (InflatedBounds.ClipLine(StartPoint, StartPoint + Dir * Length, ClippedStart, ClippedEnd))
 		{
 			// Rasterize the line over the grid
 			TVec2<int32> StartCell = FlatGrid.Cell(ClippedStart / Scale2D);
 			TVec2<int32> EndCell = FlatGrid.Cell(ClippedEnd / Scale2D);
-
 			const int32 DeltaX = FMath::Abs(EndCell[0] - StartCell[0]);
 			const int32 DeltaY = -FMath::Abs(EndCell[1] - StartCell[1]);
 			const bool bSameCell = DeltaX == 0 && DeltaY == 0;
+			FReal CurrentLength = Length;
+			if (bSameCell)
+			{
+				//start and end in same cell, so test it and any cells that the InHalfExtents overlap with
+				const FVec2 MinPoint = Start2D - Inflation2D;
+				const FVec2 MaxPoint = Start2D + Inflation2D;
+				TVec2<int32> MinCell = FlatGrid.Cell(MinPoint / Scale2D);
+				TVec2<int32> MaxCell = FlatGrid.Cell(MaxPoint / Scale2D);
+				for (int32 Y = MinCell[1]; Y <= MaxCell[1]; ++Y)
+				{
+					for (int32 X = MinCell[0]; X <= MaxCell[0]; ++X)
+					{
+						bool bContinue = Visitor.VisitSweep(Y * (GeomData.NumCols - 1) + X, CurrentLength);
 
+						if (!bContinue)
+						{
+							return true;
+						}
+					}
+				}
+				return false;
+			}
+			
+			
 			const int32 DirX = StartCell[0] < EndCell[0] ? 1 : -1;
 			const int32 DirY = StartCell[1] < EndCell[1] ? 1 : -1;
 			int32 Error = DeltaX + DeltaY;
@@ -1430,70 +1455,32 @@ namespace Chaos
 			// Tracking data for cells to query (similar to bounding volume approach)
 			F2DGridSet Seen(FlatGrid.Counts());
 			TArray<FQueueEntry> Queue;
-			Queue.Add({StartCell, -1});
+			Queue.Add({ StartCell, -1 });
 			Seen.Add(StartCell);
 
 			// Data for fast box cast
 			FVec3 Min, Max, HitPoint;
-			FReal ToI;
 			bool bParallel[3];
 			FVec3 InvDir;
 
-			FReal CurrentLength = Length;
 			FReal InvCurrentLength = 1 / CurrentLength;
-
-			for(int Axis = 0; Axis < 3; ++Axis)
+			for (int Axis = 0; Axis < 3; ++Axis)
 			{
 				bParallel[Axis] = FMath::IsNearlyZero(Dir[Axis], (FReal)1.e-8);
 				InvDir[Axis] = bParallel[Axis] ? 0 : 1 / Dir[Axis];
 			}
 
 			int32 QueueIndex = 0;
-			while(QueueIndex < Queue.Num())
+			while (QueueIndex < Queue.Num())
 			{
 				// Copy so we don't lost the entry through reallocs
 				FQueueEntry CellCoord = Queue[QueueIndex++];
 
-				if(CellCoord.ToI > CurrentLength)
+				if (CellCoord.ToI > CurrentLength)
 				{
 					continue;
 				}
 
-				if(bSameCell)
-				{
-					// Test the current cell
-					bool bContinue = Visitor.VisitSweep(CellCoord.Index[1] * (GeomData.NumCols - 1) + CellCoord.Index[0], CurrentLength);
-					
-					if(!bContinue)
-					{
-						return true;
-					}
-
-					// Flatten out a double loop and skip the centre cell
-					// to search cells immediately adjacent to the current cell
-					static const TVec2<int32> Neighbors[] =
-					{
-						{-1, -1}, {0, -1}, {1, -1},
-						{-1, 0}, {1, 0},
-						{-1, 1}, {0, 1}, {1, 1}
-					};
-
-					for(const TVec2<int32>& Neighbor : Neighbors)
-					{
-						TVec2<int32> NeighCoord = CellCoord.Index + Neighbor;
-
-						FBounds2D CellBounds;
-						if(GetCellBounds3DScaled(NeighCoord, Min, Max, HalfExtents3D) && !Seen.Contains(NeighCoord))
-						{
-							if(FAABB3(Min,Max).RaycastFast(StartPoint, Dir, InvDir, bParallel, CurrentLength, InvCurrentLength, ToI, HitPoint))
-							{
-								Seen.Add(NeighCoord);
-								Queue.Add({NeighCoord, ToI});
-							}
-						}
-					}
-				}
-				else
 				{
 					// Expand each cell along the thicken direction
 					// Although the line should minimally thicken around the perpendicular to the line direction
@@ -1503,23 +1490,21 @@ namespace Chaos
 					auto Expand = [&](const TVec2<int32>& Begin, const TVec2<int32>& Direction, const int32 NumSteps)
 					{
 						TVec2<int32> CurrentCell = Begin;
-
-						for(int32 CurrStep = 0; CurrStep < NumSteps; ++CurrStep)
+						for (int32 CurrStep = 0; CurrStep < NumSteps; ++CurrStep)
 						{
 							CurrentCell += Direction;
-
 							// Fail if we leave the grid
-							if(CurrentCell[0] < 0 || CurrentCell[1] < 0 || CurrentCell[0] > FlatGrid.Counts()[0] - 1 || CurrentCell[1] > FlatGrid.Counts()[1] - 1)
+							if (CurrentCell[0] < 0 || CurrentCell[1] < 0 || CurrentCell[0] > FlatGrid.Counts()[0] - 1 || CurrentCell[1] > FlatGrid.Counts()[1] - 1)
 							{
 								break;
 							}
 
 							// No intersections here. We set the ToI to zero to cause an intersection check to happen
 							// without any expansion when we reach this cell in the queue.
-							if(!Seen.Contains(CurrentCell))
+							if (!Seen.Contains(CurrentCell))
 							{
 								Seen.Add(CurrentCell);
-								Queue.Add({CurrentCell, 0});
+								Queue.Add({ CurrentCell, 0 });
 							}
 						}
 					};
@@ -1536,11 +1521,11 @@ namespace Chaos
 					}
 
 					// This time isn't used to reject things for this method but to flag cells that should be expanded
-					if(CellCoord.ToI < 0)
+					if (CellCoord.ToI < 0)
 					{
 						// Perform expansion for thickness
 						int32 ExpandAxis;
-						if(ThickenDir[0] == 0)
+						if (ThickenDir[0] == 0)
 						{
 							ExpandAxis = 1;
 						}
@@ -1555,26 +1540,25 @@ namespace Chaos
 						Expand(Coord, -ThickenDir, Steps);
 
 						// Walk the line and add to the queue
-						if(StartCell != EndCell)
+						if (StartCell != EndCell)
 						{
 							const int32 DoubleError = Error * 2;
-
-							if(DoubleError >= DeltaY)
+							if (DoubleError >= DeltaY)
 							{
 								Error += DeltaY;
 								StartCell[0] += DirX;
 							}
 
-							if(DoubleError <= DeltaX)
+							if (DoubleError <= DeltaX)
 							{
 								Error += DeltaX;
 								StartCell[1] += DirY;
 							}
 
-							if(!Seen.Contains(StartCell))
+							if (!Seen.Contains(StartCell))
 							{
 								Seen.Add(StartCell);
-								Queue.Add({StartCell, -1});
+								Queue.Add({ StartCell, -1 });
 							}
 						}
 					}
@@ -2151,9 +2135,14 @@ namespace Chaos
 		bool bHit = false;
 		THeightfieldSweepVisitor<QueryGeomType> SQVisitor(&GeomData, QueryGeom, StartTM, Dir, Thickness, bComputeMTD);
 		const FAABB3 QueryBounds = QueryGeom.BoundingBox();
-		const FVec3 StartPoint = StartTM.TransformPositionNoScale(QueryBounds.Center());
+		//Need inflation to consider rotation, there's probably a cheaper way to do this
+		const FVec3 RotatedExtents = StartTM.TransformVectorNoScale(QueryBounds.Extents());
+		FTransform RotationTM(StartTM.GetRotation(), FVector(0));
+		const FAABB3 RotatedBounds = QueryBounds.TransformedAABB(RotationTM);
+		
 
-		const FVec3 Inflation3D = QueryBounds.Extents() * 0.5 + FVec3(Thickness);
+		const FVec3 StartPoint = StartTM.TransformPositionNoScale(QueryBounds.Center());
+		const FVec3 Inflation3D = RotatedBounds.Extents() * 0.5 + FVec3(Thickness);
 		GridSweep(StartPoint, Dir, Length, FVec3(Inflation3D[0], Inflation3D[1], Inflation3D[2]), SQVisitor);
 
 		if(SQVisitor.OutTime <= Length)
