Chaos - Add correction terms to solver body

diff --git a/Engine/Source/Programs/HeadlessChaos/Private/HeadlessChaosTestSolverBody.cpp b/Engine/Source/Programs/HeadlessChaos/Private/HeadlessChaosTestSolverBody.cpp
--- a/Engine/Source/Programs/HeadlessChaos/Private/HeadlessChaosTestSolverBody.cpp
+++ b/Engine/Source/Programs/HeadlessChaos/Private/HeadlessChaosTestSolverBody.cpp
@@ -40,6 +40,66 @@ namespace ChaosTest
 		EXPECT_NEAR(SolverBody0.V().Z, V0.Z, KINDA_SMALL_NUMBER);
 	}
 
+	// Test that correction does not affect the implicit velocity
+	TEST(FSolverBodyTests, TestCorrection)
+	{
+		const FReal Dt = 0.1f;
+		const FVec3 X0 = FVec3(10, 0, 0);
+		const FVec3 V0 = FVec3(10, 20, 30);
+		const FVec3 W0 = FVec3(10, 0, -10);
+		const FRotation3 R0 = FRotation3::FromVector(FVec3(1.0f, -2.0f, 3.0f));
+		const FReal M0 = 100.0f;
+		const FReal I0 = M0 * 1000.0f;
+
+		const FSolverVec3 PositionDelta = FSolverVec3(10, 0, 0);
+		const FSolverVec3 PositionCorrectionDelta = FSolverVec3(5, 0, 5);
+
+		const FSolverVec3 RotationDelta = FSolverVec3(0, 10, 0);
+		const FSolverVec3 RotationCorrectionDelta = FSolverVec3(0, 5, 0);
+
+		FSolverBody SolverBody = FSolverBody::MakeInitialized();
+		SolverBody.SetX(X0);
+		SolverBody.SetR(R0);
+		SolverBody.SetV(V0);
+		SolverBody.SetW(W0);
+		SolverBody.SetP(X0);
+		SolverBody.SetQ(R0);
+		SolverBody.SetInvM(1.0f / M0);
+		SolverBody.SetInvI(1.0f / I0);
+
+		SolverBody.ApplyPositionDelta(PositionDelta);
+		SolverBody.ApplyPositionCorrectionDelta(PositionCorrectionDelta);
+
+		SolverBody.ApplyRotationDelta(RotationDelta);
+		SolverBody.ApplyRotationCorrectionDelta(RotationCorrectionDelta);
+
+		const FSolverVec3 ExpectedDP = PositionDelta + PositionCorrectionDelta;
+		const FSolverVec3 ExpectedCP = PositionCorrectionDelta;
+		const FSolverVec3 ExpectedDQ = RotationDelta + RotationCorrectionDelta;
+		const FSolverVec3 ExpectedCQ = RotationCorrectionDelta;
+		EXPECT_VECTOR_FLOAT_EQ(SolverBody.DP(), ExpectedDP);
+		EXPECT_VECTOR_FLOAT_EQ(SolverBody.CP(), ExpectedCP);
+		EXPECT_VECTOR_FLOAT_EQ(SolverBody.DQ(), ExpectedDQ);
+		EXPECT_VECTOR_FLOAT_EQ(SolverBody.CQ(), ExpectedCQ);
+
+		SolverBody.SetImplicitVelocity(Dt);
+
+		const FVec3 ExpectedVelocity = V0 + FVec3(PositionDelta / Dt);
+		const FVec3 ExpectedAngularVelocity = W0 + FVec3(RotationDelta / Dt);
+		EXPECT_VECTOR_FLOAT_EQ(SolverBody.V(), ExpectedVelocity);
+		EXPECT_VECTOR_FLOAT_EQ(SolverBody.W(), ExpectedAngularVelocity);
+
+		SolverBody.ApplyCorrections();
+
+		const FVec3 ExpectedPosition = X0 + FVec3(PositionDelta + PositionCorrectionDelta);
+		const FRotation3 ExpectedRotation = FRotation3::IntegrateRotationWithAngularVelocity(R0, RotationDelta + RotationCorrectionDelta, 1.0);
+		EXPECT_VECTOR_FLOAT_EQ(SolverBody.P(), ExpectedPosition);
+		EXPECT_FLOAT_EQ(ExpectedRotation.W, SolverBody.Q().W);
+		EXPECT_FLOAT_EQ(ExpectedRotation.X, SolverBody.Q().X);
+		EXPECT_FLOAT_EQ(ExpectedRotation.Y, SolverBody.Q().Y);
+		EXPECT_FLOAT_EQ(ExpectedRotation.Z, SolverBody.Q().Z);
+	}
+
 	class FRigidSOAsTest : public ::testing::Test
 	{
 	protected:
diff --git a/Engine/Source/Runtime/Experimental/Chaos/Public/Chaos/Evolution/SolverBody.h b/Engine/Source/Runtime/Experimental/Chaos/Public/Chaos/Evolution/SolverBody.h
--- a/Engine/Source/Runtime/Experimental/Chaos/Public/Chaos/Evolution/SolverBody.h
+++ b/Engine/Source/Runtime/Experimental/Chaos/Public/Chaos/Evolution/SolverBody.h
@@ -130,6 +130,8 @@ namespace Chaos
 		{
 			State.DP = FSolverVec3(0);
 			State.DQ = FSolverVec3(0);
+			State.CP = FSolverVec3(0);
+			State.CQ = FSolverVec3(0);
 		}
 
 		/**
@@ -140,8 +142,8 @@ namespace Chaos
 			if (IsDynamic() && (Dt != FReal(0)))
 			{
 				const FSolverReal InvDt = FSolverReal(1) / FSolverReal(Dt);
-				SetV(State.V + FVec3(State.DP * InvDt));
-				SetW(State.W + FVec3(State.DQ * InvDt));
+				SetV(State.V + FVec3((State.DP - State.CP) * InvDt));
+				SetW(State.W + FVec3((State.DQ - State.CQ) * InvDt));
 			}
 		}
 
@@ -248,15 +250,27 @@ namespace Chaos
 		inline void SetRoM(const FRotation3& InRoM) { State.RoM = InRoM; }
 
 		/**
-		 * @brief Net world-space position correction applied by the constraints
+		 * @brief Net world-space position displacement applied by the constraints
 		*/
 		inline const FSolverVec3& DP() const { return State.DP; }
 
 		/**
-		 * @brief Net world-space rotation correction applied by the constraints (axis-angle vector equivalent to angular velocity but for position)
+		 * @brief Net world-space rotation displacement applied by the constraints (axis-angle vector equivalent to angular velocity but for position)
 		*/
 		inline const FSolverVec3& DQ() const { return State.DQ; }
 
+		/**
+		 * @brief Net world-space position correction applied by the constraints
+		 * @note This only includes correction terms that do not introduce linear velocity. The full solver displacement is given by DP
+		*/
+		inline const FSolverVec3& CP() const { return State.CP; }
+
+		/**
+		 * @brief Net world-space rotation correction applied by the constraints (axis-angle vector equivalent to angular velocity but for position)
+		 * @note This only includes correction terms that do not introduce angular velocity. The full solver angular displacement is given by DQ
+		*/
+		inline const FSolverVec3& CQ() const { return State.CQ; }
+
 		/**
 		 * @brief World-space position after applying the net correction DP()
 		 * @note Calculated on demand from P() and DP() (only requires vector addition)
@@ -280,6 +294,8 @@ namespace Chaos
 			State.Q = CorrectedQ();
 			State.DP = FSolverVec3(0);
 			State.DQ = FSolverVec3(0);
+			State.CP = FSolverVec3(0);
+			State.CQ = FSolverVec3(0);
 		}
 
 		/**
@@ -343,6 +359,28 @@ namespace Chaos
 			State.DQ += DR;
 		}
 
+		/**
+		 * @brief Apply a world-space position correction delta to the solver body center of mass
+		 * This will translate the body without introducing linear velocity
+		*/
+		inline void ApplyPositionCorrectionDelta(const FSolverVec3& CP)
+		{
+			ChaosSolverCheckNaN(CP);
+			State.CP += CP;
+			State.DP += CP;
+		}
+
+		/**
+		 * @brief Apply a world-space rotation correction delta to the solver body
+		 * This will rotate the body without introducing angular velocity
+		*/
+		inline void ApplyRotationCorrectionDelta(const FSolverVec3& CR)
+		{
+			ChaosSolverCheckNaN(CR);
+			State.CQ += CR;
+			State.DQ += CR;
+		}
+
 		/**
 		 * @brief Apply a world-space velocity delta to the solver body
 		*/
@@ -400,6 +438,8 @@ namespace Chaos
 				InvI = FSolverMatrix33(0);
 				DP = FSolverVec3(0);
 				DQ = FSolverVec3(0);
+				CP = FSolverVec3(0);
+				CQ = FSolverVec3(0);
 				Level = 0;
 				X = FVec3(0);
 				R = FRotation3::FromIdentity();
@@ -427,7 +467,15 @@ namespace Chaos
 			// Net rotation delta applied by all constraints (constantly changing as we iterate over constraints)
 			FSolverVec3 DQ;
 
-			// Distance to a kinmatic body (through the contact graph). Used by collision shock propagation
+			// Net position correction delta applied by all constraints (constantly changing as we iterate over constraints)
+			// Will translate the body without introducing linear velocity
+			FSolverVec3 CP;
+
+			// Net rotation correction delta applied by all constraints (constantly changing as we iterate over constraints)
+			// Will rotate the body without introducing angular velocity
+			FSolverVec3 CQ;
+
+			// Distance to a kinematic body (through the contact graph). Used by collision shock propagation
 			int32 Level;
 
 			// World-space center of mass state at start of sub step
@@ -571,12 +619,16 @@ namespace Chaos
 		inline int32 Level() const { return Body->Level(); }
 		inline const FSolverVec3& DP() const { return Body->DP(); }
 		inline const FSolverVec3& DQ() const { return Body->DQ(); }
+		inline const FSolverVec3& CP() const { return Body->CP(); }
+		inline const FSolverVec3& CQ() const { return Body->CQ(); }
 		inline FVec3 CorrectedP() const { return Body->CorrectedP(); }
 		inline FRotation3 CorrectedQ() const { return Body->CorrectedQ(); }
 
 		inline void ApplyTransformDelta(const FSolverVec3& DP, const FSolverVec3& DR) { Body->ApplyTransformDelta(DP, DR); }
 		inline void ApplyPositionDelta(const FSolverVec3& DP) { Body->ApplyPositionDelta(DP); }
 		inline void ApplyRotationDelta(const FSolverVec3& DR) { Body->ApplyRotationDelta(DR); }
+		inline void ApplyPositionCorrectionDelta(const FSolverVec3& DP) { Body->ApplyPositionCorrectionDelta(DP); }
+		inline void ApplyRotationCorrectionDelta(const FSolverVec3& DR) { Body->ApplyRotationCorrectionDelta(DR); }
 		inline void ApplyVelocityDelta(const FSolverVec3& DV, const FSolverVec3& DW) { Body->ApplyVelocityDelta(DV, DW); }
 		inline void ApplyLinearVelocityDelta(const FSolverVec3& DV) { Body->ApplyLinearVelocityDelta(DV); }
 		inline void ApplyAngularVelocityDelta(const FSolverVec3& DW) { Body->ApplyAngularVelocityDelta(DW); }
