There is 'Network connection lost" error when player tries to build some objects
Temporary workaround for query crash
Disable an ensure, in worse case scenario a query overlap will be missed

diff --git a/Engine/Source/Runtime/Experimental/Chaos/Public/Chaos/GJK.h b/Engine/Source/Runtime/Experimental/Chaos/Public/Chaos/GJK.h
--- a/Engine/Source/Runtime/Experimental/Chaos/Public/Chaos/GJK.h
+++ b/Engine/Source/Runtime/Experimental/Chaos/Public/Chaos/GJK.h
@@ -1204,7 +1204,7 @@ namespace Chaos
 		const T SeparatedDistance = ThicknessA + ThicknessB + Epsilon;
 		while (!bIsContact && !bIsDegenerate)
 		{
-			if (!ensure(NumIterations++ < 32))	//todo: take this out
+			if (NumIterations++ >= 32)	//todo: temporary workaround for query crash
 			{
 				break;	//if taking too long just stop. This should never happen
 			}
