Chaos - fix broadphase parallel-for (knock-on from splitting of broadphase loop from narrow phases)

This bug could result is collision being randomly ignored and arbitrary objects falling through the world.

diff --git a/Engine/Source/Runtime/Experimental/Chaos/Public/Chaos/Collision/SpatialAccelerationBroadPhase.h b/Engine/Source/Runtime/Experimental/Chaos/Public/Chaos/Collision/SpatialAccelerationBroadPhase.h
--- a/Engine/Source/Runtime/Experimental/Chaos/Public/Chaos/Collision/SpatialAccelerationBroadPhase.h
+++ b/Engine/Source/Runtime/Experimental/Chaos/Public/Chaos/Collision/SpatialAccelerationBroadPhase.h
@@ -417,6 +417,8 @@ namespace Chaos
 			{
 				BroadphaseContext.Reset();
 			}
+
+			// The number of contexts that may have new overlaps in them
 			NumActiveBroadphaseContexts = 0;
 
 			// A set of contexts, one for each worker thread, containing the allocation buffers etc.
@@ -430,19 +432,19 @@ namespace Chaos
 			const auto& ContextCreator = [this, Allocator, &Settings, &OverlapView](const int32 WorkerIndex, const int32 NumWorkers) -> int32
 			{
 				// Make sure we have enough contexts 
-				// (NOTE: ContextCreator gets called for every worker, but in serial and always with the same NumWorkers)
-				NumActiveBroadphaseContexts = NumWorkers;
+				// (NOTE: ContextCreator gets called for every worker, but in serial and always with the same NumWorkers for any single ParallelFor)
 				if (NumWorkers > BroadphaseContexts.Num())
 				{
 					check(WorkerIndex == 0);	// Should be the same for all subsequent workers
 					BroadphaseContexts.SetNum(NumWorkers);
 				}
 
-				// Initialize the broadphase context for this worker
-				BroadphaseContexts[WorkerIndex].Overlaps.Reset(OverlapView.Num());
+				// Remember the number of contexts we may have added overlaps to (if there are nested parallel-fors, we need to make 
+				// sure we keep the largest NumWorkers of all the parallel-fors)
+				NumActiveBroadphaseContexts = FMath::Max(NumWorkers, NumActiveBroadphaseContexts);
+				Allocator->SetMaxContexts(NumActiveBroadphaseContexts);
 
 				// Retrieve the midphase/collision allocator for this context
-				Allocator->SetMaxContexts(NumWorkers);
 				Private::FCollisionContextAllocator* ContextAllocator = Allocator->GetContextAllocator(WorkerIndex);
 
 				// Build the collision context for the worker
