Chaos - fix CCD debug draw

diff --git a/Engine/Source/Runtime/Experimental/Chaos/Private/Chaos/CCDUtilities.cpp b/Engine/Source/Runtime/Experimental/Chaos/Private/Chaos/CCDUtilities.cpp
--- a/Engine/Source/Runtime/Experimental/Chaos/Private/Chaos/CCDUtilities.cpp
+++ b/Engine/Source/Runtime/Experimental/Chaos/Private/Chaos/CCDUtilities.cpp
@@ -352,18 +352,6 @@ namespace Chaos
 			GroupedCCDParticles[i]->Particle->X() = GroupedCCDParticles[i]->Particle->P() - GroupedCCDParticles[i]->Particle->V() * Dt;
 		}
 
-#if CHAOS_DEBUG_DRAW
-		if (CVars::ChaosSolverDrawCCDInteractions)
-		{
-			// Debugdraw the shape at the TOI=0 (black) and TOI=1 (white)
-			for (FCCDConstraint* CCDConstraint : SortedCCDConstraints)
-			{
-				DebugDraw::DrawCCDCollisionShape(FRigidTransform3::Identity, *CCDConstraint, true, FColor::Black, &CVars::ChaosSolverDebugDebugDrawSettings);
-				DebugDraw::DrawCCDCollisionShape(FRigidTransform3::Identity, *CCDConstraint, false, FColor::White, &CVars::ChaosSolverDebugDebugDrawSettings);
-			}
-		}
-#endif
-
 		// Sort constraints based on TOI
 		FReal IslandTOI = 0.f;
 		ResetIslandParticles(Island);
@@ -421,7 +409,18 @@ namespace Chaos
 			// Debugdraw the shape at the current TOI
 			if (CVars::ChaosSolverDrawCCDInteractions)
 			{
-				DebugDraw::DrawCCDCollisionShape(FRigidTransform3::Identity, *CCDConstraint, true, FColor::Magenta, &CVars::ChaosSolverDebugDebugDrawSettings);
+				if (CCDParticle0)
+				{
+					const FRigidTransform3 SweepWorldTransform0 = FRigidTransform3(FConstGenericParticleHandle(CCDParticle0->Particle)->X(), FConstGenericParticleHandle(CCDParticle0->Particle)->Q());
+					const FRigidTransform3 ShapeWorldTransformXQ0 = CCDConstraint->SweptConstraint->GetShapeRelativeTransform0() * SweepWorldTransform0;
+					DebugDraw::DrawShape(ShapeWorldTransformXQ0, CCDConstraint->SweptConstraint->GetShape0()->GetLeafGeometry(), CCDConstraint->SweptConstraint->GetShape0(), FColor::Magenta, &CVars::ChaosSolverDebugDebugDrawSettings);
+				}
+				if (CCDParticle1)
+				{
+					const FRigidTransform3 SweepWorldTransform1 = FRigidTransform3(FConstGenericParticleHandle(CCDParticle1->Particle)->X(), FConstGenericParticleHandle(CCDParticle1->Particle)->Q());
+					const FRigidTransform3 ShapeWorldTransformXR1 = CCDConstraint->SweptConstraint->GetShapeRelativeTransform1() * SweepWorldTransform1;
+					DebugDraw::DrawShape(ShapeWorldTransformXR1, CCDConstraint->SweptConstraint->GetShape1()->GetLeafGeometry(), CCDConstraint->SweptConstraint->GetShape1(), FColor::Magenta, &CVars::ChaosSolverDebugDebugDrawSettings);
+				}
 			}
 #endif
 
@@ -485,7 +484,18 @@ namespace Chaos
 		{
 			for (FCCDConstraint* CCDConstraint : SortedCCDConstraints)
 			{
-				DebugDraw::DrawCCDCollisionShape(FRigidTransform3::Identity, *CCDConstraint, false, FColor::Green, &CVars::ChaosSolverDebugDebugDrawSettings);
+				FCCDParticle* CCDParticle0 = CCDConstraint->Particle[0];
+				FCCDParticle* CCDParticle1 = CCDConstraint->Particle[1];
+				if (CCDParticle0)
+				{
+					const FRigidTransform3 ShapeWorldTransformPQ0 = CCDConstraint->SweptConstraint->GetShapeRelativeTransform0() * FConstGenericParticleHandle(CCDParticle0->Particle)->GetTransformPQ();
+					DebugDraw::DrawShape(ShapeWorldTransformPQ0, CCDConstraint->SweptConstraint->GetShape0()->GetLeafGeometry(), CCDConstraint->SweptConstraint->GetShape0(), FColor::Green, &CVars::ChaosSolverDebugDebugDrawSettings);
+				}
+				if (CCDParticle1)
+				{
+					const FRigidTransform3 ShapeWorldTransformPQ1 = CCDConstraint->SweptConstraint->GetShapeRelativeTransform1() * FConstGenericParticleHandle(CCDParticle1->Particle)->GetTransformPQ();
+					DebugDraw::DrawShape(ShapeWorldTransformPQ1, CCDConstraint->SweptConstraint->GetShape1()->GetLeafGeometry(), CCDConstraint->SweptConstraint->GetShape1(), FColor::Green, &CVars::ChaosSolverDebugDebugDrawSettings);
+				}
 			}
 		}
 #endif
@@ -507,18 +517,6 @@ namespace Chaos
 		const int32 ConstraintEnd = IslandConstraintEnd[Island];
 		check(ConstraintNum > 0);
 
-#if CHAOS_DEBUG_DRAW
-		if (CVars::ChaosSolverDrawCCDInteractions)
-		{
-			// Debugdraw the shape at the TOI=0 (black) and TOI=1 (white)
-			for (FCCDConstraint* CCDConstraint : SortedCCDConstraints)
-			{
-				DebugDraw::DrawCCDCollisionShape(FRigidTransform3::Identity, *CCDConstraint, true, FColor::Black, &CVars::ChaosSolverDebugDebugDrawSettings);
-				DebugDraw::DrawCCDCollisionShape(FRigidTransform3::Identity, *CCDConstraint, false, FColor::White, &CVars::ChaosSolverDebugDebugDrawSettings);
-			}
-		}
-#endif
-
 		// Sort constraints based on TOI
 		std::sort(SortedCCDConstraints.GetData() + ConstraintStart, SortedCCDConstraints.GetData() + ConstraintStart + ConstraintNum, CCDConstraintSortPredicate);
 		FReal IslandTOI = 0.f;
@@ -562,14 +560,6 @@ namespace Chaos
 				AdvanceParticleXToTOI(CCDParticle1, IslandTOI, Dt);
 			}
 
-#if CHAOS_DEBUG_DRAW
-			// Debugdraw the shape at the current TOI
-			if (CVars::ChaosSolverDrawCCDInteractions)
-			{
-				DebugDraw::DrawCCDCollisionShape(FRigidTransform3::Identity, *CCDConstraint, true, FColor::Magenta, &CVars::ChaosSolverDebugDebugDrawSettings);
-			}
-#endif
-
 			ApplyImpulse(CCDConstraint);
 			CCDConstraint->ProcessedCount++;
 
@@ -683,17 +673,6 @@ namespace Chaos
 		{ 
 			CCDConstraint->SweptConstraint->SetCCDResults(CCDConstraint->NetImpulse);
 		}
-
-#if CHAOS_DEBUG_DRAW
-		// Debugdraw the shapes at the final position
-		if (CVars::ChaosSolverDrawCCDInteractions)
-		{
-			for (FCCDConstraint* CCDConstraint : SortedCCDConstraints)
-			{
-				DebugDraw::DrawCCDCollisionShape(FRigidTransform3::Identity, *CCDConstraint, false, FColor::Green, &CVars::ChaosSolverDebugDebugDrawSettings);
-			}
-		}
-#endif
 	}
 
 	bool FCCDManager::UpdateParticleSweptConstraints(FCCDParticle* CCDParticle, const FReal IslandTOI, const FReal Dt)
diff --git a/Engine/Source/Runtime/Experimental/Chaos/Private/Chaos/ChaosDebugDraw.cpp b/Engine/Source/Runtime/Experimental/Chaos/Private/Chaos/ChaosDebugDraw.cpp
--- a/Engine/Source/Runtime/Experimental/Chaos/Private/Chaos/ChaosDebugDraw.cpp
+++ b/Engine/Source/Runtime/Experimental/Chaos/Private/Chaos/ChaosDebugDraw.cpp
@@ -360,6 +360,11 @@ namespace Chaos
 				return;
 			}
 
+			if (Particle == nullptr)
+			{
+				return;
+			}
+
 			if (const TPBDRigidParticleHandle<FReal, 3>* Rigid = Particle->CastToRigidParticle())
 			{
 				const TUniquePtr<TBVHParticles<FReal, 3>>& CollisionParticles = Rigid->CollisionParticles();
@@ -491,7 +496,10 @@ namespace Chaos
 
 		void DrawShapesImpl(const FGeometryParticleHandle* Particle, const FRigidTransform3& ShapeTransform, const FImplicitObject* Implicit, const FShapeOrShapesArray& Shapes, const FReal Margin, const FColor& Color, const FRealSingle Duration, const FChaosDebugDrawSettings& Settings)
 		{
-			if (!Particle || !Implicit) return;
+			if (Implicit == nullptr)
+			{
+				return;
+			}
 
 			const EImplicitObjectType PackedType = Implicit->GetType(); // Type includes scaling and instancing data
 			const EImplicitObjectType InnerType = GetInnerType(Implicit->GetType());
@@ -591,13 +599,16 @@ namespace Chaos
 			}
 
 			FColor ShapeColor = Color;
-			if (bChaosDebugDebugDrawColorShapesByShapeType)
+			if (Particle != nullptr)
 			{
-				ShapeColor = Settings.ShapesColorsPerShapeType.GetColorFromShapeType(InnerType);
-			}
-			if (bChaosDebugDebugDrawColorShapesByIsland && (FConstGenericParticleHandle(Particle)->IslandIndex() != INDEX_NONE))
-			{
-				ShapeColor = GetIslandColor(FConstGenericParticleHandle(Particle)->IslandIndex(), true);
+				if (bChaosDebugDebugDrawColorShapesByShapeType)
+				{
+					ShapeColor = Settings.ShapesColorsPerShapeType.GetColorFromShapeType(InnerType);
+				}
+				if (bChaosDebugDebugDrawColorShapesByIsland && (FConstGenericParticleHandle(Particle)->IslandIndex() != INDEX_NONE))
+				{
+					ShapeColor = GetIslandColor(FConstGenericParticleHandle(Particle)->IslandIndex(), true);
+				}
 			}
 
 			// If we get here, we have an actual shape to render
@@ -1026,24 +1037,6 @@ namespace Chaos
 			DrawCollisionImpl(SpaceTransform, ConstraintHandle->GetContact(), ColorScale, Duration, Settings);
 		}
 
-		void DrawParticleCCDCollisionShapeImpl(const FRigidTransform3& SpaceTransform, const FCCDParticle* CCDParticle, const bool bShowStartPos,  const FColor& ShapeColor, const FRealSingle Duration, const FChaosDebugDrawSettings& Settings)
-		{
-			if (CCDParticle != nullptr)
-			{
-				FConstGenericParticleHandle P0 = CCDParticle->Particle;
-				if (P0->IsDynamic() && P0->CCDEnabled())
-				{
-					FRigidTransform3 ActorTransform0 = P0->GetTransformPQ();
-					if (bShowStartPos)
-					{
-						// For CCD we want to see the initial position but with the latest rotation since that's what the sweep uses
-						ActorTransform0.SetTranslation(P0->X());
-					}
-					DrawShapesImpl(P0->Handle(), ActorTransform0 * SpaceTransform, P0->Geometry().Get(), FShapeOrShapesArray(P0->Handle()), 0.0f, ShapeColor, Duration, Settings);
-				}
-			}
-		}
-
 		void DrawParticleCCDCollisionImpulseImpl(const FRigidTransform3& SpaceTransform, const FCCDParticle* CCDParticle, const FCCDConstraint& CCDConstraint, const int32 ManifoldPointIndex, const FVec3& Impulse, const FRealSingle Duration, const FChaosDebugDrawSettings& Settings)
 		{
 			if (CCDParticle != nullptr)
@@ -1060,15 +1053,6 @@ namespace Chaos
 			}
 		}
 
-		void DrawCCDCollisionShapeImpl(const FRigidTransform3& SpaceTransform, const FCCDConstraint& CCDConstraint, const bool bShowStartPos, const FColor& ShapeColor, const FChaosDebugDrawSettings& Settings)
-		{
-			const FRealSingle Duration = ChaosDebugDrawCCDDuration;
-
-			DrawParticleCCDCollisionShapeImpl(SpaceTransform, CCDConstraint.Particle[0], bShowStartPos, ShapeColor, Duration, Settings);
-			DrawParticleCCDCollisionShapeImpl(SpaceTransform, CCDConstraint.Particle[1], bShowStartPos, ShapeColor, Duration, Settings);
-			DrawCollisionImpl(SpaceTransform, CCDConstraint.SweptConstraint, 1.0f, Duration, Settings);
-		}
-
 		void DrawCCDCollisionImpulseImpl(const FRigidTransform3& SpaceTransform, const FCCDConstraint& CCDConstraint, const int32 ManifoldPointIndex, const FVec3& Impulse, const FChaosDebugDrawSettings& Settings)
 		{
 			const FRealSingle Duration = ChaosDebugDrawCCDDuration;
@@ -1818,14 +1802,6 @@ namespace Chaos
 			}
 		}
 
-		void DrawCCDCollisionShape(const FRigidTransform3& SpaceTransform, const FCCDConstraint& CCDConstraint, const bool bShowStartPos, const FColor& ShapeColor, const FChaosDebugDrawSettings* Settings)
-		{
-			if (FDebugDrawQueue::IsDebugDrawingEnabled())
-			{
-				DrawCCDCollisionShapeImpl(SpaceTransform, CCDConstraint, bShowStartPos, ShapeColor, GetChaosDebugDrawSettings(Settings));
-			}
-		}
-
 		void DrawCCDCollisionImpulse(const FRigidTransform3& SpaceTransform, const FCCDConstraint& CCDConstraint, const int32 ManifoldPointIndex, const FVec3& Impulse, const FChaosDebugDrawSettings* Settings)
 		{
 			if (FDebugDrawQueue::IsDebugDrawingEnabled())
diff --git a/Engine/Source/Runtime/Experimental/Chaos/Private/Chaos/Collision/ParticlePairMidPhase.cpp b/Engine/Source/Runtime/Experimental/Chaos/Private/Chaos/Collision/ParticlePairMidPhase.cpp
--- a/Engine/Source/Runtime/Experimental/Chaos/Private/Chaos/Collision/ParticlePairMidPhase.cpp
+++ b/Engine/Source/Runtime/Experimental/Chaos/Private/Chaos/Collision/ParticlePairMidPhase.cpp
@@ -10,6 +10,7 @@
 #include "Chaos/Particle/ParticleUtilities.h"
 #include "Chaos/PBDCollisionConstraints.h"
 
+#include "Chaos/ChaosDebugDraw.h"
 #include "ChaosStats.h"
 
 extern bool Chaos_Collision_NarrowPhase_AABBBoundsCheck;
@@ -18,6 +19,11 @@ namespace Chaos
 {
 	namespace CVars
 	{
+		extern int32 ChaosSolverDrawCCDInteractions;
+#if CHAOS_DEBUG_DRAW
+		extern DebugDraw::FChaosDebugDrawSettings ChaosSolverDebugDebugDrawSettings;
+#endif
+
 		bool bChaos_Collision_MidPhase_EnableBoundsChecks = true;
 		FAutoConsoleVariableRef CVarChaos_Collision_EnableBoundsChecks(TEXT("p.Chaos.Collision.EnableBoundsChecks"), bChaos_Collision_MidPhase_EnableBoundsChecks, TEXT(""));
 
@@ -422,6 +428,22 @@ namespace Chaos
 			const FRigidTransform3 CCDShapeWorldTransform1 = Constraint->ImplicitTransform[1] * CCDParticleWorldTransform1;
 			const bool bDidSweep = Collisions::UpdateConstraintSwept(*Constraint.Get(), CCDShapeWorldTransform0, CCDShapeWorldTransform1, Dt);
 
+#if CHAOS_DEBUG_DRAW
+			if (CVars::ChaosSolverDrawCCDInteractions)
+			{
+				if (FConstGenericParticleHandle(Constraint->GetParticle0())->CCDEnabled())
+				{
+					DebugDraw::DrawShape(CCDShapeWorldTransform0, Implicit0, Shape0, FColor::Black, &CVars::ChaosSolverDebugDebugDrawSettings);
+					DebugDraw::DrawShape(ShapeWorldTransform0, Implicit0, Shape0, FColor::White, &CVars::ChaosSolverDebugDebugDrawSettings);
+				}
+				if (FConstGenericParticleHandle(Constraint->GetParticle1())->CCDEnabled())
+				{
+					DebugDraw::DrawShape(CCDShapeWorldTransform1, Implicit1, Shape1, FColor::Black, &CVars::ChaosSolverDebugDebugDrawSettings);
+					DebugDraw::DrawShape(ShapeWorldTransform1, Implicit1, Shape1, FColor::White, &CVars::ChaosSolverDebugDebugDrawSettings);
+				}
+			}
+#endif
+
 			// If we did get a hit but it's at TOI = 1, treat this constraint as a regular non-swept constraint (skip the rewind)
 			if ((!bDidSweep) || Constraint->GetCCDTimeOfImpact() == FReal(1))
 			{
diff --git a/Engine/Source/Runtime/Experimental/Chaos/Public/Chaos/ChaosDebugDraw.h b/Engine/Source/Runtime/Experimental/Chaos/Public/Chaos/ChaosDebugDraw.h
--- a/Engine/Source/Runtime/Experimental/Chaos/Public/Chaos/ChaosDebugDraw.h
+++ b/Engine/Source/Runtime/Experimental/Chaos/Public/Chaos/ChaosDebugDraw.h
@@ -226,7 +226,6 @@ namespace Chaos
 		CHAOS_API void DrawSpatialAccelerationStructure(const ISpatialAcceleration<FAccelerationStructureHandle, FReal, 3>& SpatialAccelerationStructure, const FChaosDebugDrawSettings* Settings = nullptr);
 		CHAOS_API void DrawSuspensionConstraints(const FRigidTransform3& SpaceTransform, const FPBDSuspensionConstraints& Constraints, const FChaosDebugDrawSettings* Settings = nullptr);
 		CHAOS_API void DrawCCDAxisThreshold(const FVec3& X, const FVec3& AxisThreshold, const FVec3& DeltaX, const FQuat& R);
-		CHAOS_API void DrawCCDCollisionShape(const FRigidTransform3& SpaceTransform, const FCCDConstraint& CCDConstraint, const bool bShowStartPos, const FColor& ShapeColor, const FChaosDebugDrawSettings* Settings = nullptr);
 		CHAOS_API void DrawCCDCollisionImpulse(const FRigidTransform3& SpaceTransform, const FCCDConstraint& CCDConstraint, const int32 ManifoldPointIndex, const FVec3& Impulse, const FChaosDebugDrawSettings* Settings = nullptr);
 #endif
 	}
