[Chaos]
* ParticleHandleFwd.h - Removed Framework/Threading.h include in favor of new include that only holds enum needed

diff --git a/Engine/Source/Runtime/Experimental/Chaos/Public/Chaos/ParticleHandleFwd.h b/Engine/Source/Runtime/Experimental/Chaos/Public/Chaos/ParticleHandleFwd.h
--- a/Engine/Source/Runtime/Experimental/Chaos/Public/Chaos/ParticleHandleFwd.h
+++ b/Engine/Source/Runtime/Experimental/Chaos/Public/Chaos/ParticleHandleFwd.h
@@ -1,8 +1,8 @@
 // Copyright Epic Games, Inc. All Rights Reserved.
 #pragma once
 
-#include <Chaos/Real.h>
-#include <Framework/Threading.h>
+#include "Chaos/Real.h"
+#include "Framework/ThreadContextEnum.h"
 #include <type_traits>
 
 // Use to define out code blocks that need to be adapted to use Particle Handles in a searchable way (better than #if 0)
@@ -116,3 +116,7 @@ namespace Chaos
 	template<EThreadContext Id>
 	using TThreadRigidParticle = std::conditional_t<Id == EThreadContext::External, FPBDRigidParticle, FPBDRigidParticleHandle>;
 }
+
+#if UE_ENABLE_INCLUDE_ORDER_DEPRECATED_IN_5_2
+#include "Framework/Threading.h"
+#endif
diff --git a/Engine/Source/Runtime/Experimental/Chaos/Public/Framework/ThreadContextEnum.h b/Engine/Source/Runtime/Experimental/Chaos/Public/Framework/ThreadContextEnum.h
new file mode 100644
--- /dev/null
+++ b/Engine/Source/Runtime/Experimental/Chaos/Public/Framework/ThreadContextEnum.h
@@ -0,0 +1,14 @@
+// Copyright Epic Games, Inc. All Rights Reserved.
+
+#pragma once
+
+namespace Chaos
+{
+
+enum class EThreadContext
+{
+	External,
+	Internal
+};
+
+}
\ No newline at end of file
diff --git a/Engine/Source/Runtime/Experimental/Chaos/Public/Framework/Threading.h b/Engine/Source/Runtime/Experimental/Chaos/Public/Framework/Threading.h
--- a/Engine/Source/Runtime/Experimental/Chaos/Public/Framework/Threading.h
+++ b/Engine/Source/Runtime/Experimental/Chaos/Public/Framework/Threading.h
@@ -3,6 +3,7 @@
 #pragma once
 
 #include "CoreMinimal.h"
+#include "Framework/ThreadContextEnum.h"
 #include "UObject/ObjectMacros.h"
 #include "PhysicsCoreTypes.h"
 #include "ChaosLog.h"
@@ -136,12 +137,6 @@ namespace Chaos
 /** Signals that we have left a write lock to control the checks above */
 #define CHAOS_RECORD_LEAVE_WRITE_LOCK Chaos::ThreadingPrivate::DecWriteDepth(this);
 
-enum class EThreadContext
-{
-	External,
-	Internal
-};
-
 #if PHYSICS_THREAD_CONTEXT
 /** Debug helper to ensure threading mistakes are caught. Do not use for ship */
 class CHAOS_API FPhysicsThreadContext : public TThreadSingleton<FPhysicsThreadContext>
