Geometry collection : fix issue where unitialized data could be interpolated on the game thread when using async physics
- fixed the case where particle was active and became inactive and final position were not properly set
- made sure the the disable state was properly set in all cases
- added a Nan detection define for making investigation of future issues easier

diff --git a/Engine/Source/Runtime/Experimental/Chaos/Private/PhysicsProxy/GeometryCollectionPhysicsProxy.cpp b/Engine/Source/Runtime/Experimental/Chaos/Private/PhysicsProxy/GeometryCollectionPhysicsProxy.cpp
--- a/Engine/Source/Runtime/Experimental/Chaos/Private/PhysicsProxy/GeometryCollectionPhysicsProxy.cpp
+++ b/Engine/Source/Runtime/Experimental/Chaos/Private/PhysicsProxy/GeometryCollectionPhysicsProxy.cpp
@@ -47,6 +47,15 @@
 #define TODO_REIMPLEMENT_RIGID_CACHING 0
 #endif
 
+#define GC_PHYSICSPROXY_CHECK_FOR_NAN_ENABLED 0
+
+#if GC_PHYSICSPROXY_CHECK_FOR_NAN_ENABLED
+#define GC_PHYSICSPROXY_CHECK_FOR_NAN(Vec) ensure(!Vec.ContainsNaN())
+#else
+#define GC_PHYSICSPROXY_CHECK_FOR_NAN(Vec)
+#endif
+
+
 float CollisionParticlesPerObjectFractionDefault = 1.0f;
 FAutoConsoleVariableRef CVarCollisionParticlesPerObjectFractionDefault(
 	TEXT("p.CollisionParticlesPerObjectFractionDefault"), 
@@ -2430,8 +2439,13 @@ void FGeometryCollectionPhysicsProxy::BufferPhysicsResults_External(Chaos::FDirt
 	{
 		FParticle& GTParticle = *GTParticles[TransformGroupIndex];
 		TargetResults.ParticleXs[TransformGroupIndex] = GTParticle.X();
+		GC_PHYSICSPROXY_CHECK_FOR_NAN(TargetResults.ParticleXs[TransformGroupIndex]);
+
 		TargetResults.ParticleRs[TransformGroupIndex] = GTParticle.R();
 		TargetResults.States[TransformGroupIndex].DynamicState = static_cast<uint8>(GetObjectStateEnumFromObjectState(GTParticle.ObjectState()));
+		TargetResults.States[TransformGroupIndex].DisabledState = GTParticle.Disabled();
+		TargetResults.States[TransformGroupIndex].DynamicInternalClusterParent = false;
+		TargetResults.States[TransformGroupIndex].HasInternalClusterParent = false;
 
 		if (LinearVelocities && AngularVelocities)
 		{
@@ -2512,6 +2526,8 @@ void FGeometryCollectionPhysicsProxy::BufferPhysicsResults_Internal(Chaos::FPBDR
 
 		for (int32 TransformGroupIndex = 0; TransformGroupIndex < NumTransformGroupElements; ++TransformGroupIndex)
 		{
+			const bool bWasActive = PhysicsThreadCollection.Active[TransformGroupIndex];
+
 			TargetResults.Transforms[TransformGroupIndex] = PhysicsThreadCollection.Transform[TransformGroupIndex];
 			TargetResults.Parent[TransformGroupIndex] = PhysicsThreadCollection.Parent[TransformGroupIndex];
 			TargetResults.InternalClusterUniqueIdx[TransformGroupIndex] = INDEX_NONE;
@@ -2545,6 +2561,8 @@ void FGeometryCollectionPhysicsProxy::BufferPhysicsResults_Internal(Chaos::FPBDR
 				// of the parent cluster.
 
 				TargetResults.ParticleXs[TransformGroupIndex] = Handle->X();
+				GC_PHYSICSPROXY_CHECK_FOR_NAN(TargetResults.ParticleXs[TransformGroupIndex]);
+
 				TargetResults.ParticleRs[TransformGroupIndex] = Handle->R();
 				TargetResults.ParticleVs[TransformGroupIndex] = Handle->V();
 				TargetResults.ParticleWs[TransformGroupIndex] = Handle->W();
@@ -2585,6 +2603,29 @@ void FGeometryCollectionPhysicsProxy::BufferPhysicsResults_Internal(Chaos::FPBDR
 			}
 			else    // Handle->Disabled()
 			{
+				// if the particle was active and now is not we still need to properly set the positions, velocities and transform values
+				// this is because interpolation may need this data when processed on the gamethread ( see PullFromPhysicsState method )
+				if (bWasActive)
+				{ 
+					TargetResults.ParticleXs[TransformGroupIndex] = Handle->X();
+					GC_PHYSICSPROXY_CHECK_FOR_NAN(TargetResults.ParticleXs[TransformGroupIndex]);
+
+					TargetResults.ParticleRs[TransformGroupIndex] = Handle->R();
+					TargetResults.ParticleVs[TransformGroupIndex] = Handle->V();
+					TargetResults.ParticleWs[TransformGroupIndex] = Handle->W();
+					FRigidTransform3 ParticleToWorld(Handle->X(), Handle->R());
+					const FTransform MassToLocal = PhysicsThreadCollection.MassToLocal[TransformGroupIndex];
+
+					TargetResults.Transforms[TransformGroupIndex] = MassToLocal.GetRelativeTransformReverse(ParticleToWorld).GetRelativeTransform(ActorToWorld);
+					TargetResults.Transforms[TransformGroupIndex].NormalizeRotation();
+					if (IsActorScaled)
+					{
+						TargetResults.Transforms[TransformGroupIndex] = MassToLocal.Inverse() * ActorScaleTransform * MassToLocal * TargetResults.Transforms[TransformGroupIndex];
+					}
+
+					PhysicsThreadCollection.Transform[TransformGroupIndex] = TargetResults.Transforms[TransformGroupIndex];
+				}
+
 				// The rigid body parent cluster has changed within the solver, and its
 				// parent body is not tracked within the geometry collection. So we need to
 				// pull the rigid bodies out of the transform hierarchy, and just drive
@@ -2593,7 +2634,7 @@ void FGeometryCollectionPhysicsProxy::BufferPhysicsResults_Internal(Chaos::FPBDR
 				{
 					if(Chaos::FPBDRigidClusteredParticleHandle* ClusterParent = ClusterParentBase->CastToClustered())
 					{
-						// syncronize parents if it has changed.
+						// synchronize parents if it has changed.
 						if(SolverClusterID[TransformGroupIndex] != ClusterParent)
 						{
 							// Force all driven rigid bodies out of the transform hierarchy
@@ -2626,6 +2667,8 @@ void FGeometryCollectionPhysicsProxy::BufferPhysicsResults_Internal(Chaos::FPBDR
 							
 							const FTransform ParticleToWorld = Handle->ChildToParent() * FRigidTransform3(ClusterParent->X(), ClusterParent->R());    // aka ClusterChildToWorld
 							TargetResults.ParticleXs[TransformGroupIndex] = ParticleToWorld.GetTranslation();
+							GC_PHYSICSPROXY_CHECK_FOR_NAN(TargetResults.ParticleXs[TransformGroupIndex]);
+
 							TargetResults.ParticleRs[TransformGroupIndex] = ParticleToWorld.GetRotation();
 							TargetResults.ParticleVs[TransformGroupIndex] = ClusterParent->V();
 							TargetResults.ParticleWs[TransformGroupIndex] = ClusterParent->W();
@@ -2801,12 +2844,18 @@ bool FGeometryCollectionPhysicsProxy::PullFromPhysicsState(const Chaos::FDirtyGe
 			const FGeometryCollectionResults& Next = NextPullData->Results;
 			for (int32 TransformGroupIndex = 0; TransformGroupIndex < NumTransforms; ++TransformGroupIndex)
 			{
-				if (!TargetResults.States[TransformGroupIndex].DisabledState)
+				const bool WasDisabled = Prev.States[TransformGroupIndex].DisabledState;
+				if (!WasDisabled)
 				{
 					FParticle& GTParticle = *GTParticles[TransformGroupIndex];
 
+					GC_PHYSICSPROXY_CHECK_FOR_NAN(Prev.ParticleXs[TransformGroupIndex]);
+					GC_PHYSICSPROXY_CHECK_FOR_NAN(Next.ParticleXs[TransformGroupIndex]);
+
 					const Chaos::FVec3 OldX = GTParticle.X();
 					const Chaos::FVec3 NewX = FMath::Lerp(Prev.ParticleXs[TransformGroupIndex], Next.ParticleXs[TransformGroupIndex], *Alpha);
+					GC_PHYSICSPROXY_CHECK_FOR_NAN(NewX);
+
 					const bool bNeedUpdateX = (NewX != OldX);
 					if (bNeedUpdateX)
 					{
