Chaos: Fixed triangle mesh bvh generation skipping large triangles above the global payload size for the intermediate AABBTree

diff --git a/Engine/Source/Runtime/Experimental/Chaos/Private/Chaos/TriangleMeshImplicitObject.cpp b/Engine/Source/Runtime/Experimental/Chaos/Private/Chaos/TriangleMeshImplicitObject.cpp
--- a/Engine/Source/Runtime/Experimental/Chaos/Private/Chaos/TriangleMeshImplicitObject.cpp
+++ b/Engine/Source/Runtime/Experimental/Chaos/Private/Chaos/TriangleMeshImplicitObject.cpp
@@ -1964,7 +1964,19 @@ void FTriangleMeshImplicitObject::RebuildBVImp(const TArray<TVec3<IdxType>>& Ele
 	{
 		BVEntries.Add({this, Tri});
 	}
-	TreeBVH.Reinitialize(BVEntries, 22);
+
+	// Override some parameters for the internal triangle mesh acceleration
+	//
+	// Children in leaf - update to tested value with higher perf for this application
+	// Tree depth - default for the tree type
+	// Max element bounds - as we use a faster tree derived from this that has no global item support
+	//						all elements must be fully placed in the tree to be picked up when we translate
+	//						them over to the fast bvh
+	constexpr static int32 MaxChildrenInLeaf = 22;
+	constexpr static int32 MaxTreeDepth = BVHType::DefaultMaxTreeDepth;
+	constexpr static Chaos::FRealSingle MaxElementBounds = std::numeric_limits<Chaos::FRealSingle>::max();
+
+	TreeBVH.Reinitialize(BVEntries, MaxChildrenInLeaf, MaxTreeDepth, MaxElementBounds);
 }
 
 FTriangleMeshImplicitObject::~FTriangleMeshImplicitObject() = default;
