[Chaos]
* Include fixes and markup to make IWYU and header units to work
* XPBDAxialSpringConstraints.h - Fixed compile error for implicit casting from float to bool

diff --git a/Engine/Source/Runtime/Experimental/Chaos/Public/Chaos/Collision/CollisionFilter.h b/Engine/Source/Runtime/Experimental/Chaos/Public/Chaos/Collision/CollisionFilter.h
--- a/Engine/Source/Runtime/Experimental/Chaos/Public/Chaos/Collision/CollisionFilter.h
+++ b/Engine/Source/Runtime/Experimental/Chaos/Public/Chaos/Collision/CollisionFilter.h
@@ -1,6 +1,8 @@
 // Copyright Epic Games, Inc. All Rights Reserved.
 #pragma once
 
+// HEADER_UNIT_SKIP - Internal
+
 #include "CoreMinimal.h"
 #include "Chaos/Core.h"
 #include "Chaos/CollisionFilterData.h"
diff --git a/Engine/Source/Runtime/Experimental/Chaos/Public/Chaos/Complex.h b/Engine/Source/Runtime/Experimental/Chaos/Public/Chaos/Complex.h
--- a/Engine/Source/Runtime/Experimental/Chaos/Public/Chaos/Complex.h
+++ b/Engine/Source/Runtime/Experimental/Chaos/Public/Chaos/Complex.h
@@ -1,6 +1,8 @@
 // Copyright Epic Games, Inc. All Rights Reserved.
 #pragma once
 
+// HEADER_UNIT_SKIP - Internal
+
 namespace Chaos
 {
 class FComplex
diff --git a/Engine/Source/Runtime/Experimental/Chaos/Public/Chaos/EPA.h b/Engine/Source/Runtime/Experimental/Chaos/Public/Chaos/EPA.h
--- a/Engine/Source/Runtime/Experimental/Chaos/Public/Chaos/EPA.h
+++ b/Engine/Source/Runtime/Experimental/Chaos/Public/Chaos/EPA.h
@@ -1,5 +1,8 @@
 // Copyright Epic Games, Inc. All Rights Reserved.
 #pragma once
+
+// HEADER_UNIT_SKIP - Internal
+
 #include "Chaos/Simplex.h"
 #include <queue>
 #include "ChaosCheck.h"
diff --git a/Engine/Source/Runtime/Experimental/Chaos/Public/Chaos/EPAVectorized.h b/Engine/Source/Runtime/Experimental/Chaos/Public/Chaos/EPAVectorized.h
--- a/Engine/Source/Runtime/Experimental/Chaos/Public/Chaos/EPAVectorized.h
+++ b/Engine/Source/Runtime/Experimental/Chaos/Public/Chaos/EPAVectorized.h
@@ -1,6 +1,7 @@
 // Copyright Epic Games, Inc. All Rights Reserved.
 #pragma once
 
+#include "Chaos/Core.h"
 #include "Chaos/SimplexVectorized.h"
 #include "Chaos/EPA.h"
 
diff --git a/Engine/Source/Runtime/Experimental/Chaos/Public/Chaos/Framework/BufferedData.h b/Engine/Source/Runtime/Experimental/Chaos/Public/Chaos/Framework/BufferedData.h
--- a/Engine/Source/Runtime/Experimental/Chaos/Public/Chaos/Framework/BufferedData.h
+++ b/Engine/Source/Runtime/Experimental/Chaos/Public/Chaos/Framework/BufferedData.h
@@ -2,6 +2,8 @@
 
 #pragma once
 
+// HEADER_UNIT_SKIP - Internal
+
 namespace Chaos
 {
 	/**
diff --git a/Engine/Source/Runtime/Experimental/Chaos/Public/Chaos/KinematicGeometryParticleBuffer.h b/Engine/Source/Runtime/Experimental/Chaos/Public/Chaos/KinematicGeometryParticleBuffer.h
--- a/Engine/Source/Runtime/Experimental/Chaos/Public/Chaos/KinematicGeometryParticleBuffer.h
+++ b/Engine/Source/Runtime/Experimental/Chaos/Public/Chaos/KinematicGeometryParticleBuffer.h
@@ -2,6 +2,8 @@
 #pragma once
 #include "GeometryParticleBuffer.h"
 
+// HEADER_UNIT_SKIP - THIS FILE IS BROKEN. GeometryParticleBuffer.h does not exist
+
 namespace Chaos
 {
 	
diff --git a/Engine/Source/Runtime/Experimental/Chaos/Public/Chaos/KinematicTargets.h b/Engine/Source/Runtime/Experimental/Chaos/Public/Chaos/KinematicTargets.h
--- a/Engine/Source/Runtime/Experimental/Chaos/Public/Chaos/KinematicTargets.h
+++ b/Engine/Source/Runtime/Experimental/Chaos/Public/Chaos/KinematicTargets.h
@@ -1,6 +1,8 @@
 // Copyright Epic Games, Inc. All Rights Reserved.
 #pragma once
 
+// HEADER_UNIT_SKIP - Internal
+
 #include "Chaos/ChaosArchive.h"
 #include "Chaos/Transform.h"
 #include "UObject/FortniteReleaseBranchCustomObjectVersion.h"
diff --git a/Engine/Source/Runtime/Experimental/Chaos/Public/Chaos/PBDAnimDriveConstraint.h b/Engine/Source/Runtime/Experimental/Chaos/Public/Chaos/PBDAnimDriveConstraint.h
--- a/Engine/Source/Runtime/Experimental/Chaos/Public/Chaos/PBDAnimDriveConstraint.h
+++ b/Engine/Source/Runtime/Experimental/Chaos/Public/Chaos/PBDAnimDriveConstraint.h
@@ -1,6 +1,8 @@
 // Copyright Epic Games, Inc. All Rights Reserved.
 #pragma once
 
+// HEADER_UNIT_SKIP - Internal
+
 #include "Chaos/Framework/Parallel.h"
 #include "Chaos/PBDSoftsEvolutionFwd.h"
 #include "Chaos/PBDStiffness.h"
diff --git a/Engine/Source/Runtime/Experimental/Chaos/Public/Chaos/PBDAxialSpringConstraints.h b/Engine/Source/Runtime/Experimental/Chaos/Public/Chaos/PBDAxialSpringConstraints.h
--- a/Engine/Source/Runtime/Experimental/Chaos/Public/Chaos/PBDAxialSpringConstraints.h
+++ b/Engine/Source/Runtime/Experimental/Chaos/Public/Chaos/PBDAxialSpringConstraints.h
@@ -1,6 +1,8 @@
 // Copyright Epic Games, Inc. All Rights Reserved.
 #pragma once
 
+// HEADER_UNIT_SKIP - Internal
+
 #include "Chaos/PBDAxialSpringConstraintsBase.h"
 #include "Chaos/PropertyCollectionAdapter.h"
 
diff --git a/Engine/Source/Runtime/Experimental/Chaos/Public/Chaos/PBDAxialSpringConstraintsBase.h b/Engine/Source/Runtime/Experimental/Chaos/Public/Chaos/PBDAxialSpringConstraintsBase.h
--- a/Engine/Source/Runtime/Experimental/Chaos/Public/Chaos/PBDAxialSpringConstraintsBase.h
+++ b/Engine/Source/Runtime/Experimental/Chaos/Public/Chaos/PBDAxialSpringConstraintsBase.h
@@ -1,6 +1,8 @@
 // Copyright Epic Games, Inc. All Rights Reserved.
 #pragma once
 
+// HEADER_UNIT_SKIP - Internal
+
 #include "Chaos/PBDSoftsEvolutionFwd.h"
 #include "Chaos/PBDSoftsSolverParticles.h"
 #include "Chaos/PBDStiffness.h"
diff --git a/Engine/Source/Runtime/Experimental/Chaos/Public/Chaos/PBDChainConstraints.h b/Engine/Source/Runtime/Experimental/Chaos/Public/Chaos/PBDChainConstraints.h
--- a/Engine/Source/Runtime/Experimental/Chaos/Public/Chaos/PBDChainConstraints.h
+++ b/Engine/Source/Runtime/Experimental/Chaos/Public/Chaos/PBDChainConstraints.h
@@ -1,6 +1,9 @@
 // Copyright Epic Games, Inc. All Rights Reserved.
 #pragma once
 
+// HEADER_UNIT_SKIP - This file does not compile anymore (Base ctor needs to be called)
+
+#include "Chaos/Evolution/IndexedConstraintContainer.h"
 #include "Chaos/PBDParticles.h"
 #include "Chaos/Framework/Parallel.h"
 #include "Chaos/PBDConstraintContainer.h"
diff --git a/Engine/Source/Runtime/Experimental/Chaos/Public/Chaos/PBDShapeConstraints.h b/Engine/Source/Runtime/Experimental/Chaos/Public/Chaos/PBDShapeConstraints.h
--- a/Engine/Source/Runtime/Experimental/Chaos/Public/Chaos/PBDShapeConstraints.h
+++ b/Engine/Source/Runtime/Experimental/Chaos/Public/Chaos/PBDShapeConstraints.h
@@ -1,6 +1,8 @@
 // Copyright Epic Games, Inc. All Rights Reserved.
 #pragma once
 
+// HEADER_UNIT_SKIP - Internal
+
 #include "Chaos/PBDShapeConstraintsBase.h"
 
 namespace Chaos::Softs
diff --git a/Engine/Source/Runtime/Experimental/Chaos/Public/Chaos/PBDShapeConstraintsBase.h b/Engine/Source/Runtime/Experimental/Chaos/Public/Chaos/PBDShapeConstraintsBase.h
--- a/Engine/Source/Runtime/Experimental/Chaos/Public/Chaos/PBDShapeConstraintsBase.h
+++ b/Engine/Source/Runtime/Experimental/Chaos/Public/Chaos/PBDShapeConstraintsBase.h
@@ -5,6 +5,7 @@
 #include "Chaos/PBDSoftsEvolutionFwd.h"
 #include "Chaos/Array.h"
 #include "Chaos/PBDParticles.h"
+#include "Chaos/PBDSoftsSolverParticles.h"
 
 #include <functional>
 
diff --git a/Engine/Source/Runtime/Experimental/Chaos/Public/Chaos/PBDTetConstraints.h b/Engine/Source/Runtime/Experimental/Chaos/Public/Chaos/PBDTetConstraints.h
--- a/Engine/Source/Runtime/Experimental/Chaos/Public/Chaos/PBDTetConstraints.h
+++ b/Engine/Source/Runtime/Experimental/Chaos/Public/Chaos/PBDTetConstraints.h
@@ -1,6 +1,8 @@
 // Copyright Epic Games, Inc. All Rights Reserved.
 #pragma once
 
+// HEADER_UNIT_SKIP - Internal
+
 #include "Chaos/PBDTetConstraintsBase.h"
 
 namespace Chaos::Softs
diff --git a/Engine/Source/Runtime/Experimental/Chaos/Public/Chaos/PBDTetConstraintsBase.h b/Engine/Source/Runtime/Experimental/Chaos/Public/Chaos/PBDTetConstraintsBase.h
--- a/Engine/Source/Runtime/Experimental/Chaos/Public/Chaos/PBDTetConstraintsBase.h
+++ b/Engine/Source/Runtime/Experimental/Chaos/Public/Chaos/PBDTetConstraintsBase.h
@@ -6,6 +6,7 @@
 #include "Chaos/PBDParticles.h"
 #include "Chaos/DynamicParticles.h"
 #include "Chaos/PBDParticles.h"
+#include "Chaos/PBDSoftsSolverParticles.h"
 
 namespace Chaos::Softs
 {
diff --git a/Engine/Source/Runtime/Experimental/Chaos/Public/Chaos/PBDVolumeConstraint.h b/Engine/Source/Runtime/Experimental/Chaos/Public/Chaos/PBDVolumeConstraint.h
--- a/Engine/Source/Runtime/Experimental/Chaos/Public/Chaos/PBDVolumeConstraint.h
+++ b/Engine/Source/Runtime/Experimental/Chaos/Public/Chaos/PBDVolumeConstraint.h
@@ -1,6 +1,8 @@
 // Copyright Epic Games, Inc. All Rights Reserved.
 #pragma once
 
+// HEADER_UNIT_SKIP - Internal
+
 #include "Chaos/PBDVolumeConstraintBase.h"
 
 namespace Chaos::Softs
diff --git a/Engine/Source/Runtime/Experimental/Chaos/Public/Chaos/PBDVolumeConstraintBase.h b/Engine/Source/Runtime/Experimental/Chaos/Public/Chaos/PBDVolumeConstraintBase.h
--- a/Engine/Source/Runtime/Experimental/Chaos/Public/Chaos/PBDVolumeConstraintBase.h
+++ b/Engine/Source/Runtime/Experimental/Chaos/Public/Chaos/PBDVolumeConstraintBase.h
@@ -1,6 +1,8 @@
 // Copyright Epic Games, Inc. All Rights Reserved.
 #pragma once
 
+// HEADER_UNIT_SKIP - Internal
+
 #include "Chaos/PBDSoftsEvolutionFwd.h"
 #include "Chaos/Array.h"
 #include "Chaos/PBDParticles.h"
diff --git a/Engine/Source/Runtime/Experimental/Chaos/Public/Chaos/PerCellBoundaryConditions.h b/Engine/Source/Runtime/Experimental/Chaos/Public/Chaos/PerCellBoundaryConditions.h
--- a/Engine/Source/Runtime/Experimental/Chaos/Public/Chaos/PerCellBoundaryConditions.h
+++ b/Engine/Source/Runtime/Experimental/Chaos/Public/Chaos/PerCellBoundaryConditions.h
@@ -1,6 +1,8 @@
 // Copyright Epic Games, Inc. All Rights Reserved.
 #pragma once
 
+// HEADER_UNIT_SKIP - Not included directly
+
 #include "Chaos/KinematicGeometryParticles.h"
 #include "Chaos/Transform.h"
 
diff --git a/Engine/Source/Runtime/Experimental/Chaos/Public/Chaos/PerParticleInitForce.h b/Engine/Source/Runtime/Experimental/Chaos/Public/Chaos/PerParticleInitForce.h
--- a/Engine/Source/Runtime/Experimental/Chaos/Public/Chaos/PerParticleInitForce.h
+++ b/Engine/Source/Runtime/Experimental/Chaos/Public/Chaos/PerParticleInitForce.h
@@ -1,6 +1,7 @@
 // Copyright Epic Games, Inc. All Rights Reserved.
 #pragma once
 
+#include "Chaos/ParticleHandle.h"
 #include "Chaos/PerParticleRule.h"
 
 namespace Chaos
diff --git a/Engine/Source/Runtime/Experimental/Chaos/Public/Chaos/PerParticlePBDCCDCollisionConstraint.h b/Engine/Source/Runtime/Experimental/Chaos/Public/Chaos/PerParticlePBDCCDCollisionConstraint.h
--- a/Engine/Source/Runtime/Experimental/Chaos/Public/Chaos/PerParticlePBDCCDCollisionConstraint.h
+++ b/Engine/Source/Runtime/Experimental/Chaos/Public/Chaos/PerParticlePBDCCDCollisionConstraint.h
@@ -1,6 +1,8 @@
 // Copyright Epic Games, Inc. All Rights Reserved.
 #pragma once
 
+// HEADER_UNIT_SKIP - Internal
+
 #include "Chaos/GeometryParticlesfwd.h"
 #include "Chaos/Transform.h"
 #include "Chaos/PBDActiveView.h"
diff --git a/Engine/Source/Runtime/Experimental/Chaos/Public/Chaos/PerParticlePBDShapeConstraints.h b/Engine/Source/Runtime/Experimental/Chaos/Public/Chaos/PerParticlePBDShapeConstraints.h
--- a/Engine/Source/Runtime/Experimental/Chaos/Public/Chaos/PerParticlePBDShapeConstraints.h
+++ b/Engine/Source/Runtime/Experimental/Chaos/Public/Chaos/PerParticlePBDShapeConstraints.h
@@ -1,8 +1,11 @@
 // Copyright Epic Games, Inc. All Rights Reserved.
 #pragma once
 
+// HEADER_UNIT_SKIP - This file does not compile anymore (Base ctor not matching)
+
 #include "Chaos/PBDShapeConstraintsBase.h"
 #include "Chaos/Framework/Parallel.h"
+#include "Chaos/PBDSoftsEvolutionFwd.h"
 
 namespace Chaos::Softs
 {
diff --git a/Engine/Source/Runtime/Experimental/Chaos/Public/Chaos/PerParticlePBDTetConstraints.h b/Engine/Source/Runtime/Experimental/Chaos/Public/Chaos/PerParticlePBDTetConstraints.h
--- a/Engine/Source/Runtime/Experimental/Chaos/Public/Chaos/PerParticlePBDTetConstraints.h
+++ b/Engine/Source/Runtime/Experimental/Chaos/Public/Chaos/PerParticlePBDTetConstraints.h
@@ -1,6 +1,7 @@
 // Copyright Epic Games, Inc. All Rights Reserved.
 #pragma once
 
+#include "Chaos/PBDSoftsSolverParticles.h"
 #include "Chaos/PBDTetConstraintsBase.h"
 
 namespace Chaos::Softs
diff --git a/Engine/Source/Runtime/Experimental/Chaos/Public/Chaos/Segment.h b/Engine/Source/Runtime/Experimental/Chaos/Public/Chaos/Segment.h
--- a/Engine/Source/Runtime/Experimental/Chaos/Public/Chaos/Segment.h
+++ b/Engine/Source/Runtime/Experimental/Chaos/Public/Chaos/Segment.h
@@ -1,6 +1,8 @@
 // Copyright Epic Games, Inc. All Rights Reserved.
 #pragma once
 
+#include "Chaos/Core.h"
+
 namespace Chaos
 {
 	template<class T>
diff --git a/Engine/Source/Runtime/Experimental/Chaos/Public/Chaos/Simplex.h b/Engine/Source/Runtime/Experimental/Chaos/Public/Chaos/Simplex.h
--- a/Engine/Source/Runtime/Experimental/Chaos/Public/Chaos/Simplex.h
+++ b/Engine/Source/Runtime/Experimental/Chaos/Public/Chaos/Simplex.h
@@ -1,6 +1,7 @@
 // Copyright Epic Games, Inc. All Rights Reserved.
 #pragma once
-#include "Chaos/Transform.h"
+
+#include "Chaos/Core.h"
 
 namespace Chaos
 {
diff --git a/Engine/Source/Runtime/Experimental/Chaos/Public/Chaos/SimplexVectorized.h b/Engine/Source/Runtime/Experimental/Chaos/Public/Chaos/SimplexVectorized.h
--- a/Engine/Source/Runtime/Experimental/Chaos/Public/Chaos/SimplexVectorized.h
+++ b/Engine/Source/Runtime/Experimental/Chaos/Public/Chaos/SimplexVectorized.h
@@ -1,6 +1,7 @@
 // Copyright Epic Games, Inc. All Rights Reserved.
 #pragma once
 
+#include "Chaos/Core.h"
 #include "Math/VectorRegister.h"
 #include "Chaos/VectorUtility.h"
 
diff --git a/Engine/Source/Runtime/Experimental/Chaos/Public/Chaos/TriangleRegister.h b/Engine/Source/Runtime/Experimental/Chaos/Public/Chaos/TriangleRegister.h
--- a/Engine/Source/Runtime/Experimental/Chaos/Public/Chaos/TriangleRegister.h
+++ b/Engine/Source/Runtime/Experimental/Chaos/Public/Chaos/TriangleRegister.h
@@ -1,6 +1,10 @@
 // Copyright Epic Games, Inc. All Rights Reserved.
 
 #pragma once
+
+// HEADER_UNIT_SKIP - Not included directly
+
+#include "Chaos/Core.h"
 #include "Math/VectorRegister.h"
 
 namespace Chaos
diff --git a/Engine/Source/Runtime/Experimental/Chaos/Public/Chaos/VectorUtility.h b/Engine/Source/Runtime/Experimental/Chaos/Public/Chaos/VectorUtility.h
--- a/Engine/Source/Runtime/Experimental/Chaos/Public/Chaos/VectorUtility.h
+++ b/Engine/Source/Runtime/Experimental/Chaos/Public/Chaos/VectorUtility.h
@@ -1,6 +1,7 @@
 // Copyright Epic Games, Inc. All Rights Reserved.
 #pragma once
 
+#include "Math/VectorRegister.h"
 
 /**
  * Cast VectorRegister4Int in VectorRegister4Float
diff --git a/Engine/Source/Runtime/Experimental/Chaos/Public/Chaos/XPBDAxialSpringConstraints.h b/Engine/Source/Runtime/Experimental/Chaos/Public/Chaos/XPBDAxialSpringConstraints.h
--- a/Engine/Source/Runtime/Experimental/Chaos/Public/Chaos/XPBDAxialSpringConstraints.h
+++ b/Engine/Source/Runtime/Experimental/Chaos/Public/Chaos/XPBDAxialSpringConstraints.h
@@ -76,11 +76,11 @@ public:
 				{
 					Particles.P(i1) -= Multiplier * Particles.InvM(i1) * Delta;
 				}
-				if (Particles.InvM(i2))
+				if (Particles.InvM(i2) != 0)
 				{
 					Particles.P(i2) += Multiplier * Particles.InvM(i2) * Barys[ConstraintIndex] * Delta;
 				}
-				if (Particles.InvM(i3))
+				if (Particles.InvM(i3) != 0)
 				{
 					Particles.P(i3) += Multiplier * Particles.InvM(i3) * ((FSolverReal)1. - Barys[ConstraintIndex]) * Delta;
 				}
@@ -101,11 +101,11 @@ public:
 				{
 					Particles.P(i1) -= Multiplier * Particles.InvM(i1) * Delta;
 				}
-				if (Particles.InvM(i2))
+				if (Particles.InvM(i2) != 0)
 				{
 					Particles.P(i2) += Multiplier * Particles.InvM(i2) * Barys[ConstraintIndex] * Delta;
 				}
-				if (Particles.InvM(i3))
+				if (Particles.InvM(i3) != 0)
 				{
 					Particles.P(i3) += Multiplier * Particles.InvM(i3) * ((FSolverReal)1. - Barys[ConstraintIndex]) * Delta;
 				}
diff --git a/Engine/Source/Runtime/Experimental/Chaos/Public/Chaos/XPBDCorotatedFiberConstraints.h b/Engine/Source/Runtime/Experimental/Chaos/Public/Chaos/XPBDCorotatedFiberConstraints.h
--- a/Engine/Source/Runtime/Experimental/Chaos/Public/Chaos/XPBDCorotatedFiberConstraints.h
+++ b/Engine/Source/Runtime/Experimental/Chaos/Public/Chaos/XPBDCorotatedFiberConstraints.h
@@ -2,6 +2,7 @@
 #pragma once
 
 #include "Chaos/PBDSpringConstraintsBase.h"
+#include "Chaos/XPBDCorotatedConstraints.h"
 #include "ChaosStats.h"
 #include "Chaos/ImplicitQRSVD.h"
 #include "Chaos/GraphColoring.h"
diff --git a/Engine/Source/Runtime/Experimental/Chaos/Public/Chaos/XPBDGridBasedCorotatedConstraints.h b/Engine/Source/Runtime/Experimental/Chaos/Public/Chaos/XPBDGridBasedCorotatedConstraints.h
--- a/Engine/Source/Runtime/Experimental/Chaos/Public/Chaos/XPBDGridBasedCorotatedConstraints.h
+++ b/Engine/Source/Runtime/Experimental/Chaos/Public/Chaos/XPBDGridBasedCorotatedConstraints.h
@@ -1,6 +1,8 @@
 // Copyright Epic Games, Inc. All Rights Reserved.
 #pragma once
 
+// HEADER_UNIT_SKIP - Not included directly
+
 #include "Chaos/PBDSpringConstraintsBase.h"
 #include "ChaosStats.h"
 #include "Chaos/ImplicitQRSVD.h"
diff --git a/Engine/Source/Runtime/Experimental/Chaos/Public/Chaos/XPBDVolumeConstraints.h b/Engine/Source/Runtime/Experimental/Chaos/Public/Chaos/XPBDVolumeConstraints.h
--- a/Engine/Source/Runtime/Experimental/Chaos/Public/Chaos/XPBDVolumeConstraints.h
+++ b/Engine/Source/Runtime/Experimental/Chaos/Public/Chaos/XPBDVolumeConstraints.h
@@ -1,6 +1,8 @@
 // Copyright Epic Games, Inc. All Rights Reserved.
 #pragma once
 
+// HEADER_UNIT_SKIP - Not included directly
+
 #include "Chaos/PBDSpringConstraintsBase.h"
 #include "ChaosStats.h"
 #include "Chaos/ImplicitQRSVD.h"
diff --git a/Engine/Source/Runtime/Experimental/Chaos/Public/GeometryCollection/ManagedArrayTypes.h b/Engine/Source/Runtime/Experimental/Chaos/Public/GeometryCollection/ManagedArrayTypes.h
--- a/Engine/Source/Runtime/Experimental/Chaos/Public/GeometryCollection/ManagedArrayTypes.h
+++ b/Engine/Source/Runtime/Experimental/Chaos/Public/GeometryCollection/ManagedArrayTypes.h
@@ -3,6 +3,7 @@
 #include "Containers/UnrealString.h"
 #include "GeometryCollection/GeometryCollectionBoneNode.h"
 #include "GeometryCollection/GeometryCollectionSection.h"
+#include "GeometryCollection/ManagedArray.h"
 #include "Math/Color.h"
 #include "Math/IntVector.h"
 #include "Math/Quat.h"
diff --git a/Engine/Source/Runtime/Experimental/ChaosCore/Public/Chaos/Core.h b/Engine/Source/Runtime/Experimental/ChaosCore/Public/Chaos/Core.h
--- a/Engine/Source/Runtime/Experimental/ChaosCore/Public/Chaos/Core.h
+++ b/Engine/Source/Runtime/Experimental/ChaosCore/Public/Chaos/Core.h
@@ -2,6 +2,8 @@
 
 #pragma once
 
+// HEADER_UNIT_SKIP - Bad include. Some headers are in Chaos while this is in ChaosCore
+
 #include "Chaos/Real.h"
 #include "Chaos/Vector.h"
 #include "Chaos/Matrix.h"
diff --git a/Engine/Source/Runtime/Experimental/ChaosCore/Public/Chaos/ImplicitQRSVD.h b/Engine/Source/Runtime/Experimental/ChaosCore/Public/Chaos/ImplicitQRSVD.h
--- a/Engine/Source/Runtime/Experimental/ChaosCore/Public/Chaos/ImplicitQRSVD.h
+++ b/Engine/Source/Runtime/Experimental/ChaosCore/Public/Chaos/ImplicitQRSVD.h
@@ -1,6 +1,8 @@
 // Copyright Epic Games, Inc. All Rights Reserved.
 #pragma once
 
+// HEADER_UNIT_SKIP - Bad include VectorTypes (living in GeometryCore)
+
 #include "Chaos/Matrix.h"
 #include "VectorTypes.h"
 #include "Chaos/Vector.h"
