Revert temporary workaround

diff --git a/Engine/Source/Runtime/Experimental/Chaos/Public/Chaos/GJK.h b/Engine/Source/Runtime/Experimental/Chaos/Public/Chaos/GJK.h
--- a/Engine/Source/Runtime/Experimental/Chaos/Public/Chaos/GJK.h
+++ b/Engine/Source/Runtime/Experimental/Chaos/Public/Chaos/GJK.h
@@ -1204,7 +1204,7 @@ namespace Chaos
 		const T SeparatedDistance = ThicknessA + ThicknessB + Epsilon;
 		while (!bIsContact && !bIsDegenerate)
 		{
-			if (NumIterations++ >= 32)	//todo: temporary workaround for query crash
+			if (!ensure(NumIterations++ < 32))	//todo: take this out
 			{
 				break;	//if taking too long just stop. This should never happen
 			}
