Constraints: Add nullptr check when looking for constraint in manager to avoid render crash.

diff --git a/Engine/Source/Runtime/Experimental/Animation/Constraints/Private/ConstraintsManager.cpp b/Engine/Source/Runtime/Experimental/Animation/Constraints/Private/ConstraintsManager.cpp
--- a/Engine/Source/Runtime/Experimental/Animation/Constraints/Private/ConstraintsManager.cpp
+++ b/Engine/Source/Runtime/Experimental/Animation/Constraints/Private/ConstraintsManager.cpp
@@ -484,7 +484,7 @@ int32 FConstraintsManagerController::GetConstraintIndex(const FName& InConstrain
 	
 	return Manager->Constraints.IndexOfByPredicate([InConstraintName](const TObjectPtr<UTickableConstraint>& Constraint)
 	{
-		return 	Constraint->GetFName() == InConstraintName;
+		return 	(Constraint && Constraint->GetFName() == InConstraintName);
 	} );
 }
 	
