Iris
-  Implemented support to allow descriptors built for structs derived from FastArraySerializer which contains additional properties to be handled by using an explicit ReplicationFragment.
- Minor cleanup in FastArrayReplicationFragment, renaming some members and methods to clarify things a bit.

(1/2)

diff --git a/Engine/Plugins/Runtime/ReplicationSystemTestPlugin/Source/Private/Tests/ReplicationState/TestFastArrayReplicationState.cpp b/Engine/Plugins/Runtime/ReplicationSystemTestPlugin/Source/Private/Tests/ReplicationState/TestFastArrayReplicationState.cpp
--- a/Engine/Plugins/Runtime/ReplicationSystemTestPlugin/Source/Private/Tests/ReplicationState/TestFastArrayReplicationState.cpp
+++ b/Engine/Plugins/Runtime/ReplicationSystemTestPlugin/Source/Private/Tests/ReplicationState/TestFastArrayReplicationState.cpp
@@ -14,6 +14,7 @@
 #include "Misc/EnumClassFlags.h"
 #include "Net/UnrealNetwork.h"
 #include "Iris/ReplicationSystem/Filtering/NetObjectFilter.h"
+#include "Iris/ReplicationSystem/FastArrayReplicationFragment.h"
 
 UTestFastArrayReplicationState_FastArray_TestClassFastArray::UTestFastArrayReplicationState_FastArray_TestClassFastArray()
 : UReplicatedTestObject()
@@ -39,14 +40,80 @@ UTestFastArrayReplicationState_FastArray_TestClassFastArrayWithExtraProperty::UT
 {
 }
 
+////////////////////////////////////////////////////////////////////////////////////////////////////////////////////////////////////////////////
+// Note: this is a test for something that should not be done. This is purely to avoid refactoring existing code before full transition to iris.
+// A FastArraySerializer struct should not contain any additional data.
+////////////////////////////////////////////////////////////////////////////////////////////////////////////////////////////////////////////////
+class FTestFastArrayReplicationState_FastArrayWithExtraProperty::FastArrayWithExtraPropertiesReplicationFragment : public UE::Net::TFastArrayReplicationFragment<typename TFastArrayTypeHelper<FTestFastArrayReplicationState_FastArrayWithExtraProperty>::FastArrayItemType, FTestFastArrayReplicationState_FastArrayWithExtraProperty>
+{
+public:
+	typedef TFastArrayReplicationFragment<typename TFastArrayTypeHelper<FTestFastArrayReplicationState_FastArrayWithExtraProperty>::FastArrayItemType, FTestFastArrayReplicationState_FastArrayWithExtraProperty> SuperT;
+
+	FastArrayWithExtraPropertiesReplicationFragment(UE::Net::EReplicationFragmentTraits InTraits, UObject* InOwner, const UE::Net::FReplicationStateDescriptor* InDescriptor)
+	: SuperT(InTraits, InOwner, InDescriptor, EAllowAdditionalPropertiesType::AllowAdditionalProperties)
+	{
+	}
+
+	virtual void ApplyReplicatedState(UE::Net::FReplicationStateApplyContext& Context) const override
+	{
+		// Get the wrapped FastArraySerializer and array
+		FTestFastArrayReplicationState_FastArrayWithExtraProperty* DstArraySerializer = this->GetFastArraySerializerFromOwner();
+
+		// Intentionally not const as we allow the src state to be modified
+		FTestFastArrayReplicationState_FastArrayWithExtraProperty* SrcArraySerializer = this->GetFastArraySerializerFromApplyContext(Context);
+
+		// Apply additional state
+		DstArraySerializer->ExtraInt = SrcArraySerializer->ExtraInt;
+
+		// Forward to fast array serializer
+		SuperT::ApplyReplicatedState(Context);
+	}
+
+	virtual bool PollReplicatedState(UE::Net::EReplicationFragmentPollFlags PollOption) override
+	{
+		// Lookup source data, we need the actual FastArraySerializer and the Array it is wrapping
+		FTestFastArrayReplicationState_FastArrayWithExtraProperty* SrcArraySerializer = this->GetFastArraySerializerFromOwner();
+	
+		// Lookup destination data
+		FTestFastArrayReplicationState_FastArrayWithExtraProperty* DstArraySerializer = this->GetFastArraySerializerFromReplicationState();
+
+		// Explicit poll of additional properties in the fastarray and mark dirty if needed
+		if (SrcArraySerializer->ExtraInt != DstArraySerializer->ExtraInt)
+		{	
+			this->MarkDirty();
+		}
+		DstArraySerializer->ExtraInt = SrcArraySerializer->ExtraInt;
+
+		return SuperT::PollReplicatedState(PollOption);
+	}
+};
+
+UE::Net::FReplicationFragment* FTestFastArrayReplicationState_FastArrayWithExtraProperty::CreateAndRegisterReplicationFragment(UObject* Owner, const UE::Net::FReplicationStateDescriptor* Descriptor, UE::Net::FFragmentRegistrationContext& Context)
+{
+	using namespace UE::Net;
+	static_assert(TFastArrayTypeHelper<FTestFastArrayReplicationState_FastArrayWithExtraProperty>::HasValidFastArrayItemType(), "Invalid FastArrayItemType detected. Make sure that FastArraySerializer has a single replicated dynamic array");
+
+	if (FTestFastArrayReplicationState_FastArrayWithExtraProperty::FastArrayWithExtraPropertiesReplicationFragment* Fragment = new FTestFastArrayReplicationState_FastArrayWithExtraProperty::FastArrayWithExtraPropertiesReplicationFragment(Context.GetFragmentTraits(), Owner, Descriptor))
+	{
+		Fragment->Register(Context, EReplicationFragmentTraits::DeleteWithInstanceProtocol);
+		return Fragment;
+	}
+	return nullptr;
+}
+
+//////////////////////////////////////////////////////////////////////////
+
 void UTestFastArrayReplicationState_FastArray_TestClassFastArrayWithExtraProperty::GetLifetimeReplicatedProps(TArray<class FLifetimeProperty>& OutLifetimeProps) const
 {
-	DOREPLIFETIME(UTestFastArrayReplicationState_FastArray_TestClassFastArrayWithExtraProperty, FastArray);
+	FDoRepLifetimeParams Params;
+	Params.CreateAndRegisterReplicationFragmentFunction = &FTestFastArrayReplicationState_FastArrayWithExtraProperty::CreateAndRegisterReplicationFragment;
+	DOREPLIFETIME_WITH_PARAMS_FAST(UTestFastArrayReplicationState_FastArray_TestClassFastArrayWithExtraProperty, FastArray, Params);
 }
 
 void UTestFastArrayReplicationState_FastArray_TestClassFastArrayWithExtraProperty::RegisterReplicationFragments(UE::Net::FFragmentRegistrationContext& Context, UE::Net::EFragmentRegistrationFlags RegistrationFlags)
 {
-	Super::RegisterReplicationFragments(Context, RegistrationFlags);
+	// Explicitly allow building descriptors for FastArray with additional properties to be allow to test this special case.
+	UE::Net::FReplicationFragmentUtil::CreateAndRegisterFragmentsForObject(this, Context, RegistrationFlags | UE::Net::EFragmentRegistrationFlags::AllowFastArraysWithAdditionalProperties);
 }
 
 void FTestFastArrayReplicationState_FastArrayItem::PostReplicatedAdd(const struct FTestFastArrayReplicationState_FastArraySerializer& InArraySerializer)
@@ -119,12 +186,11 @@ UE_NET_TEST_FIXTURE(FTestFastArrayReplicationStateContext, BuildClassDescriptorF
 
 UE_NET_TEST_FIXTURE(FTestFastArrayReplicationStateContext, BuildClassDescriptorForSinglePropertyForFastArrayWithExtraProperty)
 {
+	// We expect this test to fail as we do not pass the flag allowing fastarrays to contain extra properties
 	InitDescriptor(UTestFastArrayReplicationState_FastArray_TestClassFastArrayWithExtraProperty::StaticClass(), (int32)UTestFastArrayReplicationState_FastArray_TestClassFastArrayWithExtraProperty::ENetFields_Private::FastArray);
-
 	UE_NET_ASSERT_FALSE(EnumHasAnyFlags(Descriptor->Traits, EReplicationStateTraits::IsFastArrayReplicationState));	
 }
 
-
 UE_NET_TEST_FIXTURE(FReplicationSystemServerClientTestFixture, TestAddItem)
 {
 	UTestFastArrayReplicationState_FastArray_TestClassFastArray* ServerObject = Server->CreateObject<UTestFastArrayReplicationState_FastArray_TestClassFastArray>();
@@ -801,4 +867,71 @@ UE_NET_TEST_FIXTURE(FReplicationSystemServerClientTestFixture, TestPostReplicate
 	UE_NET_ASSERT_FALSE(ClientFastArray.bPostReplicatedReceiveWasHitWithUnresolvedReferences);
 }
 
+UE_NET_TEST_FIXTURE(FReplicationSystemServerClientTestFixture, TestFastArrayWithExtraProperty)
+{
+	// Add a client
+	FReplicationSystemTestClient* Client = CreateClient();
+
+	UTestFastArrayReplicationState_FastArray_TestClassFastArrayWithExtraProperty* ServerObject = Server->CreateObject<UTestFastArrayReplicationState_FastArray_TestClassFastArrayWithExtraProperty>();
+
+
+	// Send and deliver packet
+	Server->PreSendUpdate();
+	Server->SendAndDeliverTo(Client, true);
+	Server->PostSendUpdate();
+
+	// Store Pointer to objects
+	UTestFastArrayReplicationState_FastArray_TestClassFastArrayWithExtraProperty* ClientObject = Cast<UTestFastArrayReplicationState_FastArray_TestClassFastArrayWithExtraProperty>(Client->GetReplicationBridge()->GetReplicatedObject(ServerObject->NetRefHandle));
+
+	// Verify initial replication
+	UE_NET_ASSERT_TRUE(ClientObject!= nullptr);
+	UE_NET_ASSERT_FALSE(ClientObject->FastArray.bHitReplicatedAdd);
+	UE_NET_ASSERT_FALSE(ClientObject->FastArray.bHitReplicatedChange);
+	UE_NET_ASSERT_FALSE(ClientObject->FastArray.bHitReplicatedRemove);
+
+	// Verify that we received the data
+	UE_NET_ASSERT_EQ(ServerObject->FastArray.GetItemArray().Num(), ClientObject->FastArray.GetItemArray().Num());
+	UE_NET_ASSERT_EQ(ServerObject->FastArray.ExtraInt, ClientObject->FastArray.ExtraInt);
+
+	// Add data to array
+	FTestFastArrayReplicationState_FastArrayItem NewEntry;
+	NewEntry.bRepBool = false;
+	NewEntry.RepInt32 = 0U;
+	
+	ServerObject->FastArray.Edit().Add(NewEntry);
+
+	NewEntry.bRepBool = true;
+	NewEntry.RepInt32 = 1U;
+	
+	ServerObject->FastArray.Edit().Add(NewEntry);
+
+	// Send and deliver packet
+	Server->PreSendUpdate();
+	Server->SendAndDeliverTo(Client, true);
+	Server->PostSendUpdate();
+
+	// Verify that we received the data
+	UE_NET_ASSERT_EQ(ServerObject->FastArray.GetItemArray().Num(), ClientObject->FastArray.GetItemArray().Num());
+	UE_NET_ASSERT_EQ(ServerObject->FastArray.ExtraInt, ClientObject->FastArray.ExtraInt);
+
+	// In this case we expect order to be the same
+	UE_NET_ASSERT_EQ(ServerObject->FastArray.GetItemArray()[0].RepInt32, ClientObject->FastArray.GetItemArray()[0].RepInt32);
+
+	UE_NET_ASSERT_TRUE(ClientObject->FastArray.bHitReplicatedAdd);
+	UE_NET_ASSERT_FALSE(ClientObject->FastArray.bHitReplicatedChange);
+	UE_NET_ASSERT_FALSE(ClientObject->FastArray.bHitReplicatedRemove);
+
+	// Modify extraint
+	ServerObject->FastArray.ExtraInt = 3;
+
+	// Send and deliver packet
+	Server->PreSendUpdate();
+	Server->SendAndDeliverTo(Client, true);
+	Server->PostSendUpdate();
+
+	// Verify that we received the data
+	UE_NET_ASSERT_EQ(ServerObject->FastArray.ExtraInt, ClientObject->FastArray.ExtraInt);
+}
+
+
 }
diff --git a/Engine/Plugins/Runtime/ReplicationSystemTestPlugin/Source/Private/Tests/ReplicationState/TestFastArrayReplicationState.h b/Engine/Plugins/Runtime/ReplicationSystemTestPlugin/Source/Private/Tests/ReplicationState/TestFastArrayReplicationState.h
--- a/Engine/Plugins/Runtime/ReplicationSystemTestPlugin/Source/Private/Tests/ReplicationState/TestFastArrayReplicationState.h
+++ b/Engine/Plugins/Runtime/ReplicationSystemTestPlugin/Source/Private/Tests/ReplicationState/TestFastArrayReplicationState.h
@@ -119,6 +119,11 @@ struct FTestFastArrayReplicationState_FastArrayWithExtraProperty : public FTestF
 	{
 		return FFastArraySerializer::FastArrayDeltaSerialize<FTestFastArrayReplicationState_FastArrayItem, FTestFastArrayReplicationState_FastArrayWithExtraProperty>(Items, DeltaParms, *this);
 	}
+
+	// Note: !!This is not an example of something that should be used as we want to remove the need for users to declare custom types for FastArray-style replication
+	// As we do not really support replication of FastArrays with additional replicated properties in the derived FastArray-struct we use a custom replication fragment to poll and apply the additional replicated property
+	class FastArrayWithExtraPropertiesReplicationFragment;
+	static UE::Net::FReplicationFragment* CreateAndRegisterReplicationFragment(UObject* Owner, const UE::Net::FReplicationStateDescriptor* Descriptor, UE::Net::FFragmentRegistrationContext& Context);
 };
 
 template<>
@@ -152,6 +157,7 @@ public:
 
 	virtual void RegisterReplicationFragments(UE::Net::FFragmentRegistrationContext& Fragments, UE::Net::EFragmentRegistrationFlags RegistrationFlags) override;
 
+public:
 	UPROPERTY(Replicated)
 	FTestFastArrayReplicationState_FastArrayWithExtraProperty FastArray;
 };
@@ -159,3 +165,4 @@ public:
 
 
 
+
diff --git a/Engine/Source/Runtime/Engine/Public/Net/UnrealNetwork.h b/Engine/Source/Runtime/Engine/Public/Net/UnrealNetwork.h
--- a/Engine/Source/Runtime/Engine/Public/Net/UnrealNetwork.h
+++ b/Engine/Source/Runtime/Engine/Public/Net/UnrealNetwork.h
@@ -368,6 +368,12 @@ struct CGetFastArrayCreateReplicationFragmentFuncable
 template<typename T>
 inline typename TEnableIf<TModels<CGetFastArrayCreateReplicationFragmentFuncable, T>::Value, const FDoRepLifetimeParams>::Type FixupParams(const FDoRepLifetimeParams& Params)
 {
+	// Use passed in CreateAndRegisterReplicationFragmentFunction if set
+	if (Params.CreateAndRegisterReplicationFragmentFunction)
+	{
+		return Params;
+	}
+
 	FDoRepLifetimeParams NewParams(Params);
 	NewParams.CreateAndRegisterReplicationFragmentFunction = T::GetFastArrayCreateReplicationFragmentFunction();
 	return NewParams;
diff --git a/Engine/Source/Runtime/Experimental/Iris/Core/Private/Iris/ReplicationState/ReplicationStateDescriptorBuilder.cpp b/Engine/Source/Runtime/Experimental/Iris/Core/Private/Iris/ReplicationState/ReplicationStateDescriptorBuilder.cpp
--- a/Engine/Source/Runtime/Experimental/Iris/Core/Private/Iris/ReplicationState/ReplicationStateDescriptorBuilder.cpp
+++ b/Engine/Source/Runtime/Experimental/Iris/Core/Private/Iris/ReplicationState/ReplicationStateDescriptorBuilder.cpp
@@ -74,7 +74,8 @@ enum class EMemberPropertyTraits : uint32
 	IsSourceTriviallyDestructible		= IsSourceTriviallyConstructible << 1U,
 	IsFastArray							= IsSourceTriviallyDestructible << 1U,
 	IsNativeFastArray					= IsFastArray << 1U,
-	IsFastArrayItem						= IsNativeFastArray << 1U,
+	IsFastArrayWithExtraProperties		= IsNativeFastArray << 1U,
+	IsFastArrayItem						= IsFastArrayWithExtraProperties << 1U,
 	IsInvalidFastArray					= IsFastArrayItem << 1U,
 	HasConnectionSpecificSerialization	= IsInvalidFastArray << 1U,
 	HasPushBasedDirtiness				= HasConnectionSpecificSerialization << 1U,
@@ -126,8 +127,15 @@ public:
 	};
 
 	// Utility methods
-	static void GetSerializerTraits(FMemberProperty& OutMemberProperty, const FProperty* Property, const FPropertyNetSerializerInfo* Info);
-	static bool IsSupportedProperty(FMemberProperty& OutMemberProperty, const FProperty* Property, const TArray<FLifetimeProperty>* LifeTimeProperties = nullptr, UClass* InObjectClass = nullptr);
+	static void GetSerializerTraits(FMemberProperty& OutMemberProperty, const FProperty* Property, const FPropertyNetSerializerInfo* Info, bool bAllowFastArrayWithExtraProperties = false);
+	struct FIsSupportedPropertyParams
+	{
+		const TArray<FLifetimeProperty>* LifeTimeProperties = nullptr;
+		UClass* InObjectClass = nullptr;
+		bool bAllowFastArrayWithExtraReplicatedProperties = false;
+	};
+	static bool IsSupportedProperty(FMemberProperty& OutMemberProperty, const FProperty* Property, const FIsSupportedPropertyParams& Params);
+	static bool IsSupportedProperty(FMemberProperty& OutMemberProperty, const FProperty* Property) { FIsSupportedPropertyParams Params; return IsSupportedProperty(OutMemberProperty, Property, Params); }
 	static bool IsSupportedStructWithCustomSerializer(FMemberProperty& OutMemberProperty, const UStruct* InStruct);	
 	static bool IsStructWithCustomSerializer(const UStruct* InStruct);
 
@@ -135,11 +143,10 @@ public:
 	
 	static EMemberPropertyTraits GetConnectionFilterTrait(ELifetimeCondition Condition);
 	static EMemberPropertyTraits GetInitOnlyTrait(ELifetimeCondition Condition);
-	static EMemberPropertyTraits GetFastArrayPropertyTraits(const FNetSerializer* NetSerializer, const FProperty* Property);
+	static EMemberPropertyTraits GetFastArrayPropertyTraits(const FNetSerializer* NetSerializer, const FProperty* Property, bool bAllowFastArrayWithExtraProperties);
 	static SIZE_T GetFastArrayChangeMaskOffset(const FProperty* Property);
 	static EMemberPropertyTraits GetHasObjectReferenceTraits(const FNetSerializer* NetSerializer);
 	static void GetIrisPropertyTraits(FMemberProperty& OutMemberProperty, const FProperty* Property, const TArray<FLifetimeProperty>* LifeTimeProperties, UClass* ObjectClass);
-	static bool IsFastArraySupported(const UStruct* Struct);
 
 public:
 	enum class EDescriptorType : uint32
@@ -324,6 +331,7 @@ TRefCountPtr<const FReplicationStateDescriptor> FPropertyReplicationStateDescrip
 	FReplicationStateDescriptorBuilder::FParameters Params;
 	Params.DescriptorRegistry = DescriptorRegistry;
 	Params.EnableFastArrayHandling = EnumHasAnyFlags(MemberProperty.Traits, EMemberPropertyTraits::IsFastArray);
+	Params.AllowFastArrayWithExtraReplicatedProperties = EnumHasAnyFlags(MemberProperty.Traits, EMemberPropertyTraits::IsFastArrayWithExtraProperties);
 
 	return FReplicationStateDescriptorBuilder::CreateDescriptorForStruct(Struct, Params);
 }
@@ -1826,7 +1834,7 @@ void FPropertyReplicationStateDescriptorBuilder::GetIrisPropertyTraits(FMemberPr
 	OutMemberProperty.CreateAndRegisterReplicationFragmentFunction = CreateAndRegisterReplicationFragmentFunction;
 }
 
-void FPropertyReplicationStateDescriptorBuilder::GetSerializerTraits(FMemberProperty& OutMemberProperty, const FProperty* Property, const FPropertyNetSerializerInfo* NetSerializerInfo)
+void FPropertyReplicationStateDescriptorBuilder::GetSerializerTraits(FMemberProperty& OutMemberProperty, const FProperty* Property, const FPropertyNetSerializerInfo* NetSerializerInfo, bool bAllowFastArrayWithExtraProperties)
 {
 	const FNetSerializer* NetSerializer = NetSerializerInfo->GetNetSerializer(Property);
 
@@ -1834,10 +1842,10 @@ void FPropertyReplicationStateDescriptorBuilder::GetSerializerTraits(FMemberProp
 	OutMemberProperty.Traits |= EnumHasAnyFlags(NetSerializer->Traits, ENetSerializerTraits::HasDynamicState) ? EMemberPropertyTraits::HasDynamicState : EMemberPropertyTraits::None;
 	OutMemberProperty.Traits |= EnumHasAnyFlags(NetSerializer->Traits, ENetSerializerTraits::HasConnectionSpecificSerialization) ? EMemberPropertyTraits::HasConnectionSpecificSerialization : EMemberPropertyTraits::None;
 	OutMemberProperty.Traits |= EnumHasAnyFlags(NetSerializer->Traits, ENetSerializerTraits::UseSerializerIsEqual) ? EMemberPropertyTraits::UseSerializerIsEqual : EMemberPropertyTraits::None;
-	OutMemberProperty.Traits |= GetFastArrayPropertyTraits(NetSerializer, Property);
+	OutMemberProperty.Traits |= GetFastArrayPropertyTraits(NetSerializer, Property, bAllowFastArrayWithExtraProperties);
 }
 
-bool FPropertyReplicationStateDescriptorBuilder::IsSupportedProperty(FMemberProperty& OutMemberProperty, const FProperty* Property, const TArray<FLifetimeProperty>* LifeTimeProperties, UClass* InObjectClass)
+bool FPropertyReplicationStateDescriptorBuilder::IsSupportedProperty(FMemberProperty& OutMemberProperty, const FProperty* Property, const FIsSupportedPropertyParams& Parameters)
 {
 	// A property is supported if it has a registered SerializerInfo.
 	const FPropertyNetSerializerInfo* NetSerializerInfo = FPropertyNetSerializerInfoRegistry::FindSerializerInfo(Property);
@@ -1853,8 +1861,8 @@ bool FPropertyReplicationStateDescriptorBuilder::IsSupportedProperty(FMemberProp
 	OutMemberProperty.ExternalSizeAndAlignment = { 0, 0 };
 
 	// Init traits
-	GetIrisPropertyTraits(OutMemberProperty, Property, LifeTimeProperties, InObjectClass);
-	GetSerializerTraits(OutMemberProperty, Property, NetSerializerInfo);
+	GetIrisPropertyTraits(OutMemberProperty, Property, Parameters. LifeTimeProperties, Parameters.InObjectClass);
+	GetSerializerTraits(OutMemberProperty, Property, NetSerializerInfo, Parameters.bAllowFastArrayWithExtraReplicatedProperties);
 
 	// $IRIS: $TODO: Add property attribute to control how many changemask bits we want to use per property
 	if (EnumHasAnyFlags(OutMemberProperty.Traits, EMemberPropertyTraits::IsFastArray))
@@ -2025,25 +2033,81 @@ SIZE_T FPropertyReplicationStateDescriptorBuilder::GetFastArrayChangeMaskOffset(
 	return 0U;
 }
 
-EMemberPropertyTraits FPropertyReplicationStateDescriptorBuilder::GetFastArrayPropertyTraits(const FNetSerializer* NetSerializer, const FProperty* Property)
+EMemberPropertyTraits FPropertyReplicationStateDescriptorBuilder::GetFastArrayPropertyTraits(const FNetSerializer* NetSerializer, const FProperty* InProperty, bool bAllowFastArrayWithExtraProperties)
 {
 	if (IsStructNetSerializer(NetSerializer))
 	{
-		if (const FStructProperty* StructProperty = CastField<FStructProperty>(Property))
+		if (const FStructProperty* StructProperty = CastField<FStructProperty>(InProperty))
 		{
+			// Find traits for the FastArraySerializer
 			const UScriptStruct* Struct = StructProperty->Struct;
-			const bool bUseNativeFastArray = CVarIrisUseNativeFastArray.GetValueOnAnyThread() != 0;
-
-			if (bUseNativeFastArray && Struct->IsChildOf(FIrisFastArraySerializer::StaticStruct()))
-			{
-				return IsFastArraySupported(Struct) ? EMemberPropertyTraits::HasPushBasedDirtiness | EMemberPropertyTraits::IsFastArray | EMemberPropertyTraits::IsNativeFastArray : EMemberPropertyTraits::IsInvalidFastArray;
-			}
-			else if (Struct->IsChildOf(FFastArraySerializer::StaticStruct()))
+			if (Struct->IsChildOf(FFastArraySerializer::StaticStruct()))
 			{
-				return IsFastArraySupported(Struct) ? EMemberPropertyTraits::IsFastArray : EMemberPropertyTraits::IsInvalidFastArray;
+				FMemberProperty MemberProperty;
+				if (IsStructWithCustomSerializer(Struct))
+				{
+					return EMemberPropertyTraits::IsInvalidFastArray;
+				}
+
+				// Try to figure out if this is a FastArraySerializer we can support. It is assumed the struct has already been determined to inherit from FastArraySerializer.
+				// A proper FastArraySerializer should contain a single TArray property with elements that are a derived from FastArraySerializerItem.
+				uint32 ReplicatedPropertyCount = 0;
+				bool bHasFastArrayItem = false;
+	
+				for (TFieldIterator<FProperty> It(Struct, EFieldIteratorFlags::IncludeSuper); It; ++It)
+				{
+					const FProperty* CurrentProperty = *It;
+					if (EnumHasAnyFlags(CurrentProperty->PropertyFlags, EPropertyFlags::CPF_RepSkip))
+					{
+						continue;
+					}
+
+					++ReplicatedPropertyCount;
+
+					if (!FPropertyReplicationStateDescriptorBuilder::IsSupportedProperty(MemberProperty, CurrentProperty))
+					{
+						return EMemberPropertyTraits::IsInvalidFastArray;
+					}
+
+					if (EnumHasAnyFlags(MemberProperty.Traits, EMemberPropertyTraits::IsFastArrayItem))
+					{
+						bHasFastArrayItem = true;
+					}
+
+					// If we do not allow additional properties, report this as an invalid FastArray
+					if (!bAllowFastArrayWithExtraProperties)
+					{
+						if ((ReplicatedPropertyCount > 1) || (CastField<FArrayProperty>(MemberProperty.Property) == nullptr))
+						{
+							return EMemberPropertyTraits::IsInvalidFastArray;
+						}
+					}
+				}
+
+				if (!bHasFastArrayItem)
+				{
+					return EMemberPropertyTraits::IsInvalidFastArray;
+				}
+
+				// We have a valid FastArray, setup traits
+				EMemberPropertyTraits Traits = EMemberPropertyTraits::IsFastArray;
+
+				// If the struct is derived from FIrisFastArraySerializer and has a single property we can treat it as a native FastArray
+				const bool bUseNativeFastArray = CVarIrisUseNativeFastArray.GetValueOnAnyThread() != 0 && ReplicatedPropertyCount == 1U;
+				if (bUseNativeFastArray && Struct->IsChildOf(FIrisFastArraySerializer::StaticStruct()))
+				{
+					Traits |= EMemberPropertyTraits::HasPushBasedDirtiness | EMemberPropertyTraits::IsNativeFastArray;
+				}
+				else if (ReplicatedPropertyCount > 1U)
+				{
+					Traits |= EMemberPropertyTraits::IsFastArrayWithExtraProperties;
+				}
+
+				return Traits;
 			}
 			else if (Struct->IsChildOf(FFastArraySerializerItem::StaticStruct()))
 			{
+				// Invalidate FastArrayItem if it has a custom NetSerializer which we currently do not support, the workaround is to wrap the struct with the custom netserializer in a struct
 				if (ensureAlwaysMsgf(!IsStructWithCustomSerializer(Struct), 
 					TEXT("FPropertyReplicationStateDescriptorBuilder Iris does not support custom NetSerializers for FastArrayItems %s, if required use a property in the struct wrapping the custom NetSerializer"), *Struct->GetName())
 				)
@@ -2055,7 +2119,8 @@ EMemberPropertyTraits FPropertyReplicationStateDescriptorBuilder::GetFastArrayPr
 	}
 	else if (IsArrayPropertyNetSerializer(NetSerializer))
 	{
-		if (const FArrayProperty* ArrayProperty = CastField<FArrayProperty>(Property))
+		// If the property is an array, unwrap it and check if it is a FFastArraySerializerItem
+		if (const FArrayProperty* ArrayProperty = CastField<FArrayProperty>(InProperty))
 		{
 			if (const FStructProperty* StructProperty = CastField<FStructProperty>(ArrayProperty->Inner))
 			{
@@ -2071,50 +2136,6 @@ EMemberPropertyTraits FPropertyReplicationStateDescriptorBuilder::GetFastArrayPr
 	return EMemberPropertyTraits::None;
 }
 
-bool FPropertyReplicationStateDescriptorBuilder::IsFastArraySupported(const UStruct* Struct)
-{
-	uint32 ReplicatedPropertyCount = 0;
-
-	FMemberProperty MemberProperty;
-	if (IsStructWithCustomSerializer(Struct))
-	{
-		return false;
-	}
-
-	// Try to figure out if this is a FastArraySerializer we can support. It is assumed the struct has already been determined to inherit from FastArraySerializer.
-	// A proper implementation would contain a TArray of a struct that is a subclass of FastArraySerializerItem.
-	for (TFieldIterator<FProperty> It(Struct, EFieldIteratorFlags::IncludeSuper); It; ++It)
-	{
-		const FProperty* Property = *It;
-		if (EnumHasAnyFlags(Property->PropertyFlags, EPropertyFlags::CPF_RepSkip))
-		{
-			continue;
-		}
-
-		++ReplicatedPropertyCount;
-		if (ReplicatedPropertyCount > 1)
-		{
-			return false;
-		}
-		if (!FPropertyReplicationStateDescriptorBuilder::IsSupportedProperty(MemberProperty, Property))
-		{
-			return false;
-		}
-
-		if (!EnumHasAnyFlags(MemberProperty.Traits, EMemberPropertyTraits::IsFastArrayItem))
-		{
-			return false;
-		}
-
-		if (CastField<FArrayProperty>(MemberProperty.Property) == nullptr)
-		{
-			return false;
-		}
-	}
-
-	return true;
-}
-
 }
 
 namespace UE::Net
@@ -2128,6 +2149,7 @@ FReplicationStateDescriptorBuilder::FParameters::FParameters()
   IncludeSuper(1U),
   GetLifeTimeProperties(1U),
   EnableFastArrayHandling(1U),
+  AllowFastArrayWithExtraReplicatedProperties(0U),
   SkipCheckForCustomNetSerializerForStruct(0U),
   SinglePropertyIndex(-1)
 {
@@ -2220,11 +2242,14 @@ TRefCountPtr<const FReplicationStateDescriptor> FReplicationStateDescriptorBuild
 					continue;
 				}
 
-				// For FastArraySerializer we do not expect other properties than FastArrayItems
 				if (bIsFastArraySerializer && !EnumHasAnyFlags(MemberProperty.Traits, EMemberPropertyTraits::IsFastArrayItem))
 				{
-					UE_LOG_DESCRIPTORBUILDER_WARNING(TEXT("FReplicationStateDescriptorBuilder ::CreateDescriptorForStruct FastArraySerializerStructs can only hold a single replicated property derived from FastArraySerializerItem. Skipping property: %s.%s of type %s."), ToCStr(InStruct->GetName()), ToCStr(Property->GetName()), ToCStr(Property->GetCPPType()));
-					continue;
+					// For FastArraySerializer we do not expect other properties than FastArrayItems unless explicitly allowed.
+					if (!Parameters.AllowFastArrayWithExtraReplicatedProperties)
+					{
+						UE_LOG_DESCRIPTORBUILDER_WARNING(TEXT("FReplicationStateDescriptorBuilder ::CreateDescriptorForStruct FastArraySerializerStructs can only hold a single replicated property derived from FastArraySerializerItem. Skipping property: %s.%s of type %s."), ToCStr(InStruct->GetName()), ToCStr(Property->GetName()), ToCStr(Property->GetCPPType()));
+						continue;
+					}
 				}
 
 				for (uint32 ArrayIt = 0, ArrayEndIt = Property->ArrayDim; ArrayIt < ArrayEndIt; ++ArrayIt)
@@ -2401,6 +2426,11 @@ SIZE_T FReplicationStateDescriptorBuilder::CreateDescriptorsForClass(FResult& Cr
 	// We add separate descriptors properties with custom replication fragments
 	TArray<FPropertyReplicationStateDescriptorBuilder::FMemberProperty, TInlineAllocator<8>> CustomProperties;
 
+	FPropertyReplicationStateDescriptorBuilder::FIsSupportedPropertyParams IsSupportedPropertyParams;
+	IsSupportedPropertyParams.LifeTimeProperties = &LifeTimeProperties;
+	IsSupportedPropertyParams.InObjectClass = InObjectClass;
+	IsSupportedPropertyParams.bAllowFastArrayWithExtraReplicatedProperties = Parameters.AllowFastArrayWithExtraReplicatedProperties;
+
 	// Add properties
 	for (const FRepRecord& RepRecord : InObjectClass->ClassReps)
 	{
@@ -2416,7 +2446,7 @@ SIZE_T FReplicationStateDescriptorBuilder::CreateDescriptorsForClass(FResult& Cr
 		// Skip super class properties depending on parameters. C array elements will be treated individually.
 		if ((bIncludeSuper || Property->GetOwnerVariant() == InObjectClass))
 		{
-			if (FPropertyReplicationStateDescriptorBuilder::IsSupportedProperty(MemberProperty, Property, &LifeTimeProperties, InObjectClass))
+			if (FPropertyReplicationStateDescriptorBuilder::IsSupportedProperty(MemberProperty, Property, IsSupportedPropertyParams))
 			{
 				// Ignore disabled properties
 				if (MemberProperty.ReplicationCondition == COND_Never)
diff --git a/Engine/Source/Runtime/Experimental/Iris/Core/Private/Iris/ReplicationSystem/FastArrayReplicationFragment.cpp b/Engine/Source/Runtime/Experimental/Iris/Core/Private/Iris/ReplicationSystem/FastArrayReplicationFragment.cpp
--- a/Engine/Source/Runtime/Experimental/Iris/Core/Private/Iris/ReplicationSystem/FastArrayReplicationFragment.cpp
+++ b/Engine/Source/Runtime/Experimental/Iris/Core/Private/Iris/ReplicationSystem/FastArrayReplicationFragment.cpp
@@ -10,11 +10,12 @@
 #include "Iris/ReplicationState/ReplicationStateUtil.h"
 #include "Iris/Serialization/InternalNetSerializers.h"
 #include "Iris/Serialization/NetSerializers.h"
+#include "Iris/Serialization/InternalNetSerializerUtils.h"
 
 namespace UE::Net::Private
 {
 
-const FReplicationStateDescriptor* FFastArrayReplicationFragmentBase::GetFastArraySerializerPropertyDescriptor() const
+const FReplicationStateDescriptor* FFastArrayReplicationFragmentBase::GetFastArrayPropertyStructDescriptor() const
 {
 	const FReplicationStateMemberSerializerDescriptor& FastArrayMemberSerializerDescriptor = ReplicationStateDescriptor->MemberSerializerDescriptors[0];
 	const FStructNetSerializerConfig* FastArrayStructSerializerConfig = static_cast<const FStructNetSerializerConfig*>(FastArrayMemberSerializerDescriptor.SerializerConfig);
@@ -22,7 +23,7 @@ const FReplicationStateDescriptor* FFastArrayReplicationFragmentBase::GetFastArr
 	return FastArrayStructSerializerConfig->StateDescriptor;
 }
 
-FFastArrayReplicationFragmentBase::FFastArrayReplicationFragmentBase(EReplicationFragmentTraits InTraits, UObject* InOwner, const FReplicationStateDescriptor* InDescriptor)
+FFastArrayReplicationFragmentBase::FFastArrayReplicationFragmentBase(EReplicationFragmentTraits InTraits, UObject* InOwner, const FReplicationStateDescriptor* InDescriptor, bool bValidateDescriptor)
 : FReplicationFragment(InTraits)
 , ReplicationStateDescriptor(InDescriptor)
 , Owner(InOwner)
@@ -35,7 +36,7 @@ FFastArrayReplicationFragmentBase::FFastArrayReplicationFragmentBase(EReplicatio
 
 	if (EnumHasAnyFlags(InTraits, EReplicationFragmentTraits::CanReplicate))
 	{
-		SrcReplicationState = MakeUnique<FPropertyReplicationState>(InDescriptor);
+		ReplicationState = MakeUnique<FPropertyReplicationState>(InDescriptor);
 	}
 	
 	if (EnumHasAnyFlags(InTraits, EReplicationFragmentTraits::CanReceive))
@@ -50,7 +51,7 @@ FFastArrayReplicationFragmentBase::FFastArrayReplicationFragmentBase(EReplicatio
 			Traits |= EReplicationFragmentTraits::HasRepNotifies;
 		}
 
-		// For now we always expect pre/post operations for legacy states, we might make this
+		// For now we always expect pre/post operations for legacy states, we might make this optional
 		Traits |= EReplicationFragmentTraits::NeedsLegacyCallbacks;
 	}
 	
@@ -69,10 +70,14 @@ FFastArrayReplicationFragmentBase::FFastArrayReplicationFragmentBase(EReplicatio
 	Traits |= EReplicationFragmentTraits::HasPropertyReplicationState;
 
 	// Look up the descriptor of the struct
-	const FReplicationStateDescriptor* FastArrayStructDescriptor = GetFastArraySerializerPropertyDescriptor();
+	const FReplicationStateDescriptor* FastArrayStructDescriptor = GetFastArrayPropertyStructDescriptor();
 
-	// And we are looking for the offset of the first and only member of the struct which is the relative offest to the FastArray
-	WrappedArrayOffsetRelativeFastArraySerializerProperty = FastArrayStructDescriptor->MemberDescriptors[0].ExternalMemberOffset;
+	// And we are looking for the offset of ItemArray which is the relative offset from the FastArray struct
+	const uint32 ItemArrayMemberIndex = FFastArrayReplicationFragmentHelper::GetFastArrayStructItemArrayMemberIndex(FastArrayStructDescriptor);
+	WrappedArrayOffsetRelativeFastArraySerializerProperty = FastArrayStructDescriptor->MemberDescriptors[ItemArrayMemberIndex].ExternalMemberOffset;
+
+	// Validate the expected member count
+	ensureMsgf(!bValidateDescriptor || FastArrayStructDescriptor->MemberCount == 1U, TEXT("A FastArray using the default FastArrayReplicationFragment is expected to have a single replicated array property derived from FastArraySerializerItem"));
 
 #if WITH_PUSH_MODEL
 	if (EnumHasAnyFlags(InDescriptor->Traits, EReplicationStateTraits::HasPushBasedDirtiness))
@@ -88,7 +93,10 @@ const FReplicationStateDescriptor* FFastArrayReplicationFragmentBase::GetArrayEl
 {
 	const FReplicationStateDescriptor* FastArrayDescriptor = ReplicationStateDescriptor.GetReference();
 	const FReplicationStateDescriptor* FastArrayStructDescriptor = static_cast<const FStructNetSerializerConfig*>(FastArrayDescriptor->MemberSerializerDescriptors[0].SerializerConfig)->StateDescriptor;
-	const FReplicationStateDescriptor* ElementStateDescriptor = static_cast<const FArrayPropertyNetSerializerConfig*>(FastArrayStructDescriptor->MemberSerializerDescriptors[0].SerializerConfig)->StateDescriptor;
+
+	const uint32 ItemArrayMemberIndex = FFastArrayReplicationFragmentHelper::GetFastArrayStructItemArrayMemberIndex(FastArrayStructDescriptor);
+	const FReplicationStateDescriptor* ElementStateDescriptor = static_cast<const FArrayPropertyNetSerializerConfig*>(FastArrayStructDescriptor->MemberSerializerDescriptors[ItemArrayMemberIndex].SerializerConfig)->StateDescriptor;
+
 	const FReplicationStateMemberSerializerDescriptor& ElementSerializerDescriptor = ElementStateDescriptor->MemberSerializerDescriptors[0];
 	
 	return static_cast<const FStructNetSerializerConfig*>(ElementSerializerDescriptor.SerializerConfig)->StateDescriptor;
@@ -107,7 +115,7 @@ bool FFastArrayReplicationFragmentBase::InternalCompareArrayElement(const FRepli
 void FFastArrayReplicationFragmentBase::Register(FFragmentRegistrationContext& Context, EReplicationFragmentTraits InTraits)
 {
 	this->Traits |= InTraits;
-	Context.RegisterReplicationFragment(this, ReplicationStateDescriptor.GetReference(), SrcReplicationState ? SrcReplicationState->GetStateBuffer() : nullptr);
+	Context.RegisterReplicationFragment(this, ReplicationStateDescriptor.GetReference(), ReplicationState ? ReplicationState->GetStateBuffer() : nullptr);
 }
 
 void FFastArrayReplicationFragmentBase::CollectOwner(FReplicationStateOwnerCollector* Owners) const
@@ -142,10 +150,21 @@ void FFastArrayReplicationFragmentBase::ReplicatedStateToString(FStringBuilderBa
 	ReceivedState.ToString(StringBuilder, bIncludeAll);
 }
 
-////////////////////////////////////////////////////////////////////////////////////////////////////////////////////////////////////////////////////
-// NativeFastArrayReplicationFragment
-////////////////////////////////////////////////////////////////////////////////////////////////////////////////////////////////////////////////////
-const FReplicationStateDescriptor* FNativeFastArrayReplicationFragmentBase::GetFastArraySerializerPropertyDescriptor() const
+uint32 FFastArrayReplicationFragmentHelper::GetFastArrayStructItemArrayMemberIndex(const FReplicationStateDescriptor* StructDescriptor)
+{
+	// We know that one of the members we are looking for is the item array and there will only be one of them.
+	uint32 MemberIndex = 0U;
+	for (;;)
+	{
+		if (IsArrayPropertyNetSerializer(StructDescriptor->MemberSerializerDescriptors[MemberIndex].Serializer))
+		{
+			return MemberIndex;
+		}
+		++MemberIndex;
+	}
+}
+
+const FReplicationStateDescriptor* FNativeFastArrayReplicationFragmentBase::GetFastArrayPropertyStructDescriptor() const
 {
 	const FReplicationStateMemberSerializerDescriptor& FastArrayMemberSerializerDescriptor = ReplicationStateDescriptor->MemberSerializerDescriptors[0];
 	const FStructNetSerializerConfig* FastArrayStructSerializerConfig = static_cast<const FStructNetSerializerConfig*>(FastArrayMemberSerializerDescriptor.SerializerConfig);
@@ -153,6 +172,10 @@ const FReplicationStateDescriptor* FNativeFastArrayReplicationFragmentBase::GetF
 	return FastArrayStructSerializerConfig->StateDescriptor;
 }
 
+////////////////////////////////////////////////////////////////////////////////////////////////////////////////////////////////////////////////////
+// NativeFastArrayReplicationFragment
+////////////////////////////////////////////////////////////////////////////////////////////////////////////////////////////////////////////////////
+
 FNativeFastArrayReplicationFragmentBase::FNativeFastArrayReplicationFragmentBase(EReplicationFragmentTraits InTraits, UObject* InOwner, const FReplicationStateDescriptor* InDescriptor)
 : FReplicationFragment(InTraits)
 , ReplicationStateDescriptor(InDescriptor)
@@ -189,17 +212,21 @@ FNativeFastArrayReplicationFragmentBase::FNativeFastArrayReplicationFragmentBase
 	}
 
 	// Look up the descriptor of the struct
-	const FReplicationStateDescriptor* FastArrayStructDescriptor = GetFastArraySerializerPropertyDescriptor();
+	const FReplicationStateDescriptor* FastArrayStructDescriptor = GetFastArrayPropertyStructDescriptor();
 
-	// And we are looking for the offset of the first and only member of the struct which is the relative offest to the FastArray
-	WrappedArrayOffsetRelativeFastArraySerializerProperty = FastArrayStructDescriptor->MemberDescriptors[0].ExternalMemberOffset;
+	// And we are looking for the offset of ItemArray which is the relative offset from the FastArray struct
+	const uint32 ItemArrayMemberIndex = FFastArrayReplicationFragmentHelper::GetFastArrayStructItemArrayMemberIndex(FastArrayStructDescriptor);
+	WrappedArrayOffsetRelativeFastArraySerializerProperty = FastArrayStructDescriptor->MemberDescriptors[ItemArrayMemberIndex].ExternalMemberOffset;
 }
 
 const FReplicationStateDescriptor* FNativeFastArrayReplicationFragmentBase::GetArrayElementDescriptor() const
 {
 	const FReplicationStateDescriptor* FastArrayDescriptor = ReplicationStateDescriptor.GetReference();
 	const FReplicationStateDescriptor* FastArrayStructDescriptor = static_cast<const FStructNetSerializerConfig*>(FastArrayDescriptor->MemberSerializerDescriptors[0].SerializerConfig)->StateDescriptor;
-	const FReplicationStateDescriptor* ElementStateDescriptor = static_cast<const FArrayPropertyNetSerializerConfig*>(FastArrayStructDescriptor->MemberSerializerDescriptors[0].SerializerConfig)->StateDescriptor;
+
+	const uint32 ItemArrayMemberIndex = FFastArrayReplicationFragmentHelper::GetFastArrayStructItemArrayMemberIndex(FastArrayStructDescriptor);
+	const FReplicationStateDescriptor* ElementStateDescriptor = static_cast<const FArrayPropertyNetSerializerConfig*>(FastArrayStructDescriptor->MemberSerializerDescriptors[ItemArrayMemberIndex].SerializerConfig)->StateDescriptor;
+
 	const FReplicationStateMemberSerializerDescriptor& ElementSerializerDescriptor = ElementStateDescriptor->MemberSerializerDescriptors[0];
 	
 	return static_cast<const FStructNetSerializerConfig*>(ElementSerializerDescriptor.SerializerConfig)->StateDescriptor;
@@ -207,7 +234,7 @@ const FReplicationStateDescriptor* FNativeFastArrayReplicationFragmentBase::GetA
 
 void FNativeFastArrayReplicationFragmentBase::InternalDequantizeFastArray(FNetSerializationContext& Context, uint8* RESTRICT DstExternalBuffer, const uint8* RESTRICT SrcInternalBuffer, const FReplicationStateDescriptor* Descriptor)
 {
-	const uint32 MemberIndex = 0u;
+	const uint32 MemberIndex = FFastArrayReplicationFragmentHelper::GetFastArrayStructItemArrayMemberIndex(Descriptor);
 
 	const FReplicationStateMemberDescriptor* MemberDescriptors = Descriptor->MemberDescriptors;
 	const FReplicationStateMemberSerializerDescriptor* MemberSerializerDescriptors = Descriptor->MemberSerializerDescriptors;
diff --git a/Engine/Source/Runtime/Experimental/Iris/Core/Private/Iris/ReplicationSystem/ReplicationFragmentUtil.cpp b/Engine/Source/Runtime/Experimental/Iris/Core/Private/Iris/ReplicationSystem/ReplicationFragmentUtil.cpp
--- a/Engine/Source/Runtime/Experimental/Iris/Core/Private/Iris/ReplicationSystem/ReplicationFragmentUtil.cpp
+++ b/Engine/Source/Runtime/Experimental/Iris/Core/Private/Iris/ReplicationSystem/ReplicationFragmentUtil.cpp
@@ -30,6 +30,14 @@ uint32 FReplicationFragmentUtil::CreateAndRegisterFragmentsForObject(UObject* Ob
 		}
 	}
 
+	// Pass-on that we allow FastAarrays with extra replicated properties for this object.
+	// NOTE: this is not perfect as it allows this for all FastArray properties of the object but as there 
+	// is further validation in the actual ReplicationStateFragment implementation for FastArrays it is good enough.
+	if (EnumHasAnyFlags(RegistrationFlags, EFragmentRegistrationFlags::AllowFastArraysWithAdditionalProperties))
+	{
+		BuilderParameters.AllowFastArrayWithExtraReplicatedProperties = 1U;
+	}
+
 	FReplicationStateDescriptorBuilder::FResult Result;
 	FReplicationStateDescriptorBuilder::CreateDescriptorsForClass(Result, Object->GetClass(), BuilderParameters);
 
diff --git a/Engine/Source/Runtime/Experimental/Iris/Core/Public/Iris/ReplicationState/ReplicationStateDescriptorBuilder.h b/Engine/Source/Runtime/Experimental/Iris/Core/Public/Iris/ReplicationState/ReplicationStateDescriptorBuilder.h
--- a/Engine/Source/Runtime/Experimental/Iris/Core/Public/Iris/ReplicationState/ReplicationStateDescriptorBuilder.h
+++ b/Engine/Source/Runtime/Experimental/Iris/Core/Public/Iris/ReplicationState/ReplicationStateDescriptorBuilder.h
@@ -32,6 +32,8 @@ public:
 		uint32 GetLifeTimeProperties : 1;
 		// If EnableFastArrayHandling is true and the struct inherits from FFastArraySerializer then special logic will be applied which allows it to be bound to a FastArrayReplicationFragment.
 		uint32 EnableFastArrayHandling : 1;
+		// If AllowFastArrayWithExtraReplicatedProperties is set to true, we will allow building descriptors for fastarrays with more than a single property.
+		uint32 AllowFastArrayWithExtraReplicatedProperties : 1;
 		// In SkipCheckForCustomNetSerializer is true the descriptor will be built using the underlying representation even if a CustomNetSerializer is registered for the struct
 		uint32 SkipCheckForCustomNetSerializerForStruct : 1;
 		 // If SinglePropertyIndex != -1 we will create a state that only includes the specified property.
diff --git a/Engine/Source/Runtime/Experimental/Iris/Core/Public/Iris/ReplicationSystem/FastArrayReplicationFragment.h b/Engine/Source/Runtime/Experimental/Iris/Core/Public/Iris/ReplicationSystem/FastArrayReplicationFragment.h
--- a/Engine/Source/Runtime/Experimental/Iris/Core/Public/Iris/ReplicationSystem/FastArrayReplicationFragment.h
+++ b/Engine/Source/Runtime/Experimental/Iris/Core/Public/Iris/ReplicationSystem/FastArrayReplicationFragment.h
@@ -15,29 +15,41 @@ namespace UE::Net
 {
 
 /**
- * TFastArrayReplicationFragment - Binds a typed FastArray to a FReplicationfragments
+ * TFastArrayReplicationFragment - Binds a typed FastArray to a FReplicationfragment
  * Used to support FFastArray-based serialization with no required code modifications
  * Backed by a PropertyReplicationState which means that we will have to poll source data for dirtiness,
  * in the case of FastArrays this involves comparing the replication key of the array and its items.
  */
 template <typename FastArrayItemType, typename FastArrayType>
-class TFastArrayReplicationFragment final : public Private::FFastArrayReplicationFragmentBase
+class TFastArrayReplicationFragment : public Private::FFastArrayReplicationFragmentBase
 {
 public:
 	typedef TArray<FastArrayItemType> ItemArrayType;
 	TFastArrayReplicationFragment(EReplicationFragmentTraits InTraits, UObject* InOwner, const FReplicationStateDescriptor* InDescriptor) : FFastArrayReplicationFragmentBase(InTraits, InOwner, InDescriptor) {}
 
 protected:
+	enum EAllowAdditionalPropertiesType { AllowAdditionalProperties };
+	TFastArrayReplicationFragment(EReplicationFragmentTraits InTraits, UObject* InOwner, const FReplicationStateDescriptor* InDescriptor, const EAllowAdditionalPropertiesType) : FFastArrayReplicationFragmentBase(InTraits, InOwner, InDescriptor, false) {}
+
 	virtual void ApplyReplicatedState(FReplicationStateApplyContext& Context) const override;
 	virtual bool PollReplicatedState(EReplicationFragmentPollFlags PollOption) override;
 
 	bool PollAllState(bool bForceFullCompare = false);
 
+	// Returns true if the FastArray is dirty
 	bool IsDirty() const;
 
-private:
-	// Get the wrapped FastArraySerializerProperty
-	inline FastArrayType* GetWrappedFastArraySerializer() const;
+	// Mark FastArray dirty
+	void MarkDirty();
+
+	// Get FastArraySerializer from owner
+	inline FastArrayType* GetFastArraySerializerFromOwner() const;
+
+	// Get FastArraySerializer from our cached ReplicationState
+	inline FastArrayType* GetFastArraySerializerFromReplicationState() const;
+
+	// Get FastArraySerialzier from received state
+	inline FastArrayType* GetFastArraySerializerFromApplyContext(FReplicationStateApplyContext& Context) const;	
 };
 
 /**
@@ -68,9 +80,10 @@ protected:
 
 	bool IsDirty() const;
 
+	// Get FastArraySerializer from owner
+	inline FastArrayType* GetFastArraySerializerFromOwner() const;
+
 private:
-	// Get the wrapped FastArraySerializerProperty
-	inline FastArrayType* GetWrappedFastArraySerializer() const;
 	PollingPolicyType PollingPolicy;
 };
 
@@ -81,31 +94,37 @@ template <typename FastArrayItemType, typename FastArrayType>
 void TFastArrayReplicationFragment<FastArrayItemType, FastArrayType>::ApplyReplicatedState(FReplicationStateApplyContext& Context) const
 {
 	// Get the wrapped FastArraySerializer and array
-	FastArrayType* DstArraySerializer = GetWrappedFastArraySerializer();
+	FastArrayType* DstArraySerializer = GetFastArraySerializerFromOwner();
 	ItemArrayType* DstWrappedArray = reinterpret_cast<ItemArrayType*>(reinterpret_cast<uint8*>(DstArraySerializer) + WrappedArrayOffsetRelativeFastArraySerializerProperty);
 
-	// Intentionally not const as we allow the src state to be modified
-	FastArrayType* SrcArraySerializer;
-	const ItemArrayType* SrcWrappedArray;
-
 	// Get received array data
-	{
-		uint8* ReceivedStateBuffer = Context.StateBufferData.ExternalStateBuffer;
-		const SIZE_T FastArraySerializerMemberOffset = ReplicationStateDescriptor->MemberDescriptors[0].ExternalMemberOffset;
-		SrcArraySerializer = reinterpret_cast<FastArrayType*>(ReceivedStateBuffer + FastArraySerializerMemberOffset);
-		SrcWrappedArray = reinterpret_cast<const ItemArrayType*>(ReceivedStateBuffer + FastArraySerializerMemberOffset + WrappedArrayOffsetRelativeFastArraySerializerProperty);
-	}
+	// Intentionally not const as we allow the src state to be modified	
+	FastArrayType* SrcArraySerializer = GetFastArraySerializerFromApplyContext(Context);
+	const ItemArrayType* SrcWrappedArray = reinterpret_cast<ItemArrayType*>(reinterpret_cast<uint8*>(SrcArraySerializer) + WrappedArrayOffsetRelativeFastArraySerializerProperty);
 
 	// Apply state and issue callbacks etc
 	Private::FFastArrayReplicationFragmentHelper::ApplyReplicatedState(DstArraySerializer, DstWrappedArray, SrcArraySerializer, SrcWrappedArray, GetArrayElementDescriptor(), Context);
 }
 
 template <typename FastArrayItemType, typename FastArrayType>
-FastArrayType* TFastArrayReplicationFragment<FastArrayItemType, FastArrayType>::GetWrappedFastArraySerializer() const
+FastArrayType* TFastArrayReplicationFragment<FastArrayItemType, FastArrayType>::GetFastArraySerializerFromOwner() const
 {
 	return reinterpret_cast<FastArrayType*>(reinterpret_cast<uint8*>(Owner) + ReplicationStateDescriptor->MemberProperties[0]->GetOffset_ForGC());
 }
 
+template <typename FastArrayItemType, typename FastArrayType>
+inline FastArrayType* TFastArrayReplicationFragment<FastArrayItemType, FastArrayType>::GetFastArraySerializerFromApplyContext(FReplicationStateApplyContext& Context) const
+{
+	return reinterpret_cast<FastArrayType*>(Context.StateBufferData.ExternalStateBuffer + ReplicationStateDescriptor->MemberDescriptors[0].ExternalMemberOffset);	
+}
+
+template <typename FastArrayItemType, typename FastArrayType>
+FastArrayType* TFastArrayReplicationFragment<FastArrayItemType, FastArrayType>::GetFastArraySerializerFromReplicationState() const
+{
+	checkSlow(ReplicationState);
+	return ReplicationState ? reinterpret_cast<FastArrayType*>(ReplicationState->GetStateBuffer() + ReplicationStateDescriptor->MemberDescriptors[0].ExternalMemberOffset) : nullptr;
+}
+
 template <typename FastArrayItemType, typename FastArrayType>
 bool TFastArrayReplicationFragment<FastArrayItemType, FastArrayType>::PollReplicatedState(EReplicationFragmentPollFlags PollOption)
 {
@@ -126,11 +145,11 @@ template <typename FastArrayItemType, typename FastArrayType>
 bool TFastArrayReplicationFragment<FastArrayItemType, FastArrayType>::PollAllState(bool bForceFullCompare)
 {
 	// Lookup source data, we need the actual FastArraySerializer and the Array it is wrapping
-	FastArrayType* SrcArraySerializer = GetWrappedFastArraySerializer();
+	FastArrayType* SrcArraySerializer = GetFastArraySerializerFromOwner();
 	ItemArrayType* SrcWrappedArray = reinterpret_cast<ItemArrayType*>(reinterpret_cast<uint8*>(SrcArraySerializer) + WrappedArrayOffsetRelativeFastArraySerializerProperty);
 	
 	// Lookup destination data
-	FastArrayType* DstArraySerializer = reinterpret_cast<FastArrayType*>(SrcReplicationState->GetStateBuffer() + ReplicationStateDescriptor->MemberDescriptors[0].ExternalMemberOffset);
+	FastArrayType* DstArraySerializer = GetFastArraySerializerFromReplicationState();
 	ItemArrayType* DstWrappedArray = reinterpret_cast<ItemArrayType*>(reinterpret_cast<uint8*>(DstArraySerializer) + WrappedArrayOffsetRelativeFastArraySerializerProperty);
 
 	// Check if we can early out
@@ -154,7 +173,7 @@ bool TFastArrayReplicationFragment<FastArrayItemType, FastArrayType>::PollAllSta
 	// We currently use a simple modulo scheme for bits in the changemask
 	// A single bit might represent several entries in the array which all will be considered dirty, it is up to the serializer to handle this
 	// The first bit is used by the owning property we need to offset by one and deduct one from the usable bits
-	FNetBitArrayView MemberChangeMask = UE::Net::Private::GetMemberChangeMask(SrcReplicationState->GetStateBuffer(), ReplicationStateDescriptor);
+	FNetBitArrayView MemberChangeMask = UE::Net::Private::GetMemberChangeMask(ReplicationState->GetStateBuffer(), ReplicationStateDescriptor);
 	
 	const FReplicationStateMemberChangeMaskDescriptor& MemberChangeMaskDescriptor = ReplicationStateDescriptor->MemberChangeMaskDescriptors[0];
 	const uint32 ChangeMaskBitOffset = MemberChangeMaskDescriptor.BitOffset + FIrisFastArraySerializer::IrisFastArrayChangeMaskBitOffset;
@@ -197,10 +216,10 @@ bool TFastArrayReplicationFragment<FastArrayItemType, FastArrayType>::PollAllSta
 	DstArraySerializer->ArrayReplicationKey = SrcArraySerializer->ArrayReplicationKey;
 
 	// Mark the NetObject as dirty
-	if (bMarkArrayDirty && SrcReplicationState->IsCustomConditionEnabled(FIrisFastArraySerializer::IrisFastArrayPropertyBitIndex))
+	if (bMarkArrayDirty && ReplicationState->IsCustomConditionEnabled(FIrisFastArraySerializer::IrisFastArrayPropertyBitIndex))
 	{
 		MemberChangeMask.SetBit(FIrisFastArraySerializer::IrisFastArrayPropertyBitIndex);
-		MarkNetObjectStateDirty(UE::Net::Private::GetReplicationStateHeader(SrcReplicationState->GetStateBuffer(), ReplicationStateDescriptor));
+		MarkNetObjectStateDirty(UE::Net::Private::GetReplicationStateHeader(ReplicationState->GetStateBuffer(), ReplicationStateDescriptor));
 	}
 
 	return IsDirty();
@@ -209,10 +228,22 @@ bool TFastArrayReplicationFragment<FastArrayItemType, FastArrayType>::PollAllSta
 template <typename FastArrayItemType, typename FastArrayType>
 bool TFastArrayReplicationFragment<FastArrayItemType, FastArrayType>::IsDirty() const
 {
-	FReplicationStateHeader& ReplicationStateHeader = Private::GetReplicationStateHeader(SrcReplicationState->GetStateBuffer(), ReplicationStateDescriptor);
+	FReplicationStateHeader& ReplicationStateHeader = Private::GetReplicationStateHeader(ReplicationState->GetStateBuffer(), ReplicationStateDescriptor);
 	return Private::FReplicationStateHeaderAccessor::GetIsStateDirty(ReplicationStateHeader);
 }
 
+template <typename FastArrayItemType, typename FastArrayType>
+void TFastArrayReplicationFragment<FastArrayItemType, FastArrayType>::MarkDirty()
+{
+	// Mark the NetObject as dirty
+	if (ReplicationState->IsCustomConditionEnabled(FIrisFastArraySerializer::IrisFastArrayPropertyBitIndex))
+	{
+		FNetBitArrayView MemberChangeMask = UE::Net::Private::GetMemberChangeMask(ReplicationState->GetStateBuffer(), ReplicationStateDescriptor);
+		MemberChangeMask.SetBit(FIrisFastArraySerializer::IrisFastArrayPropertyBitIndex);
+		MarkNetObjectStateDirty(UE::Net::Private::GetReplicationStateHeader(ReplicationState->GetStateBuffer(), ReplicationStateDescriptor));
+	}
+}
+
 /*
  * TNativeFastArrayReplicationFragment implementation
  */
@@ -231,7 +262,7 @@ template <typename FastArrayItemType, typename FastArrayType, typename PollingPo
 void TNativeFastArrayReplicationFragment<FastArrayItemType, FastArrayType, PollingPolicyType>::Register(FFragmentRegistrationContext& Context, EReplicationFragmentTraits InTraits)
 {
 	Traits |= InTraits;
-	Context.RegisterReplicationFragment(this, ReplicationStateDescriptor.GetReference(), reinterpret_cast<uint8*>(GetWrappedFastArraySerializer()));
+	Context.RegisterReplicationFragment(this, ReplicationStateDescriptor.GetReference(), reinterpret_cast<uint8*>(GetFastArraySerializerFromOwner()));
 }
 
 template <typename FastArrayItemType, typename FastArrayType, typename PollingPolicyType>
@@ -247,7 +278,7 @@ void TNativeFastArrayReplicationFragment<FastArrayItemType, FastArrayType, Polli
 		{
 			FastArrayType ReceivedState;
 			FastArrayType DefaultState;
-			InternalDequantizeFastArray(*Context.NetSerializationContext, reinterpret_cast<uint8*>(&ReceivedState), Context.StateBufferData.RawStateBuffer, GetFastArraySerializerPropertyDescriptor());
+			InternalDequantizeFastArray(*Context.NetSerializationContext, reinterpret_cast<uint8*>(&ReceivedState), Context.StateBufferData.RawStateBuffer, GetFastArrayPropertyStructDescriptor());
 			if (Descriptor->MemberProperties[0]->Identical(&ReceivedState, &DefaultState))
 			{
 				return;
@@ -277,7 +308,7 @@ bool TNativeFastArrayReplicationFragment<FastArrayItemType, FastArrayType, Polli
 	if (FPollingState* PollingState = PollingPolicy.GetPollingState())
 	{
 		// Get the source FastArraySerializer and array
-		FastArrayType* SrcArraySerializer = GetWrappedFastArraySerializer();
+		FastArrayType* SrcArraySerializer = GetFastArraySerializerFromOwner();
 		ItemArrayType* SrcWrappedArray = reinterpret_cast<ItemArrayType*>(reinterpret_cast<uint8*>(SrcArraySerializer) + WrappedArrayOffsetRelativeFastArraySerializerProperty);
 
 		// Check if we can early out
@@ -355,7 +386,7 @@ bool TNativeFastArrayReplicationFragment<FastArrayItemType, FastArrayType, Polli
 template <typename FastArrayItemType, typename FastArrayType, typename PollingPolicyType>
 bool TNativeFastArrayReplicationFragment<FastArrayItemType, FastArrayType, PollingPolicyType>::IsDirty() const
 {
-	FastArrayType* SrcArraySerializer = GetWrappedFastArraySerializer();
+	FastArrayType* SrcArraySerializer = GetFastArraySerializerFromOwner();
 	const FReplicationStateHeader& ReplicationStateHeader = Private::FIrisFastArraySerializerPrivateAccessor::GetReplicationStateHeader(*SrcArraySerializer);
 	return Private::FReplicationStateHeaderAccessor::GetIsStateDirty(ReplicationStateHeader);
 }
@@ -368,7 +399,7 @@ void TNativeFastArrayReplicationFragment<FastArrayItemType, FastArrayType, Polli
 	FastArrayType ReceivedState;
 
 	// Dequantize into temporary array
-	InternalDequantizeFastArray(*Context.NetSerializationContext, reinterpret_cast<uint8*>(&ReceivedState), Context.StateBufferData.RawStateBuffer, GetFastArraySerializerPropertyDescriptor());
+	InternalDequantizeFastArray(*Context.NetSerializationContext, reinterpret_cast<uint8*>(&ReceivedState), Context.StateBufferData.RawStateBuffer, GetFastArrayPropertyStructDescriptor());
 
 	// Get the wrapped FastArraySerializer and array
 	FastArrayType* DstArraySerializer = reinterpret_cast<FastArrayType*>(reinterpret_cast<uint8*>(Owner) + ReplicationStateDescriptor->MemberProperties[0]->GetOffset_ForGC());
@@ -389,14 +420,14 @@ void TNativeFastArrayReplicationFragment<FastArrayItemType, FastArrayType, Polli
 	FastArrayType ReceivedState;
 
 	// Dequantize into temporary array
-	InternalDequantizeFastArray(*Context.NetSerializationContext, reinterpret_cast<uint8*>(&ReceivedState), Context.StateBufferData.RawStateBuffer, GetFastArraySerializerPropertyDescriptor());
+	InternalDequantizeFastArray(*Context.NetSerializationContext, reinterpret_cast<uint8*>(&ReceivedState), Context.StateBufferData.RawStateBuffer, GetFastArrayPropertyStructDescriptor());
 
 	// Output state to string
-	ToString(StringBuilder, reinterpret_cast<uint8*>(&ReceivedState), GetFastArraySerializerPropertyDescriptor());
+	ToString(StringBuilder, reinterpret_cast<uint8*>(&ReceivedState), GetFastArrayPropertyStructDescriptor());
 }
 
 template <typename FastArrayItemType, typename FastArrayType, typename PollingPolicyType>
-FastArrayType* TNativeFastArrayReplicationFragment<FastArrayItemType, FastArrayType, PollingPolicyType>::GetWrappedFastArraySerializer() const
+FastArrayType* TNativeFastArrayReplicationFragment<FastArrayItemType, FastArrayType, PollingPolicyType>::GetFastArraySerializerFromOwner() const
 {
 	return reinterpret_cast<FastArrayType*>(reinterpret_cast<uint8*>(Owner) + ReplicationStateDescriptor->MemberProperties[0]->GetOffset_ForGC());
 }
@@ -407,7 +438,7 @@ template <typename FastArrayType>
 FReplicationFragment* CreateAndRegisterFragment(UObject* Owner, const FReplicationStateDescriptor* Descriptor, FFragmentRegistrationContext& Context)
 {
 	using namespace UE::Net;
-	static_assert(TFastArrayTypeHelper<FastArrayType>::HasValidFastArrayItemType(), "Invalid FastArrayItemType detected. Make sure that FastArraySerializer has a single replicated dynamic array");
+	static_assert(TFastArrayTypeHelper<FastArrayType>::HasValidFastArrayItemType(), "Invalid FastArrayItemType detected. Make sure that FastArraySerializer has a single replicated property that is a dynamic array of the expected type");
 
 	if constexpr (TIsDerivedFrom<FastArrayType, FIrisFastArraySerializer>::IsDerived)
 	{
diff --git a/Engine/Source/Runtime/Experimental/Iris/Core/Public/Iris/ReplicationSystem/Private/FastArrayReplicationFragmentInternal.h b/Engine/Source/Runtime/Experimental/Iris/Core/Public/Iris/ReplicationSystem/Private/FastArrayReplicationFragmentInternal.h
--- a/Engine/Source/Runtime/Experimental/Iris/Core/Public/Iris/ReplicationSystem/Private/FastArrayReplicationFragmentInternal.h
+++ b/Engine/Source/Runtime/Experimental/Iris/Core/Public/Iris/ReplicationSystem/Private/FastArrayReplicationFragmentInternal.h
@@ -9,11 +9,13 @@
 #include "Iris/ReplicationState/PropertyReplicationState.h"
 #include "Iris/ReplicationState/ReplicationStateDescriptor.h"
 #include "Net/Core/Trace/NetDebugName.h"
+#include "Net/Core/NetBitArray.h"
 
-namespace UE {
-namespace Net {
+namespace UE::Net
+{
 
-namespace FastArrayPollingPolicies {
+namespace FastArrayPollingPolicies
+{
 
 struct FPollingState
 {
@@ -48,7 +50,8 @@ public:
 
 } // End of namespace FastArrayPollingPolicies
 
-namespace Private {
+namespace Private
+{
 
 /** Utility methods to behave similar to FastArraySerializer */
 struct FFastArrayReplicationFragmentHelper
@@ -67,6 +70,9 @@ struct FFastArrayReplicationFragmentHelper
 	/** Compare array element, only replicated items will be compared */
 	IRISCORE_API static bool InternalCompareArrayElement(const FReplicationStateDescriptor* ArrayElementDescriptor, void* RESTRICT Dst, const void* RESTRICT Src);
 
+	/** Find the member index of the FastArrayIteArray, used to support FastArrayNetSerializers with extra properties */
+	IRISCORE_API static uint32 GetFastArrayStructItemArrayMemberIndex(const FReplicationStateDescriptor* StructDescriptor);
+
 	/**
 	 * Conditionally invoke PostReplicatedReceive method depending on if it is defined or not
 	 * We only want to do this for FastArrays that define PostReplicatedReceive since it might require extra work to calculate the required parameters
@@ -89,16 +95,12 @@ public:
 	IRISCORE_API void Register(FFragmentRegistrationContext& Context, EReplicationFragmentTraits InTraits);
 
 protected:
-	IRISCORE_API FFastArrayReplicationFragmentBase(EReplicationFragmentTraits InTraits, UObject* InOwner, const FReplicationStateDescriptor* InDescriptor);
-
-	IRISCORE_API const FReplicationStateDescriptor* GetFastArraySerializerPropertyDescriptor() const;
+	IRISCORE_API FFastArrayReplicationFragmentBase(EReplicationFragmentTraits InTraits, UObject* InOwner, const FReplicationStateDescriptor* InDescriptor, bool bValidateDescriptor = true);
+	IRISCORE_API const FReplicationStateDescriptor* GetFastArrayPropertyStructDescriptor() const;
 	IRISCORE_API const FReplicationStateDescriptor* GetArrayElementDescriptor() const;
-
 	IRISCORE_API virtual void CollectOwner(FReplicationStateOwnerCollector* Owners) const override;
 	IRISCORE_API virtual void CallRepNotifies(FReplicationStateApplyContext& Context) override;	
-
 	IRISCORE_API virtual void ReplicatedStateToString(FStringBuilderBase& StringBuilder, FReplicationStateApplyContext& Context, EReplicationStateToStringFlags Flags) const override;
-
 	IRISCORE_API static void InternalCopyArrayElement(const FReplicationStateDescriptor* ArrayElementDescriptor, void* RESTRICT Dst, const void* RESTRICT Src);
 	IRISCORE_API static bool InternalCompareArrayElement(const FReplicationStateDescriptor* ArrayElementDescriptor, void* RESTRICT Dst, const void* RESTRICT Src);
 
@@ -107,7 +109,7 @@ protected:
 	TRefCountPtr<const FReplicationStateDescriptor> ReplicationStateDescriptor;
 
 	// This is the source state from which we source our state data
-	TUniquePtr<FPropertyReplicationState> SrcReplicationState;
+	TUniquePtr<FPropertyReplicationState> ReplicationState;
 
 	// Owner
 	UObject* Owner;
@@ -121,11 +123,9 @@ class FNativeFastArrayReplicationFragmentBase : public FReplicationFragment
 {
 protected:
 	IRISCORE_API FNativeFastArrayReplicationFragmentBase(EReplicationFragmentTraits InTraits, UObject* InOwner, const FReplicationStateDescriptor* InDescriptor);
-
-	IRISCORE_API virtual void CollectOwner(FReplicationStateOwnerCollector* Owners) const override;
-
-	IRISCORE_API const FReplicationStateDescriptor* GetFastArraySerializerPropertyDescriptor() const;
+	IRISCORE_API const FReplicationStateDescriptor* GetFastArrayPropertyStructDescriptor() const;
 	IRISCORE_API const FReplicationStateDescriptor* GetArrayElementDescriptor() const;
+	IRISCORE_API virtual void CollectOwner(FReplicationStateOwnerCollector* Owners) const override;
 
 	// Dequantize state into DstExternalBuffer, note it is expected to be initialized
 	IRISCORE_API static void InternalDequantizeFastArray(FNetSerializationContext& Context, uint8* RESTRICT DstExternalBuffer, const uint8* RESTRICT SrcInternalBuffer, const FReplicationStateDescriptor* FastArrayPropertyDescriptor);
@@ -275,4 +275,4 @@ void FFastArrayReplicationFragmentHelper::ApplyReplicatedState(FastArrayType* Ds
 	CallPostReplicatedReceiveOrNot(*DstArraySerializer, Context.bHasUnresolvableReferences);
 }
 
-}}} // End of namespaces
+}} // End of namespaces
diff --git a/Engine/Source/Runtime/Experimental/Iris/Core/Public/Iris/ReplicationSystem/ReplicationFragment.h b/Engine/Source/Runtime/Experimental/Iris/Core/Public/Iris/ReplicationSystem/ReplicationFragment.h
--- a/Engine/Source/Runtime/Experimental/Iris/Core/Public/Iris/ReplicationSystem/ReplicationFragment.h
+++ b/Engine/Source/Runtime/Experimental/Iris/Core/Public/Iris/ReplicationSystem/ReplicationFragment.h
@@ -174,8 +174,13 @@ protected:
 enum class EFragmentRegistrationFlags : uint32
 {
 	None = 0U,
+	
+	// Indicates that we only should register RPC:s
 	RegisterRPCsOnly = 1U,
+	// Indicates that this objects should use the CDO for class defaults instead of the archetype
 	InitializeDefaultStateFromClassDefaults = RegisterRPCsOnly << 1U,
+	// Allow building descriptors for FastArrays that contain additional properties, NOTE: This should be avoided as fastarrays normally only should contain a single replicated property.
+	AllowFastArraysWithAdditionalProperties = InitializeDefaultStateFromClassDefaults << 1U,
 };
 ENUM_CLASS_FLAGS(EFragmentRegistrationFlags);
 
