Refactor Modular Vehicle into Component and Simulation parts
Refactored controls so it uses the timed network messages which keeps simulation deterministic across client/server

diff --git a/Engine/Source/Runtime/Experimental/Chaos/Private/PhysicsProxy/GeometryCollectionPhysicsProxy.cpp b/Engine/Source/Runtime/Experimental/Chaos/Private/PhysicsProxy/GeometryCollectionPhysicsProxy.cpp
--- a/Engine/Source/Runtime/Experimental/Chaos/Private/PhysicsProxy/GeometryCollectionPhysicsProxy.cpp
+++ b/Engine/Source/Runtime/Experimental/Chaos/Private/PhysicsProxy/GeometryCollectionPhysicsProxy.cpp
@@ -2894,9 +2894,14 @@ bool FGeometryCollectionPhysicsProxy::PullFromPhysicsState(const Chaos::FDirtyGe
 		}
 		else
 		{
+			TManagedArray<bool>* AnimationsActive =
+				GameThreadCollection.FindAttribute<bool>("AnimateTransformAttribute", FGeometryCollection::TransformGroup);
+
 			for (int32 TransformGroupIndex = 0; TransformGroupIndex < NumTransforms; ++TransformGroupIndex)
 			{
-				if (!TargetResults.States[TransformGroupIndex].DisabledState)
+				bool bAnimatingWhileDisabled = AnimationsActive ? (*AnimationsActive)[TransformGroupIndex] : false;
+
+				if (!TargetResults.States[TransformGroupIndex].DisabledState || bAnimatingWhileDisabled)
 				{
 					FParticle& GTParticle = *GTParticles[TransformGroupIndex];
 
