geometry collections should immediately add itself to the external acceleration structure

For large scenes, it may take multiple seconds (if not longer) for a geometry collection to show up in the external acceleration structure causing it to be ignored by movement sweep checks (and others). This is because the geometry collection relies on the internal -> external acceleration structure copy which may take awhile to complete. The single particle physics proxy (i.e. static meshes) will add itself manually to the external acceleration structure in FInitBodiesHelperBase::InitBodies via FChaosScene::AddActorsToScene_AssumesLocked.

diff --git a/Engine/Source/Runtime/Engine/Private/PhysicsEngine/Experimental/PhysScene_Chaos.cpp b/Engine/Source/Runtime/Engine/Private/PhysicsEngine/Experimental/PhysScene_Chaos.cpp
--- a/Engine/Source/Runtime/Engine/Private/PhysicsEngine/Experimental/PhysScene_Chaos.cpp
+++ b/Engine/Source/Runtime/Engine/Private/PhysicsEngine/Experimental/PhysScene_Chaos.cpp
@@ -13,6 +13,7 @@
 #include "PhysicsReplication.h"
 #include "PhysicsEngine/ConstraintInstance.h"
 #include "PhysicsEngine/PhysicsCollisionHandler.h"
+#include "Chaos/PhysicsObjectInternalInterface.h"
 
 #include "ChaosSolversModule.h"
 
@@ -1693,7 +1694,7 @@ void FPhysScene_Chaos::OnSyncBodies(Chaos::FPhysicsSolverBase* Solver)
 	TArray<FPhysScenePendingComponentTransform_Chaos> PendingTransforms;
 	TSet<FGeometryCollectionPhysicsProxy*> GCProxies;
 
-	FPhysicsCommand::ExecuteWrite(this, [&Solver, &PendingTransforms](FPhysScene* PhysScene)
+	FPhysicsCommand::ExecuteWrite(this, [this, &Solver, &PendingTransforms](FPhysScene* PhysScene)
 	{
 		auto RigidLambda = [&PhysScene, &PendingTransforms](Chaos::FSingleParticlePhysicsProxy* Proxy)
 		{
@@ -1766,7 +1767,21 @@ void FPhysScene_Chaos::OnSyncBodies(Chaos::FPhysicsSolverBase* Solver)
 
 		};
 
-		Solver->PullPhysicsStateForEachDirtyProxy_External(RigidLambda, ConstraintLambda);
+		auto GeometryCollectionLambda = [this](FGeometryCollectionPhysicsProxy* Proxy)
+		{
+			// Use the interface interface since we're already locked.
+			Chaos::FWritePhysicsObjectInterface_Internal Interface = Chaos::FPhysicsObjectInternalInterface::GetWrite();
+			TArray<Chaos::FPhysicsObjectHandle> Handles = Proxy->GetAllPhysicsObjects().FilterByPredicate(
+				[&Interface](Chaos::FPhysicsObjectHandle Handle)
+				{
+					return !Interface.AreAllDisabled({ &Handle, 1 });
+				}
+			);
+
+			Interface.AddToSpatialAcceleration(Handles, GetSpacialAcceleration());
+		};
+
+		Solver->PullPhysicsStateForEachDirtyProxy_External(RigidLambda, ConstraintLambda, GeometryCollectionLambda);
 
 	});
 
diff --git a/Engine/Source/Runtime/Experimental/Chaos/Private/Chaos/PhysicsObjectInterface.cpp b/Engine/Source/Runtime/Experimental/Chaos/Private/Chaos/PhysicsObjectInterface.cpp
--- a/Engine/Source/Runtime/Experimental/Chaos/Private/Chaos/PhysicsObjectInterface.cpp
+++ b/Engine/Source/Runtime/Experimental/Chaos/Private/Chaos/PhysicsObjectInterface.cpp
@@ -158,6 +158,22 @@ namespace Chaos
 		return FQuat::Identity;
 	}
 
+	template<EThreadContext Id>
+	FSpatialAccelerationIdx FReadPhysicsObjectInterface<Id>::GetSpatialIndex(FPhysicsObjectHandle Object)
+	{
+		if (!Object)
+		{
+			return {};
+		}
+
+		if (TThreadParticle<Id>* Particle = Object->GetParticle<Id>())
+		{
+			return Particle->SpatialIdx();
+		}
+
+		return {};
+	}
+
 	template<EThreadContext Id>
 	TArray<FPerShapeData*> FReadPhysicsObjectInterface<Id>::GetAllShapes(TArrayView<FPhysicsObjectHandle> InObjects)
 	{
@@ -467,6 +483,12 @@ namespace Chaos
 		return AggregateResult;
 	}
 
+	template<EThreadContext Id>
+	FAccelerationStructureHandle FReadPhysicsObjectInterface<Id>::CreateAccelerationStructureHandle(FPhysicsObjectHandle InObject)
+	{
+		return FAccelerationStructureHandle{InObject->GetParticle<Id>()};
+	}
+
 	template<EThreadContext Id>
 	void FWritePhysicsObjectInterface<Id>::PutToSleep(TArrayView<FPhysicsObjectHandle> InObjects)
 	{
diff --git a/Engine/Source/Runtime/Experimental/Chaos/Private/PBDRigidsSolver.cpp b/Engine/Source/Runtime/Experimental/Chaos/Private/PBDRigidsSolver.cpp
--- a/Engine/Source/Runtime/Experimental/Chaos/Private/PBDRigidsSolver.cpp
+++ b/Engine/Source/Runtime/Experimental/Chaos/Private/PBDRigidsSolver.cpp
@@ -1513,7 +1513,7 @@ namespace Chaos
 	// however.
 	void FPBDRigidsSolver::UpdateGameThreadStructures()
 	{
-		PullPhysicsStateForEachDirtyProxy_External([](auto){}, [](auto) {});
+		PullPhysicsStateForEachDirtyProxy_External([](auto){}, [](auto) {}, [](auto) {});
 	}
 
 	int32 FPBDRigidsSolver::NumJointConstraints() const
diff --git a/Engine/Source/Runtime/Experimental/Chaos/Private/PhysicsProxy/GeometryCollectionPhysicsProxy.cpp b/Engine/Source/Runtime/Experimental/Chaos/Private/PhysicsProxy/GeometryCollectionPhysicsProxy.cpp
--- a/Engine/Source/Runtime/Experimental/Chaos/Private/PhysicsProxy/GeometryCollectionPhysicsProxy.cpp
+++ b/Engine/Source/Runtime/Experimental/Chaos/Private/PhysicsProxy/GeometryCollectionPhysicsProxy.cpp
@@ -4177,6 +4177,19 @@ void FGeometryCollectionPhysicsProxy::FieldForcesUpdateCallback(Chaos::FPBDRigid
 	}
 }
 
+TArray<Chaos::FPhysicsObjectHandle> FGeometryCollectionPhysicsProxy::GetAllPhysicsObjects()
+{
+	TArray<Chaos::FPhysicsObjectHandle> Handles;
+	Handles.Reserve(PhysicsObjects.Num());
+
+	for (Chaos::FPhysicsObjectUniquePtr& Object : PhysicsObjects)
+	{
+		Handles.Add(Object.Get());
+	}
+
+	return Handles;
+}
+
 Chaos::FPhysicsObjectHandle FGeometryCollectionPhysicsProxy::GetPhysicsObjectByIndex(int32 Index) const
 {
 	if (!PhysicsObjects.IsValidIndex(Index))
diff --git a/Engine/Source/Runtime/Experimental/Chaos/Public/Chaos/Framework/PhysicsSolverBase.h b/Engine/Source/Runtime/Experimental/Chaos/Public/Chaos/Framework/PhysicsSolverBase.h
--- a/Engine/Source/Runtime/Experimental/Chaos/Public/Chaos/Framework/PhysicsSolverBase.h
+++ b/Engine/Source/Runtime/Experimental/Chaos/Public/Chaos/Framework/PhysicsSolverBase.h
@@ -492,7 +492,11 @@ namespace Chaos
 		}
 
 		/** Used to update external thread data structures. RigidFunc allows per dirty rigid code to execute. Include PhysicsSolverBaseImpl.h to call this function*/
+		template <typename RigidLambda, typename ConstraintLambda, typename GeometryCollectionLambda>
+		void PullPhysicsStateForEachDirtyProxy_External(const RigidLambda& RigidFunc, const ConstraintLambda& ConstraintFunc, const GeometryCollectionLambda& GeometryCollectionFunc);
+
 		template <typename RigidLambda, typename ConstraintLambda>
+		UE_DEPRECATED(5.4, "Use PullPhysicsStateForEachDirtyProxy_External with the additional GeometryCollectionLambda parameter instead.")
 		void PullPhysicsStateForEachDirtyProxy_External(const RigidLambda& RigidFunc, const ConstraintLambda& ConstraintFunc);
 
 		bool IsUsingAsyncResults() const
diff --git a/Engine/Source/Runtime/Experimental/Chaos/Public/Chaos/PhysicsObjectInterface.h b/Engine/Source/Runtime/Experimental/Chaos/Public/Chaos/PhysicsObjectInterface.h
--- a/Engine/Source/Runtime/Experimental/Chaos/Public/Chaos/PhysicsObjectInterface.h
+++ b/Engine/Source/Runtime/Experimental/Chaos/Public/Chaos/PhysicsObjectInterface.h
@@ -2,7 +2,9 @@
 
 #pragma once
 
+#include "Chaos/ParticleHandle.h"
 #include "Chaos/PhysicsObject.h"
+#include "Chaos/ISpatialAcceleration.h"
 #include "Framework/Threading.h"
 #include "Containers/ArrayView.h"
 #include "Math/MathFwd.h"
@@ -54,6 +56,7 @@ namespace Chaos
 		FVector GetCoM(FPhysicsObjectHandle Object);
 		FVector GetWorldCoM(FPhysicsObjectHandle Object);
 		FQuat GetR(FPhysicsObjectHandle Object);
+		FSpatialAccelerationIdx GetSpatialIndex(FPhysicsObjectHandle Object);
 
 		TArray<FPerShapeData*> GetAllShapes(TArrayView<FPhysicsObjectHandle> InObjects);
 		bool GetPhysicsObjectOverlap(FPhysicsObjectHandle ObjectA, FPhysicsObjectHandle ObjectB, bool bTraceComplex, Chaos::FMTDInfo* OutMTD);
@@ -68,7 +71,7 @@ namespace Chaos
 		FBox GetBounds(TArrayView<FPhysicsObjectHandle> InObjects);
 		FBox GetWorldBounds(TArrayView<FPhysicsObjectHandle> InObjects);
 		FClosestPhysicsObjectResult GetClosestPhysicsBodyFromLocation(TArrayView<FPhysicsObjectHandle> InObjects, const FVector& WorldLocation);
-
+		FAccelerationStructureHandle CreateAccelerationStructureHandle(FPhysicsObjectHandle Handle);
 		friend class FPhysicsObjectInterface;
 	protected:
 		FReadPhysicsObjectInterface() = default;
@@ -89,6 +92,23 @@ namespace Chaos
 		void AddForce(TArrayView<FPhysicsObjectHandle> InObjects, const FVector& Force, bool bInvalidate);
 		void AddTorque(TArrayView<FPhysicsObjectHandle> InObjects, const FVector& Torque, bool bInvalidate);
 
+		template<typename TPayloadType, typename T, int d>
+		void AddToSpatialAcceleration(TArrayView<FPhysicsObjectHandle> InObjects, ISpatialAcceleration<TPayloadType, T, d>* SpatialAcceleration)
+		{
+			if (!SpatialAcceleration)
+			{
+				return;
+			}
+
+			for (FPhysicsObjectHandle Handle : InObjects)
+			{
+				const FBox WorldBounds = this->GetWorldBounds({ &Handle, 1 });
+				const FAABB3 ChaosWorldBounds{ WorldBounds.Min, WorldBounds.Max };
+				FAccelerationStructureHandle AccelerationHandle = this->CreateAccelerationStructureHandle(Handle);
+				SpatialAcceleration->UpdateElementIn(AccelerationHandle, ChaosWorldBounds, true, this->GetSpatialIndex(Handle));
+			}
+		}
+
 		friend class FPhysicsObjectInterface;
 	protected:
 		FWritePhysicsObjectInterface() = default;
diff --git a/Engine/Source/Runtime/Experimental/Chaos/Public/Chaos/PhysicsSolverBaseImpl.h b/Engine/Source/Runtime/Experimental/Chaos/Public/Chaos/PhysicsSolverBaseImpl.h
--- a/Engine/Source/Runtime/Experimental/Chaos/Public/Chaos/PhysicsSolverBaseImpl.h
+++ b/Engine/Source/Runtime/Experimental/Chaos/Public/Chaos/PhysicsSolverBaseImpl.h
@@ -16,8 +16,8 @@
 namespace Chaos
 {
 	// Pulls physics state for each dirty particle and allows caller to do additional work if needed
-	template <typename RigidLambda, typename ConstraintLambda>
-	void FPhysicsSolverBase::PullPhysicsStateForEachDirtyProxy_External(const RigidLambda& RigidFunc, const ConstraintLambda& ConstraintFunc)
+	template <typename RigidLambda, typename ConstraintLambda, typename GeometryCollectionLambda>
+	void FPhysicsSolverBase::PullPhysicsStateForEachDirtyProxy_External(const RigidLambda& RigidFunc, const ConstraintLambda& ConstraintFunc, const GeometryCollectionLambda& GeometryCollectionFunc)
 	{
 		using namespace Chaos;
 
@@ -33,7 +33,7 @@ namespace Chaos
 			//case 5: prev has no dirty data and next was overwritten. In this case do nothing as the overwritten data wins, and also particle may be deleted
 
 			TArray<const FChaosInterpolationResults*> ResultsPerChannel = PullResultsManager->PullAsyncPhysicsResults_External(ResultsTime);
-			for(int32 ChannelIdx = 0; ChannelIdx < ResultsPerChannel.Num(); ++ChannelIdx)
+			for (int32 ChannelIdx = 0; ChannelIdx < ResultsPerChannel.Num(); ++ChannelIdx)
 			{
 				const FChaosInterpolationResults& Results = *ResultsPerChannel[ChannelIdx];
 				LatestData = Results.Next;
@@ -51,13 +51,16 @@ namespace Chaos
 						}
 					}
 				}
-				
+
 				// geometry collections
 				for (const FChaosGeometryCollectionInterpolationData& GCInterp : Results.GeometryCollectionInterpolations)
 				{
 					if (FGeometryCollectionPhysicsProxy* Proxy = GCInterp.Prev.GetProxy())
 					{
-						Proxy->PullFromPhysicsState(GCInterp.Prev, SolverTimestamp, &GCInterp.Next, &Results.Alpha);
+						if (Proxy->PullFromPhysicsState(GCInterp.Prev, SolverTimestamp, &GCInterp.Next, &Results.Alpha))
+						{
+							GeometryCollectionFunc(Proxy);
+						}
 					}
 				}
 			}
@@ -75,31 +78,34 @@ namespace Chaos
 
 			// Single particles
 			for (const FChaosRigidInterpolationData& RigidInterp : Results.RigidInterpolations)
+			{
+				if (FSingleParticlePhysicsProxy* Proxy = RigidInterp.Prev.GetProxy())
 				{
-					if (FSingleParticlePhysicsProxy* Proxy = RigidInterp.Prev.GetProxy())
+					if (Proxy->PullFromPhysicsState(RigidInterp.Next, SolverTimestamp))
 					{
-						if (Proxy->PullFromPhysicsState(RigidInterp.Next, SolverTimestamp))
-						{
-							RigidFunc(Proxy);
-						}
+						RigidFunc(Proxy);
 					}
 				}
+			}
 
 			// geometry collections
 			for (const FChaosGeometryCollectionInterpolationData& GCInterp : Results.GeometryCollectionInterpolations)
 			{
 				if (FGeometryCollectionPhysicsProxy* Proxy = GCInterp.Prev.GetProxy())
 				{
-					Proxy->PullFromPhysicsState(GCInterp.Next, SolverTimestamp);
+					if (Proxy->PullFromPhysicsState(GCInterp.Next, SolverTimestamp))
+					{
+						GeometryCollectionFunc(Proxy);
+					}
 				}
 			}
-			
+
 		}
 
 		//no interpolation for GC or joints at the moment
-		if(LatestData)
+		if (LatestData)
 		{
-			 const int32 SyncTimestamp = LatestData->SolverTimestamp;
+			const int32 SyncTimestamp = LatestData->SolverTimestamp;
 
 			//
 			// @todo(chaos) : Add Dirty Constraints Support
@@ -126,4 +132,10 @@ namespace Chaos
 		}
 	}
 
+	template <typename RigidLambda, typename ConstraintLambda>
+	void FPhysicsSolverBase::PullPhysicsStateForEachDirtyProxy_External(const RigidLambda& RigidFunc, const ConstraintLambda& ConstraintFunc)
+	{
+		auto GeometryCollectionFunc = [](FGeometryCollectionPhysicsProxy* Proxy) {};
+		PullPhysicsStateForEachDirtyProxy_External(RigidFunc, ConstraintFunc, GeometryCollectionFunc);
+	}
 }
diff --git a/Engine/Source/Runtime/Experimental/Chaos/Public/PhysicsProxy/GeometryCollectionPhysicsProxy.h b/Engine/Source/Runtime/Experimental/Chaos/Public/PhysicsProxy/GeometryCollectionPhysicsProxy.h
--- a/Engine/Source/Runtime/Experimental/Chaos/Public/PhysicsProxy/GeometryCollectionPhysicsProxy.h
+++ b/Engine/Source/Runtime/Experimental/Chaos/Public/PhysicsProxy/GeometryCollectionPhysicsProxy.h
@@ -429,6 +429,7 @@ public:
 		PostPhysicsSyncCallback = Callback;
 	}
 	
+	TArray<Chaos::FPhysicsObjectHandle> GetAllPhysicsObjects();
 	Chaos::FPhysicsObjectHandle GetPhysicsObjectByIndex(int32 Index) const;
 	int32 GetNumParticles() const { return NumParticles; }
 protected:
diff --git a/Engine/Source/Runtime/Experimental/GeometryCollectionEngine/Private/GeometryCollection/GeometryCollectionComponent.cpp b/Engine/Source/Runtime/Experimental/GeometryCollectionEngine/Private/GeometryCollection/GeometryCollectionComponent.cpp
--- a/Engine/Source/Runtime/Experimental/GeometryCollectionEngine/Private/GeometryCollection/GeometryCollectionComponent.cpp
+++ b/Engine/Source/Runtime/Experimental/GeometryCollectionEngine/Private/GeometryCollection/GeometryCollectionComponent.cpp
@@ -5,6 +5,7 @@
 #include "AI/NavigationSystemHelpers.h"
 #include "Async/ParallelFor.h"
 #include "Chaos/ChaosPhysicalMaterial.h"
+#include "Chaos/ChaosScene.h"
 #include "ChaosSolversModule.h"
 #include "ChaosStats.h"
 #include "ComponentRecreateRenderStateContext.h"
@@ -34,6 +35,7 @@
 #include "PhysicalMaterials/PhysicalMaterial.h"
 #include "Physics/Experimental/PhysScene_Chaos.h"
 #include "Physics/PhysicsFiltering.h"
+#include "PhysicsEngine/PhysicsObjectExternalInterface.h"
 #include "PhysicsField/PhysicsFieldComponent.h"
 #include "PhysicsProxy/GeometryCollectionPhysicsProxy.h"
 #include "PhysicsSolver.h"
@@ -2633,7 +2635,8 @@ void UGeometryCollectionComponent::RegisterAndInitializePhysicsProxy()
 	}
 #endif
 	PhysicsProxy = new FGeometryCollectionPhysicsProxy(this, *DynamicCollection, SimulationParameters, InitialSimFilter, InitialQueryFilter, CollectorGuid);
-	PhysicsProxy->SetPostPhysicsSyncCallback([this]() { UpdateAttachedChildrenTransform(); }); 
+	PhysicsProxy->SetPostPhysicsSyncCallback([this]() { UpdateAttachedChildrenTransform(); });
+
 	if (GetIsReplicated())
 	{
 		// using net mode and not local role because at this time in the initialization client and server both have an authority local role
@@ -2682,6 +2685,15 @@ void UGeometryCollectionComponent::RegisterAndInitializePhysicsProxy()
 		}
 	}
 
+	// We need to add the geometry collection into the external acceleration structure so that it's immediately available for queries instead of waiting for the sync from the physics thread (which could take awhile).
+	// Just adding the root particle should be sufficient since that'll be the only particle we'd expect any collisions with right after initialization.
+	if (Chaos::FPhysicsObjectHandle RootObject = GetPhysicsObjectByName(NAME_None))
+	{
+		TArrayView<Chaos::FPhysicsObjectHandle> Handles{ &RootObject, 1 };
+		FLockedWritePhysicsObjectExternalInterface Interface = FPhysicsObjectExternalInterface::LockWrite(Handles);
+		Interface->AddToSpatialAcceleration(Handles, Scene->GetSpacialAcceleration());
+	}
+
 	RegisterForEvents();
 	SetAsyncPhysicsTickEnabled(GetIsReplicated());
 }
diff --git a/Engine/Source/Runtime/PhysicsCore/Private/ChaosScene.cpp b/Engine/Source/Runtime/PhysicsCore/Private/ChaosScene.cpp
--- a/Engine/Source/Runtime/PhysicsCore/Private/ChaosScene.cpp
+++ b/Engine/Source/Runtime/PhysicsCore/Private/ChaosScene.cpp
@@ -377,7 +377,7 @@ void FChaosScene::StartFrame()
 
 void FChaosScene::OnSyncBodies(Chaos::FPhysicsSolverBase* Solver)
 {
-	Solver->PullPhysicsStateForEachDirtyProxy_External([](auto){}, [](auto) {});
+	Solver->PullPhysicsStateForEachDirtyProxy_External([](auto) {}, [](auto) {}, [](auto) {});
 }
 
 bool FChaosScene::AreAnyTasksPending() const
