Chaos
- Fix potential uninitialized PreObjectState in Rigid Particle Handle
- Add IsSleeping() to GenericParticleHandle

diff --git a/Engine/Source/Runtime/Experimental/Chaos/Public/Chaos/ParticleHandle.h b/Engine/Source/Runtime/Experimental/Chaos/Public/Chaos/ParticleHandle.h
--- a/Engine/Source/Runtime/Experimental/Chaos/Public/Chaos/ParticleHandle.h
+++ b/Engine/Source/Runtime/Experimental/Chaos/Public/Chaos/ParticleHandle.h
@@ -877,6 +877,7 @@ protected:
 		SetAcceleration(TVector<T, d>(0));
 		SetAngularAcceleration(TVector<T, d>(0));
 		SetObjectStateLowLevel(Params.bStartSleeping ? EObjectStateType::Sleeping : EObjectStateType::Dynamic);
+		SetPreObjectStateLowLevel(ObjectState());
 		SetIslandIndex(INDEX_NONE);
 		SetSleepType(ESleepType::MaterialSleep);
 		SetInvIConditioning(TVec3<FRealSingle>(1));
@@ -1417,6 +1418,7 @@ public:
 	bool IsStatic() const { return (MHandle->ObjectState() == EObjectStateType::Static); }
 	bool IsKinematic() const { return (MHandle->ObjectState() == EObjectStateType::Kinematic); }
 	bool IsDynamic() const { return (MHandle->ObjectState() == EObjectStateType::Dynamic) || (MHandle->ObjectState() == EObjectStateType::Sleeping); }
+	bool IsSleeping() const { return (MHandle->ObjectState() == EObjectStateType::Sleeping); }
 
 	const FKinematicGeometryParticleHandle* CastToKinematicParticle() const { return MHandle->CastToKinematicParticle(); }
 	FKinematicGeometryParticleHandle* CastToKinematicParticle() { return MHandle->CastToKinematicParticle(); }
