geometry collection crumble events should not get spammed after the last crumble (i.e. do not send crumble events for disabled clusters)

diff --git a/Engine/Source/Runtime/Experimental/Chaos/Private/Chaos/PBDRigidClustering.cpp b/Engine/Source/Runtime/Experimental/Chaos/Private/Chaos/PBDRigidClustering.cpp
--- a/Engine/Source/Runtime/Experimental/Chaos/Private/Chaos/PBDRigidClustering.cpp
+++ b/Engine/Source/Runtime/Experimental/Chaos/Private/Chaos/PBDRigidClustering.cpp
@@ -1332,6 +1332,11 @@ namespace Chaos
 
 	bool FRigidClustering::BreakCluster(FPBDRigidClusteredParticleHandle* ClusteredParticle)
 	{
+		if (!ClusteredParticle || ClusteredParticle->Disabled())
+		{
+			return false;
+		}
+
 		// max strain will allow to unconditionally release the children when strain is evaluated
 		constexpr FReal MaxStrain = TNumericLimits<FReal>::Max();
 		if (TArray<FPBDRigidParticleHandle*>* ChildrenHandles = GetChildrenMap().Find(ClusteredParticle))
