fix space transform usage in FReadPhysicsObjectInterface<Id>::GetPhysicsObjectOverlap and Overlap_GeomInternal to properly account for transformed implicit geometry

diff --git a/Engine/Source/Runtime/Engine/Private/PhysicsEngine/Experimental/PhysInterface_Chaos.cpp b/Engine/Source/Runtime/Engine/Private/PhysicsEngine/Experimental/PhysInterface_Chaos.cpp
--- a/Engine/Source/Runtime/Engine/Private/PhysicsEngine/Experimental/PhysInterface_Chaos.cpp
+++ b/Engine/Source/Runtime/Engine/Private/PhysicsEngine/Experimental/PhysInterface_Chaos.cpp
@@ -1045,7 +1045,7 @@ bool Overlap_GeomInternal(const FBodyInstance* InInstance, const Chaos::FImplici
 				if (OutOptResult)
 				{
 					Chaos::FMTDInfo MTDInfo;
-					if (Chaos::Utilities::CastHelper(InGeom, ActorTM, [&](const auto& Downcast, const auto& FullActorTM) { return Chaos::OverlapQuery(*Shape->GetGeometry(), FullActorTM, Downcast, GeomTransform, /*Thickness=*/0, &MTDInfo); }))
+					if (Chaos::Utilities::CastHelper(InGeom, GeomTransform, [&](const auto& Downcast, const auto& FullGeomTransform) { return Chaos::OverlapQuery(*Shape->GetGeometry(), ActorTM, Downcast, FullGeomTransform, /*Thickness=*/0, &MTDInfo); }))
 					{
 						OutOptResult->Distance = MTDInfo.Penetration;
 						OutOptResult->Direction = MTDInfo.Normal;
@@ -1054,7 +1054,7 @@ bool Overlap_GeomInternal(const FBodyInstance* InInstance, const Chaos::FImplici
 				}
 				else	//question: why do we even allow user to not pass in MTD info?
 				{
-					if (Chaos::Utilities::CastHelper(InGeom, ActorTM, [&](const auto& Downcast, const auto& FullActorTM) { return Chaos::OverlapQuery(*Shape->GetGeometry(), FullActorTM, Downcast, GeomTransform); }))
+					if (Chaos::Utilities::CastHelper(InGeom, GeomTransform, [&](const auto& Downcast, const auto& FullGeomTransform) { return Chaos::OverlapQuery(*Shape->GetGeometry(), ActorTM, Downcast, FullGeomTransform); }))
 					{
 						return true;
 					}
diff --git a/Engine/Source/Runtime/Experimental/Chaos/Private/Chaos/PhysicsObjectInterface.cpp b/Engine/Source/Runtime/Experimental/Chaos/Private/Chaos/PhysicsObjectInterface.cpp
--- a/Engine/Source/Runtime/Experimental/Chaos/Private/Chaos/PhysicsObjectInterface.cpp
+++ b/Engine/Source/Runtime/Experimental/Chaos/Private/Chaos/PhysicsObjectInterface.cpp
@@ -215,55 +215,51 @@ namespace Chaos
 			return false;
 		}
 
-		for (FPerShapeData* A : ShapesA)
+		for (FPerShapeData* B : ShapesB)
 		{
-			if (!A)
+			if (!B)
 			{
 				continue;
 			}
 
-			FCollisionFilterData ShapeFilterA = A->GetQueryData();
-			const bool bShapeIsComplexA = (ShapeFilterA.Word3 & static_cast<uint8>(EFilterFlags::ComplexCollision)) != 0;
-			const bool bShapeIsSimpleA = (ShapeFilterA.Word3 & static_cast<uint8>(EFilterFlags::SimpleCollision)) != 0;
-			const bool bShouldTraceA = (bTraceComplex && bShapeIsComplexA) || (!bTraceComplex && bShapeIsSimpleA);
-			if (!bShouldTraceA)
+			const TSerializablePtr<FImplicitObject> GeomB = B->GetGeometry();
+			if (!GeomB || !GeomB->IsConvex())
 			{
 				continue;
 			}
 
-			const FAABB3 BoxShapeA = A->GetGeometry()->CalculateTransformedBounds(TransformA);
-
-			for (FPerShapeData* B : ShapesB)
+			const FAABB3 BoxShapeB = GeomB->CalculateTransformedBounds(TransformB);
+			// At this point on, this function should be mirror the Overlap_GeomInternal function in PhysInterface_Chaos.cpp.
+			// ShapeA is equivalent to InInstance and GeomB is equivalent to InGeom.
+			
+			for (FPerShapeData* A : ShapesA)
 			{
-				if (!B)
+				if (!A)
 				{
 					continue;
 				}
 
-				FCollisionFilterData ShapeFilterB = B->GetQueryData();
-				const bool bShapeIsComplexB = (ShapeFilterB.Word3 & static_cast<uint8>(EFilterFlags::ComplexCollision)) != 0;
-				const bool bShapeIsSimpleB = (ShapeFilterB.Word3 & static_cast<uint8>(EFilterFlags::SimpleCollision)) != 0;
-				const bool bShouldTraceB = (bTraceComplex && bShapeIsComplexB) || (!bTraceComplex && bShapeIsSimpleB);
-
-				if (!bShouldTraceB)
+				FCollisionFilterData ShapeFilter = A->GetQueryData();
+				const bool bShapeIsComplex = (ShapeFilter.Word3 & static_cast<uint8>(EFilterFlags::ComplexCollision)) != 0;
+				const bool bShapeIsSimple = (ShapeFilter.Word3 & static_cast<uint8>(EFilterFlags::SimpleCollision)) != 0;
+				const bool bShouldTrace = (bTraceComplex && bShapeIsComplex) || (!bTraceComplex && bShapeIsSimple);
+				if (!bShouldTrace)
 				{
 					continue;
 				}
 
-				const FAABB3 BoxShapeB = B->GetGeometry()->CalculateTransformedBounds(TransformB);
-
+				const FAABB3 BoxShapeA = A->GetGeometry()->CalculateTransformedBounds(TransformA);
 				if (!BoxShapeA.Intersects(BoxShapeB))
 				{
 					continue;
 				}
 
 				const bool bOverlap = Chaos::Utilities::CastHelper(
-					*B->GetGeometry(),
-					// ...Why is this transform passed?
+					*GeomB,
 					TransformB,
-					[A, &TransformA, &TransformB, OutMTD](const auto& Downcast, const auto&)
+					[A, &TransformA, OutMTD](const auto& Downcast, const auto& FullTransformB)
 					{
-						return Chaos::OverlapQuery(*A->GetGeometry(), TransformA, Downcast, TransformB, 0, OutMTD);
+						return Chaos::OverlapQuery(*A->GetGeometry(), TransformA, Downcast, FullTransformB, 0, OutMTD);
 					}
 				);
 
@@ -427,8 +423,8 @@ namespace Chaos
 			FBox ParticleBox(ForceInit);
 			if (const FImplicitObject* Geometry = Particle->Geometry().Get(); Geometry && Geometry->HasBoundingBox())
 			{
-				const Chaos::FAABB3 Box = Geometry->BoundingBox().TransformedAABB(WorldTransform);
-				ParticleBox = FBox{ Box.Min(), Box.Max() };
+				const Chaos::FAABB3 WorldBox = Geometry->CalculateTransformedBounds(WorldTransform);
+				ParticleBox = FBox{ WorldBox.Min(), WorldBox.Max() };
 			}
 
 			if (ParticleBox.IsValid)
