Add midphase contact modifier which allows users to force CCD on or off per midphase pair.

diff --git a/Engine/Source/Runtime/Experimental/Chaos/Private/Chaos/Collision/ParticlePairMidPhase.cpp b/Engine/Source/Runtime/Experimental/Chaos/Private/Chaos/Collision/ParticlePairMidPhase.cpp
--- a/Engine/Source/Runtime/Experimental/Chaos/Private/Chaos/Collision/ParticlePairMidPhase.cpp
+++ b/Engine/Source/Runtime/Experimental/Chaos/Private/Chaos/Collision/ParticlePairMidPhase.cpp
@@ -847,6 +847,11 @@ namespace Chaos
 		Flags.bIsSleeping = false;
 	}
 
+	void FParticlePairMidPhase::ResetModifications()
+	{
+		Flags.bIsCCDActive = Flags.bIsCCD;
+	}
+
 	void FParticlePairMidPhase::Init(
 		FGeometryParticleHandle* InParticle0,
 		FGeometryParticleHandle* InParticle1,
@@ -859,7 +864,16 @@ namespace Chaos
 		Particle1 = InParticle1;
 		Key = InKey;
 
-		Flags.bIsCCD = Context.GetSettings().bAllowCCD && (FConstGenericParticleHandle(Particle0)->CCDEnabled() || FConstGenericParticleHandle(Particle1)->CCDEnabled());
+		// If CCD is allowed in the current context and for at least one of
+		// the particles involved, enable it for this midphase.
+		//
+		// bIsCCDActive is reset to bIsCCD each frame, but can be overridden
+		// by modifiers.
+		const bool bIsCCD = Context.GetSettings().bAllowCCD && (
+			FConstGenericParticleHandle(Particle0)->CCDEnabled() ||
+			FConstGenericParticleHandle(Particle1)->CCDEnabled());
+		Flags.bIsCCD = bIsCCD;
+		Flags.bIsCCDActive = bIsCCD;
 
 		BuildDetectors();
 
@@ -925,7 +939,11 @@ namespace Chaos
 
 	bool FParticlePairMidPhase::ShouldEnableCCD(const FReal Dt)
 	{
-		if (Flags.bIsCCD)
+		// bIsCCDActive is set to bIsCCD at the beginning of every frame, but may be
+		// overridden in midphase modification or potentially other systems which run
+		// in between mid and narrow phase. bIsCCDActive indicates the final
+		// overridden value so we use that here instead of bIsCCD.
+		if (Flags.bIsCCDActive)
 		{
 			FConstGenericParticleHandle ConstParticle0 = FConstGenericParticleHandle(Particle0);
 			FConstGenericParticleHandle ConstParticle1 = FConstGenericParticleHandle(Particle1);
diff --git a/Engine/Source/Runtime/Experimental/Chaos/Private/Chaos/MidPhaseModification.cpp b/Engine/Source/Runtime/Experimental/Chaos/Private/Chaos/MidPhaseModification.cpp
new file mode 100644
--- /dev/null
+++ b/Engine/Source/Runtime/Experimental/Chaos/Private/Chaos/MidPhaseModification.cpp
@@ -0,0 +1,57 @@
+// Copyright Epic Games, Inc. All Rights Reserved.
+
+#include "Chaos/MidPhaseModification.h"
+#include "Chaos/Collision/PBDCollisionConstraint.h"
+#include "Chaos/Collision/ParticlePairMidPhase.h"
+#include "Chaos/Particle/ParticleUtilities.h"
+#include "Chaos/PBDCollisionConstraints.h"
+#include "Chaos/ParticleHandleFwd.h"
+
+namespace Chaos
+{
+	void FMidPhaseModifier::DisableCCD()
+	{
+		MidPhase->SetCCDIsActive(false);
+	}
+
+	void FMidPhaseModifier::GetParticles(
+		const FGeometryParticleHandle** Particle0,
+		const FGeometryParticleHandle** Particle1) const
+	{
+		*Particle0 = MidPhase->GetParticle0();
+		*Particle1 = MidPhase->GetParticle1();
+	}
+
+	const FGeometryParticleHandle* FMidPhaseModifier::GetOtherParticle(const FGeometryParticleHandle* InParticle) const
+	{
+		if (MidPhase)
+		{
+			const FGeometryParticleHandle* Particle0 = MidPhase->GetParticle0();
+			const FGeometryParticleHandle* Particle1 = MidPhase->GetParticle1();
+			if (InParticle == Particle0)
+			{
+				return Particle1;
+			}
+			else if (InParticle == Particle1)
+			{
+				return Particle0;
+			}
+		}
+		return nullptr;
+	}
+
+	FMidPhaseModifierParticleIterator FMidPhaseModifierParticleRange::begin() const
+	{
+		return FMidPhaseModifierParticleIterator(Accessor, Particle);
+	}
+
+	FMidPhaseModifierParticleIterator FMidPhaseModifierParticleRange::end() const
+	{
+		return FMidPhaseModifierParticleIterator(Accessor, Particle, Particle->ParticleCollisions().Num());
+	}
+
+	FMidPhaseModifierParticleRange FMidPhaseModifierAccessor::GetMidPhases(FGeometryParticleHandle* Particle)
+	{
+		return FMidPhaseModifierParticleRange(this, Particle);
+	}
+}
diff --git a/Engine/Source/Runtime/Experimental/Chaos/Private/Chaos/PBDCollisionConstraints.cpp b/Engine/Source/Runtime/Experimental/Chaos/Private/Chaos/PBDCollisionConstraints.cpp
--- a/Engine/Source/Runtime/Experimental/Chaos/Private/Chaos/PBDCollisionConstraints.cpp
+++ b/Engine/Source/Runtime/Experimental/Chaos/Private/Chaos/PBDCollisionConstraints.cpp
@@ -4,6 +4,7 @@
 
 #include "Chaos/ChaosPerfTest.h"
 #include "Chaos/ContactModification.h"
+#include "Chaos/MidPhaseModification.h"
 #include "Chaos/PBDCollisionConstraintsContact.h"
 #include "Chaos/CollisionResolution.h"
 #include "Chaos/Collision/CollisionPruning.h"
@@ -331,6 +332,15 @@ namespace Chaos
 		}
 	}
 
+	void FPBDCollisionConstraints::ApplyMidPhaseModifier(const TArray<ISimCallbackObject*>& MidPhaseModifiers, FReal Dt)
+	{
+		FMidPhaseModifierAccessor ModifierAccessor;
+		for(ISimCallbackObject* ModifierCallback : MidPhaseModifiers)
+		{
+			ModifierCallback->MidPhaseModification_Internal(ModifierAccessor);
+		}
+	}
+
 	void FPBDCollisionConstraints::ApplyCollisionModifier(const TArray<ISimCallbackObject*>& CollisionModifiers, FReal Dt)
 	{
 		if (GetConstraints().Num() > 0)
diff --git a/Engine/Source/Runtime/Experimental/Chaos/Private/Chaos/PBDRigidsEvolutionGBF.cpp b/Engine/Source/Runtime/Experimental/Chaos/Private/Chaos/PBDRigidsEvolutionGBF.cpp
--- a/Engine/Source/Runtime/Experimental/Chaos/Private/Chaos/PBDRigidsEvolutionGBF.cpp
+++ b/Engine/Source/Runtime/Experimental/Chaos/Private/Chaos/PBDRigidsEvolutionGBF.cpp
@@ -108,6 +108,7 @@ namespace Chaos
 		DECLARE_CYCLE_STAT(TEXT("FPBDRigidsEvolutionGBF::KinematicTargets"), STAT_Evolution_KinematicTargets, STATGROUP_Chaos);
 		DECLARE_CYCLE_STAT(TEXT("FPBDRigidsEvolutionGBF::PostIntegrateCallback"), STAT_Evolution_PostIntegrateCallback, STATGROUP_Chaos);
 		DECLARE_CYCLE_STAT(TEXT("FPBDRigidsEvolutionGBF::CollisionModifierCallback"), STAT_Evolution_CollisionModifierCallback, STATGROUP_Chaos);
+		DECLARE_CYCLE_STAT(TEXT("FPBDRigidsEvolutionGBF::MidPhaseModifierCallback"), STAT_Evolution_MidPhaseModifierCallback, STATGROUP_Chaos);
 		DECLARE_CYCLE_STAT(TEXT("FPBDRigidsEvolutionGBF::CCD"), STAT_Evolution_CCD, STATGROUP_Chaos);
 		DECLARE_CYCLE_STAT(TEXT("FPBDRigidsEvolutionGBF::CCDCorrection"), STAT_Evolution_CCDCorrection, STATGROUP_Chaos);
 		DECLARE_CYCLE_STAT(TEXT("FPBDRigidsEvolutionGBF::GraphColor"), STAT_Evolution_GraphColor, STATGROUP_Chaos);
@@ -367,6 +368,11 @@ void FPBDRigidsEvolutionGBF::AdvanceOneTimeStepImpl(const FReal Dt, const FSubSt
 
 		CollisionDetector.RunBroadPhase(Dt, GetCurrentStepResimCache());
 
+		{
+			SCOPE_CYCLE_COUNTER(STAT_Evolution_MidPhaseModifierCallback);
+			CollisionConstraints.ApplyMidPhaseModifier(*MidPhaseModifiers, Dt);
+		}
+
 		CollisionDetector.RunNarrowPhase(Dt, GetCurrentStepResimCache());
 	}
 
@@ -621,7 +627,7 @@ void FPBDRigidsEvolutionGBF::TestModeResetCollisions()
 	}
 }
 
-FPBDRigidsEvolutionGBF::FPBDRigidsEvolutionGBF(FPBDRigidsSOAs& InParticles,THandleArray<FChaosPhysicsMaterial>& SolverPhysicsMaterials, const TArray<ISimCallbackObject*>* InCollisionModifiers, bool InIsSingleThreaded)
+FPBDRigidsEvolutionGBF::FPBDRigidsEvolutionGBF(FPBDRigidsSOAs& InParticles, THandleArray<FChaosPhysicsMaterial>& SolverPhysicsMaterials, const TArray<ISimCallbackObject*>* InMidPhaseModifiers, const TArray<ISimCallbackObject*>* InCollisionModifiers, bool InIsSingleThreaded)
 	: Base(InParticles, SolverPhysicsMaterials, InIsSingleThreaded)
 	, Clustering(*this, Particles.GetClusteredParticles())
 	, CollisionConstraints(InParticles, Collided, PhysicsMaterials, PerParticlePhysicsMaterials, &SolverPhysicsMaterials, CalculateNumCollisionsPerBlock(), DefaultRestitutionThreshold)
@@ -630,6 +636,7 @@ FPBDRigidsEvolutionGBF::FPBDRigidsEvolutionGBF(FPBDRigidsSOAs& InParticles,THand
 	, PostIntegrateCallback(nullptr)
 	, PreApplyCallback(nullptr)
 	, CurrentStepResimCacheImp(nullptr)
+	, MidPhaseModifiers(InMidPhaseModifiers)
 	, CollisionModifiers(InCollisionModifiers)
 	, CCDManager()
 	, bIsDeterministic(false)
diff --git a/Engine/Source/Runtime/Experimental/Chaos/Private/PBDRigidsSolver.cpp b/Engine/Source/Runtime/Experimental/Chaos/Private/PBDRigidsSolver.cpp
--- a/Engine/Source/Runtime/Experimental/Chaos/Private/PBDRigidsSolver.cpp
+++ b/Engine/Source/Runtime/Experimental/Chaos/Private/PBDRigidsSolver.cpp
@@ -489,7 +489,7 @@ namespace Chaos
 		, FloorHeight(0.f)
 		, bIsDeterministic(false)
 		, Particles(UniqueIndices)
-		, MEvolution(new FPBDRigidsEvolution(Particles, SimMaterials, &ContactModifiers, BufferingModeIn == Chaos::EMultiBufferMode::Single))
+		, MEvolution(new FPBDRigidsEvolution(Particles, SimMaterials, &MidPhaseModifiers, &ContactModifiers, BufferingModeIn == Chaos::EMultiBufferMode::Single))
 		, MEventManager(new FEventManager(BufferingModeIn))
 		, MSolverEventFilters(new FSolverEventFilters())
 		, MDirtyParticlesBuffer(new FDirtyParticlesBuffer(BufferingModeIn, BufferingModeIn == Chaos::EMultiBufferMode::Single))
@@ -780,7 +780,7 @@ namespace Chaos
 		SetMaxDeltaTime_External(1.0f);
 		SetMinDeltaTime_External(UE_SMALL_NUMBER);
 		SetMaxSubSteps_External(1);
-		MEvolution = TUniquePtr<FPBDRigidsEvolution>(new FPBDRigidsEvolution(Particles, SimMaterials, &ContactModifiers, BufferMode == EMultiBufferMode::Single)); 
+		MEvolution = TUniquePtr<FPBDRigidsEvolution>(new FPBDRigidsEvolution(Particles, SimMaterials, &MidPhaseModifiers, &ContactModifiers, BufferMode == EMultiBufferMode::Single)); 
 
 		PerSolverField = MakeUnique<FPerSolverFieldSystem>();
 
@@ -1281,6 +1281,11 @@ namespace Chaos
 				SimCallbackObjects.Add(SimCallbackObject);
 			}
 
+			if (SimCallbackObject->HasOption(ESimCallbackOptions::MidPhaseModification))
+			{
+				MidPhaseModifiers.Add(SimCallbackObject);
+			}
+
 			if (SimCallbackObject->HasOption(ESimCallbackOptions::ContactModification))
 			{
 				ContactModifiers.Add(SimCallbackObject);
diff --git a/Engine/Source/Runtime/Experimental/Chaos/Public/Chaos/Collision/ParticleCollisions.h b/Engine/Source/Runtime/Experimental/Chaos/Public/Chaos/Collision/ParticleCollisions.h
--- a/Engine/Source/Runtime/Experimental/Chaos/Public/Chaos/Collision/ParticleCollisions.h
+++ b/Engine/Source/Runtime/Experimental/Chaos/Public/Chaos/Collision/ParticleCollisions.h
@@ -58,6 +58,18 @@ namespace Chaos
 		*/
 		void RemoveMidPhase(FGeometryParticleHandle* InParticle, FParticlePairMidPhase* InMidPhase);
 
+		/**
+		 * @brief Get a midphase by its index
+		 */
+		inline FParticlePairMidPhase* GetMidPhase(const int32 InIndex)
+		{
+			if (MidPhases.IsValidIndex(InIndex))
+			{
+				return MidPhases[InIndex].Value;
+			}
+			return nullptr;
+		}
+
 		/**
 		 * @brief Find the mid phase with the matching key
 		 * @param InKey The internal key from a FCollisionParticlePairKey
diff --git a/Engine/Source/Runtime/Experimental/Chaos/Public/Chaos/Collision/ParticlePairMidPhase.h b/Engine/Source/Runtime/Experimental/Chaos/Public/Chaos/Collision/ParticlePairMidPhase.h
--- a/Engine/Source/Runtime/Experimental/Chaos/Public/Chaos/Collision/ParticlePairMidPhase.h
+++ b/Engine/Source/Runtime/Experimental/Chaos/Public/Chaos/Collision/ParticlePairMidPhase.h
@@ -371,6 +371,11 @@ namespace Chaos
 		*/
 		void Reset();
 
+		/**
+		 *
+		 */
+		void ResetModifications();
+
 		/**
 		 * @brief Create collision constraints for all colliding shape pairs
 		*/
@@ -433,6 +438,14 @@ namespace Chaos
 			}
 		}
 
+		/**
+		 *
+		 */
+		void SetCCDIsActive(bool bCCDIsActive)
+		{
+			Flags.bIsCCDActive = bCCDIsActive;
+		}
+
 	private:
 		/**
 		 * @brief Build the list of potentially colliding shape pairs.
@@ -466,7 +479,8 @@ namespace Chaos
 			FFlags() : Bits(0) {}
 			struct
 			{
-				uint32 bIsCCD : 1;
+				uint32 bIsCCD : 1;       // True if CCD is supported by either particle
+				uint32 bIsCCDActive : 1; // True if CCD is active for this midphase on this frame. Resets to bIsCCD each frame.
 				uint32 bIsInitialized : 1;
 				uint32 bIsSleeping : 1;
 			};
diff --git a/Engine/Source/Runtime/Experimental/Chaos/Public/Chaos/Collision/SpatialAccelerationBroadPhase.h b/Engine/Source/Runtime/Experimental/Chaos/Public/Chaos/Collision/SpatialAccelerationBroadPhase.h
--- a/Engine/Source/Runtime/Experimental/Chaos/Public/Chaos/Collision/SpatialAccelerationBroadPhase.h
+++ b/Engine/Source/Runtime/Experimental/Chaos/Public/Chaos/Collision/SpatialAccelerationBroadPhase.h
@@ -372,11 +372,11 @@ namespace Chaos
 				CSV_SCOPED_TIMING_STAT(PhysicsVerbose, DetectCollisions_AssignMidPhases);
 
 				// Find or assign a midphase to each overlapping particle pair
-				const auto& AssignMidphasesWorker = [this, Dt](const int32 ContextIndex)
+				const auto& AssignMidPhasesWorker = [this, Dt](const int32 ContextIndex)
 				{
-					AssignMidphases(BroadphaseContexts[ContextIndex]);
+					AssignMidPhases(BroadphaseContexts[ContextIndex]);
 				};
-				PhysicsParallelFor(NumActiveBroadphaseContexts, AssignMidphasesWorker, bDisableCollisionParallelFor);
+				PhysicsParallelFor(NumActiveBroadphaseContexts, AssignMidPhasesWorker, bDisableCollisionParallelFor);
 
 				// Merge all the midphases from each worker into the primary allocator
 				Allocator->ProcessNewMidPhases();
@@ -389,11 +389,11 @@ namespace Chaos
 		void ProduceCollisions(FReal Dt)
 		{
 			CSV_SCOPED_TIMING_STAT(PhysicsVerbose, DetectCollisions_MidPhase);
-			const auto& ProcessMidphasesWorker = [this, Dt](const int32 ContextIndex)
+			const auto& ProcessMidPhasesWorker = [this, Dt](const int32 ContextIndex)
 			{
-				ProcessMidphases(Dt, BroadphaseContexts[ContextIndex]);
+				ProcessMidPhases(Dt, BroadphaseContexts[ContextIndex]);
 			};
-			PhysicsParallelFor(NumActiveBroadphaseContexts, ProcessMidphasesWorker, bDisableCollisionParallelFor);
+			PhysicsParallelFor(NumActiveBroadphaseContexts, ProcessMidPhasesWorker, bDisableCollisionParallelFor);
 		}
 	
 		
@@ -580,7 +580,7 @@ namespace Chaos
 
 		// Find or assign midphases to each of the overlapping particle pairs
 		// @todo(chaos): optimize
-		void AssignMidphases(Private::FBroadPhaseContext& BroadphaseContext)
+		void AssignMidPhases(Private::FBroadPhaseContext& BroadphaseContext)
 		{
 			Private::FCollisionContextAllocator* ContextAllocator = BroadphaseContext.CollisionContext.GetAllocator();
 
@@ -590,12 +590,17 @@ namespace Chaos
 			{
 				Private::FBroadPhaseOverlap& Overlap = BroadphaseContext.Overlaps[OverlapIndex];
 
-				BroadphaseContext.MidPhases[OverlapIndex] = ContextAllocator->GetMidPhase(Overlap.Particles[0], Overlap.Particles[1], Overlap.Particles[Overlap.SearchParticleIndex], BroadphaseContext.CollisionContext);
+				// Get the midphase for this pair
+				FParticlePairMidPhase* MidPhase = ContextAllocator->GetMidPhase(Overlap.Particles[0], Overlap.Particles[1], Overlap.Particles[Overlap.SearchParticleIndex], BroadphaseContext.CollisionContext);
+				BroadphaseContext.MidPhases[OverlapIndex] = MidPhase;
+
+				// Reset any modifications that may have occurred to this midphase on the previous frame
+				MidPhase->ResetModifications();
 			}
 		}
 
 		// Process all the midphases: generate constraints and execute the narrowphase
-		void ProcessMidphases(const FReal Dt, const Private::FBroadPhaseContext& BroadphaseContext)
+		void ProcessMidPhases(const FReal Dt, const Private::FBroadPhaseContext& BroadphaseContext)
 		{
 			// Prefetch initial set of MidPhases
 			const int32 PrefetchLookahead = 4;
diff --git a/Engine/Source/Runtime/Experimental/Chaos/Public/Chaos/Framework/PhysicsSolverBase.h b/Engine/Source/Runtime/Experimental/Chaos/Public/Chaos/Framework/PhysicsSolverBase.h
--- a/Engine/Source/Runtime/Experimental/Chaos/Public/Chaos/Framework/PhysicsSolverBase.h
+++ b/Engine/Source/Runtime/Experimental/Chaos/Public/Chaos/Framework/PhysicsSolverBase.h
@@ -610,6 +610,7 @@ namespace Chaos
 	TUniquePtr<FPendingSpatialDataQueue> PendingSpatialOperations_External;
 
 	TArray<ISimCallbackObject*> SimCallbackObjects;
+	TArray<ISimCallbackObject*> MidPhaseModifiers;
 	TArray<ISimCallbackObject*> ContactModifiers;
 	TArray<ISimCallbackObject*> RegistrationWatchers;
 	TArray<ISimCallbackObject*> UnregistrationWatchers;
diff --git a/Engine/Source/Runtime/Experimental/Chaos/Public/Chaos/MidPhaseModification.h b/Engine/Source/Runtime/Experimental/Chaos/Public/Chaos/MidPhaseModification.h
new file mode 100644
--- /dev/null
+++ b/Engine/Source/Runtime/Experimental/Chaos/Public/Chaos/MidPhaseModification.h
@@ -0,0 +1,147 @@
+// Copyright Epic Games, Inc. All Rights Reserved.
+#pragma once
+
+#include "Chaos/Core.h"
+#include "Chaos/Collision/ParticlePairMidPhase.h"
+#include "ParticleHandleFwd.h"
+
+namespace Chaos
+{
+	class FCollisionConstraintAllocator;	// Container for solver's constraints and midphase pairs
+	class FParticlePairMidPhase;			// The underlying midphase type
+	class FMidPhaseModifier;			// Modifier for a particular midphase
+	class FMidPhaseModifierAccessor;		// Access class for all midphase pair modifiers
+
+	/*
+	 * Wrapper around a MidPhase object which exposes limited access and manipulation functions.
+	 * Also contains the Modifier, which is responsible for iterating over PairModifiers, and
+	 * storing and executing requested modifications.
+	 */
+	class CHAOS_API FMidPhaseModifier
+	{
+	public:
+		FMidPhaseModifier()
+			: MidPhase(nullptr)
+			, Accessor(nullptr)
+		{}
+
+		FMidPhaseModifier(FParticlePairMidPhase* InMidPhase, FMidPhaseModifierAccessor* InAccessor)
+			: MidPhase(InMidPhase)
+			, Accessor(InAccessor)
+		{}
+
+		//
+		// Modifying functions
+		//
+
+		// Disable CCD for this pair
+		void DisableCCD();
+
+		//
+		// Accessor functions
+		//
+
+		void GetParticles(const FGeometryParticleHandle** Particle0, const FGeometryParticleHandle** Particle1) const;
+
+		const FGeometryParticleHandle* GetOtherParticle(const FGeometryParticleHandle* InParticle) const;
+
+
+	private:
+
+		FParticlePairMidPhase* MidPhase;
+		FMidPhaseModifierAccessor* Accessor;
+	};
+
+	/*
+	 * Class for iterating over midphases involving a specific particle
+	 */
+	class CHAOS_API FMidPhaseModifierParticleIterator
+	{
+	public:
+		FMidPhaseModifierParticleIterator(const FMidPhaseModifierParticleIterator& Other)
+			: Particle(Other.Particle)
+			, Accessor(Other.Accessor)
+			, MidPhaseIndex(Other.MidPhaseIndex)
+		{ }
+
+		FMidPhaseModifierParticleIterator(
+			FMidPhaseModifierAccessor* InAccessor,
+			TGeometryParticleHandle<FReal, 3>* InParticle,
+			int32 InMidPhaseIndex = 0)
+			: Particle(InParticle)
+			, Accessor(InAccessor)
+			, MidPhaseIndex(InMidPhaseIndex)
+			, PairModifier(Particle->ParticleCollisions().GetMidPhase(MidPhaseIndex), Accessor)
+		{ }
+
+		FMidPhaseModifier& operator*()
+		{
+			return PairModifier;
+		}
+
+		FMidPhaseModifier* operator->()
+		{
+			return &PairModifier;
+		}
+
+		FMidPhaseModifierParticleIterator& operator++()
+		{
+			++MidPhaseIndex;
+			PairModifier = FMidPhaseModifier(Particle->ParticleCollisions().GetMidPhase(MidPhaseIndex), Accessor);
+			return *this;
+		}
+
+		explicit operator bool() const
+		{
+			return MidPhaseIndex < Particle->ParticleCollisions().Num();
+		}
+
+		bool operator==(const FMidPhaseModifierParticleIterator& Other)
+		{
+			return MidPhaseIndex == Other.MidPhaseIndex;
+		}
+
+		bool operator!=(const FMidPhaseModifierParticleIterator& Other)
+		{
+			return !operator==(Other);
+		}
+
+	private:
+		TGeometryParticleHandle<FReal, 3>* Particle;
+		FMidPhaseModifierAccessor* Accessor;
+		int32 MidPhaseIndex;
+		FMidPhaseModifier PairModifier;
+
+		friend class FMidPhaseModifierParticleRange;
+	};
+
+	class CHAOS_API FMidPhaseModifierParticleRange
+	{
+	public:
+		FMidPhaseModifierParticleRange(
+			FMidPhaseModifierAccessor* InAccessor, FGeometryParticleHandle* InParticle)
+			: Accessor(InAccessor)
+			, Particle(InParticle) { }
+		FMidPhaseModifierParticleIterator begin() const;
+		FMidPhaseModifierParticleIterator end() const;
+	private:
+		FMidPhaseModifierAccessor* Accessor;
+		FGeometryParticleHandle* Particle;
+	};
+
+	/*
+	 * Provides interface for accessing midphase pair modifiers
+	 */
+	class CHAOS_API FMidPhaseModifierAccessor
+	{
+	public:
+		friend class FMidPhaseModifier;
+		friend class FMidPhasePairModifierIterator;
+
+		// Get an object which allows for range iteration over the list of
+		// midphases for a particle
+		FMidPhaseModifierParticleRange GetMidPhases(FGeometryParticleHandle* Particle);
+	};
+
+	using FMidPhaseModifierCallback = TFunction<void(const FCollisionContactModifier&)>;
+}
diff --git a/Engine/Source/Runtime/Experimental/Chaos/Public/Chaos/PBDCollisionConstraints.h b/Engine/Source/Runtime/Experimental/Chaos/Public/Chaos/PBDCollisionConstraints.h
--- a/Engine/Source/Runtime/Experimental/Chaos/Public/Chaos/PBDCollisionConstraints.h
+++ b/Engine/Source/Runtime/Experimental/Chaos/Public/Chaos/PBDCollisionConstraints.h
@@ -103,6 +103,11 @@ public:
 	 */
 	void DetectProbeCollisions(FReal Dt);
 
+	/**
+	 * Apply modifiers to particle pair midphases
+	 */
+	void ApplyMidPhaseModifier(const TArray<ISimCallbackObject*>& MidPhaseModifiers, FReal Dt);
+
 	/**
 	 * Apply modifiers to the constraints and specify which constraints should be disabled.
 	 * You would probably call this in the PostComputeCallback. Prefer this to calling RemoveConstraints in a loop,
diff --git a/Engine/Source/Runtime/Experimental/Chaos/Public/Chaos/PBDRigidsEvolutionGBF.h b/Engine/Source/Runtime/Experimental/Chaos/Public/Chaos/PBDRigidsEvolutionGBF.h
--- a/Engine/Source/Runtime/Experimental/Chaos/Public/Chaos/PBDRigidsEvolutionGBF.h
+++ b/Engine/Source/Runtime/Experimental/Chaos/Public/Chaos/PBDRigidsEvolutionGBF.h
@@ -65,7 +65,7 @@ namespace Chaos
 		static constexpr FRealSingle DefaultCollisionMaxPushOutVelocity = 1000.0f;
 		static constexpr int32 DefaultRestitutionThreshold = 1000;
 
-		CHAOS_API FPBDRigidsEvolutionGBF(FPBDRigidsSOAs& InParticles, THandleArray<FChaosPhysicsMaterial>& SolverPhysicsMaterials, const TArray<ISimCallbackObject*>* InCollisionModifiers = nullptr, bool InIsSingleThreaded = false);
+		CHAOS_API FPBDRigidsEvolutionGBF(FPBDRigidsSOAs& InParticles, THandleArray<FChaosPhysicsMaterial>& SolverPhysicsMaterials, const TArray<ISimCallbackObject*>* InMidPhaseModifiers = nullptr, const TArray<ISimCallbackObject*>* InCollisionModifiers = nullptr, bool InIsSingleThreaded = false);
 		CHAOS_API ~FPBDRigidsEvolutionGBF();
 
 		FORCEINLINE void SetPostIntegrateCallback(const FPBDRigidsEvolutionCallback& Cb)
@@ -358,6 +358,7 @@ namespace Chaos
 		FPBDRigidsEvolutionCallback PreApplyCallback;
 		FPBDRigidsEvolutionInternalHandleCallback InternalParticleInitilization;
 		FEvolutionResimCache* CurrentStepResimCacheImp;
+		const TArray<ISimCallbackObject*>* MidPhaseModifiers;
 		const TArray<ISimCallbackObject*>* CollisionModifiers;
 
 		FCCDManager CCDManager;
diff --git a/Engine/Source/Runtime/Experimental/Chaos/Public/Chaos/SimCallbackObject.h b/Engine/Source/Runtime/Experimental/Chaos/Public/Chaos/SimCallbackObject.h
--- a/Engine/Source/Runtime/Experimental/Chaos/Public/Chaos/SimCallbackObject.h
+++ b/Engine/Source/Runtime/Experimental/Chaos/Public/Chaos/SimCallbackObject.h
@@ -10,6 +10,7 @@
 namespace Chaos
 {
 class FPhysicsSolverBase;
+class FMidPhaseModifierAccessor;
 class FCollisionContactModifier;
 class FSingleParticlePhysicsProxy;
 
@@ -21,11 +22,12 @@ namespace Utilities
 enum class ESimCallbackOptions : uint8
 {
 	Presimulate				= 1 << 0,
-	ContactModification		= 1 << 1,
-	ParticleRegister		= 1 << 2,
-	ParticleUnregister		= 1 << 3,
-	RunOnFrozenGameThread	= 1 << 4,
-	Rewind					= 1 << 5
+	MidPhaseModification	= 1 << 1,
+	ContactModification		= 1 << 2,
+	ParticleRegister		= 1 << 3,
+	ParticleUnregister		= 1 << 4,
+	RunOnFrozenGameThread	= 1 << 5,
+	Rewind					= 1 << 6
 };
 ENUM_CLASS_FLAGS(ESimCallbackOptions)
 
@@ -65,6 +67,11 @@ public:
 		OnPreSimulate_Internal();
 	}
 
+	void MidPhaseModification_Internal(FMidPhaseModifierAccessor& Modifier)
+	{
+		OnMidPhaseModification_Internal(Modifier);
+	}
+
 	void ContactModification_Internal(FCollisionContactModifier& Modifier)
 	{
 		OnContactModification_Internal(Modifier);
@@ -154,6 +161,16 @@ private:
 	*/
 	virtual void OnPreSimulate_Internal() = 0;
 
+	/**
+	* Called once per simulation step. Allows user to modify midphase pairs
+	*
+	* NOTE: you must explicitly request midphase modification when registering the callback for this to be called
+	*/
+	virtual void OnMidPhaseModification_Internal(FMidPhaseModifierAccessor& Modifier)
+	{
+		check(false);
+	}
+
 	/**
 	* Called once per simulation step. Allows user to modify contacts
 	*
