fix MTD computation in Overlap_GeomInternal and FReadPhysicsObjectInterface<Id>::GetPhysicsObjectOverlap to take the maximum distance instead of accepting the first overlap as the correct MTD.

diff --git a/Engine/Source/Runtime/Engine/Private/PhysicsEngine/Experimental/PhysInterface_Chaos.cpp b/Engine/Source/Runtime/Engine/Private/PhysicsEngine/Experimental/PhysInterface_Chaos.cpp
--- a/Engine/Source/Runtime/Engine/Private/PhysicsEngine/Experimental/PhysInterface_Chaos.cpp
+++ b/Engine/Source/Runtime/Engine/Private/PhysicsEngine/Experimental/PhysInterface_Chaos.cpp
@@ -1031,6 +1031,12 @@ bool Overlap_GeomInternal(const FBodyInstance* InInstance, const Chaos::FImplici
 
 	const FTransform ActorTM(RigidBody->GetGameThreadAPI().R(), RigidBody->GetGameThreadAPI().X());
 
+	if (OutOptResult)
+	{
+		OutOptResult->Distance = 0.0;
+	}
+
+	bool bHasOverlap = false;
 	// Iterate over each shape
 	for (int32 ShapeIdx = 0; ShapeIdx < NumShapes; ++ShapeIdx)
 	{
@@ -1050,9 +1056,12 @@ bool Overlap_GeomInternal(const FBodyInstance* InInstance, const Chaos::FImplici
 					Chaos::FMTDInfo MTDInfo;
 					if (Chaos::Utilities::CastHelper(InGeom, GeomTransform, [&](const auto& Downcast, const auto& FullGeomTransform) { return Chaos::OverlapQuery(*Shape->GetGeometry(), ActorTM, Downcast, FullGeomTransform, /*Thickness=*/0, &MTDInfo); }))
 					{
-						OutOptResult->Distance = MTDInfo.Penetration;
-						OutOptResult->Direction = MTDInfo.Normal;
-						return true;	//question: should we take most shallow penetration?
+						bHasOverlap = true;
+						if (MTDInfo.Penetration > OutOptResult->Distance)
+						{
+							OutOptResult->Distance = MTDInfo.Penetration;
+							OutOptResult->Direction = MTDInfo.Normal;
+						}
 					}
 				}
 				else	//question: why do we even allow user to not pass in MTD info?
@@ -1067,7 +1076,7 @@ bool Overlap_GeomInternal(const FBodyInstance* InInstance, const Chaos::FImplici
 		}
 	}
 
-	return false;
+	return bHasOverlap;
 }
 
 bool FPhysInterface_Chaos::Overlap_Geom(const FBodyInstance* InBodyInstance, const FPhysicsGeometryCollection& InGeometry, const FTransform& InShapeTransform, FMTDResult* OutOptResult, bool bTraceComplex)
diff --git a/Engine/Source/Runtime/Experimental/Chaos/Private/Chaos/PhysicsObjectInterface.cpp b/Engine/Source/Runtime/Experimental/Chaos/Private/Chaos/PhysicsObjectInterface.cpp
--- a/Engine/Source/Runtime/Experimental/Chaos/Private/Chaos/PhysicsObjectInterface.cpp
+++ b/Engine/Source/Runtime/Experimental/Chaos/Private/Chaos/PhysicsObjectInterface.cpp
@@ -1,5 +1,4 @@
 // Copyright Epic Games, Inc. All Rights Reserved.
-
 #include "Chaos/PhysicsObjectInterface.h"
 
 #include "Chaos/CastingUtilities.h"
@@ -215,6 +214,12 @@ namespace Chaos
 			return false;
 		}
 
+		if (OutMTD)
+		{
+			OutMTD->Penetration = 0.0;
+		}
+
+		bool bFoundOverlap = false;
 		for (FPerShapeData* B : ShapesB)
 		{
 			if (!B)
@@ -254,25 +259,35 @@ namespace Chaos
 					continue;
 				}
 
+				Chaos::FMTDInfo TmpMTDInfo;
 				const bool bOverlap = Chaos::Utilities::CastHelper(
 					*GeomB,
 					TransformB,
-					[A, &TransformA, OutMTD](const auto& Downcast, const auto& FullTransformB)
+					[A, &TransformA, OutMTD, &TmpMTDInfo](const auto& Downcast, const auto& FullTransformB)
 					{
-						return Chaos::OverlapQuery(*A->GetGeometry(), TransformA, Downcast, FullTransformB, 0, OutMTD);
+						return Chaos::OverlapQuery(*A->GetGeometry(), TransformA, Downcast, FullTransformB, 0, OutMTD ? &TmpMTDInfo : nullptr);
 					}
 				);
 
-				// Is this actually correct? For now I'm mirroring the behavior in Overlap_GeomInternal
-				// but if the goal is to compute the MTD, shouldn't we need to iterate over all pairs of shapes?
 				if (bOverlap)
 				{
-					return true;
+					bFoundOverlap = true;
+					if (!OutMTD)
+					{
+						// Don't care about the MTD so as soon as we find an overlap we can return.
+						return true;
+					}
+					else if (TmpMTDInfo.Penetration > OutMTD->Penetration)
+					{
+						// TODO: I don't think this math is actually correct to find the total overlap between the two objects composed of multiple shapes.
+						// If we need to find the MTD we need to find the largest overlap.
+						*OutMTD = TmpMTDInfo;
+					}
 				}
 			}
 		}
 
-		return false;
+		return bFoundOverlap;
 	}
 
 	template<EThreadContext Id>
