Removed TIsSame usage from Chaos and GeometryCore (part of a larger change to deprecate TIsSame).

diff --git a/Engine/Source/Runtime/Experimental/Chaos/Public/Chaos/AABBTree.h b/Engine/Source/Runtime/Experimental/Chaos/Public/Chaos/AABBTree.h
--- a/Engine/Source/Runtime/Experimental/Chaos/Public/Chaos/AABBTree.h
+++ b/Engine/Source/Runtime/Experimental/Chaos/Public/Chaos/AABBTree.h
@@ -733,8 +733,8 @@ public:
 	static constexpr int32 DefaultMaxChildrenInLeaf = 12;
 	static constexpr int32 DefaultMaxTreeDepth = 16;
 	static constexpr int32 DefaultMaxNumToProcess = 0; // 0 special value for processing all without timeslicing
-	static constexpr ESpatialAcceleration StaticType = TIsSame<TAABBTreeLeafArray<TPayloadType>, TLeafType>::Value ? ESpatialAcceleration::AABBTree : 
-		(TIsSame<TBoundingVolume<TPayloadType>, TLeafType>::Value ? ESpatialAcceleration::AABBTreeBV : ESpatialAcceleration::Unknown);
+	static constexpr ESpatialAcceleration StaticType = std::is_same_v<TAABBTreeLeafArray<TPayloadType>, TLeafType> ? ESpatialAcceleration::AABBTree :
+		(std::is_same_v<TBoundingVolume<TPayloadType>, TLeafType> ? ESpatialAcceleration::AABBTreeBV : ESpatialAcceleration::Unknown);
 	TAABBTree()
 		: ISpatialAcceleration<TPayloadType, T, 3>(StaticType)
 		, bDynamicTree(false)
diff --git a/Engine/Source/Runtime/Experimental/Chaos/Public/Chaos/HeightField.h b/Engine/Source/Runtime/Experimental/Chaos/Public/Chaos/HeightField.h
--- a/Engine/Source/Runtime/Experimental/Chaos/Public/Chaos/HeightField.h
+++ b/Engine/Source/Runtime/Experimental/Chaos/Public/Chaos/HeightField.h
@@ -275,10 +275,10 @@ namespace Chaos
 
 			// Only supporting unsigned int types for the height range - really no difference using
 			// this or signed but this is a little nicer overall
-			static_assert(TIsSame<StorageType, uint8>::Value || 
-				TIsSame<StorageType, uint16>::Value || 
-				TIsSame<StorageType, uint32>::Value || 
-				TIsSame<StorageType, uint64>::Value,
+			static_assert(std::is_same_v<StorageType, uint8> ||
+				std::is_same_v<StorageType, uint16> ||
+				std::is_same_v<StorageType, uint32> ||
+				std::is_same_v<StorageType, uint64>,
 				"Expected unsigned integer type for heightfield data storage");
 
 			// Data sizes to validate during serialization
diff --git a/Engine/Source/Runtime/Experimental/Chaos/Public/Chaos/ImplicitObjectScaled.h b/Engine/Source/Runtime/Experimental/Chaos/Public/Chaos/ImplicitObjectScaled.h
--- a/Engine/Source/Runtime/Experimental/Chaos/Public/Chaos/ImplicitObjectScaled.h
+++ b/Engine/Source/Runtime/Experimental/Chaos/Public/Chaos/ImplicitObjectScaled.h
@@ -210,11 +210,12 @@ public:
 
 	static const TImplicitObjectInstanced<TConcrete>& AsInstancedChecked(const FImplicitObject& Obj)
 	{
-		if(TIsSame<TConcrete,FImplicitObject>::Value)
+		if constexpr (std::is_same_v<TConcrete,FImplicitObject>)
 		{
 			//can cast any instanced to ImplicitObject base
 			check(IsInstanced(Obj.GetType()));
-		} else
+		}
+		else
 		{
 			check(StaticType() == Obj.GetType());
 		}
@@ -516,7 +517,7 @@ public:
 
 	static const TImplicitObjectScaled<TConcrete>& AsScaledChecked(const FImplicitObject& Obj)
 	{
-		if (TIsSame<TConcrete, FImplicitObject>::Value)
+		if constexpr (std::is_same_v<TConcrete, FImplicitObject>)
 		{
 			//can cast any scaled to ImplicitObject base
 			check(IsScaled(Obj.GetType()));
@@ -530,7 +531,7 @@ public:
 
 	static TImplicitObjectScaled<TConcrete>& AsScaledChecked(FImplicitObject& Obj)
 	{
-		if (TIsSame<TConcrete, FImplicitObject>::Value)
+		if constexpr (std::is_same_v<TConcrete, FImplicitObject>)
 		{
 			//can cast any scaled to ImplicitObject base
 			check(IsScaled(Obj.GetType()));
@@ -544,7 +545,7 @@ public:
 
 	static const TImplicitObjectScaled<TConcrete>* AsScaled(const FImplicitObject& Obj)
 	{
-		if (TIsSame<TConcrete, FImplicitObject>::Value)
+		if constexpr (std::is_same_v<TConcrete, FImplicitObject>)
 		{
 			//can cast any scaled to ImplicitObject base
 			return IsScaled(Obj.GetType()) ? static_cast<const TImplicitObjectScaled<TConcrete>*>(&Obj) : nullptr;
@@ -557,7 +558,7 @@ public:
 
 	static TImplicitObjectScaled<TConcrete>* AsScaled(FImplicitObject& Obj)
 	{
-		if (TIsSame<TConcrete, FImplicitObject>::Value)
+		if constexpr (std::is_same_v<TConcrete, FImplicitObject>)
 		{
 			//can cast any scaled to ImplicitObject base
 			return IsScaled(Obj.GetType()) ? static_cast<TImplicitObjectScaled<TConcrete>*>(&Obj) : nullptr;
diff --git a/Engine/Source/Runtime/Experimental/Chaos/Public/Chaos/SpatialAccelerationCollection.h b/Engine/Source/Runtime/Experimental/Chaos/Public/Chaos/SpatialAccelerationCollection.h
--- a/Engine/Source/Runtime/Experimental/Chaos/Public/Chaos/SpatialAccelerationCollection.h
+++ b/Engine/Source/Runtime/Experimental/Chaos/Public/Chaos/SpatialAccelerationCollection.h
@@ -354,13 +354,13 @@ struct TSpatialAccelerationCollectionHelper
 };
 
 template <typename SpatialAccelerationCollection>
-typename TEnableIf<TIsSame<typename SpatialAccelerationCollection::TPayloadType, FAccelerationStructureHandle>::Value, void>::Type PBDComputeConstraintsLowLevel_Helper(FReal Dt, const SpatialAccelerationCollection& Accel, FSpatialAccelerationBroadPhase& BroadPhase, Private::FCollisionConstraintAllocator* Allocator, const FCollisionDetectorSettings& Settings, IResimCacheBase* ResimCache)
+typename TEnableIf<std::is_same_v<typename SpatialAccelerationCollection::TPayloadType, FAccelerationStructureHandle>, void>::Type PBDComputeConstraintsLowLevel_Helper(FReal Dt, const SpatialAccelerationCollection& Accel, FSpatialAccelerationBroadPhase& BroadPhase, Private::FCollisionConstraintAllocator* Allocator, const FCollisionDetectorSettings& Settings, IResimCacheBase* ResimCache)
 {
 	BroadPhase.ProduceOverlaps(Dt, Accel, Allocator, Settings, ResimCache);
 }
 
 template <typename SpatialAccelerationCollection>
-typename TEnableIf<!TIsSame<typename SpatialAccelerationCollection::TPayloadType, FAccelerationStructureHandle>::Value, void>::Type PBDComputeConstraintsLowLevel_Helper(FReal Dt, const SpatialAccelerationCollection& Accel, FSpatialAccelerationBroadPhase& BroadPhase, Private::FCollisionConstraintAllocator* Allocator, const FCollisionDetectorSettings& Settings, IResimCacheBase* ResimCache)
+typename TEnableIf<!std::is_same_v<typename SpatialAccelerationCollection::TPayloadType, FAccelerationStructureHandle>, void>::Type PBDComputeConstraintsLowLevel_Helper(FReal Dt, const SpatialAccelerationCollection& Accel, FSpatialAccelerationBroadPhase& BroadPhase, Private::FCollisionConstraintAllocator* Allocator, const FCollisionDetectorSettings& Settings, IResimCacheBase* ResimCache)
 {
 }
 
diff --git a/Engine/Source/Runtime/Experimental/ChaosCore/Public/Chaos/Vector.h b/Engine/Source/Runtime/Experimental/ChaosCore/Public/Chaos/Vector.h
--- a/Engine/Source/Runtime/Experimental/ChaosCore/Public/Chaos/Vector.h
+++ b/Engine/Source/Runtime/Experimental/ChaosCore/Public/Chaos/Vector.h
@@ -1311,7 +1311,7 @@ PRAGMA_ENABLE_DEPRECATION_WARNINGS
 	template<typename T, int d>
 	inline FArchive& SerializeReal(FArchive& Ar, TVector<T, d>& ValueIn)
 	{
-		static_assert(TIsSame<T, float>::Value || TIsSame<T, double>::Value, "only float or double are supported by this function");
+		static_assert(std::is_same_v<T, float> || std::is_same_v<T, double>, "only float or double are supported by this function");
 		for (int32 Idx = 0; Idx < d; ++Idx)
 		{
 			FRealSingle RealSingle = (FRealSingle)ValueIn[Idx];
diff --git a/Engine/Source/Runtime/Experimental/Iris/Core/Private/Iris/Serialization/BitfieldNetSerializer.cpp b/Engine/Source/Runtime/Experimental/Iris/Core/Private/Iris/Serialization/BitfieldNetSerializer.cpp
--- a/Engine/Source/Runtime/Experimental/Iris/Core/Private/Iris/Serialization/BitfieldNetSerializer.cpp
+++ b/Engine/Source/Runtime/Experimental/Iris/Core/Private/Iris/Serialization/BitfieldNetSerializer.cpp
@@ -68,7 +68,7 @@ void FBitfieldNetSerializer::Dequantize(FNetSerializationContext& Context, const
 
 bool FBitfieldNetSerializer::IsEqual(FNetSerializationContext& Context, const FNetIsEqualArgs& Args)
 {
-	static_assert(TIsSame<SourceType, QuantizedType>::Value, "FBitfieldNetSerializer::IsEqual needs to be re-implemented");
+	static_assert(std::is_same_v<SourceType, QuantizedType>, "FBitfieldNetSerializer::IsEqual needs to be re-implemented");
 
 	const SourceType Value0 = *reinterpret_cast<const SourceType*>(Args.Source0);
 	const SourceType Value1 = *reinterpret_cast<const SourceType*>(Args.Source1);
diff --git a/Engine/Source/Runtime/Experimental/Iris/Core/Public/Iris/Serialization/NetSerializerBuilder.inl b/Engine/Source/Runtime/Experimental/Iris/Core/Public/Iris/Serialization/NetSerializerBuilder.inl
--- a/Engine/Source/Runtime/Experimental/Iris/Core/Public/Iris/Serialization/NetSerializerBuilder.inl
+++ b/Engine/Source/Runtime/Experimental/Iris/Core/Public/Iris/Serialization/NetSerializerBuilder.inl
@@ -131,44 +131,44 @@ private:
 	template<typename> struct FTypeCheck;
 
 	// Version check
-	template<typename U> static ETrueType TestHasVersion(typename TEnableIf<TIsSame<decltype(&FVersion::Version), decltype(&U::Version)>::Value>::Type*);
+	template<typename U> static ETrueType TestHasVersion(typename TEnableIf<std::is_same_v<decltype(&FVersion::Version), decltype(&U::Version)>>::Type*);
 	template<typename> static EFalseType TestHasVersion(...);
 
 	// Traits
 	template<typename U> static ETrueType TestHasCustomNetReferenceIsPresent(FTypeCheck<decltype(&U::bHasCustomNetReference)>*);
 	template<typename> static EFalseType TestHasCustomNetReferenceIsPresent(...);
 
-	template<typename U> static ETrueType TestHasCustomNetReferenceIsBool(typename TEnableIf<TIsSame<decltype(&FTraits::bHasCustomNetReference), decltype(&U::bHasCustomNetReference)>::Value>::Type*);
+	template<typename U> static ETrueType TestHasCustomNetReferenceIsBool(typename TEnableIf<std::is_same_v<decltype(&FTraits::bHasCustomNetReference), decltype(&U::bHasCustomNetReference)>>::Type*);
 	template<typename> static EFalseType TestHasCustomNetReferenceIsBool(...);
 
 	template<typename U> static ETrueType TestUseSerializerIsEqualIsPresent(FTypeCheck<decltype(&U::bUseSerializerIsEqual)>*);
 	template<typename> static EFalseType TestUseSerializerIsEqualIsPresent(...);
 
-	template<typename U> static ETrueType TestUseSerializerIsEqualIsBool(typename TEnableIf<TIsSame<decltype(&FTraits::bUseSerializerIsEqual), decltype(&U::bUseSerializerIsEqual)>::Value>::Type*);
+	template<typename U> static ETrueType TestUseSerializerIsEqualIsBool(typename TEnableIf<std::is_same_v<decltype(&FTraits::bUseSerializerIsEqual), decltype(&U::bUseSerializerIsEqual)>>::Type*);
 	template<typename> static EFalseType TestUseSerializerIsEqualIsBool(...);
 
 	template<typename U> static ETrueType TestIsForwardingSerializerIsPresent(FTypeCheck<decltype(&U::bIsForwardingSerializer)>*);
 	template<typename> static EFalseType TestIsForwardingSerializerIsPresent(...);
 
-	template<typename U> static ETrueType TestIsForwardingSerializerIsBool(typename TEnableIf<TIsSame<decltype(&FTraits::bIsForwardingSerializer), decltype(&U::bIsForwardingSerializer)>::Value>::Type*);
+	template<typename U> static ETrueType TestIsForwardingSerializerIsBool(typename TEnableIf<std::is_same_v<decltype(&FTraits::bIsForwardingSerializer), decltype(&U::bIsForwardingSerializer)>>::Type*);
 	template<typename> static EFalseType TestIsForwardingSerializerIsBool(...);
 
 	template<typename U> static ETrueType TestHasConnectionSpecificSerializationIsPresent(FTypeCheck<decltype(&U::bHasConnectionSpecificSerialization)>*);
 	template<typename> static EFalseType TestHasConnectionSpecificSerializationIsPresent(...);
 
-	template<typename U> static ETrueType TestHasConnectionSpecificSerializationIsBool(typename TEnableIf<TIsSame<decltype(&FTraits::bHasConnectionSpecificSerialization), decltype(&U::bHasConnectionSpecificSerialization)>::Value>::Type*);
+	template<typename U> static ETrueType TestHasConnectionSpecificSerializationIsBool(typename TEnableIf<std::is_same_v<decltype(&FTraits::bHasConnectionSpecificSerialization), decltype(&U::bHasConnectionSpecificSerialization)>>::Type*);
 	template<typename> static EFalseType TestHasConnectionSpecificSerializationIsBool(...);
 
 	template<typename U> static ETrueType TestHasDynamicStateIsPresent(FTypeCheck<decltype(&U::bHasDynamicState)>*);
 	template<typename> static EFalseType TestHasDynamicStateIsPresent(...);
 
-	template<typename U> static ETrueType TestHasDynamicStateIsBool(typename TEnableIf<TIsSame<decltype(&FTraits::bHasDynamicState), decltype(&U::bHasDynamicState)>::Value>::Type*);
+	template<typename U> static ETrueType TestHasDynamicStateIsBool(typename TEnableIf<std::is_same_v<decltype(&FTraits::bHasDynamicState), decltype(&U::bHasDynamicState)>>::Type*);
 	template<typename> static EFalseType TestHasDynamicStateIsBool(...);
 
 	template<typename U> static ETrueType TestUseDefaultDeltaIsPresent(FTypeCheck<decltype(&U::bUseDefaultDelta)>*);
 	template<typename> static EFalseType TestUseDefaultDeltaIsPresent(...);
 
-	template<typename U> static ETrueType TestUseDefaultDeltaIsBool(typename TEnableIf<TIsSame<decltype(&FTraits::bUseDefaultDelta), decltype(&U::bUseDefaultDelta)>::Value>::Type*);
+	template<typename U> static ETrueType TestUseDefaultDeltaIsBool(typename TEnableIf<std::is_same_v<decltype(&FTraits::bUseDefaultDelta), decltype(&U::bUseDefaultDelta)>>::Type*);
 	template<typename> static EFalseType TestUseDefaultDeltaIsBool(...);
 
 	// Type checks
@@ -421,7 +421,7 @@ public:
 	static constexpr uint32 GetQuantizedTypeSize() { return sizeof(typename NetSerializerImpl::QuantizedType); }
 
 	template<typename T = void, typename U = typename TEnableIf<!HasQuantizedType && HasSourceType, T>::Type, bool V = true, char W = 0>
-	static constexpr uint32 GetQuantizedTypeSize() { return TIsSame<void, typename NetSerializerImpl::SourceType>::Value ? uint32(0) : sizeof(typename TChooseClass<TIsSame<void, typename NetSerializerImpl::SourceType>::Value, uint8, typename NetSerializerImpl::SourceType>::Result); }
+	static constexpr uint32 GetQuantizedTypeSize() { return std::is_same_v<void, typename NetSerializerImpl::SourceType> ? uint32(0) : sizeof(typename TChooseClass<std::is_same_v<void, typename NetSerializerImpl::SourceType>, uint8, typename NetSerializerImpl::SourceType>::Result); }
 
 	template<typename T = void, typename U = typename TEnableIf<!(HasSourceType || HasQuantizedType), T>::Type, char V = 0>
 	static constexpr uint32 GetQuantizedTypeSize() { return 0; }
@@ -430,7 +430,7 @@ public:
 	static constexpr uint32 GetQuantizedTypeAlignment() { return alignof(typename NetSerializerImpl::QuantizedType); }
 
 	template<typename T = void, typename U = typename TEnableIf<!HasQuantizedType && HasSourceType, T>::Type, bool V = true, char W = 0>
-	static constexpr uint32 GetQuantizedTypeAlignment() { return alignof(typename TChooseClass<TIsSame<void, typename NetSerializerImpl::SourceType>::Value, uint8, typename NetSerializerImpl::SourceType>::Result); }
+	static constexpr uint32 GetQuantizedTypeAlignment() { return alignof(typename TChooseClass<std::is_same_v<void, typename NetSerializerImpl::SourceType>, uint8, typename NetSerializerImpl::SourceType>::Result); }
 
 	template<typename T = void, typename U = typename TEnableIf<!(HasSourceType || HasQuantizedType), T>::Type, char V = 0>
 	static constexpr uint32 GetQuantizedTypeAlignment() { return 1; }
diff --git a/Engine/Source/Runtime/GeometryCore/Public/CompGeom/ExactPredicates.h b/Engine/Source/Runtime/GeometryCore/Public/CompGeom/ExactPredicates.h
--- a/Engine/Source/Runtime/GeometryCore/Public/CompGeom/ExactPredicates.h
+++ b/Engine/Source/Runtime/GeometryCore/Public/CompGeom/ExactPredicates.h
@@ -123,7 +123,7 @@ RealType Orient3(const TVector<RealType>& A, const TVector<RealType>& B, const T
 	double PB[3]{ (double)B.X, (double)B.Y, (double)B.Z };
 	double PC[3]{ (double)C.X, (double)C.Y, (double)C.Z };
 	double PD[3]{ (double)D.X, (double)D.Y, (double)D.Z };
-	if constexpr (TIsSame<RealType, double>::Value)
+	if constexpr (std::is_same_v<RealType, double>)
 	{
 		return Orient3D(PA, PB, PC, PD);
 	}
@@ -160,7 +160,7 @@ RealType Facing3(const TVector<RealType>& A, const TVector<RealType>& B, const T
 	double PB[3]{ (double)B.X, (double)B.Y, (double)B.Z };
 	double PC[3]{ (double)C.X, (double)C.Y, (double)C.Z };
 	double Dir[3]{ (double)Direction.X, (double)Direction.Y, (double)Direction.Z };
-	if constexpr (TIsSame<RealType, double>::Value)
+	if constexpr (std::is_same_v<RealType, double>)
 	{
 		return Facing3D(PA, PB, PC, Dir);
 	}
@@ -187,7 +187,7 @@ RealType InCircle2(const TVector2<RealType>& A, const TVector2<RealType>& B, con
 	double PB[2]{ (double)B.X, (double)B.Y };
 	double PC[2]{ (double)C.X, (double)C.Y };
 	double PD[2]{ (double)D.X, (double)D.Y };
-	if constexpr (TIsSame<RealType, double>::Value)
+	if constexpr (std::is_same_v<RealType, double>)
 	{
 		return InCircle(PA, PB, PC, PD);
 	}
diff --git a/Engine/Source/Runtime/GeometryCore/Public/Util/DynamicVector.h b/Engine/Source/Runtime/GeometryCore/Public/Util/DynamicVector.h
--- a/Engine/Source/Runtime/GeometryCore/Public/Util/DynamicVector.h
+++ b/Engine/Source/Runtime/GeometryCore/Public/Util/DynamicVector.h
@@ -417,11 +417,11 @@ private:
 				// If the currently stored type is one of these types, and the archive is from before bulk serialization for these types was enabled,
 				// we need to still use per element serialization for legacy data.
 				constexpr bool bIsLWCBulkSerializedDoubleType =
-					TIsSame<Type, FVector2d>::Value ||
-					TIsSame<Type, FVector3d>::Value ||
-					TIsSame<Type, FVector4d>::Value ||
-					TIsSame<Type, FQuat4d>::Value ||
-					TIsSame<Type, FTransform3d>::Value;
+					std::is_same_v<Type, FVector2d> ||
+					std::is_same_v<Type, FVector3d> ||
+					std::is_same_v<Type, FVector4d> ||
+					std::is_same_v<Type, FQuat4d> ||
+					std::is_same_v<Type, FTransform3d>;
 				const bool bUseBulkSerialization = TCanBulkSerialize<Type>::Value && !(bIsLWCBulkSerializedDoubleType && Ar.UEVer() <
 					EUnrealEngineObjectUE5Version::LARGE_WORLD_COORDINATES);
 
