Fix Capsule Overlap

diff --git a/Engine/Source/Programs/HeadlessChaos/Private/HeadlessChaosTestHeightfield.cpp b/Engine/Source/Programs/HeadlessChaos/Private/HeadlessChaosTestHeightfield.cpp
--- a/Engine/Source/Programs/HeadlessChaos/Private/HeadlessChaosTestHeightfield.cpp
+++ b/Engine/Source/Programs/HeadlessChaos/Private/HeadlessChaosTestHeightfield.cpp
@@ -4,6 +4,7 @@
 #include "HeadlessChaos.h"
 #include "HeadlessChaosTestUtility.h"
 #include "Modules/ModuleManager.h"
+#include "Chaos/GeometryQueries.h"
 #include "Chaos/HeightField.h"
 
 namespace ChaosTest {
@@ -932,6 +933,78 @@ namespace ChaosTest {
 		SweepBoxTest();
 	}
 
+
+	void OverlapConsistentTest()
+	{
+		const FReal CountToWorldScale = 1;
+		const int32 Columns = 10;
+		const int32 Rows = 10;
+
+		TArray<FReal> Heights = CreateMountain(Columns, Rows);
+
+		TArray<FReal> HeightsCopy = Heights;
+		FVec3 Scale(1.0, 1.0, 1.0);
+		FHeightField Heightfield(MoveTemp(HeightsCopy), TArray<uint8>(), Rows, Columns, Scale);
+		const auto& Bounds = Heightfield.BoundingBox();	//Current API forces us to do this to cache the bounds
+		{
+			TBox<FReal, 3> Box(FVec3(-0.5, -0.5, -0.5), FVec3(0.5, 0.5, 0.5));
+			FCapsule Capsule(FVec3(0.0, 0.0, 0.0), FVec3(0.0, 0.0, 9.0), 1.14);
+			Chaos::FSphere Sphere1(FVec3(0.0, 0.0, -2.0), 0.6);
+
+			FMTDInfo OutBoxMTD;
+			FMTDInfo OutCapsuleMTD;
+			FMTDInfo OutSphereMTD;
+
+			for (int32 Row = 0; Row < Rows; ++Row)
+			{
+				for (int32 Col = 0; Col < Columns; ++Col)
+				{
+					for (FReal Height = 0; Height < 12; Height+=0.5)
+					{
+
+						const FVec3 Translation(Col, Row, Height);
+						FRigidTransform3 QueryTM(Translation, TRotation<FReal, 3>::Identity);
+						FVec3 Dir(1, 0, 0);
+						bool BoxResult = Heightfield.OverlapGeom(Box, QueryTM, 0.0, nullptr);
+						bool BoxResultMTD = Heightfield.OverlapGeom(Box, QueryTM, 0.0, &OutBoxMTD);
+						// EXPECT_EQ(BoxResult, BoxResultMTD); // TODO Enable test
+
+						bool CapsuleResult = Heightfield.OverlapGeom(Capsule, QueryTM, 0.0, nullptr);
+						bool CapsuleResultMTD = Heightfield.OverlapGeom(Capsule, QueryTM, 0.0, &OutCapsuleMTD);
+						EXPECT_EQ(CapsuleResult, CapsuleResultMTD);
+
+						bool SphereResult = Heightfield.OverlapGeom(Sphere1, QueryTM, 0.0, nullptr);
+						bool SphereResultMTD = Heightfield.OverlapGeom(Sphere1, QueryTM, 0.0, &OutSphereMTD);
+						EXPECT_EQ(SphereResult, SphereResultMTD);
+					}
+				}
+			}
+		}
+		// Testing capsule upside down
+		{
+			FCapsule Capsule(FVec3(0.0, 0.0, 9.0), FVec3(0.0, 0.0, 0.0), 1.14);
+			FMTDInfo OutCapsuleMTD;
+
+			for (int32 Row = 0; Row < Rows; ++Row)
+			{
+				for (int32 Col = 0; Col < Columns; ++Col)
+				{
+					for (FReal Height = 0; Height < 12; Height += 0.5)
+					{
+
+						const FVec3 Translation(Col, Row, Height);
+						FRigidTransform3 QueryTM(Translation, TRotation<FReal, 3>::Identity);
+						FVec3 Dir(1, 0, 0);
+
+						bool CapsuleResult = Heightfield.OverlapGeom(Capsule, QueryTM, 0.0, nullptr);
+						bool CapsuleResultMTD = Heightfield.OverlapGeom(Capsule, QueryTM, 0.0, &OutCapsuleMTD);
+						EXPECT_EQ(CapsuleResult, CapsuleResultMTD);
+					}
+				}
+			}
+		}
+	}
+
 	void OverlapTest()
 	{
 		const FReal CountToWorldScale = 1;
@@ -1099,6 +1172,7 @@ namespace ChaosTest {
 		ChaosTest::RaycastVariousWalkOnHeightField();
 		ChaosTest::SweepTest();
 		ChaosTest::OverlapTest();
+		ChaosTest::OverlapConsistentTest();
 		EditHeights();
 		SUCCEED();
 	}
diff --git a/Engine/Source/Runtime/Experimental/Chaos/Private/Chaos/Collision/TriangleOverlap.cpp b/Engine/Source/Runtime/Experimental/Chaos/Private/Chaos/Collision/TriangleOverlap.cpp
--- a/Engine/Source/Runtime/Experimental/Chaos/Private/Chaos/Collision/TriangleOverlap.cpp
+++ b/Engine/Source/Runtime/Experimental/Chaos/Private/Chaos/Collision/TriangleOverlap.cpp
@@ -36,37 +36,106 @@ namespace Chaos
 		// Mask to compare the distance of the triangle edge plane to the capsule radius.
 		constexpr static int32 AllEdges = 0b1111;
 		{
-			const VectorRegister4Float PlaneCA = VectorNormalize(VectorCross(CA, Normal));
-			const VectorRegister4Float EdgeX1 = VectorDot3(AX1, PlaneCA);
-			const VectorRegister4Float EdgeX2 = VectorDot3(AX2, PlaneCA);
-			const VectorRegister4Float Edge = VectorMoveLh(EdgeX1, EdgeX2);
-			const VectorRegister4Float EdgeGTRadius = VectorCompareGT(Edge, RadiusSimd);
-			if (VectorMaskBits(EdgeGTRadius) == AllEdges)
+			VectorRegister4Float PlaneCA = VectorNormalize(VectorCross(CA, Normal));
+			FRealSingle EdgeX1 = VectorDot3Scalar(AX1, PlaneCA);
+			FRealSingle EdgeX2 = VectorDot3Scalar(AX2, PlaneCA);
+			FRealSingle Edge = FMath::Min<FRealSingle>(EdgeX1, EdgeX2);
+
+			if (Edge > Radius)
 			{
 				return false;
 			}
+			// Make a check on another similar separating plane, but instead of using the triangle normal
+			// use the capsule and the edge.
+			// Make sure the other plane is on the right side of the triangle
+			if (Edge > 0.0f)
+			{
+				VectorRegister4Float OtherNormal = VectorCross(AX1, CA);
+				PlaneCA = VectorNormalize(VectorCross(CA, OtherNormal));
+				EdgeX1 = VectorDot3Scalar(AX1, PlaneCA);
+				EdgeX2 = VectorDot3Scalar(AX2, PlaneCA);
+				Edge = FMath::Min<FRealSingle>(EdgeX1, EdgeX2);
+				if (Edge > Radius)
+				{
+					return false;
+				}
+
+				OtherNormal = VectorCross(AX2, CA);
+				PlaneCA = VectorNormalize(VectorCross(CA, OtherNormal));
+				EdgeX1 = VectorDot3Scalar(AX1, PlaneCA);
+				EdgeX2 = VectorDot3Scalar(AX2, PlaneCA);
+				Edge = FMath::Min<FRealSingle>(EdgeX1, EdgeX2);
+				if (Edge > Radius)
+				{
+					return false;
+				}
+			}
 		}
 		{
-			const VectorRegister4Float PlaneAB = VectorNormalize(VectorCross(AB, Normal));
-			const VectorRegister4Float EdgeX1 = VectorDot3(BX1, PlaneAB);
-			const VectorRegister4Float EdgeX2 = VectorDot3(VectorSubtract(X2, B), PlaneAB);
-			const VectorRegister4Float Edge = VectorMoveLh(EdgeX1, EdgeX2);
-			const VectorRegister4Float EdgeGTRadius = VectorCompareGT(Edge, RadiusSimd);
-			if (VectorMaskBits(EdgeGTRadius) == AllEdges)
+			VectorRegister4Float PlaneAB = VectorNormalize(VectorCross(AB, Normal));
+			FRealSingle EdgeX1 = VectorDot3Scalar(BX1, PlaneAB);
+			FRealSingle EdgeX2 = VectorDot3Scalar(VectorSubtract(X2, B), PlaneAB);
+			FRealSingle Edge = FMath::Min<FRealSingle>(EdgeX1, EdgeX2);
+			if (Edge > Radius)
 			{
 				return false;
 			}
+			
+			if (Edge > 0.0f)
+			{
+				VectorRegister4Float OtherNormal = VectorCross(BX1, AB);
+				PlaneAB = VectorNormalize(VectorCross(AB, OtherNormal));
+				EdgeX1 = VectorDot3Scalar(BX1, PlaneAB);
+				EdgeX2 = VectorDot3Scalar(VectorSubtract(X2, B), PlaneAB);
+				Edge = FMath::Min<FRealSingle>(EdgeX1, EdgeX2);
+				if (Edge > Radius)
+				{
+					return false;
+				}
+
+				OtherNormal = VectorCross(VectorSubtract(X2, B), AB);
+				PlaneAB = VectorNormalize(VectorCross(AB, OtherNormal));
+				EdgeX1 = VectorDot3Scalar(BX1, PlaneAB);
+				EdgeX2 = VectorDot3Scalar(VectorSubtract(X2, B), PlaneAB);
+				Edge = FMath::Min<FRealSingle>(EdgeX1, EdgeX2);
+				if (Edge > Radius)
+				{
+					return false;
+				}
+			}
 		}
 		{
-			const VectorRegister4Float PlaneBC = VectorNormalize(VectorCross(BC, Normal));
-			const VectorRegister4Float EdgeX1 = VectorDot3(CX1, PlaneBC);
-			const VectorRegister4Float EdgeX2 = VectorDot3(VectorSubtract(X2, C), PlaneBC);
-			const VectorRegister4Float Edge = VectorMoveLh(EdgeX1, EdgeX2);
-			const VectorRegister4Float EdgeGTRadius = VectorCompareGT(Edge, RadiusSimd);
-			if (VectorMaskBits(EdgeGTRadius) == AllEdges)
+			VectorRegister4Float PlaneBC = VectorNormalize(VectorCross(BC, Normal));
+			FRealSingle EdgeX1 = VectorDot3Scalar(CX1, PlaneBC);
+			FRealSingle EdgeX2 = VectorDot3Scalar(VectorSubtract(X2, C), PlaneBC);
+			FRealSingle Edge = FMath::Min<FRealSingle>(EdgeX1, EdgeX2);
+			if (Edge > Radius)
 			{
 				return false;
 			}
+
+			if (Edge > 0.0f)
+			{
+				VectorRegister4Float OtherNormal = VectorCross(CX1, BC);
+				PlaneBC = VectorNormalize(VectorCross(BC, OtherNormal));
+				EdgeX1 = VectorDot3Scalar(CX1, PlaneBC);
+				EdgeX2 = VectorDot3Scalar(VectorSubtract(X2, C), PlaneBC);
+				Edge = FMath::Min<FRealSingle>(EdgeX1, EdgeX2);
+				if (Edge > Radius)
+				{
+					return false;
+				}
+
+				OtherNormal = VectorCross(VectorSubtract(X2, C), BC);
+				PlaneBC = VectorNormalize(VectorCross(BC, OtherNormal));
+				EdgeX1 = VectorDot3Scalar(CX1, PlaneBC);
+				EdgeX2 = VectorDot3Scalar(VectorSubtract(X2, C), PlaneBC);
+				Edge = FMath::Min<FRealSingle>(EdgeX1, EdgeX2);
+				if (Edge > Radius)
+				{
+					return false;
+				}
+			}
 		}
 
 		// Plane Triangle
