add profiling scope to GetPhysicsObjectOverlap

diff --git a/Engine/Source/Runtime/Experimental/Chaos/Private/Chaos/PhysicsObjectInterface.cpp b/Engine/Source/Runtime/Experimental/Chaos/Private/Chaos/PhysicsObjectInterface.cpp
--- a/Engine/Source/Runtime/Experimental/Chaos/Private/Chaos/PhysicsObjectInterface.cpp
+++ b/Engine/Source/Runtime/Experimental/Chaos/Private/Chaos/PhysicsObjectInterface.cpp
@@ -201,6 +201,7 @@ namespace Chaos
 	template<EThreadContext Id>
 	bool FReadPhysicsObjectInterface<Id>::GetPhysicsObjectOverlap(FPhysicsObjectHandle ObjectA, FPhysicsObjectHandle ObjectB, bool bTraceComplex, Chaos::FMTDInfo* OutMTD)
 	{
+		TRACE_CPUPROFILER_EVENT_SCOPE(FReadPhysicsObjectInterface<Id>::GetPhysicsObjectOverlap);
 		TArray<FPerShapeData*> ShapesA = GetAllShapes({ &ObjectA, 1 });
 		const FTransform TransformA(GetR(ObjectA), GetX(ObjectA));
 		const FBox BoxA = GetWorldBounds({&ObjectA, 1});
