Chaos - fix HeadlessChaos crash

diff --git a/Engine/Source/Runtime/Experimental/Chaos/Private/Chaos/PBDRigidsEvolutionGBF.cpp b/Engine/Source/Runtime/Experimental/Chaos/Private/Chaos/PBDRigidsEvolutionGBF.cpp
--- a/Engine/Source/Runtime/Experimental/Chaos/Private/Chaos/PBDRigidsEvolutionGBF.cpp
+++ b/Engine/Source/Runtime/Experimental/Chaos/Private/Chaos/PBDRigidsEvolutionGBF.cpp
@@ -368,6 +368,7 @@ void FPBDRigidsEvolutionGBF::AdvanceOneTimeStepImpl(const FReal Dt, const FSubSt
 
 		CollisionDetector.RunBroadPhase(Dt, GetCurrentStepResimCache());
 
+		if (MidPhaseModifiers)
 		{
 			SCOPE_CYCLE_COUNTER(STAT_Evolution_MidPhaseModifierCallback);
 			CollisionConstraints.ApplyMidPhaseModifier(*MidPhaseModifiers, Dt);
