Chaos - Fix IgnoreCollisionManager so that WakeParticle does not reset ignore collision flags from joints.

The bug was related to the GT and PT both separately manipulating a flag that indicates that a particle has an entry in the ignore table. The flag is used by the broadphase. This flag was not written back to the GT when modified on the PT so any subsequent GT change (to any value in the interop data block containing the flag) would overwrite the PT changes. Explicitly waking a particle from the GT was causing this to happen.

The fix is to remove the GT-controlled flag and have it fully controlled on the PT side based on whether the particle is in the Ignore Manager or not.

This fix can be improved architecturally by breaking the IgnoreCollisionManager into 2 parts: one for the physics thread and one for the interop.

diff --git a/Engine/Source/Programs/HeadlessChaos/Private/HeadlessChaosTestIgnoreManager.cpp b/Engine/Source/Programs/HeadlessChaos/Private/HeadlessChaosTestIgnoreManager.cpp
--- a/Engine/Source/Programs/HeadlessChaos/Private/HeadlessChaosTestIgnoreManager.cpp
+++ b/Engine/Source/Programs/HeadlessChaos/Private/HeadlessChaosTestIgnoreManager.cpp
@@ -111,15 +111,12 @@ namespace ChaosTest
 		const FUniqueIdx NonCollidableId = NonCollidableParticle.UniqueIdx();
 		const FUniqueIdx FloorId = FloorParticle.UniqueIdx();
 
-		// Set up the ignore manager
-		FIgnoreCollisionManager& IgnoreManager = Solver->GetEvolution()->GetBroadPhase().GetIgnoreCollisionManager();
-		IgnoreManager.AddIgnoreCollisionsFor(NonCollidableId, FloorId);
-
 		// Flush the solver
 		FlushSolver(Solver);
 
-		// Tell the broadphase we're interested in checking the ignore manager for this object
-		NonCollidableProxy->GetHandle_LowLevel()->CastToRigidParticle()->AddCollisionConstraintFlag(ECollisionConstraintFlags::CCF_BroadPhaseIgnoreCollisions);
+		// Set up the ignore manager
+		FIgnoreCollisionManager& IgnoreManager = Solver->GetEvolution()->GetBroadPhase().GetIgnoreCollisionManager();
+		IgnoreManager.AddIgnoreCollisions(NonCollidableProxy->GetHandle_LowLevel(), FloorProxy->GetHandle_LowLevel());
 
 		RunSolver(Solver, 1.0f, 10);
 
@@ -197,16 +194,13 @@ namespace ChaosTest
 		const FUniqueIdx NonCollidableId = NonCollidableParticle.UniqueIdx();
 		const FUniqueIdx FloorId = FloorParticle.UniqueIdx();
 
-		// Set up the ignore manager (Twice as if two systems requested it)
-		FIgnoreCollisionManager& IgnoreManager = Solver->GetEvolution()->GetBroadPhase().GetIgnoreCollisionManager();
-		IgnoreManager.AddIgnoreCollisionsFor(NonCollidableId, FloorId);
-		IgnoreManager.AddIgnoreCollisionsFor(NonCollidableId, FloorId);
-
 		// Flush the solver
 		FlushSolver(Solver);
 
-		// Tell the broadphase we're interested in checking the ignore manager for this object
-		NonCollidableProxy->GetHandle_LowLevel()->CastToRigidParticle()->AddCollisionConstraintFlag(ECollisionConstraintFlags::CCF_BroadPhaseIgnoreCollisions);
+		// Set up the ignore manager (Twice as if two systems requested it)
+		FIgnoreCollisionManager& IgnoreManager = Solver->GetEvolution()->GetBroadPhase().GetIgnoreCollisionManager();
+		IgnoreManager.AddIgnoreCollisions(NonCollidableProxy->GetHandle_LowLevel(), FloorProxy->GetHandle_LowLevel());
+		IgnoreManager.AddIgnoreCollisions(NonCollidableProxy->GetHandle_LowLevel(), FloorProxy->GetHandle_LowLevel());
 
 		const float Dt = 1.0f;
 		const int32 Steps = 10;
@@ -234,7 +228,7 @@ namespace ChaosTest
 		EXPECT_EQ(Hits[NonCollidableId], 0);	// Non collidable ignored all collisions with the floor
 
 		// Remove one source
-		IgnoreManager.RemoveIgnoreCollisionsFor(NonCollidableId, FloorId);
+		IgnoreManager.RemoveIgnoreCollisions(NonCollidableProxy->GetHandle_LowLevel(), FloorProxy->GetHandle_LowLevel());
 
 		Callback->CollisionPairs.Reset();
 		RunSolver(Solver, 1.0f, 10);
@@ -248,7 +242,7 @@ namespace ChaosTest
 		EXPECT_EQ(Hits[NonCollidableId], 0);	// Non collidable ignored all collisions with the floor
 
 		// Remove final source
-		IgnoreManager.RemoveIgnoreCollisionsFor(NonCollidableId, FloorId);
+		IgnoreManager.RemoveIgnoreCollisions(NonCollidableProxy->GetHandle_LowLevel(), FloorProxy->GetHandle_LowLevel());
 
 		Callback->CollisionPairs.Reset();
 		RunSolver(Solver, 1.0f, 10);
diff --git a/Engine/Source/Runtime/Experimental/Chaos/Private/Chaos/Collision/CollisionConstraintFlags.cpp b/Engine/Source/Runtime/Experimental/Chaos/Private/Chaos/Collision/CollisionConstraintFlags.cpp
--- a/Engine/Source/Runtime/Experimental/Chaos/Private/Chaos/Collision/CollisionConstraintFlags.cpp
+++ b/Engine/Source/Runtime/Experimental/Chaos/Private/Chaos/Collision/CollisionConstraintFlags.cpp
@@ -4,6 +4,11 @@
 #include "Chaos/GeometryParticles.h"
 #include "Chaos/ParticleHandle.h"
 
+// @todo(chaos): remove after IgnoreCollisionManager refactor
+#include "PBDRigidsSolver.h"
+#include "PhysicsProxy/SingleParticlePhysicsProxy.h"
+#include "Chaos/PBDRigidsEvolutionGBF.h"
+
 namespace Chaos
 {
 
@@ -34,7 +39,7 @@ namespace Chaos
 		return Entries ? Entries->Num() : 0;
 	}
 
-	void FIgnoreCollisionManager::AddIgnoreCollisionsFor(FHandleID Body0, FHandleID Body1)
+	void FIgnoreCollisionManager::AddIgnoreCollisionsImpl(FHandleID Body0, FHandleID Body1)
 	{
 		TArray<FIgnoreEntry>& Entries = IgnoreCollisionsList.FindOrAdd(Body0);
 		FIgnoreEntry* Entry = Entries.FindByPredicate([&Body1](const FIgnoreEntry& Entry)
@@ -52,7 +57,7 @@ namespace Chaos
 		}
 	}
 
-	int32 FIgnoreCollisionManager::RemoveIgnoreCollisionsFor(FHandleID Body0, FHandleID Body1)
+	int32 FIgnoreCollisionManager::RemoveIgnoreCollisionsImpl(FHandleID Body0, FHandleID Body1)
 	{
 		TArray<FIgnoreEntry>* Entries = IgnoreCollisionsList.Find(Body0);
 
@@ -86,6 +91,88 @@ namespace Chaos
 		return 0;
 	}
 
+	bool FIgnoreCollisionManager::IgnoresCollision(const FGeometryParticleHandle* Particle0, const FGeometryParticleHandle* Particle1) const
+	{
+		return IgnoresCollision(Particle0->UniqueIdx(), Particle1->UniqueIdx());
+	}
+
+	void FIgnoreCollisionManager::SetIgnoreCollisionFlag(FPBDRigidParticleHandle* Rigid, const bool bUsesIgnoreCollisionManager)
+	{
+		if (bUsesIgnoreCollisionManager)
+		{
+			Rigid->SetUseIgnoreCollisionManager();
+		}
+		else
+		{
+			Rigid->ClearUseIgnoreCollisionManager();
+		}
+	}
+
+	void FIgnoreCollisionManager::AddIgnoreCollisions(FGeometryParticleHandle* Particle0, FGeometryParticleHandle* Particle1)
+	{
+		if (Particle0 && Particle1)
+		{
+			FPBDRigidParticleHandle* Rigid0 = Particle0->CastToRigidParticle();
+			FPBDRigidParticleHandle* Rigid1 = Particle1->CastToRigidParticle();
+			if (Rigid0 || Rigid1)
+			{
+				const FUniqueIdx ID0 = Particle0->UniqueIdx();
+				const FUniqueIdx ID1 = Particle1->UniqueIdx();
+
+				if (Rigid0)
+				{
+					AddIgnoreCollisionsImpl(ID0, ID1);
+					SetIgnoreCollisionFlag(Rigid0, true);
+				}
+
+				if (Rigid1)
+				{
+					AddIgnoreCollisionsImpl(ID1, ID0);
+					SetIgnoreCollisionFlag(Rigid1, true);
+				}
+			}
+		}
+	}
+
+	void FIgnoreCollisionManager::RemoveIgnoreCollisions(FGeometryParticleHandle* Particle0, FGeometryParticleHandle* Particle1)
+	{
+		if (Particle0 && Particle1)
+		{
+			FPBDRigidParticleHandle* Rigid0 = Particle0->CastToRigidParticle();
+			FPBDRigidParticleHandle* Rigid1 = Particle1->CastToRigidParticle();
+			if (Rigid0 || Rigid1)
+			{
+				const FUniqueIdx ID0 = Particle0->UniqueIdx();
+				const FUniqueIdx ID1 = Particle1->UniqueIdx();
+
+				if (Rigid0)
+				{
+					if (RemoveIgnoreCollisionsImpl(ID0, ID1) == 0)
+					{
+						SetIgnoreCollisionFlag(Rigid0, false);
+					}
+				}
+
+				if (Rigid1)
+				{
+					if (RemoveIgnoreCollisionsImpl(ID1, ID0) == 0)
+					{
+						SetIgnoreCollisionFlag(Rigid1, false);
+					}
+				}
+			}
+		}
+	}
+
+	FGeometryParticleHandle* FIgnoreCollisionManager::GetParticleHandle(FHandleID Body, FPBDRigidsSolver& Solver)
+	{
+		if (FSingleParticlePhysicsProxy* Proxy = Solver.GetParticleProxy_PT(Body))
+		{
+			return Proxy->GetHandle_LowLevel();
+		}
+		return nullptr;
+	}
+
 	void FIgnoreCollisionManager::PopStorageData_Internal(int32 ExternalTimestamp)
 	{
 		FStorageData* StorageData;
@@ -121,7 +208,7 @@ namespace Chaos
 		}
 	}
 
-	void FIgnoreCollisionManager::ProcessPendingQueues()
+	void FIgnoreCollisionManager::ProcessPendingQueues(FPBDRigidsSolver& Solver)
 	{
 
 		// remove particles that have been created and destroyed
@@ -164,9 +251,10 @@ namespace Chaos
 					{
 						FUniqueIdx ID0 = Elem.Key;
 						FUniqueIdx ID1 = Elem.Value[Index];
-						
-						AddIgnoreCollisionsFor(ID0, ID1);
-						AddIgnoreCollisionsFor(ID1, ID0);
+
+						FGeometryParticleHandle* Particle0 = GetParticleHandle(ID0, Solver);
+						FGeometryParticleHandle* Particle1 = GetParticleHandle(ID1, Solver);
+						AddIgnoreCollisions(Particle0, Particle1);
 
 						Elem.Value.RemoveAtSwap(Index, 1);
 					}
@@ -190,6 +278,14 @@ namespace Chaos
 			for (auto& Idx : PendingDeactivations)
 			{
 				IgnoreCollisionsList.Remove(Idx);
+
+				if (FGeometryParticleHandle* Particle = GetParticleHandle(Idx, Solver))
+				{
+					if (FPBDRigidParticleHandle* Rigid = Particle->CastToRigidParticle())
+					{
+						SetIgnoreCollisionFlag(Rigid, false);
+					}
+				}
 			}
 			PendingDeactivations.Empty();
 		}
diff --git a/Engine/Source/Runtime/Experimental/Chaos/Private/Chaos/ParticleHandle.cpp b/Engine/Source/Runtime/Experimental/Chaos/Private/Chaos/ParticleHandle.cpp
--- a/Engine/Source/Runtime/Experimental/Chaos/Private/Chaos/ParticleHandle.cpp
+++ b/Engine/Source/Runtime/Experimental/Chaos/Private/Chaos/ParticleHandle.cpp
@@ -235,39 +235,4 @@ namespace Chaos
 		return Idx;
 	}
 
-	template <>
-	void TPBDRigidParticleHandleImp<FReal, 3, true>::SetDynamicMisc(const FParticleDynamicMisc& DynamicMisc, FPBDRigidsEvolutionBase& Evolution)
-	{
-	
-		if (Disabled() != DynamicMisc.Disabled())
-		{
-			if (DynamicMisc.Disabled())
-			{
-				Evolution.DisableParticle(this);
-			}
-			else
-			{
-				Evolution.EnableParticle(this);
-			}
-		}		
-
-		const bool bDirtyInertiaConditioning = (ObjectState() != DynamicMisc.ObjectState());
-		if (bDirtyInertiaConditioning)
-		{
-			SetInertiaConditioningDirty();
-		}
-
-		SetLinearEtherDrag(DynamicMisc.LinearEtherDrag());
-		SetAngularEtherDrag(DynamicMisc.AngularEtherDrag());
-		SetMaxLinearSpeedSq(DynamicMisc.MaxLinearSpeedSq());
-		SetMaxAngularSpeedSq(DynamicMisc.MaxAngularSpeedSq());
-		SetCollisionGroup(DynamicMisc.CollisionGroup());
-		SetDisabled(DynamicMisc.Disabled());
-		SetCollisionConstraintFlags(DynamicMisc.CollisionConstraintFlags());
-		SetControlFlags(DynamicMisc.ControlFlags());
-
-		Evolution.SetParticleObjectState(this, DynamicMisc.ObjectState());
-		Evolution.SetParticleSleepType(this, DynamicMisc.SleepType());
-	}
-
 }
diff --git a/Engine/Source/Runtime/Experimental/Chaos/Private/PBDRigidsSolver.cpp b/Engine/Source/Runtime/Experimental/Chaos/Private/PBDRigidsSolver.cpp
--- a/Engine/Source/Runtime/Experimental/Chaos/Private/PBDRigidsSolver.cpp
+++ b/Engine/Source/Runtime/Experimental/Chaos/Private/PBDRigidsSolver.cpp
@@ -333,7 +333,7 @@ namespace Chaos
 					GeoclObj->FieldParameterUpdateCallback(MSolver);
 				}
 
-				MSolver->GetEvolution()->GetBroadPhase().GetIgnoreCollisionManager().ProcessPendingQueues();
+				MSolver->GetEvolution()->GetBroadPhase().GetIgnoreCollisionManager().ProcessPendingQueues(*MSolver);
 			}
 
 			{
@@ -1131,7 +1131,7 @@ namespace Chaos
 				{
 					RewindData->PushGTDirtyData(*Manager, DataIdx, Dirty, ShapeDirtyData);
 				}
-				Proxy->PushToPhysicsState(*Manager, DataIdx, Dirty, ShapeDirtyData, *GetEvolution(), ExternalDt);
+				Proxy->PushToPhysicsState(*Manager, DataIdx, Dirty, ShapeDirtyData, ExternalDt);
 			}
 			else
 			{
@@ -1712,18 +1712,24 @@ CSV_CUSTOM_STAT(PhysicsCounters, Name, Value, ECsvCustomStatOp::Set);
 #endif
 	}
 
+	FSingleParticlePhysicsProxy* FPBDRigidsSolver::GetParticleProxy_PT(const FUniqueIdx& Idx)
+	{
+		return SingleParticlePhysicsProxies_PT.IsValidIndex(Idx.Idx) ? SingleParticlePhysicsProxies_PT[Idx.Idx] : nullptr;
+	}
+
+	const FSingleParticlePhysicsProxy* FPBDRigidsSolver::GetParticleProxy_PT(const FUniqueIdx& Idx) const
+	{
+		return SingleParticlePhysicsProxies_PT.IsValidIndex(Idx.Idx) ? SingleParticlePhysicsProxies_PT[Idx.Idx] : nullptr;
+	}
+
 	FSingleParticlePhysicsProxy* FPBDRigidsSolver::GetParticleProxy_PT(const FGeometryParticleHandle& Handle)
 	{
-		const FUniqueIdx UniqueIdx = Handle.UniqueIdx();
-		return SingleParticlePhysicsProxies_PT.IsValidIndex(UniqueIdx.Idx)
-			? SingleParticlePhysicsProxies_PT[UniqueIdx.Idx] : nullptr;
+		return GetParticleProxy_PT(Handle.UniqueIdx());
 	}
 
 	const FSingleParticlePhysicsProxy* FPBDRigidsSolver::GetParticleProxy_PT(const FGeometryParticleHandle& Handle) const
 	{
-		const FUniqueIdx UniqueIdx = Handle.UniqueIdx();
-		return SingleParticlePhysicsProxies_PT.IsValidIndex(UniqueIdx.Idx)
-			? SingleParticlePhysicsProxies_PT[UniqueIdx.Idx] : nullptr;
+		return GetParticleProxy_PT(Handle.UniqueIdx());
 	}
 
 	void FPBDRigidsSolver::PostEvolutionVDBPush() const
@@ -1830,6 +1836,46 @@ CSV_CUSTOM_STAT(PhysicsCounters, Name, Value, ECsvCustomStatOp::Set);
 		GetEvolution()->SetIsDeterministic(IsDetemerministic());
 	}
 
+	void FPBDRigidsSolver::SetParticleDynamicMisc(FPBDRigidParticleHandle* Rigid, const FParticleDynamicMisc& DynamicMisc)
+	{
+		if (Rigid == nullptr)
+		{
+			return;
+		}
+
+		// Enable or disable the particle
+		if (Rigid->Disabled() != DynamicMisc.Disabled())
+		{
+			if (DynamicMisc.Disabled())
+			{
+				GetEvolution()->DisableParticle(Rigid);
+			}
+			else
+			{
+				GetEvolution()->EnableParticle(Rigid);
+			}
+		}
+
+		// If we changed kinematics we need to rebuild the inertia conditioning
+		const bool bDirtyInertiaConditioning = (Rigid->ObjectState() != DynamicMisc.ObjectState());
+		if (bDirtyInertiaConditioning)
+		{
+			Rigid->SetInertiaConditioningDirty();
+		}
+
+		Rigid->SetLinearEtherDrag(DynamicMisc.LinearEtherDrag());
+		Rigid->SetAngularEtherDrag(DynamicMisc.AngularEtherDrag());
+		Rigid->SetMaxLinearSpeedSq(DynamicMisc.MaxLinearSpeedSq());
+		Rigid->SetMaxAngularSpeedSq(DynamicMisc.MaxAngularSpeedSq());
+		Rigid->SetCollisionGroup(DynamicMisc.CollisionGroup());
+		Rigid->SetDisabled(DynamicMisc.Disabled());
+		Rigid->SetCollisionConstraintFlags(DynamicMisc.CollisionConstraintFlags());
+		Rigid->SetControlFlags(DynamicMisc.ControlFlags());
+
+		GetEvolution()->SetParticleObjectState(Rigid, DynamicMisc.ObjectState());
+		GetEvolution()->SetParticleSleepType(Rigid, DynamicMisc.SleepType());
+	}
+
 
 	Chaos::FClusterCreationParameters::EConnectionMethod ToInternalConnectionMethod(EClusterUnionMethod InMethod)
 	{
diff --git a/Engine/Source/Runtime/Experimental/Chaos/Private/PhysicsProxy/JointConstraintProxy.cpp b/Engine/Source/Runtime/Experimental/Chaos/Private/PhysicsProxy/JointConstraintProxy.cpp
--- a/Engine/Source/Runtime/Experimental/Chaos/Private/PhysicsProxy/JointConstraintProxy.cpp
+++ b/Engine/Source/Runtime/Experimental/Chaos/Private/PhysicsProxy/JointConstraintProxy.cpp
@@ -117,35 +117,7 @@ void FJointConstraintPhysicsProxy::DestroyOnPhysicsThread(FPBDRigidsSolver* InSo
 		// If this constraint disables collisions - need to restore the collisions after we destroy the constraint
 		if(bValidPair && !Settings.bCollisionEnabled)
 		{
-			FPBDRigidParticleHandle* Rigid0 = Particles[0]->CastToRigidParticle();
-			FPBDRigidParticleHandle* Rigid1 = Particles[1]->CastToRigidParticle();
-
-			if(Rigid0 || Rigid1)
-			{
-				// At least one rigid so collisions are possible
-				const FUniqueIdx Id0 = Particles[0]->UniqueIdx();
-				const FUniqueIdx Id1 = Particles[1]->UniqueIdx();
-
-				FIgnoreCollisionManager& IgnoreManager = InSolver->GetEvolution()->GetBroadPhase().GetIgnoreCollisionManager();
-
-				// Remove the ignore flags for rigid particles. If we end up with no more ignore entries for this particle
-				// then remove the broadphase flag to no longer check the ignore manager for this particles.
-				if(Rigid0)
-				{
-					if(IgnoreManager.RemoveIgnoreCollisionsFor(Id0, Id1) == 0)
-					{
-						Rigid0->RemoveCollisionConstraintFlag(ECollisionConstraintFlags::CCF_BroadPhaseIgnoreCollisions);
-					}
-				}
-
-				if(Rigid1)
-				{
-					if(IgnoreManager.RemoveIgnoreCollisionsFor(Id1, Id0) == 0)
-					{
-						Rigid1->RemoveCollisionConstraintFlag(ECollisionConstraintFlags::CCF_BroadPhaseIgnoreCollisions);
-					}
-				}
-			}
+			InSolver->GetEvolution()->GetBroadPhase().GetIgnoreCollisionManager().RemoveIgnoreCollisions(Particles[0], Particles[1]);
 		}
 
 		// Ensure that our connected particles are aware that this constraint no longer exists
@@ -196,64 +168,21 @@ void FJointConstraintPhysicsProxy::PushStateOnPhysicsThread(FPBDRigidsSolver* In
 			const FPBDJointSettings& JointSettingsBuffer = *Data;
 			const FPBDJointSettings& CurrentConstraintSettings = Constraint_PT->GetSettings();
 
+			// Handle changes to the CollisionEnabled flag
 			if (CurrentConstraintSettings.bCollisionEnabled != JointSettingsBuffer.bCollisionEnabled)
 			{
-				const TVector<FGeometryParticleHandle*, 2>& BasePairs = Constraint_PT->GetConstrainedParticles();
-				FGeometryParticleHandle * Handle0 = BasePairs[0];
-				FGeometryParticleHandle* Handle1 = BasePairs[1];
-
-				// Three pieces of state to update on the physics thread. 
-				// .. Mask on the particle array
-				// .. Constraint collisions enabled array
-				// .. IgnoreCollisionsManager
-				if (Handle0 && Handle1)
+				const TVector<FGeometryParticleHandle*, 2>& Particles = Constraint_PT->GetConstrainedParticles();
+				if (JointSettingsBuffer.bCollisionEnabled)
 				{
-					FPBDRigidParticleHandle* RigidHandle0 = Handle0->CastToRigidParticle();
-					FPBDRigidParticleHandle* RigidHandle1 = Handle1->CastToRigidParticle();
-
-					// As long as one particle is a rigid we can add the ignore entry, one particle can be a static
-					if (RigidHandle0 || RigidHandle1)
-					{
-						const FUniqueIdx ID0 = Handle0->UniqueIdx();
-						const FUniqueIdx ID1 = Handle1->UniqueIdx();
-						FIgnoreCollisionManager& IgnoreCollisionManager = InSolver->GetEvolution()->GetBroadPhase().GetIgnoreCollisionManager();
-
-						// For rigid/dynamic particles, add the broadphase flag and the IDs to check for disabled collisions
-						if(RigidHandle0)
-						{
-							if(JointSettingsBuffer.bCollisionEnabled)
-							{
-								if(IgnoreCollisionManager.RemoveIgnoreCollisionsFor(ID0, ID1) == 0)
-								{
-									RigidHandle0->RemoveCollisionConstraintFlag(ECollisionConstraintFlags::CCF_BroadPhaseIgnoreCollisions);
-								}
-							}
-							else
-							{
-								RigidHandle0->AddCollisionConstraintFlag(ECollisionConstraintFlags::CCF_BroadPhaseIgnoreCollisions);
-								IgnoreCollisionManager.AddIgnoreCollisionsFor(ID0, ID1);
-							}
-						}
-
-						if(RigidHandle1)
-						{
-							if (JointSettingsBuffer.bCollisionEnabled)
-							{
-								if(IgnoreCollisionManager.RemoveIgnoreCollisionsFor(ID1, ID0) == 0)
-								{
-									RigidHandle1->RemoveCollisionConstraintFlag(ECollisionConstraintFlags::CCF_BroadPhaseIgnoreCollisions);
-								}
-							}
-							else
-							{
-								RigidHandle1->AddCollisionConstraintFlag(ECollisionConstraintFlags::CCF_BroadPhaseIgnoreCollisions);
-								IgnoreCollisionManager.AddIgnoreCollisionsFor(ID1, ID0);
-							}
-						}
-					}
+					InSolver->GetEvolution()->GetBroadPhase().GetIgnoreCollisionManager().RemoveIgnoreCollisions(Particles[0], Particles[1]);
+				}
+				else
+				{
+					InSolver->GetEvolution()->GetBroadPhase().GetIgnoreCollisionManager().AddIgnoreCollisions(Particles[0], Particles[1]);
 				}
 			}
 
+			// Update the joint settings
 			Constraint_PT->SetSettings(JointSettingsBuffer);
 		}
 	}
diff --git a/Engine/Source/Runtime/Experimental/Chaos/Private/PhysicsProxy/SingleParticlePhysicsProxy.cpp b/Engine/Source/Runtime/Experimental/Chaos/Private/PhysicsProxy/SingleParticlePhysicsProxy.cpp
--- a/Engine/Source/Runtime/Experimental/Chaos/Private/PhysicsProxy/SingleParticlePhysicsProxy.cpp
+++ b/Engine/Source/Runtime/Experimental/Chaos/Private/PhysicsProxy/SingleParticlePhysicsProxy.cpp
@@ -42,8 +42,8 @@ FSingleParticlePhysicsProxy::~FSingleParticlePhysicsProxy()
 CHAOS_API int32 ForceNoCollisionIntoSQ = 0;
 FAutoConsoleVariableRef CVarForceNoCollisionIntoSQ(TEXT("p.ForceNoCollisionIntoSQ"), ForceNoCollisionIntoSQ, TEXT("When enabled, all particles end up in sq structure, even ones with no collision"));
 
-template <Chaos::EParticleType ParticleType, typename TEvolution>
-void PushToPhysicsStateImp(const Chaos::FDirtyPropertiesManager& Manager, Chaos::FGeometryParticleHandle* Handle, int32 DataIdx, const Chaos::FDirtyProxy& Dirty, Chaos::FShapeDirtyData* ShapesData, TEvolution& Evolution, bool bResimInitialized, Chaos::FReal ExternalDt)
+template <Chaos::EParticleType ParticleType>
+void PushToPhysicsStateImp(const Chaos::FDirtyPropertiesManager& Manager, Chaos::FGeometryParticleHandle* Handle, int32 DataIdx, const Chaos::FDirtyProxy& Dirty, Chaos::FShapeDirtyData* ShapesData, Chaos::FPBDRigidsSolver& Solver, bool bResimInitialized, Chaos::FReal ExternalDt)
 {
 	using namespace Chaos;
 	constexpr bool bHasKinematicData = ParticleType != EParticleType::Static;
@@ -51,6 +51,7 @@ void PushToPhysicsStateImp(const Chaos::FDirtyPropertiesManager& Manager, Chaos:
 	auto KinematicHandle = bHasKinematicData ? static_cast<Chaos::FKinematicGeometryParticleHandle*>(Handle) : nullptr;
 	auto RigidHandle = bHasDynamicData ? static_cast<Chaos::FPBDRigidParticleHandle*>(Handle) : nullptr;
 	const FDirtyChaosProperties& ParticleData = Dirty.PropertyData;
+	FPBDRigidsEvolutionGBF& Evolution = *Solver.GetEvolution();
 
 	if (bResimInitialized)	//todo: assumes particles are always initialized as enabled. This is not true in future versions of code, so check PushData
 	{
@@ -116,7 +117,7 @@ void PushToPhysicsStateImp(const Chaos::FDirtyPropertiesManager& Manager, Chaos:
 
 			if(auto NewData = ParticleData.FindDynamicMisc(Manager,DataIdx))
 			{
-				RigidHandle->SetDynamicMisc(*NewData, Evolution);				
+				Solver.SetParticleDynamicMisc(RigidHandle, *NewData);
 			}
 		}
 
@@ -183,18 +184,18 @@ void PushToPhysicsStateImp(const Chaos::FDirtyPropertiesManager& Manager, Chaos:
 // TGeometryParticle<FReal, 3> template specialization 
 //
 
-void FSingleParticlePhysicsProxy::PushToPhysicsState(const Chaos::FDirtyPropertiesManager& Manager, int32 DataIdx, const Chaos::FDirtyProxy& Dirty, Chaos::FShapeDirtyData* ShapesData, Chaos::FPBDRigidsEvolutionGBF& Evolution, Chaos::FReal ExternalDt)
+void FSingleParticlePhysicsProxy::PushToPhysicsState(const Chaos::FDirtyPropertiesManager& Manager, int32 DataIdx, const Chaos::FDirtyProxy& Dirty, Chaos::FShapeDirtyData* ShapesData, Chaos::FReal ExternalDt)
 {
 	using namespace Chaos;
-	const int32 CurFrame = static_cast<FPBDRigidsSolver*>(Solver)->GetCurrentFrame();
-	const FRewindData* RewindData = static_cast<FPBDRigidsSolver*>(Solver)->GetRewindData();
+	FPBDRigidsSolver& RigidsSolver = *static_cast<FPBDRigidsSolver*>(Solver);
+	const int32 CurFrame = RigidsSolver.GetCurrentFrame();
+	const FRewindData* RewindData = RigidsSolver.GetRewindData();
 	const bool bResimInitialized = RewindData && RewindData->IsResim() && CurFrame == InitializedOnStep;
 	switch(Dirty.PropertyData.GetParticleBufferType())
 	{
-		
-	case EParticleType::Static: PushToPhysicsStateImp<EParticleType::Static>(Manager, Handle, DataIdx, Dirty, ShapesData, Evolution, bResimInitialized, ExternalDt); break;
-	case EParticleType::Kinematic: PushToPhysicsStateImp<EParticleType::Kinematic>(Manager, Handle, DataIdx, Dirty, ShapesData, Evolution, bResimInitialized, ExternalDt); break;
-	case EParticleType::Rigid: PushToPhysicsStateImp<EParticleType::Rigid>(Manager, Handle, DataIdx, Dirty, ShapesData, Evolution, bResimInitialized, ExternalDt); break;
+	case EParticleType::Static: PushToPhysicsStateImp<EParticleType::Static>(Manager, Handle, DataIdx, Dirty, ShapesData, RigidsSolver, bResimInitialized, ExternalDt); break;
+	case EParticleType::Kinematic: PushToPhysicsStateImp<EParticleType::Kinematic>(Manager, Handle, DataIdx, Dirty, ShapesData, RigidsSolver, bResimInitialized, ExternalDt); break;
+	case EParticleType::Rigid: PushToPhysicsStateImp<EParticleType::Rigid>(Manager, Handle, DataIdx, Dirty, ShapesData, RigidsSolver, bResimInitialized, ExternalDt); break;
 	default: check(false); //unexpected path
 	}
 }
diff --git a/Engine/Source/Runtime/Experimental/Chaos/Private/RewindData.cpp b/Engine/Source/Runtime/Experimental/Chaos/Private/RewindData.cpp
--- a/Engine/Source/Runtime/Experimental/Chaos/Private/RewindData.cpp
+++ b/Engine/Source/Runtime/Experimental/Chaos/Private/RewindData.cpp
@@ -201,7 +201,7 @@ bool FRewindData::RewindToFrame(int32 Frame)
 		bAnyChange |= RewindHelper(PTParticle->CastToKinematicParticle(), bResimAsFollower, History.Velocities, [](auto Particle, const auto& Data) {Particle->SetVelocities(Data); });
 		bAnyChange |= RewindHelper(PTParticle->CastToKinematicParticle(), bResimAsFollower, History.KinematicTarget, [](auto Particle, const auto& Data) {Particle->SetKinematicTarget(Data); });
 		bAnyChange |= RewindHelper(PTParticle->CastToRigidParticle(),bResimAsFollower,  History.Dynamics, [](auto Particle, const auto& Data) {Particle->SetDynamics(Data); });
-		bAnyChange |= RewindHelper(PTParticle->CastToRigidParticle(),bResimAsFollower,  History.DynamicsMisc, [Evolution = Solver->GetEvolution()](auto Particle, const auto& Data) {Particle->SetDynamicMisc(Data, *Evolution); });
+		bAnyChange |= RewindHelper(PTParticle->CastToRigidParticle(),bResimAsFollower,  History.DynamicsMisc, [this](auto Particle, const auto& Data) {Solver->SetParticleDynamicMisc(Particle, Data); });
 		bAnyChange |= RewindHelper(PTParticle->CastToRigidParticle(),bResimAsFollower,  History.MassProps, [](auto Particle, const auto& Data) {Particle->SetMassProps(Data); });
 
 		if (!bResimAsFollower)
diff --git a/Engine/Source/Runtime/Experimental/Chaos/Public/Chaos/Collision/CollisionConstraintFlags.h b/Engine/Source/Runtime/Experimental/Chaos/Public/Chaos/Collision/CollisionConstraintFlags.h
--- a/Engine/Source/Runtime/Experimental/Chaos/Public/Chaos/Collision/CollisionConstraintFlags.h
+++ b/Engine/Source/Runtime/Experimental/Chaos/Public/Chaos/Collision/CollisionConstraintFlags.h
@@ -14,17 +14,24 @@
 
 namespace Chaos
 {
+	// @todo(chaos): remove these when FIgnoreCollisionManager is split into Core and Proxy classes (see todo comments below)
+	class FPBDRigidsSolver;
+
 	/**
 	 * @brief Flags for user-control over per-particle collision behaviour
+	 * @todo(chaos): move SmoothEdgeCollisions to FRigidParticleControlFlags and remove ECollisionConstraintFlags
 	*/
 	enum class ECollisionConstraintFlags : uint32
 	{
 		CCF_None                       = 0x0,
-		CCF_BroadPhaseIgnoreCollisions = 0x1,
+		CCF_BroadPhaseIgnoreCollisions = 0x1,	// Not used
 		CCF_SmoothEdgeCollisions       = 0x2,
 		CCF_DummyFlag
 	};
 
+	// @todo(chaos): this should be broken out into FIgnoreCollisionManager (PT-only) and FIgnoreCollisionManagerProxy (interop).
+	// E.g., The interop could store sets of particle proxy pairs to enable/disable and call the ParticleHandle API on the PT object. 
+	// Or it could use the solver to map IDs to particle handles like we currently do.
 	class CHAOS_API FIgnoreCollisionManager
 	{
 	public:
@@ -53,22 +60,39 @@ namespace Chaos
 			StorageDataProducer = GetNewStorageData();
 		}
 
+		//
+		//
+		// PHYSICS THREAD API
+		//
+		//
+
 		bool ContainsHandle(FHandleID Body0) const;
 
 		bool IgnoresCollision(FHandleID Body0, FHandleID Body1) const;
 
 		int32 NumIgnoredCollision(FHandleID Body0) const;
 
-		void AddIgnoreCollisionsFor(FHandleID Body0, FHandleID Body1);
+		/**
+		 * Return true if we should ignore collisions between the two particles
+		 */
+		bool IgnoresCollision(const FGeometryParticleHandle* Particle0, const FGeometryParticleHandle* Particle1) const;
+
+		/**
+		 * Add an ignore entry for collisions of Particle0->Particle1.
+		 * Note, this will add an entry for both particles against the other if both particles are dynamic or kinematic.
+		 */
+		void AddIgnoreCollisions(FGeometryParticleHandle* Particle0, FGeometryParticleHandle* Particle1);
 
 		/**
-		 * Remove an ignore entry for collisions of Body0->Body1
-		 * Note, a reversed ignore entry could exist
-		 * @param Body0 First body in the collision
-		 * @param Body1 Second body in the collision
-		 * @return Number of collisions Body0 ignores after the removal
+		 * Remove an ignore entry for collisions of Particle0->Particle1 (in both directions if applicable)
 		 */
-		int32 RemoveIgnoreCollisionsFor(FHandleID Body0, FHandleID Body1);
+		void RemoveIgnoreCollisions(FGeometryParticleHandle* Particle0, FGeometryParticleHandle* Particle1);
+
+		//
+		//
+		// GAME THREAD API
+		//
+		//
 
 		FPendingMap& GetPendingActivationsForGameThread(int32 ExternalTimestamp) 
 		{
@@ -108,18 +132,36 @@ namespace Chaos
 			}
 		}
 
+		//
+		//
+		// INTEROP API
+		//
+		//
+
 		/*
 		*
 		*/
-		void ProcessPendingQueues();
+		void ProcessPendingQueues(FPBDRigidsSolver& Solver);
 
 		/*
 		*
 		*/
 		void PopStorageData_Internal(int32 ExternalTimestamp);
 
+
+		// DEPRECATED
+		UE_DEPRECATED(5.2, "The physics thread uses the ParticleHandle API. See AddIgnoreCollisions")
+		void AddIgnoreCollisionsFor(FHandleID Body0, FHandleID Body1) { }
+		UE_DEPRECATED(5.2, "The physics thread uses the ParticleHandle API. See RemoveIgnoreCollisions")
+		int32 RemoveIgnoreCollisionsFor(FHandleID Body0, FHandleID Body1) { return 0; }
+
 	private:
 
+		FGeometryParticleHandle* GetParticleHandle(FHandleID Body, FPBDRigidsSolver& Solver);
+		void SetIgnoreCollisionFlag(FPBDRigidParticleHandle* Rigid, const bool bUsesIgnoreCollisionManager);
+		void AddIgnoreCollisionsImpl(FHandleID Body0, FHandleID Body1);
+		int32 RemoveIgnoreCollisionsImpl(FHandleID Body0, FHandleID Body1);
+
 		// Multiple sources can request an ignore pair, so we handle a simple count of
 		// ignore requests to ensure we don't prematurely remove an ignore pair
 		// #CHAOSTODO replace bi-directional map to pair hash to reduce required storage and dependency
diff --git a/Engine/Source/Runtime/Experimental/Chaos/Public/Chaos/Collision/SpatialAccelerationBroadPhase.h b/Engine/Source/Runtime/Experimental/Chaos/Public/Chaos/Collision/SpatialAccelerationBroadPhase.h
--- a/Engine/Source/Runtime/Experimental/Chaos/Public/Chaos/Collision/SpatialAccelerationBroadPhase.h
+++ b/Engine/Source/Runtime/Experimental/Chaos/Public/Chaos/Collision/SpatialAccelerationBroadPhase.h
@@ -94,7 +94,7 @@ namespace Chaos
 				bIsKinematic1 = Rigid1->IsKinematic();
 				bIsDynamicAsleep1 = !bIsKinematic1 && Rigid1->IsSleeping();
 				bIsDynamicAwake1 = !bIsKinematic1 && !bIsDynamicAsleep1;
-				bUseIgnoreCollisionManager1 = Rigid1->HasCollisionConstraintFlag(ECollisionConstraintFlags::CCF_BroadPhaseIgnoreCollisions);
+				bUseIgnoreCollisionManager1 = Rigid1->UseIgnoreCollisionManager();
 				bDisabled1 = Rigid1->Disabled();
 				CollisionGroup1 = Rigid1->CollisionGroup();
 			}
@@ -111,7 +111,7 @@ namespace Chaos
 				bIsKinematic2 = Rigid2->IsKinematic();
 				bIsDynamicAsleep2 = !bIsKinematic2 && Rigid2->IsSleeping();
 				bIsDynamicAwake2 = !bIsKinematic2 && !bIsDynamicAsleep2;
-				bUseIgnoreCollisionManager2 = Rigid2->HasCollisionConstraintFlag(ECollisionConstraintFlags::CCF_BroadPhaseIgnoreCollisions);
+				bUseIgnoreCollisionManager2 = Rigid2->UseIgnoreCollisionManager();
 				bDisabled2 = Rigid2->Disabled();
 				CollisionGroup2 = Rigid2->CollisionGroup();
 			}
@@ -133,7 +133,7 @@ namespace Chaos
 			// Is this particle interaction governed by the IgnoreCollisionManager? If so, check to see if interaction is allowed
 			if (bUseIgnoreCollisionManager1 || bUseIgnoreCollisionManager2)
 			{
-				if (IgnoreCollisionManager.IgnoresCollision(Particle1->UniqueIdx(), Particle2->UniqueIdx()))
+				if (IgnoreCollisionManager.IgnoresCollision(Particle1, Particle2))
 				{
 					return false;
 				}
diff --git a/Engine/Source/Runtime/Experimental/Chaos/Public/Chaos/ParticleHandle.h b/Engine/Source/Runtime/Experimental/Chaos/Public/Chaos/ParticleHandle.h
--- a/Engine/Source/Runtime/Experimental/Chaos/Public/Chaos/ParticleHandle.h
+++ b/Engine/Source/Runtime/Experimental/Chaos/Public/Chaos/ParticleHandle.h
@@ -1022,7 +1022,8 @@ public:
 		SetInertiaConditioningDirty();
 	}
 
-	void SetDynamicMisc(const FParticleDynamicMisc& DynamicMisc, FPBDRigidsEvolutionBase& Evolution);
+	UE_DEPRECATED(5.2, "Moved to FPBDRigidsSolver")
+	void SetDynamicMisc(const FParticleDynamicMisc& DynamicMisc, FPBDRigidsEvolutionBase& Evolution) {}
 
 	void ResetSmoothedVelocities()
 	{
@@ -1170,6 +1171,21 @@ public:
 		PBDRigidParticles->TransientFlags(ParticleIdx).ClearInertiaConditioningDirty();
 	}
 
+	inline bool UseIgnoreCollisionManager() const
+	{
+		return PBDRigidParticles->TransientFlags(ParticleIdx).GetUseIgnoreCollisionManager();
+	}
+
+	inline void SetUseIgnoreCollisionManager()
+	{
+		PBDRigidParticles->TransientFlags(ParticleIdx).SetUseIgnoreCollisionManager();
+	}
+
+	inline void ClearUseIgnoreCollisionManager()
+	{
+		PBDRigidParticles->TransientFlags(ParticleIdx).ClearUseIgnoreCollisionManager();
+	}
+
 	ESleepType SleepType() const { return PBDRigidParticles->SleepType(ParticleIdx);}
 
 	void SetSleepType(ESleepType SleepType){ PBDRigidParticles->SetSleepType(ParticleIdx, SleepType); }
@@ -1970,6 +1986,31 @@ public:
 		}
 	}
 
+	inline bool UseIgnoreCollisionManager() const
+	{
+		if (auto RigidHandle = MHandle->CastToRigidParticle())
+		{
+			return RigidHandle->UseIgnoreCollisionManager();
+		}
+		return false;
+	}
+
+	inline void SetUseIgnoreCollisionManager()
+	{
+		if (auto RigidHandle = MHandle->CastToRigidParticle())
+		{
+			RigidHandle->SetUseIgnoreCollisionManager();
+		}
+	}
+
+	inline void ClearUseIgnoreCollisionManager()
+	{
+		if (auto RigidHandle = MHandle->CastToRigidParticle())
+		{
+			RigidHandle->ClearUseIgnoreCollisionManager();
+		}
+	}
+
 	FReal LinearEtherDrag() const
 	{
 		if (auto RigidHandle = MHandle->CastToRigidParticle())
diff --git a/Engine/Source/Runtime/Experimental/Chaos/Public/Chaos/RigidParticleControlFlags.h b/Engine/Source/Runtime/Experimental/Chaos/Public/Chaos/RigidParticleControlFlags.h
--- a/Engine/Source/Runtime/Experimental/Chaos/Public/Chaos/RigidParticleControlFlags.h
+++ b/Engine/Source/Runtime/Experimental/Chaos/Public/Chaos/RigidParticleControlFlags.h
@@ -101,14 +101,21 @@ namespace Chaos
 		FStorage GetFlags() const { return Bits; }
 		void SetFlags(const FStorage InFlags) { Bits = InFlags; }
 
+		// Set to true when some property of the particle changes and we should regenerate the inertia conditioning vector
 		bool GetInertiaConditioningDirty() const { return Flags.bInertiaConditioningDirty; }
 		void SetInertiaConditioningDirty() { Flags.bInertiaConditioningDirty = true; }
 		void ClearInertiaConditioningDirty() { Flags.bInertiaConditioningDirty = false; }
 
+		// Set to true when the particle has one or more entries in the CollisionIgnoreManager
+		bool GetUseIgnoreCollisionManager() const { return Flags.bUseIgnoreCollisionManager; }
+		void SetUseIgnoreCollisionManager() { Flags.bUseIgnoreCollisionManager = true; }
+		void ClearUseIgnoreCollisionManager() { Flags.bUseIgnoreCollisionManager = false; }
+
 	private:
 		struct FFlags
 		{
 			FStorage bInertiaConditioningDirty : 1;
+			FStorage bUseIgnoreCollisionManager : 1;
 			// Add new properties above this line
 			// Change FStorage typedef if we exceed the max bits
 		};
diff --git a/Engine/Source/Runtime/Experimental/Chaos/Public/PBDRigidsSolver.h b/Engine/Source/Runtime/Experimental/Chaos/Public/PBDRigidsSolver.h
--- a/Engine/Source/Runtime/Experimental/Chaos/Public/PBDRigidsSolver.h
+++ b/Engine/Source/Runtime/Experimental/Chaos/Public/PBDRigidsSolver.h
@@ -302,9 +302,15 @@ namespace Chaos
 		void UpdateExpensiveStatCounters() const;
 
 		// Access particle proxy from physics thread useful for cross thread communication
+		FSingleParticlePhysicsProxy* GetParticleProxy_PT(const FUniqueIdx& Idx);
+		const FSingleParticlePhysicsProxy* GetParticleProxy_PT(const FUniqueIdx& Idx) const;
 		FSingleParticlePhysicsProxy* GetParticleProxy_PT(const FGeometryParticleHandle& Handle);
 		const FSingleParticlePhysicsProxy* GetParticleProxy_PT(const FGeometryParticleHandle& Handle) const;
 
+		// Interop utilities
+		void SetParticleDynamicMisc(FPBDRigidParticleHandle* Rigid, const FParticleDynamicMisc& DynamicMisc);
+
+
 	private:
 
 		/**/
diff --git a/Engine/Source/Runtime/Experimental/Chaos/Public/PhysicsProxy/SingleParticlePhysicsProxy.h b/Engine/Source/Runtime/Experimental/Chaos/Public/PhysicsProxy/SingleParticlePhysicsProxy.h
--- a/Engine/Source/Runtime/Experimental/Chaos/Public/PhysicsProxy/SingleParticlePhysicsProxy.h
+++ b/Engine/Source/Runtime/Experimental/Chaos/Public/PhysicsProxy/SingleParticlePhysicsProxy.h
@@ -120,7 +120,7 @@ public:
 
 	// Threading API
 
-	void PushToPhysicsState(const FDirtyPropertiesManager& Manager,int32 DataIdx,const FDirtyProxy& Dirty,FShapeDirtyData* ShapesData, FPBDRigidsEvolutionGBF& Evolution, FReal ExternalDt);
+	void PushToPhysicsState(const FDirtyPropertiesManager& Manager,int32 DataIdx,const FDirtyProxy& Dirty,FShapeDirtyData* ShapesData, FReal ExternalDt);
 
 	/**/
 	void ClearAccumulatedData();
diff --git a/Engine/Source/Runtime/PhysicsCore/Private/ChaosEngineInterface.cpp b/Engine/Source/Runtime/PhysicsCore/Private/ChaosEngineInterface.cpp
--- a/Engine/Source/Runtime/PhysicsCore/Private/ChaosEngineInterface.cpp
+++ b/Engine/Source/Runtime/PhysicsCore/Private/ChaosEngineInterface.cpp
@@ -326,12 +326,10 @@ void FChaosEngineInterface::AddDisabledCollisionsFor_AssumesLocked(const TMap<FP
 
 			if (Chaos::FPBDRigidParticle* Rigid0 = ActorReference->GetParticle_LowLevel()->CastToRigidParticle())
 			{
-				Rigid0->AddCollisionConstraintFlag(Chaos::ECollisionConstraintFlags::CCF_BroadPhaseIgnoreCollisions);
 				for (auto Handle1 : Elem.Value)
 				{
 					if (Chaos::FPBDRigidParticle* Rigid1 = Handle1->GetParticle_LowLevel()->CastToRigidParticle())
 					{
-						Rigid1->AddCollisionConstraintFlag(Chaos::ECollisionConstraintFlags::CCF_BroadPhaseIgnoreCollisions);
 						DisabledCollisions.Add(Handle1->GetGameThreadAPI().UniqueIdx());
 					}
 				}
