Chaos - optimize MidPhase ResetModifications

Moved it to just after the narrowphase where the cache is hot.

diff --git a/Engine/Source/Runtime/Experimental/Chaos/Private/Chaos/Collision/ParticlePairMidPhase.cpp b/Engine/Source/Runtime/Experimental/Chaos/Private/Chaos/Collision/ParticlePairMidPhase.cpp
--- a/Engine/Source/Runtime/Experimental/Chaos/Private/Chaos/Collision/ParticlePairMidPhase.cpp
+++ b/Engine/Source/Runtime/Experimental/Chaos/Private/Chaos/Collision/ParticlePairMidPhase.cpp
@@ -843,13 +843,18 @@ namespace Chaos
 		MultiShapePairDetectors.Reset();
 
 		Flags.bIsCCD = false;
-		Flags.bIsInitialized = false;
+		Flags.bIsCCDActive = false;
 		Flags.bIsSleeping = false;
+		Flags.bIsModified = false;
 	}
 
 	void FParticlePairMidPhase::ResetModifications()
 	{
-		Flags.bIsCCDActive = Flags.bIsCCD;
+		if (Flags.bIsModified)
+		{
+			Flags.bIsCCDActive = Flags.bIsCCD;
+			Flags.bIsModified = false;
+		}
 	}
 
 	void FParticlePairMidPhase::Init(
@@ -878,8 +883,6 @@ namespace Chaos
 		BuildDetectors();
 
 		InitThresholds();
-
-		Flags.bIsInitialized = true;
 	}
 
 	void FParticlePairMidPhase::BuildDetectors()
@@ -1026,6 +1029,10 @@ namespace Chaos
 			NumActiveConstraints += MultiShapePair.GenerateCollisions(CullDistance, Dt, Context);
 		}
 
+		// Reset any modifications applied by the MidPhaseModifier.
+		// In principle this belongs at the start of the frame but it is hard to avoid the cache miss there.
+		ResetModifications();
+
 		LastUsedEpoch = Context.GetAllocator()->GetCurrentEpoch();
 	}
 
diff --git a/Engine/Source/Runtime/Experimental/Chaos/Public/Chaos/Collision/ParticlePairMidPhase.h b/Engine/Source/Runtime/Experimental/Chaos/Public/Chaos/Collision/ParticlePairMidPhase.h
--- a/Engine/Source/Runtime/Experimental/Chaos/Public/Chaos/Collision/ParticlePairMidPhase.h
+++ b/Engine/Source/Runtime/Experimental/Chaos/Public/Chaos/Collision/ParticlePairMidPhase.h
@@ -439,10 +439,11 @@ namespace Chaos
 		}
 
 		/**
-		 *
+		 * Override the CCD condition for this mid-phase. Used by the MidPhase modifier and gets reset every frame.
 		 */
 		void SetCCDIsActive(bool bCCDIsActive)
 		{
+			Flags.bIsModified = true;
 			Flags.bIsCCDActive = bCCDIsActive;
 		}
 
@@ -480,9 +481,9 @@ namespace Chaos
 			struct
 			{
 				uint32 bIsCCD : 1;       // True if CCD is supported by either particle
-				uint32 bIsCCDActive : 1; // True if CCD is active for this midphase on this frame. Resets to bIsCCD each frame.
-				uint32 bIsInitialized : 1;
+				uint32 bIsCCDActive : 1; // True if CCD is active for this midphase on this frame. This can be changes by modifiers and resets to bIsCCD each frame.
 				uint32 bIsSleeping : 1;
+				uint32 bIsModified : 1;  // True if a modifier applied any changes to this midphase
 			};
 			uint32 Bits;
 		} Flags; // 4 bytes
diff --git a/Engine/Source/Runtime/Experimental/Chaos/Public/Chaos/Collision/SpatialAccelerationBroadPhase.h b/Engine/Source/Runtime/Experimental/Chaos/Public/Chaos/Collision/SpatialAccelerationBroadPhase.h
--- a/Engine/Source/Runtime/Experimental/Chaos/Public/Chaos/Collision/SpatialAccelerationBroadPhase.h
+++ b/Engine/Source/Runtime/Experimental/Chaos/Public/Chaos/Collision/SpatialAccelerationBroadPhase.h
@@ -593,9 +593,6 @@ namespace Chaos
 				// Get the midphase for this pair
 				FParticlePairMidPhase* MidPhase = ContextAllocator->GetMidPhase(Overlap.Particles[0], Overlap.Particles[1], Overlap.Particles[Overlap.SearchParticleIndex], BroadphaseContext.CollisionContext);
 				BroadphaseContext.MidPhases[OverlapIndex] = MidPhase;
-
-				// Reset any modifications that may have occurred to this midphase on the previous frame
-				MidPhase->ResetModifications();
 			}
 		}
 
