Iris - Support custom ReplicationFragments for structs with a custom NetSerializer.

diff --git a/Engine/Source/Runtime/Experimental/Iris/Core/Private/Iris/ReplicationState/ReplicationStateDescriptorBuilder.cpp b/Engine/Source/Runtime/Experimental/Iris/Core/Private/Iris/ReplicationState/ReplicationStateDescriptorBuilder.cpp
--- a/Engine/Source/Runtime/Experimental/Iris/Core/Private/Iris/ReplicationState/ReplicationStateDescriptorBuilder.cpp
+++ b/Engine/Source/Runtime/Experimental/Iris/Core/Private/Iris/ReplicationState/ReplicationStateDescriptorBuilder.cpp
@@ -1870,6 +1870,12 @@ bool FPropertyReplicationStateDescriptorBuilder::IsSupportedProperty(FMemberProp
 		OutMemberProperty.ChangeMaskBits = 1U + FIrisFastArraySerializer::IrisFastArrayChangeMaskBits;
 	}
 
+	// Structs with custom serializers may have a default custom replication fragment.
+	if (!OutMemberProperty.CreateAndRegisterReplicationFragmentFunction)
+	{
+		OutMemberProperty.CreateAndRegisterReplicationFragmentFunction = NetSerializerInfo->GetCreateAndRegisterReplicationFragmentFunction();
+	}
+
 	return true;
 }
 
@@ -1899,6 +1905,7 @@ bool FPropertyReplicationStateDescriptorBuilder::IsSupportedStructWithCustomSeri
 	OutMemberProperty.ExternalSizeAndAlignment = { (SIZE_T)InStruct->GetStructureSize(), (SIZE_T)InStruct->GetMinAlignment() };
 	OutMemberProperty.ReplicationCondition = ELifetimeCondition::COND_None;
 	OutMemberProperty.PropertyRepNotifyFunction = nullptr;
+	OutMemberProperty.CreateAndRegisterReplicationFragmentFunction = NetSerializerInfo->GetCreateAndRegisterReplicationFragmentFunction();
 	
 	// Init traits
 	OutMemberProperty.Traits = EMemberPropertyTraits::None;
diff --git a/Engine/Source/Runtime/Experimental/Iris/Core/Public/Iris/ReplicationState/PropertyNetSerializerInfoRegistry.h b/Engine/Source/Runtime/Experimental/Iris/Core/Public/Iris/ReplicationState/PropertyNetSerializerInfoRegistry.h
--- a/Engine/Source/Runtime/Experimental/Iris/Core/Public/Iris/ReplicationState/PropertyNetSerializerInfoRegistry.h
+++ b/Engine/Source/Runtime/Experimental/Iris/Core/Public/Iris/ReplicationState/PropertyNetSerializerInfoRegistry.h
@@ -13,7 +13,11 @@
 namespace UE::Net
 {
 
+class FFragmentRegistrationContext;
+class FReplicationFragment;
+struct FReplicationStateDescriptor;
 enum class EReplicationStateTraits : uint32;
+typedef FReplicationFragment* (*CreateAndRegisterReplicationFragmentFunc)(UObject*, const FReplicationStateDescriptor*, FFragmentRegistrationContext&);
 
 }
 
@@ -32,6 +36,8 @@ struct FPropertyNetSerializerInfo
 	virtual const FNetSerializer* GetNetSerializer(const FProperty* Property) const { return nullptr; }
 	virtual bool CanUseDefaultConfig(const FProperty* Property) const { return true; }
 	virtual const FNetSerializerConfig* BuildNetSerializerConfig(void* NetSerializerConfigBuffer, const FProperty* Property) const { return nullptr; }
+	/** Custom replication fragments are currently only supported by structs with a custom NetSerializer. See UE_NET_IMPLEMENT_NAMED_STRUCT_NETSERIALIZER_WITH_CUSTOM_FRAGMENT_INFO. */
+	virtual CreateAndRegisterReplicationFragmentFunc GetCreateAndRegisterReplicationFragmentFunction() const { return nullptr; }
 };
 
 /**
@@ -102,6 +108,19 @@ struct FNamedStructPropertyNetSerializerInfo : public TSimplePropertyNetSerializ
 		const FStructProperty* StructProp = CastFieldChecked<const FStructProperty>(Property);
 		return StructProp->Struct->GetFName() == PropertyFName;
 	};
+
+	virtual CreateAndRegisterReplicationFragmentFunc GetCreateAndRegisterReplicationFragmentFunction() const
+	{
+		return CreateAndRegisterReplicationFragmentFunction;
+	}
+
+	void SetCreateAndRegisterReplicationFragmentFunction(CreateAndRegisterReplicationFragmentFunc InCreateAndRegisterReplicationFragmentFunction)
+	{
+		CreateAndRegisterReplicationFragmentFunction = InCreateAndRegisterReplicationFragmentFunction;
+	}
+
+private:
+	CreateAndRegisterReplicationFragmentFunc CreateAndRegisterReplicationFragmentFunction = nullptr;
 };
 
 struct FLastResortPropertyNetSerializerInfo : public FPropertyNetSerializerInfo
@@ -148,6 +167,14 @@ const UE::Net::FPropertyNetSerializerInfo& GetPropertyNetSerializerInfo_##Proper
 #define UE_NET_IMPLEMENT_NAMED_STRUCT_NETSERIALIZER_INFO(Name, SerializerName) \
 const UE::Net::FPropertyNetSerializerInfo& GetPropertyNetSerializerInfo_##Name() { static UE::Net::FNamedStructPropertyNetSerializerInfo StaticInstance(Name, UE_NET_GET_SERIALIZER(SerializerName)); return StaticInstance; };
 
+#define UE_NET_IMPLEMENT_NAMED_STRUCT_NETSERIALIZER_WITH_CUSTOM_FRAGMENT_INFO(Name, SerializerName, CreateAndRegisterReplicationFragmentFunction) \
+const UE::Net::FPropertyNetSerializerInfo& GetPropertyNetSerializerInfo_##Name() \
+{ \
+	static UE::Net::FNamedStructPropertyNetSerializerInfo StaticInstance(Name, UE_NET_GET_SERIALIZER(SerializerName)); \
+	StaticInstance.SetCreateAndRegisterReplicationFragmentFunction(CreateAndRegisterReplicationFragmentFunction); \
+	return StaticInstance; \
+};
+
 // Implement FPropertyNetSerializerInfo for cases where LastResortPropertyNetSerializer is needed for struct with custom serialization.
 #define UE_NET_IMPLEMENT_NAMED_STRUCT_LASTRESORT_NETSERIALIZER_INFO(StructName) \
 const UE::Net::FPropertyNetSerializerInfo& GetPropertyNetSerializerInfo_##StructName() { static UE::Net::FNamedStructLastResortPropertyNetSerializerInfo StaticInstance(StructName); return StaticInstance; };
