fix where the initial game thread particle positions for geometry collections will not match the initial physics thread particle positions

this fixes an issue where certain geometry collections will see a initial "snap" when it is being constrained by a joint constraint when both are placed in the level since the joint constraint will compute reference frames using the initial game thread particle position

diff --git a/Engine/Source/Runtime/Experimental/Chaos/Private/PhysicsProxy/GeometryCollectionPhysicsProxy.cpp b/Engine/Source/Runtime/Experimental/Chaos/Private/PhysicsProxy/GeometryCollectionPhysicsProxy.cpp
--- a/Engine/Source/Runtime/Experimental/Chaos/Private/PhysicsProxy/GeometryCollectionPhysicsProxy.cpp
+++ b/Engine/Source/Runtime/Experimental/Chaos/Private/PhysicsProxy/GeometryCollectionPhysicsProxy.cpp
@@ -483,7 +483,12 @@ void FGeometryCollectionPhysicsProxy::Initialize(Chaos::FPBDRigidsEvolutionBase
 
 			GTParticles[Index]->SetUniqueIdx(Evolution->GenerateUniqueIdx());
 
-			const FTransform& T = Parameters.WorldTransform * GameThreadCollection.Transform[Index];
+			// Note that this transform must match the physics thread transform computation for initialization.
+			// Take for example, a geometry collection in the editor that is linked with a joint constraint. The joint
+			// constraint will query the game thread particle position/rotation for the geometry collection to compute its
+			// reference frame. If that position/rotation does not match up with the physics thread's position/rotation,
+			// the geometry collection particle will have an added velocity computed by the joint constraint solver.
+			const FTransform& T = GameThreadCollection.MassToLocal[Index] * GameThreadCollection.Transform[Index] * Parameters.WorldTransform;
 			P->SetX(T.GetTranslation(), false);
 			P->SetR(T.GetRotation(), false);
 			P->SetM(Mass[Index] * MassScale);
