add an option to output an axis-aligned overlap heuristic in FReadPhysicsObjectInterface<Id>::GetPhysicsObjectOverlap

diff --git a/Engine/Source/Runtime/Experimental/Chaos/Private/Chaos/PhysicsObjectInterface.cpp b/Engine/Source/Runtime/Experimental/Chaos/Private/Chaos/PhysicsObjectInterface.cpp
--- a/Engine/Source/Runtime/Experimental/Chaos/Private/Chaos/PhysicsObjectInterface.cpp
+++ b/Engine/Source/Runtime/Experimental/Chaos/Private/Chaos/PhysicsObjectInterface.cpp
@@ -199,7 +199,7 @@ namespace Chaos
 	}
 
 	template<EThreadContext Id>
-	bool FReadPhysicsObjectInterface<Id>::GetPhysicsObjectOverlap(FPhysicsObjectHandle ObjectA, FPhysicsObjectHandle ObjectB, bool bTraceComplex, Chaos::FMTDInfo* OutMTD)
+	bool FReadPhysicsObjectInterface<Id>::GetPhysicsObjectOverlap(FPhysicsObjectHandle ObjectA, FPhysicsObjectHandle ObjectB, bool bTraceComplex, Chaos::FOverlapInfo& OutOverlap)
 	{
 		TRACE_CPUPROFILER_EVENT_SCOPE(FReadPhysicsObjectInterface<Id>::GetPhysicsObjectOverlap);
 		TArray<FPerShapeData*> ShapesA = GetAllShapes({ &ObjectA, 1 });
@@ -215,11 +215,17 @@ namespace Chaos
 			return false;
 		}
 
-		if (OutMTD)
+		if (OutOverlap.MTD)
 		{
-			OutMTD->Penetration = 0.0;
+			OutOverlap.MTD->Penetration = 0.0;
 		}
 
+		if (OutOverlap.AxisOverlap)
+		{
+			*OutOverlap.AxisOverlap = FBox{ EForceInit::ForceInitToZero };
+		}
+		const bool bComputeMTD = OutOverlap.MTD != nullptr;
+
 		bool bFoundOverlap = false;
 		for (FPerShapeData* B : ShapesB)
 		{
@@ -264,25 +270,31 @@ namespace Chaos
 				const bool bOverlap = Chaos::Utilities::CastHelper(
 					*GeomB,
 					TransformB,
-					[A, &TransformA, OutMTD, &TmpMTDInfo](const auto& Downcast, const auto& FullTransformB)
+					[A, &TransformA, bComputeMTD, &TmpMTDInfo](const auto& Downcast, const auto& FullTransformB)
 					{
-						return Chaos::OverlapQuery(*A->GetGeometry(), TransformA, Downcast, FullTransformB, 0, OutMTD ? &TmpMTDInfo : nullptr);
+						return Chaos::OverlapQuery(*A->GetGeometry(), TransformA, Downcast, FullTransformB, 0, bComputeMTD ? &TmpMTDInfo : nullptr);
 					}
 				);
 
 				if (bOverlap)
 				{
 					bFoundOverlap = true;
-					if (!OutMTD)
+					if (!OutOverlap.MTD && !OutOverlap.AxisOverlap)
 					{
-						// Don't care about the MTD so as soon as we find an overlap we can return.
+						// Don't care about computing extra overlap information so we can exit early.
 						return true;
 					}
-					else if (TmpMTDInfo.Penetration > OutMTD->Penetration)
+
+					if (OutOverlap.MTD && TmpMTDInfo.Penetration > OutOverlap.MTD->Penetration)
 					{
-						// TODO: I don't think this math is actually correct to find the total overlap between the two objects composed of multiple shapes.
 						// If we need to find the MTD we need to find the largest overlap.
-						*OutMTD = TmpMTDInfo;
+						*OutOverlap.MTD = TmpMTDInfo;
+					}
+
+					if (OutOverlap.AxisOverlap)
+					{
+						const FAABB3 Intersection = BoxShapeA.GetIntersection(BoxShapeB);
+						*OutOverlap.AxisOverlap += FBox{ Intersection.Min(), Intersection.Max() };
 					}
 				}
 			}
diff --git a/Engine/Source/Runtime/Experimental/Chaos/Public/Chaos/PhysicsObjectInterface.h b/Engine/Source/Runtime/Experimental/Chaos/Public/Chaos/PhysicsObjectInterface.h
--- a/Engine/Source/Runtime/Experimental/Chaos/Public/Chaos/PhysicsObjectInterface.h
+++ b/Engine/Source/Runtime/Experimental/Chaos/Public/Chaos/PhysicsObjectInterface.h
@@ -41,6 +41,12 @@ namespace Chaos
 	class FPerShapeData;
 	struct FMTDInfo;
 
+	struct CHAOS_API FOverlapInfo
+	{
+		FMTDInfo* MTD = nullptr;
+		FBox* AxisOverlap = nullptr;
+	};
+
 	/**
 	 * FReadPhysicsObjectInterface will assume that these operations are safe to call (i.e. the relevant scenes have been read locked on the game thread).
 	 */
@@ -59,7 +65,7 @@ namespace Chaos
 		FSpatialAccelerationIdx GetSpatialIndex(FPhysicsObjectHandle Object);
 
 		TArray<FPerShapeData*> GetAllShapes(TArrayView<FPhysicsObjectHandle> InObjects);
-		bool GetPhysicsObjectOverlap(FPhysicsObjectHandle ObjectA, FPhysicsObjectHandle ObjectB, bool bTraceComplex, Chaos::FMTDInfo* OutMTD);
+		bool GetPhysicsObjectOverlap(FPhysicsObjectHandle ObjectA, FPhysicsObjectHandle ObjectB, bool bTraceComplex, Chaos::FOverlapInfo& OutOverlap);
 
 		bool AreAllValid(TArrayView<FPhysicsObjectHandle> InObjects);
 		bool AreAllKinematic(TArrayView<FPhysicsObjectHandle> InObjects);
