Chaos - minor broadphase optimization and cleanup
- move overlap filtering from AABB visitor to AssignMidPhases
- change a template parameter to a variable (no perf change)

diff --git a/Engine/Source/Runtime/Experimental/Chaos/Public/Chaos/Collision/SpatialAccelerationBroadPhase.h b/Engine/Source/Runtime/Experimental/Chaos/Public/Chaos/Collision/SpatialAccelerationBroadPhase.h
--- a/Engine/Source/Runtime/Experimental/Chaos/Public/Chaos/Collision/SpatialAccelerationBroadPhase.h
+++ b/Engine/Source/Runtime/Experimental/Chaos/Public/Chaos/Collision/SpatialAccelerationBroadPhase.h
@@ -23,53 +23,6 @@ namespace Chaos
 
 	namespace Private
 	{
-		/**
-		 * Output element from the broadphase
-		*/
-		class FBroadPhaseOverlap
-		{
-		public:
-			FBroadPhaseOverlap()
-			{
-			}
-
-			FBroadPhaseOverlap(FGeometryParticleHandle* Particle0, FGeometryParticleHandle* Particle1, const int32 InSearchParticleIndex)
-				: Particles{ Particle0, Particle1 }
-				, SearchParticleIndex(InSearchParticleIndex)
-			{
-			}
-
-			void SwapOrder()
-			{
-				Swap(Particles[0], Particles[1]);
-			}
-
-			FGeometryParticleHandle* Particles[2];
-			int32 SearchParticleIndex;
-		};
-
-		/**
-		 * Per-thread context for the broadphase.
-		*/
-		class FBroadPhaseContext
-		{
-		public:
-			FBroadPhaseContext()
-			{
-			}
-
-			void Reset()
-			{
-				Overlaps.Reset();
-				MidPhases.Reset();
-				CollisionContext.Reset();
-			}
-
-			TArray<FBroadPhaseOverlap> Overlaps;
-			TArray<FParticlePairMidPhase*> MidPhases;
-			FCollisionContext CollisionContext;
-		};
-
 		/**
 		 * Check whether the two particles are allowed to collide, and also whether we should flip their order.
 		 */
@@ -224,22 +177,80 @@ namespace Chaos
 			// Since particles can change type (between Kinematic and Dynamic) we may visit them in different orders at different times, but if we allow
 			// that it would break Resim and constraint reuse. Also, if only one particle is Dynamic, we want it in first position. This isn't a strtct 
 			// requirement but some downstream systems assume this is true (e.g., CCD, TriMesh collision).
-			bOutSwapOrder = ShouldSwapParticleOrder((bIsDynamicAwake1 || bIsDynamicAsleep1), (bIsDynamicAwake2 || bIsDynamicAsleep2), bIsParticle1Preferred);
+			if (bAcceptParticlePair)
+			{
+				bOutSwapOrder = ShouldSwapParticleOrder((bIsDynamicAwake1 || bIsDynamicAsleep1), (bIsDynamicAwake2 || bIsDynamicAsleep2), bIsParticle1Preferred);
+			}
 
 			return bAcceptParticlePair;
 		}
 
+		/**
+		 * Output element from the broadphase
+		*/
+		class FBroadPhaseOverlap
+		{
+		public:
+			FBroadPhaseOverlap()
+			{
+			}
+
+			FBroadPhaseOverlap(FGeometryParticleHandle* Particle0, FGeometryParticleHandle* Particle1, const int32 InSearchParticleIndex)
+				: Particles{ Particle0, Particle1 }
+				, SearchParticleIndex(InSearchParticleIndex)
+				, bCollisionsEnabled(true)
+			{
+			}
+
+			void ApplyFilter(const FIgnoreCollisionManager& IgnoreCollisionManager, const bool bIsResimming)
+			{
+				bool bSwapOrder = false;
+
+				bCollisionsEnabled = Private::ParticlePairCollisionAllowed(Particles[0], Particles[1], IgnoreCollisionManager, bIsResimming, bSwapOrder);
+
+				if (bSwapOrder)
+				{
+					Swap(Particles[0], Particles[1]);
+					SearchParticleIndex = 1 - SearchParticleIndex;
+				}
+			}
+
+			FGeometryParticleHandle* Particles[2];
+			int32 SearchParticleIndex;
+			bool bCollisionsEnabled;
+		};
+
+		/**
+		 * Per-thread context for the broadphase.
+		*/
+		class FBroadPhaseContext
+		{
+		public:
+			FBroadPhaseContext()
+			{
+			}
+
+			void Reset()
+			{
+				Overlaps.Reset();
+				MidPhases.Reset();
+				CollisionContext.Reset();
+			}
+
+			TArray<FBroadPhaseOverlap> Overlaps;
+			TArray<FParticlePairMidPhase*> MidPhases;
+			FCollisionContext CollisionContext;
+		};
+
 		/**
 		 * A visitor for the spatial partitioning system used to build the set of objects overlapping a bounding box.
 		 */
 		struct FSimOverlapVisitor
 		{
-			FSimOverlapVisitor(FGeometryParticleHandle* ParticleHandle, const FCollisionFilterData& InSimFilterData, const FIgnoreCollisionManager& InIgnoreCollisionManager, const bool bInIsResimming, Private::FBroadPhaseContext& InContext)
+			FSimOverlapVisitor(FGeometryParticleHandle* ParticleHandle, const FCollisionFilterData& InSimFilterData, Private::FBroadPhaseContext& InContext)
 				: Context(InContext)
-				, IgnoreCollisionManager(InIgnoreCollisionManager)
 				, SimFilterData(InSimFilterData)
 				, ParticleUniqueIdx(ParticleHandle ? ParticleHandle->UniqueIdx() : FUniqueIdx(0))
-				, bIsResimming(bInIsResimming)
 				, AccelerationHandle(ParticleHandle)
 			{
 			}
@@ -249,19 +260,7 @@ namespace Chaos
 				FGeometryParticleHandle* Particle1 = AccelerationHandle.GetGeometryParticleHandle_PhysicsThread();
 				FGeometryParticleHandle* Particle2 = Instance.Payload.GetGeometryParticleHandle_PhysicsThread();
 			
-				bool bSwapOrder = false;
-				const bool bAccept = Private::ParticlePairCollisionAllowed(Particle1, Particle2, IgnoreCollisionManager, bIsResimming, bSwapOrder);
-				
-				if (bAccept)
-				{
-					int32 SearchParticleIndex = 0;
-					if (bSwapOrder)
-					{
-						Swap(Particle1, Particle2);
-						SearchParticleIndex = 1;
-					}
-					Context.Overlaps.Emplace(Particle1, Particle2, SearchParticleIndex);
-				}
+				Context.Overlaps.Emplace(Particle1, Particle2, 0);
 
 				return true;
 			}
@@ -299,10 +298,8 @@ namespace Chaos
 
 		private:
 			Private::FBroadPhaseContext& Context;
-			const FIgnoreCollisionManager& IgnoreCollisionManager;
 			FCollisionFilterData SimFilterData;
 			FUniqueIdx ParticleUniqueIdx; // unique id of the particle visiting, used to skip self intersection as early as possible
-			bool bIsResimming;
 
 			/** Handle to be stored to retrieve the payload on which we are querying the acceleration structure*/
 			FAccelerationStructureHandle AccelerationHandle;
@@ -321,6 +318,8 @@ namespace Chaos
 		FSpatialAccelerationBroadPhase(const FPBDRigidsSOAs& InParticles)
 			: Particles(InParticles)
 			, SpatialAcceleration(nullptr)
+			, NumActiveBroadphaseContexts(0)
+			, bNeedsResim(false)
 		{
 		}
 
@@ -336,8 +335,7 @@ namespace Chaos
 			FReal Dt, 
 			Private::FCollisionConstraintAllocator* Allocator,
 			const FCollisionDetectorSettings& Settings,
-			IResimCacheBase* ResimCache
-			)
+			IResimCacheBase* ResimCache)
 		{
 			SCOPE_CYCLE_COUNTER(STAT_Collisions_SpatialBroadPhase);
 
@@ -347,6 +345,8 @@ namespace Chaos
 				return;
 			}
 
+			bNeedsResim = ResimCache && ResimCache->IsResimming();
+
 			if (const auto AABBTree = SpatialAcceleration->template As<TAABBTree<FAccelerationStructureHandle, TAABBTreeLeafArray<FAccelerationStructureHandle>>>())
 			{
 				ProduceOverlaps(Dt, *AABBTree, Allocator, Settings, ResimCache);
@@ -404,7 +404,7 @@ namespace Chaos
 		 * @param InSpatialAcceleration Spatial acceleration (AABB, bounding volumes...) to be used for broadphase collision detection
 		 * @param NarrowPhase Narrowphase collision detection that will be executed on each potential pairs coming from the broadphase detection
 		 * */
-		template<bool bNeedsResim, bool bOnlyRigid, typename ViewType, typename SpatialAccelerationType>
+		template<bool bOnlyRigid, typename ViewType, typename SpatialAccelerationType>
 		void ComputeParticlesOverlaps(
 			ViewType& OverlapView, 
 			FReal Dt,
@@ -459,7 +459,7 @@ namespace Chaos
 				CSV_SCOPED_TIMING_STAT(PhysicsVerbose, DetectCollisions_AABBTree);
 				const auto& ProduceOverlapsWorker = [this, Dt, &InSpatialAcceleration](auto& Particle1, const int32 ActiveIdxIdx, const int32 ContextIndex)
 				{
-					ProduceParticleOverlaps<bNeedsResim, bOnlyRigid>(Dt, Particle1.Handle(), InSpatialAcceleration, BroadphaseContexts[ContextIndex]);
+					ProduceParticleOverlaps<bOnlyRigid>(Dt, Particle1.Handle(), InSpatialAcceleration, BroadphaseContexts[ContextIndex]);
 				};
 				OverlapView.ParallelFor(ContextCreator, ProduceOverlapsWorker, bDisableCollisionParallelFor);
 			}
@@ -479,9 +479,7 @@ namespace Chaos
 			// visiting pairs that cannot collide (e.g., kinemtic-kinematic, or kinematic-sleeping for
 			// stationary kinematics)
 			const bool bDisableParallelFor = bDisableCollisionParallelFor;
-			auto* EvolutionResimCache = static_cast<FEvolutionResimCache*>(ResimCache);
-			const bool bResimSkipCollision = ResimCache && ResimCache->IsResimming();
-			if(!bResimSkipCollision)
+			if(!bNeedsResim)
 			{
 				const TParticleView<TPBDRigidParticles<FReal, 3>>& DynamicSleepingView = Particles.GetNonDisabledDynamicView();
 				const TParticleView<TKinematicGeometryParticles<FReal, 3>>& DynamicMovingKinematicView = Particles.GetActiveDynamicMovingKinematicParticlesView();
@@ -491,18 +489,18 @@ namespace Chaos
 				// and sleeping-dynamics, even though this means we visit sleeping pairs.
 				if(DynamicSleepingView.Num() < DynamicMovingKinematicView.Num())
 				{
-					ComputeParticlesOverlaps<false, true>(DynamicSleepingView, Dt, InSpatialAcceleration, Allocator, Settings);
+					ComputeParticlesOverlaps<true>(DynamicSleepingView, Dt, InSpatialAcceleration, Allocator, Settings);
 				}
 				else
 				{
-					ComputeParticlesOverlaps<false, false>(DynamicMovingKinematicView, Dt, InSpatialAcceleration, Allocator, Settings);
+					ComputeParticlesOverlaps<false>(DynamicMovingKinematicView, Dt, InSpatialAcceleration, Allocator, Settings);
 				}
 			}
 			else
 			{
 				const TParticleView<TGeometryParticles<FReal, 3>>& DesyncedView = ResimCache->GetDesyncedView();
 				
-				ComputeParticlesOverlaps<true,false>(DesyncedView, Dt, InSpatialAcceleration, Allocator, Settings);
+				ComputeParticlesOverlaps<false>(DesyncedView, Dt, InSpatialAcceleration, Allocator, Settings);
 			}
 		}
 
@@ -511,7 +509,7 @@ namespace Chaos
 	private:
 
 		// Generate the set of particles that overlap the specified particle and are allowed to collide with it
-		template<bool bIsResimming, bool bOnlyRigid, typename T_SPATIALACCELERATION>
+		template<bool bOnlyRigid, typename T_SPATIALACCELERATION>
 		void ProduceParticleOverlaps(
 		    FReal Dt,
 		    FGeometryParticleHandle* Particle1,
@@ -541,7 +539,7 @@ namespace Chaos
 
 			// @todo(chaos): This should already be handled by selecting the appropriate particle view in the calling function. GeometryCollections?
 			const bool bHasValidState = bOnlyRigid ? (bIsDynamicAwake1 || bIsDynamicAsleep1) : (bIsDynamicAwake1 || bIsKinematic1);
-			if (!bHasValidState && !bIsResimming)
+			if (!bHasValidState && !bNeedsResim)
 			{
 				return;
 			}
@@ -558,7 +556,7 @@ namespace Chaos
 
 				{
 					PHYSICS_CSV_SCOPED_EXPENSIVE(PhysicsVerbose, DetectCollisions_BroadPhase);
-					Private::FSimOverlapVisitor OverlapVisitor(Particle1, ParticleSimData, IgnoreCollisionManager, bIsResimming, Context);
+					Private::FSimOverlapVisitor OverlapVisitor(Particle1, ParticleSimData, Context);
 					InSpatialAcceleration.Overlap(Box1, OverlapVisitor);
 				}
 			}
@@ -586,14 +584,25 @@ namespace Chaos
 
 			BroadphaseContext.MidPhases.SetNum(BroadphaseContext.Overlaps.Num());
 
+			int32 MidPhaseIndex = 0;
 			for (int32 OverlapIndex = 0, OverlapEnd = BroadphaseContext.Overlaps.Num(); OverlapIndex < OverlapEnd; ++OverlapIndex)
 			{
 				Private::FBroadPhaseOverlap& Overlap = BroadphaseContext.Overlaps[OverlapIndex];
 
-				// Get the midphase for this pair
-				FParticlePairMidPhase* MidPhase = ContextAllocator->GetMidPhase(Overlap.Particles[0], Overlap.Particles[1], Overlap.Particles[Overlap.SearchParticleIndex], BroadphaseContext.CollisionContext);
-				BroadphaseContext.MidPhases[OverlapIndex] = MidPhase;
+				// Check to see if the two particles are allowed to collide.
+				// NOTE: this may also swap the order of the particles.
+				Overlap.ApplyFilter(IgnoreCollisionManager, bNeedsResim);
+
+				if (Overlap.bCollisionsEnabled)
+				{
+					// Get the midphase for this pair
+					FParticlePairMidPhase* MidPhase = ContextAllocator->GetMidPhase(Overlap.Particles[0], Overlap.Particles[1], Overlap.Particles[Overlap.SearchParticleIndex], BroadphaseContext.CollisionContext);
+					BroadphaseContext.MidPhases[MidPhaseIndex] = MidPhase;
+					++MidPhaseIndex;
+				}
 			}
+
+			BroadphaseContext.MidPhases.SetNum(MidPhaseIndex);
 		}
 
 		// Process all the midphases: generate constraints and execute the narrowphase
@@ -625,8 +634,9 @@ namespace Chaos
 		const FPBDRigidsSOAs& Particles;
 		const FAccelerationStructure* SpatialAcceleration;
 		TArray<Private::FBroadPhaseContext> BroadphaseContexts;
-		int32 NumActiveBroadphaseContexts;
 		FIgnoreCollisionManager IgnoreCollisionManager;
+		int32 NumActiveBroadphaseContexts;
+		bool bNeedsResim;
 	};
 
 }
