Chaos - minor collision solver cleanup to make it easier to isolate the vectorized solver code

diff --git a/Engine/Source/Runtime/Experimental/Chaos/Private/Chaos/Collision/PBDCollisionSolver.cpp b/Engine/Source/Runtime/Experimental/Chaos/Private/Chaos/Collision/PBDCollisionSolver.cpp
--- a/Engine/Source/Runtime/Experimental/Chaos/Private/Chaos/Collision/PBDCollisionSolver.cpp
+++ b/Engine/Source/Runtime/Experimental/Chaos/Private/Chaos/Collision/PBDCollisionSolver.cpp
@@ -22,6 +22,7 @@ namespace Chaos
 	namespace CVars
 	{
 		extern int32 Chaos_Collision_UseShockPropagation;
+		extern bool bChaos_PBDCollisionSolver_UseJacobiPairSolver;
 
 		bool bChaos_PBDCollisionSolver_Position_SolveEnabled = true;
 		float Chaos_PBDCollisionSolver_Position_MinInvMassScale = 0.77f;
@@ -224,5 +225,65 @@ namespace Chaos
 			}
 		}
 
+		//////////////////////////////////////////////////////////////////////////////////////////////////
+		//////////////////////////////////////////////////////////////////////////////////////////////////
+		//////////////////////////////////////////////////////////////////////////////////////////////////
+		//////////////////////////////////////////////////////////////////////////////////////////////////
+
+		void FPBDCollisionSolverHelper::SolvePositionNoFriction(const TArrayView<FPBDCollisionSolver>& CollisionSolvers, const FSolverReal Dt, const FSolverReal MaxPushOut)
+		{
+			if (!CVars::bChaos_PBDCollisionSolver_UseJacobiPairSolver)
+			{
+				for (FPBDCollisionSolver& CollisionSolver : CollisionSolvers)
+				{
+					CollisionSolver.SolvePositionNoFriction(Dt, MaxPushOut);
+				}
+			}
+			else
+			{
+				for (FPBDCollisionSolver& CollisionSolver : CollisionSolvers)
+				{
+					CollisionSolver.SolvePositionNoFrictionJacobi(Dt, MaxPushOut);
+				}
+			}
+		}
+
+		void FPBDCollisionSolverHelper::SolvePositionWithFriction(const TArrayView<FPBDCollisionSolver>& CollisionSolvers, const FSolverReal Dt, const FSolverReal MaxPushOut)
+		{
+			if (!CVars::bChaos_PBDCollisionSolver_UseJacobiPairSolver)
+			{
+				for (FPBDCollisionSolver& CollisionSolver : CollisionSolvers)
+				{
+					CollisionSolver.SolvePositionWithFriction(Dt, MaxPushOut);
+				}
+			}
+			else
+			{
+				for (FPBDCollisionSolver& CollisionSolver : CollisionSolvers)
+				{
+					CollisionSolver.SolvePositionWithFrictionJacobi(Dt, MaxPushOut);
+				}
+			}
+		}
+
+		void FPBDCollisionSolverHelper::SolveVelocity(const TArrayView<FPBDCollisionSolver>& CollisionSolvers, const FSolverReal Dt, const bool bApplyDynamicFriction)
+		{
+			if (!CVars::bChaos_PBDCollisionSolver_UseJacobiPairSolver)
+			{
+				for (FPBDCollisionSolver& CollisionSolver : CollisionSolvers)
+				{
+					CollisionSolver.SolveVelocity(Dt, bApplyDynamicFriction);
+				}
+			}
+			else
+			{
+				for (FPBDCollisionSolver& CollisionSolver : CollisionSolvers)
+				{
+					CollisionSolver.SolveVelocityJacobi(Dt, bApplyDynamicFriction);
+				}
+			}
+		}
+
+
 	}	// namespace Private
 }	// namespace Chaos
\ No newline at end of file
diff --git a/Engine/Source/Runtime/Experimental/Chaos/Private/Chaos/Collision/SolverCollisionContainer.cpp b/Engine/Source/Runtime/Experimental/Chaos/Private/Chaos/Collision/SolverCollisionContainer.cpp
--- a/Engine/Source/Runtime/Experimental/Chaos/Private/Chaos/Collision/SolverCollisionContainer.cpp
+++ b/Engine/Source/Runtime/Experimental/Chaos/Private/Chaos/Collision/SolverCollisionContainer.cpp
@@ -514,8 +514,7 @@ namespace Chaos
 		}
 	}
 
-	// @todo(chaos): parallel version of SolvePosition
-	bool FPBDCollisionContainerSolver::SolvePositionImpl(const FReal Dt, const int32 It, const int32 NumIts, const int32 BeginIndex, const int32 EndIndex, const FPBDCollisionSolverSettings& SolverSettings)
+	bool FPBDCollisionContainerSolver::SolvePositionImpl(const FReal InDt, const int32 It, const int32 NumIts, const int32 BeginIndex, const int32 EndIndex, const FPBDCollisionSolverSettings& SolverSettings)
 	{
 		SCOPE_CYCLE_COUNTER(STAT_Collisions_Apply);
 		if (!bChaos_PBDCollisionSolver_Position_SolveEnabled)
@@ -523,13 +522,12 @@ namespace Chaos
 			return false;
 		}
 
-		UpdatePositionShockPropagation(Dt, It, NumIts, BeginIndex, EndIndex, SolverSettings);
-
-		// Only apply friction for the last few (tunable) iterations
-		const bool bApplyStaticFriction = (It >= (NumIts - SolverSettings.NumPositionFrictionIterations));
+		if (EndIndex <= BeginIndex)
+		{
+			return false;
+		}
 
-		// Adjust max pushout to attempt to make it iteration count independent
-		const FReal MaxPushOut = (SolverSettings.MaxPushOutVelocity > 0) ? (SolverSettings.MaxPushOutVelocity * Dt) / FReal(NumIts) : 0;
+		UpdatePositionShockPropagation(InDt, It, NumIts, BeginIndex, EndIndex, SolverSettings);
 
 		// We run collision detection here under two conditions (normally it is run after Integration and before the constraint solver phase):
 		// 1) When deferring collision detection until the solver phase for better joint-collision behaviour (RBAN). In this case, we only do this on the first iteration.
@@ -537,108 +535,49 @@ namespace Chaos
 		const bool bRunDeferredCollisionDetection = (It == 0) && ConstraintContainer.GetDetectorSettings().bDeferNarrowPhase;
 		if (bRunDeferredCollisionDetection || bPerIterationCollisionDetection)
 		{
-			UpdateCollisions(Dt, BeginIndex, EndIndex);
-		}
-
-		// Apply the position correction
-		if (bApplyStaticFriction)
-		{
-			return SolvePositionWithFrictionImpl(Dt, BeginIndex, EndIndex, MaxPushOut);
+			UpdateCollisions(InDt, BeginIndex, EndIndex);
 		}
-		else
-		{
-			return SolvePositionNoFrictionImpl(Dt, BeginIndex, EndIndex, MaxPushOut);
-		}
-	}
 
-	// Solve position with friction (last few iterations each tick)
-	bool FPBDCollisionContainerSolver::SolvePositionWithFrictionImpl(const FReal InDt, const int32 BeginIndex, const int32 EndIndex, const FReal InMaxPushOut)
-	{
-		if (EndIndex == BeginIndex)
-		{
-			return false;
-		}
+		// Only apply friction for the last few (tunable) iterations
+		// Adjust max pushout to attempt to make it iteration count independent
 		const FSolverReal Dt = FSolverReal(InDt);
-		const FSolverReal MaxPushOut = FSolverReal(InMaxPushOut);
+		const bool bApplyStaticFriction = (It >= (NumIts - SolverSettings.NumPositionFrictionIterations));
+		const FSolverReal MaxPushOut = (SolverSettings.MaxPushOutVelocity > 0) ? (FSolverReal(SolverSettings.MaxPushOutVelocity) * Dt) / FSolverReal(NumIts) : FSolverReal(0);
 
-		if (!CVars::bChaos_PBDCollisionSolver_UseJacobiPairSolver)
+		// Apply the position correction
+		if (bApplyStaticFriction)
 		{
-			for (int32 SolverIndex = BeginIndex; SolverIndex < EndIndex; ++SolverIndex)
-			{
-				CollisionSolvers[SolverIndex].SolvePositionWithFriction(Dt, MaxPushOut);
-			}
+			Private::FPBDCollisionSolverHelper::SolvePositionWithFriction(MakeArrayView(&CollisionSolvers[BeginIndex], EndIndex - BeginIndex), Dt, MaxPushOut);
 		}
 		else
 		{
-			for (int32 SolverIndex = BeginIndex; SolverIndex < EndIndex; ++SolverIndex)
-			{
-				CollisionSolvers[SolverIndex].SolvePositionWithFrictionJacobi(Dt, MaxPushOut);
-			}
+			Private::FPBDCollisionSolverHelper::SolvePositionNoFriction(MakeArrayView(&CollisionSolvers[BeginIndex], EndIndex - BeginIndex), Dt, MaxPushOut);
 		}
 
 		return true;
 	}
 
-	// Solve position without friction (first few iterations each tick)
-	bool FPBDCollisionContainerSolver::SolvePositionNoFrictionImpl(const FReal InDt, const int32 BeginIndex, const int32 EndIndex, const FReal InMaxPushOut)
+	bool FPBDCollisionContainerSolver::SolveVelocityImpl(const FReal InDt, const int32 It, const int32 NumIts, const int32 BeginIndex, const int32 EndIndex, const FPBDCollisionSolverSettings& SolverSettings)
 	{
-		if (EndIndex == BeginIndex)
+		SCOPE_CYCLE_COUNTER(STAT_Collisions_ApplyPushOut);
+		if (!bChaos_PBDCollisionSolver_Velocity_SolveEnabled)
 		{
 			return false;
 		}
-		const FSolverReal Dt = FSolverReal(InDt);
-		const FSolverReal MaxPushOut = FSolverReal(InMaxPushOut);
 
-		if (!CVars::bChaos_PBDCollisionSolver_UseJacobiPairSolver)
-		{
-			for (int32 SolverIndex = BeginIndex; SolverIndex < EndIndex; ++SolverIndex)
-			{
-				CollisionSolvers[SolverIndex].SolvePositionNoFriction(Dt, MaxPushOut);
-			}
-		}
-		else
-		{
-			for (int32 SolverIndex = BeginIndex; SolverIndex < EndIndex; ++SolverIndex)
-			{
-				CollisionSolvers[SolverIndex].SolvePositionNoFrictionJacobi(Dt, MaxPushOut);
-			}
-		}
-
-		return true;
-	}
-
-	bool FPBDCollisionContainerSolver::SolveVelocityImpl(const FReal InDt, const int32 It, const int32 NumIts, const int32 BeginIndex, const int32 EndIndex, const FPBDCollisionSolverSettings& SolverSettings)
-	{
-		SCOPE_CYCLE_COUNTER(STAT_Collisions_ApplyPushOut);
-		if (!bChaos_PBDCollisionSolver_Velocity_SolveEnabled)
+		if (EndIndex <= BeginIndex)
 		{
 			return false;
 		}
-		const FSolverReal Dt = FSolverReal(InDt);
 
-		UpdateVelocityShockPropagation(Dt, It, NumIts, BeginIndex, EndIndex, SolverSettings);
+		UpdateVelocityShockPropagation(InDt, It, NumIts, BeginIndex, EndIndex, SolverSettings);
 
+		const FSolverReal Dt = FSolverReal(InDt);
 		const bool bApplyDynamicFriction = (It >= NumIts - SolverSettings.NumVelocityFrictionIterations);
 
-		// Apply the velocity correction
-		// @todo(chaos): parallel version of SolveVelocity
-		bool bNeedsAnotherIteration = false;
-		if (!CVars::bChaos_PBDCollisionSolver_UseJacobiPairSolver)
-		{
-			for (int32 SolverIndex = BeginIndex; SolverIndex < EndIndex; ++SolverIndex)
-			{
-				bNeedsAnotherIteration |= CollisionSolvers[SolverIndex].SolveVelocity(Dt, bApplyDynamicFriction);
-			}
-		}
-		else
-		{
-			for (int32 SolverIndex = BeginIndex; SolverIndex < EndIndex; ++SolverIndex)
-			{
-				bNeedsAnotherIteration |= CollisionSolvers[SolverIndex].SolveVelocityJacobi(Dt, bApplyDynamicFriction);
-			}
-		}
+		Private::FPBDCollisionSolverHelper::SolveVelocity(MakeArrayView(&CollisionSolvers[BeginIndex], EndIndex - BeginIndex), Dt, bApplyDynamicFriction);
 
-		return bNeedsAnotherIteration;
+		return true;
 	}
 
 	void FPBDCollisionContainerSolver::UpdateCollisions(const FReal InDt, const int32 BeginIndex, const int32 EndIndex)
diff --git a/Engine/Source/Runtime/Experimental/Chaos/Public/Chaos/Collision/PBDCollisionSolver.h b/Engine/Source/Runtime/Experimental/Chaos/Public/Chaos/Collision/PBDCollisionSolver.h
--- a/Engine/Source/Runtime/Experimental/Chaos/Public/Chaos/Collision/PBDCollisionSolver.h
+++ b/Engine/Source/Runtime/Experimental/Chaos/Public/Chaos/Collision/PBDCollisionSolver.h
@@ -100,17 +100,9 @@ namespace Chaos
 			*/
 			bool ShouldSolveVelocity() const;
 
+			// World-space contact data
 			FWorldContactPoint WorldContact;
 
-			// World-space body-relative contact points
-			//FSolverVec3 RelativeContactPosition0;
-			//FSolverVec3 RelativeContactPosition1;
-
-			// World-space contact axes (normal and 2 tangents)
-			//FSolverVec3 WorldContactNormal;
-			//FSolverVec3 WorldContactTangentU;
-			//FSolverVec3 WorldContactTangentV;
-
 			// I^-1.(R x A) for each body where A is each axis (Normal, TangentU, TangentV)
 			FSolverVec3 WorldContactNormalAngular0;
 			FSolverVec3 WorldContactTangentUAngular0;
@@ -124,14 +116,6 @@ namespace Chaos
 			FSolverReal ContactMassTangentU;
 			FSolverReal ContactMassTangentV;
 
-			// Initial world-space contact separation that we are trying to correct
-			//FSolverReal WorldContactDeltaNormal;
-			//FSolverReal WorldContactDeltaTangentU;
-			//FSolverReal WorldContactDeltaTangentV;
-
-			// Desired final normal velocity, taking Restitution into account
-			//FSolverReal WorldContactVelocityTargetNormal;
-
 			// Solver outputs
 			FSolverReal NetPushOutNormal;
 			FSolverReal NetPushOutTangentU;
@@ -144,10 +128,6 @@ namespace Chaos
 			// Equal to (NormalPushOut / TangentialPushOut) before clamping to the friction cone.
 			// Used to move the static friction anchors to the edge of the cone in Scatter.
 			FSolverReal StaticFrictionRatio;
-
-			// Whether this point is enabled
-			// @todo(chaos): consider decoupling the solver manifold points array size from the constraint manifold points array
-			//uint8 bEnabled : 1;
 		};
 
 		/**
@@ -1360,6 +1340,18 @@ namespace Chaos
 			// @todo(chaos): support early-out in velocity solve if necessary
 			return true;
 		}
-	}
-}
+
+		/**
+		 * A helper for solving arrays of constraints
+		 */
+		class FPBDCollisionSolverHelper
+		{
+		public:
+			static void SolvePositionNoFriction(const TArrayView<FPBDCollisionSolver>& CollisionSolvers, const FSolverReal Dt, const FSolverReal MaxPushOut);
+			static void SolvePositionWithFriction(const TArrayView<FPBDCollisionSolver>& CollisionSolvers, const FSolverReal Dt, const FSolverReal MaxPushOut);
+			static void SolveVelocity(const TArrayView<FPBDCollisionSolver>& CollisionSolvers, const FSolverReal Dt, const bool bApplyDynamicFriction);
+		};
+
+	}	// namespace Private
+}	// namespace Chaos
 
diff --git a/Engine/Source/Runtime/Experimental/Chaos/Public/Chaos/Collision/SolverCollisionContainer.h b/Engine/Source/Runtime/Experimental/Chaos/Public/Chaos/Collision/SolverCollisionContainer.h
--- a/Engine/Source/Runtime/Experimental/Chaos/Public/Chaos/Collision/SolverCollisionContainer.h
+++ b/Engine/Source/Runtime/Experimental/Chaos/Public/Chaos/Collision/SolverCollisionContainer.h
@@ -86,8 +86,6 @@ namespace Chaos
 		void UpdatePositionShockPropagation(const FReal Dt, const int32 It, const int32 NumIts, const int32 BeginIndex, const int32 EndIndex, const FPBDCollisionSolverSettings& SolverSettings);
 		void UpdateVelocityShockPropagation(const FReal Dt, const int32 It, const int32 NumIts, const int32 BeginIndex, const int32 EndIndex, const FPBDCollisionSolverSettings& SolverSettings);
 		bool SolvePositionImpl(const FReal Dt, const int32 It, const int32 NumIts, const int32 BeginIndex, const int32 EndIndex, const FPBDCollisionSolverSettings& SolverSettings);
-		bool SolvePositionWithFrictionImpl(const FReal Dt, const int32 BeginIndex, const int32 EndIndex, const FReal MaxPushOut);
-		bool SolvePositionNoFrictionImpl(const FReal Dt, const int32 BeginIndex, const int32 EndIndex, const FReal MaxPushOut);
 		bool SolveVelocityImpl(const FReal Dt, const int32 It, const int32 NumIts, const int32 BeginIndex, const int32 EndIndex, const FPBDCollisionSolverSettings& SolverSettings);
 		void UpdateCollisions(const FReal InDt, const int32 BeginIndex, const int32 EndIndex);
 
