implement CopyWithScale for instanced implicit objects and scaled implicit objects

diff --git a/Engine/Source/Runtime/Experimental/Chaos/Public/Chaos/ImplicitObjectScaled.h b/Engine/Source/Runtime/Experimental/Chaos/Public/Chaos/ImplicitObjectScaled.h
--- a/Engine/Source/Runtime/Experimental/Chaos/Public/Chaos/ImplicitObjectScaled.h
+++ b/Engine/Source/Runtime/Experimental/Chaos/Public/Chaos/ImplicitObjectScaled.h
@@ -203,6 +203,8 @@ public:
 		return TUniquePtr<FImplicitObject>(CopyHelper(this));
 	}
 
+	virtual TUniquePtr<FImplicitObject> CopyWithScale(const FVec3& Scale) const override;
+
 	virtual FString ToString() const override
 	{
 		return FString::Printf(TEXT("Instanced %s, Margin %f"), *MObject->ToString(), GetMargin());
@@ -1007,6 +1009,13 @@ public:
 		return TUniquePtr<FImplicitObject>(CopyHelper(this));
 	}
 
+	virtual TUniquePtr<FImplicitObject> CopyWithScale(const FVec3& Scale) const override
+	{
+		TImplicitObjectScaled<TConcrete, bInstanced>* Obj = CopyHelper(this);
+		Obj->SetScale(Scale);
+		return TUniquePtr<FImplicitObject>(Obj);
+	}
+
 	virtual FString ToString() const override
 	{
 		return FString::Printf(TEXT("Scaled %s, Scale: [%f, %f, %f], Margin: %f"), *MObject->ToString(), MScale.X, MScale.Y, MScale.Z, GetMargin());
@@ -1073,6 +1082,11 @@ using TImplicitObjectScaledNonSerializable = TImplicitObjectScaled<TConcrete, fa
 template <typename T, int d>
 using TImplicitObjectScaledGeneric = TImplicitObjectScaled<FImplicitObject>;
 
+template<typename TConcrete>
+TUniquePtr<FImplicitObject> TImplicitObjectInstanced<TConcrete>::CopyWithScale(const FVec3& Scale) const
+{
+	return TUniquePtr<FImplicitObject>(new TImplicitObjectScaled<TConcrete, true>(MObject, Scale, 0.0));
+}
 
 /**
  * @brief Remove the Instanced or Scaled wrapper from an ImplicitObject of a known inner type and extract the instance properties
