Iris
* Added assert to tell if the netdriver is missing the required DataStreamChannel definition in it's config

diff --git a/Engine/Source/Runtime/Engine/Private/Net/Iris/ReplicationSystem/ActorReplicationBridge.cpp b/Engine/Source/Runtime/Engine/Private/Net/Iris/ReplicationSystem/ActorReplicationBridge.cpp
--- a/Engine/Source/Runtime/Engine/Private/Net/Iris/ReplicationSystem/ActorReplicationBridge.cpp
+++ b/Engine/Source/Runtime/Engine/Private/Net/Iris/ReplicationSystem/ActorReplicationBridge.cpp
@@ -12,6 +12,7 @@
 #include "Iris/Core/IrisProfiler.h"
 #include "Iris/Core/NetObjectReference.h"
 #include "Iris/ReplicationSystem/NetToken.h"
+#include "Iris/ReplicationSystem/ObjectReplicationBridgeConfig.h"
 #include "Iris/ReplicationSystem/StringTokenStore.h"
 #include "Iris/ReplicationSystem/ReplicationOperations.h"
 #include "Iris/ReplicationSystem/ReplicationSystem.h"
@@ -847,6 +848,18 @@ void UActorReplicationBridge::SetNetDriver(UNetDriver* const InNetDriver)
 	if (InNetDriver != nullptr)
 	{
 		MaxPollFrequency = static_cast<float>(FPlatformMath::Max(InNetDriver->NetServerMaxTickRate, 0));
+
+		const FName RequiredChannelName = UObjectReplicationBridgeConfig::GetConfig()->GetRequiredNetDriverChannelClassName();
+		
+		if (!RequiredChannelName.IsNone())
+		{
+			const bool bRequiredChannelIsConfigured = InNetDriver->ChannelDefinitions.FindByPredicate([RequiredChannelName](const FChannelDefinition& rhs)
+				{
+					return rhs.ClassName == RequiredChannelName;
+				}) != nullptr;
+
+			checkf(bRequiredChannelIsConfigured, TEXT("ObjectReplication needs the netdriver channel %s to work. Add this channel to the netdriver channel definitions config"), *RequiredChannelName.ToString());
+		}
 	}
 }
 
diff --git a/Engine/Source/Runtime/Experimental/Iris/Core/Private/Iris/ReplicationSystem/ObjectReplicationBridge.cpp b/Engine/Source/Runtime/Experimental/Iris/Core/Private/Iris/ReplicationSystem/ObjectReplicationBridge.cpp
--- a/Engine/Source/Runtime/Experimental/Iris/Core/Private/Iris/ReplicationSystem/ObjectReplicationBridge.cpp
+++ b/Engine/Source/Runtime/Experimental/Iris/Core/Private/Iris/ReplicationSystem/ObjectReplicationBridge.cpp
@@ -1292,7 +1292,7 @@ void UObjectReplicationBridge::LoadConfig()
 	// Reset PathNameCache
 	ConfigClassPathNameCache.Empty();
 
-	const UObjectReplicationBridgeConfig* BridgeConfig = GetDefault<UObjectReplicationBridgeConfig>();
+	const UObjectReplicationBridgeConfig* BridgeConfig = UObjectReplicationBridgeConfig::GetConfig();
 
 	// Load poll configs
 	
diff --git a/Engine/Source/Runtime/Experimental/Iris/Core/Private/Iris/ReplicationSystem/ObjectReplicationBridgeConfig.cpp b/Engine/Source/Runtime/Experimental/Iris/Core/Private/Iris/ReplicationSystem/ObjectReplicationBridgeConfig.cpp
--- a/Engine/Source/Runtime/Experimental/Iris/Core/Private/Iris/ReplicationSystem/ObjectReplicationBridgeConfig.cpp
+++ b/Engine/Source/Runtime/Experimental/Iris/Core/Private/Iris/ReplicationSystem/ObjectReplicationBridgeConfig.cpp
@@ -8,6 +8,11 @@ UObjectReplicationBridgeConfig::UObjectReplicationBridgeConfig()
 {
 }
 
+const UObjectReplicationBridgeConfig* UObjectReplicationBridgeConfig::GetConfig()
+{
+	return GetDefault<UObjectReplicationBridgeConfig>();
+}
+
 TConstArrayView<FObjectReplicationBridgePollConfig> UObjectReplicationBridgeConfig::GetPollConfigs() const
 {
 	return MakeArrayView(PollConfigs);
diff --git a/Engine/Source/Runtime/Experimental/Iris/Core/Public/Iris/ReplicationSystem/ObjectReplicationBridgeConfig.h b/Engine/Source/Runtime/Experimental/Iris/Core/Public/Iris/ReplicationSystem/ObjectReplicationBridgeConfig.h
--- a/Engine/Source/Runtime/Experimental/Iris/Core/Public/Iris/ReplicationSystem/ObjectReplicationBridgeConfig.h
+++ b/Engine/Source/Runtime/Experimental/Iris/Core/Public/Iris/ReplicationSystem/ObjectReplicationBridgeConfig.h
@@ -80,11 +80,15 @@ class UObjectReplicationBridgeConfig : public UObject
 	GENERATED_BODY()
 
 public:
+
+	IRISCORE_API static const UObjectReplicationBridgeConfig* GetConfig();
+
 	IRISCORE_API TConstArrayView<FObjectReplicationBridgePollConfig> GetPollConfigs() const;
 	IRISCORE_API TConstArrayView<FObjectReplicationBridgeFilterConfig> GetFilterConfigs() const;
 	IRISCORE_API TConstArrayView<FObjectReplicationBridgePrioritizerConfig> GetPrioritizerConfigs() const;
 	IRISCORE_API TConstArrayView<FObjectReplicationBridgeDeltaCompressionConfig> GetDeltaCompressionConfigs() const;
 	FName GetDefaultSpatialFilterName() const;
+	FName GetRequiredNetDriverChannelClassName() const;
 
 protected:
 	UObjectReplicationBridgeConfig();
@@ -123,9 +127,20 @@ private:
 	 */
 	UPROPERTY(Config)
 	FName DefaultSpatialFilterName;
+
+	/**
+	 * The name of the channel class required for object replication to work.
+	 */
+	UPROPERTY(Config)
+	FName RequiredNetDriverChannelClassName;
 };
 
 inline FName UObjectReplicationBridgeConfig::GetDefaultSpatialFilterName() const
 {
 	return DefaultSpatialFilterName;
 }
+
+inline FName UObjectReplicationBridgeConfig::GetRequiredNetDriverChannelClassName() const
+{
+	return RequiredNetDriverChannelClassName;
+}
