Chaos - ISPC for collision solver (WIP disabled)

Currently this is just the cpp version converted as-is to ISPC. This is not useful yet as the math is almost entirely scalar and only operates on one constraint at a time.

diff --git a/Engine/Source/Runtime/Experimental/Chaos/Private/Chaos/Collision/PBDCollisionSolver.cpp b/Engine/Source/Runtime/Experimental/Chaos/Private/Chaos/Collision/PBDCollisionSolver.cpp
--- a/Engine/Source/Runtime/Experimental/Chaos/Private/Chaos/Collision/PBDCollisionSolver.cpp
+++ b/Engine/Source/Runtime/Experimental/Chaos/Private/Chaos/Collision/PBDCollisionSolver.cpp
@@ -13,6 +13,10 @@
 #include "Chaos/PBDCollisionConstraintsContact.h"
 #include "Chaos/Utilities.h"
 
+#if INTEL_ISPC
+#include "PBDCollisionSolver.ispc.generated.h"
+#endif
+
 //PRAGMA_DISABLE_OPTIMIZATION
 
 DEFINE_LOG_CATEGORY(LogChaosCollision);
@@ -24,6 +28,10 @@ namespace Chaos
 		extern int32 Chaos_Collision_UseShockPropagation;
 		extern bool bChaos_PBDCollisionSolver_UseJacobiPairSolver;
 
+		// NOTE: WORK IN PROGRESS. Not ready to be enabled.
+		bool bChaos_PBDCollisionSolver_ISPC = false;
+		FAutoConsoleVariableRef CVarChaos_PBDCollisionSolver_ISPC(TEXT("p.Chaos.PBDCollisionSolver.ISPC"), bChaos_PBDCollisionSolver_ISPC, TEXT(""));
+
 		bool bChaos_PBDCollisionSolver_Position_SolveEnabled = true;
 		float Chaos_PBDCollisionSolver_Position_MinInvMassScale = 0.77f;
 		float Chaos_PBDCollisionSolver_Position_StaticFrictionStiffness = 1.0f;
@@ -232,6 +240,14 @@ namespace Chaos
 
 		void FPBDCollisionSolverHelper::SolvePositionNoFriction(const TArrayView<FPBDCollisionSolver>& CollisionSolvers, const FSolverReal Dt, const FSolverReal MaxPushOut)
 		{
+#if INTEL_ISPC
+			if (CVars::bChaos_PBDCollisionSolver_ISPC)
+			{
+				ispc::SolvePositionNoFriction((ispc::FPBDCollisionSolver*)CollisionSolvers.GetData(), CollisionSolvers.Num(), Dt, MaxPushOut);
+				return;
+			}
+#endif
+
 			if (!CVars::bChaos_PBDCollisionSolver_UseJacobiPairSolver)
 			{
 				for (FPBDCollisionSolver& CollisionSolver : CollisionSolvers)
@@ -284,6 +300,16 @@ namespace Chaos
 			}
 		}
 
+		void FPBDCollisionSolverHelper::CheckISPC()
+		{
+#if INTEL_ISPC
+			check(sizeof(ispc::FPBDCollisionSolver) == sizeof(Private::FPBDCollisionSolver));
+			check(sizeof(ispc::FPBDCollisionSolverManifoldPoint) == sizeof(Private::FPBDCollisionSolverManifoldPoint));
+			check(sizeof(ispc::FWorldContactPoint) == sizeof(FWorldContactPoint));
+			check(sizeof(ispc::FConstraintSolverBody) == sizeof(FConstraintSolverBody));
+			check(sizeof(ispc::FSolverBody) == sizeof(FSolverBody));
+#endif
+		}
 
 	}	// namespace Private
 }	// namespace Chaos
\ No newline at end of file
diff --git a/Engine/Source/Runtime/Experimental/Chaos/Private/Chaos/Collision/PBDCollisionSolver.ispc b/Engine/Source/Runtime/Experimental/Chaos/Private/Chaos/Collision/PBDCollisionSolver.ispc
new file mode 100644
--- /dev/null
+++ b/Engine/Source/Runtime/Experimental/Chaos/Private/Chaos/Collision/PBDCollisionSolver.ispc
@@ -0,0 +1,266 @@
+// Copyright Epic Games, Inc. All Rights Reserved.
+
+#define EXPLICIT_VECTOR4 1
+
+#include "Math/Vector.isph"
+#include "Math/Quat.isph"
+#include "Math/Transform.isph"
+#include "Chaos/Matrix33.isph"
+
+export struct FSolverBody
+{
+	// World-space inverse inertia
+	// @todo(chaos): do we need this, or should we force all systems to use the FConstraintSolverBody decorator?
+	uniform FMatrix33f InvI;
+
+	// Actor-space center of mass rotation
+	uniform FVector4d RoM;
+
+	// World-space rotation of mass at start of sub step
+	uniform FVector4d R;
+
+	// Predicted world-space center of mass rotation (post-integration, pre-constraint-solve)
+	uniform FVector4d Q;
+
+	// Actor-space center of mass location
+	uniform FVector3d CoM;
+
+	// World-space center of mass state at start of sub step
+	uniform FVector3d X;
+
+	// Predicted world-space center of mass position (post-integration, pre-constraint-solve)
+	uniform FVector3d P;
+
+
+	// Local-space inverse inertia (diagonal, so only 3 elements)
+	uniform FVector3f InvILocal;
+
+	// Inverse mass
+	uniform float InvM;
+
+	// Net position delta applied by all constraints (constantly changing as we iterate over constraints)
+	uniform FVector3f DP;
+
+	// Net rotation delta applied by all constraints (constantly changing as we iterate over constraints)
+	uniform FVector3f DQ;
+
+	// Net position correction delta applied by all constraints (constantly changing as we iterate over constraints)
+	// Will translate the body without introducing linear velocity
+	uniform FVector3f CP;
+
+	// Net rotation correction delta applied by all constraints (constantly changing as we iterate over constraints)
+	// Will rotate the body without introducing angular velocity
+	uniform FVector3f CQ;
+
+	// World-space center of mass velocity
+	uniform FVector3f V;
+
+	// World-space center of mass angular velocity
+	uniform FVector3f W;
+
+	// Distance to a kinematic body (through the contact graph). Used by collision shock propagation
+	uniform int32 Level;
+
+	uniform float AlignmentPadding[2];
+};
+
+static uniform bool Body_IsDynamic(const uniform FSolverBody* uniform Body)
+{ 
+	return (Body->InvM > FLOAT_SMALL_NUMBER); 
+}
+
+static void Body_ApplyPositionDelta(uniform FSolverBody* uniform Body, const uniform FVector3f DP)
+{
+	Body->DP = Body->DP + DP;
+}
+
+static void Body_ApplyRotationDelta(uniform FSolverBody* uniform Body, const uniform FVector3f DQ)
+{
+	Body->DQ = Body->DQ + DQ;
+}
+
+export struct FConstraintSolverBody
+{
+	// The body we decorate
+	uniform FSolverBody* uniform Body;
+
+	// The body modifiers
+	uniform float InvMScale;
+	uniform float InvIScale;
+	uniform float ShockPropagationScale;
+};
+
+export struct FWorldContactPoint
+{
+	// World-space contact point relative to each particle's center of mass
+	uniform FVector3f RelativeContactPoints[2];
+
+	// World-space contact normal and tangents
+	uniform FVector3f ContactNormal;
+	uniform FVector3f ContactTangentU;
+	uniform FVector3f ContactTangentV;
+
+	// Errors to correct along each of the contact axes
+	uniform float ContactDeltaNormal;
+	uniform float ContactDeltaTangentU;
+	uniform float ContactDeltaTangentV;
+
+	// Target velocity along the normal direction
+	uniform float ContactTargetVelocityNormal;
+};
+
+export struct FPBDCollisionSolverManifoldPoint
+{
+	// World-space contact data
+	uniform FWorldContactPoint WorldContact;
+
+	// I^-1.(R x A) for each body where A is each axis (Normal, TangentU, TangentV)
+	uniform FVector3f WorldContactNormalAngular0;
+	uniform FVector3f WorldContactTangentUAngular0;
+	uniform FVector3f WorldContactTangentVAngular0;
+	uniform FVector3f WorldContactNormalAngular1;
+	uniform FVector3f WorldContactTangentUAngular1;
+	uniform FVector3f WorldContactTangentVAngular1;
+
+	// Contact mass (for non-friction)
+	uniform float ContactMassNormal;
+	uniform float ContactMassTangentU;
+	uniform float ContactMassTangentV;
+
+	// Solver outputs
+	uniform float NetPushOutNormal;
+	uniform float NetPushOutTangentU;
+	uniform float NetPushOutTangentV;
+	uniform float NetImpulseNormal;
+	uniform float NetImpulseTangentU;
+	uniform float NetImpulseTangentV;
+
+	// A measure of how much we exceeded the static friction threshold.
+	// Equal to (NormalPushOut / TangentialPushOut) before clamping to the friction cone.
+	// Used to move the static friction anchors to the edge of the cone in Scatter.
+	uniform float StaticFrictionRatio;
+};
+
+export struct FPBDCollisionSolver
+{
+	// Static Friction in the position-solve phase
+	uniform float StaticFriction;
+
+	// Dynamic Friction in the position-solve phase
+	uniform float DynamicFriction;
+
+	// Dynamic Friction in the velocity-solve phase
+	uniform float VelocityFriction;
+
+	// Solver stiffness (scales all pushout and impulses)
+	uniform float Stiffness;
+
+	// Bodies and contacts
+	uniform FConstraintSolverBody SolverBodies[2];
+	uniform FPBDCollisionSolverManifoldPoint* uniform ManifoldPoints;
+	uniform int32 NumManifoldPoints;
+	uniform int32 MaxManifoldPoints;
+};
+
+static uniform float ManifoldPoint_CalculateContactPositionErrorNormal(
+	uniform FPBDCollisionSolverManifoldPoint* uniform SolverManifoldPoint, 
+	const uniform FSolverBody* uniform Body0, 
+	const uniform FSolverBody* uniform Body1, 
+	const uniform float MaxPushOut)
+{
+	// Linear version: calculate the contact delta assuming linear motion after applying a positional impulse at the contact point. There will be an error that depends on the size of the rotation.
+	const uniform FVector3f ContactDelta0 = Body0->DP + VectorCross(Body0->DQ, SolverManifoldPoint->WorldContact.RelativeContactPoints[0]);
+	const uniform FVector3f ContactDelta1 = Body1->DP + VectorCross(Body1->DQ, SolverManifoldPoint->WorldContact.RelativeContactPoints[1]);
+	const uniform FVector3f ContactDelta = ContactDelta0 - ContactDelta1;
+
+	uniform float ContactDeltaNormal = SolverManifoldPoint->WorldContact.ContactDeltaNormal + VectorDot(ContactDelta, SolverManifoldPoint->WorldContact.ContactNormal);
+
+	// NOTE: ContactDeltaNormal is negative for penetration
+	// NOTE: MaxPushOut == 0 disables the pushout limits
+	if ((MaxPushOut > 0) && (ContactDeltaNormal < -MaxPushOut))
+	{
+		ContactDeltaNormal = -MaxPushOut;
+	}
+
+	return ContactDeltaNormal;
+}
+
+static uniform float CalculatePositionCorrectionNormal(
+	const uniform float Stiffness,
+	const uniform float ContactDeltaNormal,
+	const uniform float ContactMassNormal,
+	const uniform float NetPushOutNormal)
+{
+	const uniform float PushOutNormal = -Stiffness * ContactDeltaNormal * ContactMassNormal;
+
+	// The total pushout so far this sub-step
+	// Unilateral constraint: Net-negative impulses not allowed (negative incremental impulses are allowed as long as the net is positive)
+	if ((NetPushOutNormal + PushOutNormal) > 0)
+	{
+		return PushOutNormal;
+	}
+	else
+	{
+		return -NetPushOutNormal;
+	}
+}
+
+static void CollisionSolver_ApplyPositionCorrectionNormal(
+	uniform FPBDCollisionSolver* uniform Solver,
+	const uniform float Stiffness,
+	const uniform float ContactDeltaNormal,
+	uniform FPBDCollisionSolverManifoldPoint* uniform ManifoldPoint,
+	uniform FSolverBody* uniform Body0,
+	uniform FSolverBody* uniform Body1)
+{
+	const uniform float PushOutNormal = CalculatePositionCorrectionNormal(Stiffness, ContactDeltaNormal, ManifoldPoint->ContactMassNormal, ManifoldPoint->NetPushOutNormal);
+
+	ManifoldPoint->NetPushOutNormal = ManifoldPoint->NetPushOutNormal + PushOutNormal;
+
+	// Update the particle state based on the pushout
+	if (Body_IsDynamic(Body0))
+	{
+		const uniform FVector3f DX0 = ManifoldPoint->WorldContact.ContactNormal * (Body0->InvM * PushOutNormal);
+		const uniform FVector3f DR0 = ManifoldPoint->WorldContactNormalAngular0 * PushOutNormal;
+		Body_ApplyPositionDelta(Body0, DX0);
+		Body_ApplyRotationDelta(Body0, DR0);
+	}
+	if (Body_IsDynamic(Body1))
+	{
+		const uniform FVector3f DX1 = ManifoldPoint->WorldContact.ContactNormal * (Body1->InvM * -PushOutNormal);
+		const uniform FVector3f DR1 = ManifoldPoint->WorldContactNormalAngular1 * -PushOutNormal;
+		Body_ApplyPositionDelta(Body1, DX1);
+		Body_ApplyRotationDelta(Body1, DR1);
+	}
+}
+
+static void CollisionSolver_SolvePositionNoFriction(
+	uniform FPBDCollisionSolver* uniform Solver, 
+	const uniform float Dt, 
+	const uniform float MaxPushOut)
+{
+	for (uniform int PointIndex = 0; PointIndex < Solver->NumManifoldPoints; ++PointIndex)
+	{
+		uniform FPBDCollisionSolverManifoldPoint* uniform SolverManifoldPoint = &Solver->ManifoldPoints[PointIndex];
+
+		const uniform float ContactDeltaNormal = ManifoldPoint_CalculateContactPositionErrorNormal(SolverManifoldPoint, Solver->SolverBodies[0].Body, Solver->SolverBodies[1].Body, MaxPushOut);
+
+		const uniform bool bProcessManifoldPoint = (ContactDeltaNormal < 0) || (SolverManifoldPoint->NetPushOutNormal > FLOAT_SMALL_NUMBER);
+		if (bProcessManifoldPoint)
+		{
+			CollisionSolver_ApplyPositionCorrectionNormal(Solver, Solver->Stiffness, ContactDeltaNormal, SolverManifoldPoint, Solver->SolverBodies[0].Body, Solver->SolverBodies[1].Body);
+		}
+	}
+}
+
+export void SolvePositionNoFriction(
+	uniform FPBDCollisionSolver* uniform Solvers, 
+	const uniform int NumSolvers, 
+	const uniform float Dt, 
+	const uniform float MaxPushOut)
+{
+	for (uniform int SolverIndex = 0; SolverIndex < NumSolvers; ++SolverIndex)
+	{
+		CollisionSolver_SolvePositionNoFriction(&Solvers[SolverIndex], Dt, MaxPushOut);
+	}
+}
\ No newline at end of file
diff --git a/Engine/Source/Runtime/Experimental/Chaos/Private/Chaos/Collision/SolverCollisionContainer.cpp b/Engine/Source/Runtime/Experimental/Chaos/Private/Chaos/Collision/SolverCollisionContainer.cpp
--- a/Engine/Source/Runtime/Experimental/Chaos/Private/Chaos/Collision/SolverCollisionContainer.cpp
+++ b/Engine/Source/Runtime/Experimental/Chaos/Private/Chaos/Collision/SolverCollisionContainer.cpp
@@ -288,6 +288,9 @@ namespace Chaos
 		, ConstraintContainer(InConstraintContainer)
 		, bPerIterationCollisionDetection(false)
 	{
+#if !UE_BUILD_SHIPPING && !UE_BUILD_TEST && INTEL_ISPC
+		Private::FPBDCollisionSolverHelper::CheckISPC();
+#endif
 	}
 
 	FPBDCollisionContainerSolver::~FPBDCollisionContainerSolver()
diff --git a/Engine/Source/Runtime/Experimental/Chaos/Public/Chaos/Collision/PBDCollisionSolver.h b/Engine/Source/Runtime/Experimental/Chaos/Public/Chaos/Collision/PBDCollisionSolver.h
--- a/Engine/Source/Runtime/Experimental/Chaos/Public/Chaos/Collision/PBDCollisionSolver.h
+++ b/Engine/Source/Runtime/Experimental/Chaos/Public/Chaos/Collision/PBDCollisionSolver.h
@@ -1350,6 +1350,8 @@ namespace Chaos
 			static void SolvePositionNoFriction(const TArrayView<FPBDCollisionSolver>& CollisionSolvers, const FSolverReal Dt, const FSolverReal MaxPushOut);
 			static void SolvePositionWithFriction(const TArrayView<FPBDCollisionSolver>& CollisionSolvers, const FSolverReal Dt, const FSolverReal MaxPushOut);
 			static void SolveVelocity(const TArrayView<FPBDCollisionSolver>& CollisionSolvers, const FSolverReal Dt, const bool bApplyDynamicFriction);
+
+			static void CheckISPC();
 		};
 
 	}	// namespace Private
diff --git a/Engine/Source/Runtime/Experimental/Chaos/Public/Chaos/Evolution/SolverBody.h b/Engine/Source/Runtime/Experimental/Chaos/Public/Chaos/Evolution/SolverBody.h
--- a/Engine/Source/Runtime/Experimental/Chaos/Public/Chaos/Evolution/SolverBody.h
+++ b/Engine/Source/Runtime/Experimental/Chaos/Public/Chaos/Evolution/SolverBody.h
@@ -433,34 +433,55 @@
 
 			void Init()
 			{
+				InvI = FSolverMatrix33(0);
+				RoM = FRotation3::FromIdentity();
+				R = FRotation3::FromIdentity();
+				Q = FRotation3::FromIdentity();
+				CoM = FVec3(0);
+				X = FVec3(0);
+				P = FVec3(0);
 				InvILocal = FSolverVec3(0);
 				InvM = 0;
-				InvI = FSolverMatrix33(0);
 				DP = FSolverVec3(0);
 				DQ = FSolverVec3(0);
 				CP = FSolverVec3(0);
 				CQ = FSolverVec3(0);
-				Level = 0;
-				X = FVec3(0);
-				R = FRotation3::FromIdentity();
-				P = FVec3(0);
-				Q = FRotation3::FromIdentity();
 				V = FSolverVec3(0);
 				W = FSolverVec3(0);
-				CoM = FVec3(0);
-				RoM = FRotation3::FromIdentity();
+				Level = 0;
 			}
 
+			// World-space inverse inertia
+			// NOTE: Matrix and Rotation are alignas(16) so beware of padding
+			// @todo(chaos): do we need this, or should we force all systems to use the FConstraintSolverBody decorator?
+			FSolverMatrix33 InvI;
+
+			// Actor-space center of mass rotation
+			FRotation3 RoM;
+
+			// World-space rotation of mass at start of sub step
+			FRotation3 R;
+
+			// Predicted world-space center of mass rotation (post-integration, pre-constraint-solve)
+			FRotation3 Q;
+
+			// Actor-space center of mass location
+			FVec3 CoM;
+
+			// World-space center of mass state at start of sub step
+			FVec3 X;
+
+			// Predicted world-space center of mass position (post-integration, pre-constraint-solve)
+			FVec3 P;
+
+
+
 			// Local-space inverse inertia (diagonal, so only 3 elements)
 			FSolverVec3 InvILocal;
 
 			// Inverse mass
 			FSolverReal InvM;
 
-			// World-space inverse inertia
-			// @todo(chaos): do we need this, or should we force all systems to use the FConstraintSolverBody decorator?
-			FSolverMatrix33 InvI;
-
 			// Net position delta applied by all constraints (constantly changing as we iterate over constraints)
 			FSolverVec3 DP;
 
@@ -475,32 +496,14 @@
 			// Will rotate the body without introducing angular velocity
 			FSolverVec3 CQ;
 
-			// Distance to a kinematic body (through the contact graph). Used by collision shock propagation
-			int32 Level;
-
-			// World-space center of mass state at start of sub step
-			FVec3 X;
-
-			// World-space rotation of mass at start of sub step
-			FRotation3 R;
-
-			// Predicted world-space center of mass position (post-integration, pre-constraint-solve)
-			FVec3 P;
-
-			// Predicted world-space center of mass rotation (post-integration, pre-constraint-solve)
-			FRotation3 Q;
-
 			// World-space center of mass velocity
 			FSolverVec3 V;
 
 			// World-space center of mass angular velocity
 			FSolverVec3 W;
 
-			// Actor-space center of mass location
-			FVec3 CoM;
-
-			// Actor-space center of mass rotation
-			FRotation3 RoM;
+			// Distance to a kinematic body (through the contact graph). Used by collision shock propagation
+			int32 Level;
 		};
 
 		FState State;
