Chaos - more Collision Solver ISPC (WIP Disabled)

diff --git a/Engine/Source/Runtime/Experimental/Chaos/Private/Chaos/Collision/PBDCollisionSolver.cpp b/Engine/Source/Runtime/Experimental/Chaos/Private/Chaos/Collision/PBDCollisionSolver.cpp
--- a/Engine/Source/Runtime/Experimental/Chaos/Private/Chaos/Collision/PBDCollisionSolver.cpp
+++ b/Engine/Source/Runtime/Experimental/Chaos/Private/Chaos/Collision/PBDCollisionSolver.cpp
@@ -21,6 +21,9 @@
 
 DEFINE_LOG_CATEGORY(LogChaosCollision);
 
+// TEMP: to be removed
+DECLARE_CYCLE_STAT(TEXT("SolvePositionNoFriction"), STAT_SolvePositionNoFriction, STATGROUP_ChaosConstraintSolver);
+
 namespace Chaos
 {
 	namespace CVars
@@ -28,9 +31,18 @@ namespace Chaos
 		extern int32 Chaos_Collision_UseShockPropagation;
 		extern bool bChaos_PBDCollisionSolver_UseJacobiPairSolver;
 
-		// NOTE: WORK IN PROGRESS. Not ready to be enabled.
-		bool bChaos_PBDCollisionSolver_ISPC = false;
-		FAutoConsoleVariableRef CVarChaos_PBDCollisionSolver_ISPC(TEXT("p.Chaos.PBDCollisionSolver.ISPC"), bChaos_PBDCollisionSolver_ISPC, TEXT(""));
+		//
+		// Solver Settings
+		//
+
+		bool bChaos_PBDCollisionSolver_ISPC = false;	// NOTE: WORK IN PROGRESS. Not ready to be enabled.
+		bool bChaos_PBDCollisionSolver_UseJacobiPairSolver = false;
+		FAutoConsoleVariableRef CVarChaos_PBDCollisionSolver_ISPC(TEXT("p.Chaos.PBDCollisionSolver.ISPC"), bChaos_PBDCollisionSolver_ISPC, TEXT("Use ISPC collision solver (WIP)"));
+		FAutoConsoleVariableRef CVarChaosPBDCollisionSolverUseJacobiPairSolver(TEXT("p.Chaos.PBDCollisionSolver.UseJacobi"), bChaos_PBDCollisionSolver_UseJacobiPairSolver, TEXT("Whether to use the Jacobi collision pair solver"));
+
+		//
+		// Position Solver Settings
+		//
 
 		bool bChaos_PBDCollisionSolver_Position_SolveEnabled = true;
 		float Chaos_PBDCollisionSolver_Position_MinInvMassScale = 0.77f;
@@ -44,8 +56,12 @@ namespace Chaos
 		FAutoConsoleVariableRef CVarChaos_PBDCollisionSolver_Position_PositionSolverTolerance(TEXT("p.Chaos.PBDCollisionSolver.Position.PositionTolerance"), Chaos_PBDCollisionSolver_Position_PositionSolverTolerance, TEXT(""));
 		FAutoConsoleVariableRef CVarChaos_PBDCollisionSolver_Position_RotationSolverTolerance(TEXT("p.Chaos.PBDCollisionSolver.Position.RotationTolerance"), Chaos_PBDCollisionSolver_Position_RotationSolverTolerance, TEXT(""));
 
+		//
+		// Velocity Solver Settings
+		//
+
 		bool bChaos_PBDCollisionSolver_Velocity_SolveEnabled = true;
-		// If this is the same as Chaos_PBDCollisionSolver_Position_MinInvMassScale and all velocity iterations have shockpropagation, we avoid recalculating constraiunt-space mass
+		// If Chaos_PBDCollisionSolver_Velocity_MinInvMassScale is the same as Chaos_PBDCollisionSolver_Position_MinInvMassScale and all velocity iterations have shockpropagation, we avoid recalculating constraint-space mass
 		float Chaos_PBDCollisionSolver_Velocity_MinInvMassScale = Chaos_PBDCollisionSolver_Position_MinInvMassScale;
 		bool bChaos_PBDCollisionSolver_Velocity_FrictionEnabled = true;
 		bool bChaos_PBDCollisionSolver_Velocity_AveragePointEnabled = true;
@@ -240,10 +256,28 @@ namespace Chaos
 
 		void FPBDCollisionSolverHelper::SolvePositionNoFriction(const TArrayView<FPBDCollisionSolver>& CollisionSolvers, const FSolverReal Dt, const FSolverReal MaxPushOut)
 		{
+			SCOPE_CYCLE_COUNTER(STAT_SolvePositionNoFriction);
+
 #if INTEL_ISPC
 			if (CVars::bChaos_PBDCollisionSolver_ISPC)
 			{
-				ispc::SolvePositionNoFriction((ispc::FPBDCollisionSolver*)CollisionSolvers.GetData(), CollisionSolvers.Num(), Dt, MaxPushOut);
+				if (!CVars::bChaos_PBDCollisionSolver_UseJacobiPairSolver)
+				{
+					ispc::SolvePositionNoFriction(
+						(ispc::FPBDCollisionSolver*)CollisionSolvers.GetData(), 
+						CollisionSolvers.Num(), 
+						Dt, 
+						MaxPushOut);
+				}
+				else
+				{
+					ispc::SolvePositionNoFrictionJacobi(
+						(ispc::FPBDCollisionSolver*)CollisionSolvers.GetData(), 
+						CollisionSolvers.Num(), 
+						Dt, 
+						MaxPushOut, 
+						CVars::Chaos_PBDCollisionSolver_JacobiStiffness);
+				}
 				return;
 			}
 #endif
diff --git a/Engine/Source/Runtime/Experimental/Chaos/Private/Chaos/Collision/PBDCollisionSolver.ispc b/Engine/Source/Runtime/Experimental/Chaos/Private/Chaos/Collision/PBDCollisionSolver.ispc
--- a/Engine/Source/Runtime/Experimental/Chaos/Private/Chaos/Collision/PBDCollisionSolver.ispc
+++ b/Engine/Source/Runtime/Experimental/Chaos/Private/Chaos/Collision/PBDCollisionSolver.ispc
@@ -7,162 +7,176 @@
 #include "Math/Transform.isph"
 #include "Chaos/Matrix33.isph"
 
+//
+//
+// STRUCTURES
+//
+//
+
 export struct FSolverBody
 {
 	// World-space inverse inertia
 	// @todo(chaos): do we need this, or should we force all systems to use the FConstraintSolverBody decorator?
-	uniform FMatrix33f InvI;
+	FMatrix33f InvI;
 
 	// Actor-space center of mass rotation
-	uniform FVector4d RoM;
+	FVector4d RoM;
 
 	// World-space rotation of mass at start of sub step
-	uniform FVector4d R;
+	FVector4d R;
 
 	// Predicted world-space center of mass rotation (post-integration, pre-constraint-solve)
-	uniform FVector4d Q;
+	FVector4d Q;
 
 	// Actor-space center of mass location
-	uniform FVector3d CoM;
+	FVector3d CoM;
 
 	// World-space center of mass state at start of sub step
-	uniform FVector3d X;
+	FVector3d X;
 
 	// Predicted world-space center of mass position (post-integration, pre-constraint-solve)
-	uniform FVector3d P;
+	FVector3d P;
 
 
 	// Local-space inverse inertia (diagonal, so only 3 elements)
-	uniform FVector3f InvILocal;
+	FVector3f InvILocal;
 
 	// Inverse mass
-	uniform float InvM;
+	float InvM;
 
 	// Net position delta applied by all constraints (constantly changing as we iterate over constraints)
-	uniform FVector3f DP;
+	FVector3f DP;
 
 	// Net rotation delta applied by all constraints (constantly changing as we iterate over constraints)
-	uniform FVector3f DQ;
+	FVector3f DQ;
 
 	// Net position correction delta applied by all constraints (constantly changing as we iterate over constraints)
 	// Will translate the body without introducing linear velocity
-	uniform FVector3f CP;
+	FVector3f CP;
 
 	// Net rotation correction delta applied by all constraints (constantly changing as we iterate over constraints)
 	// Will rotate the body without introducing angular velocity
-	uniform FVector3f CQ;
+	FVector3f CQ;
 
 	// World-space center of mass velocity
-	uniform FVector3f V;
+	FVector3f V;
 
 	// World-space center of mass angular velocity
-	uniform FVector3f W;
+	FVector3f W;
 
 	// Distance to a kinematic body (through the contact graph). Used by collision shock propagation
-	uniform int32 Level;
+	int32 Level;
 
-	uniform float AlignmentPadding[2];
+	float AlignmentPadding[2];
 };
 
-static uniform bool Body_IsDynamic(const uniform FSolverBody* uniform Body)
-{ 
-	return (Body->InvM > FLOAT_SMALL_NUMBER); 
-}
-
-static void Body_ApplyPositionDelta(uniform FSolverBody* uniform Body, const uniform FVector3f DP)
-{
-	Body->DP = Body->DP + DP;
-}
-
-static void Body_ApplyRotationDelta(uniform FSolverBody* uniform Body, const uniform FVector3f DQ)
-{
-	Body->DQ = Body->DQ + DQ;
-}
-
 export struct FConstraintSolverBody
 {
 	// The body we decorate
-	uniform FSolverBody* uniform Body;
+	FSolverBody* Body;
 
 	// The body modifiers
-	uniform float InvMScale;
-	uniform float InvIScale;
-	uniform float ShockPropagationScale;
+	float InvMScale;
+	float InvIScale;
+	float ShockPropagationScale;
 };
 
 export struct FWorldContactPoint
 {
 	// World-space contact point relative to each particle's center of mass
-	uniform FVector3f RelativeContactPoints[2];
+	FVector3f RelativeContactPoints[2];
 
 	// World-space contact normal and tangents
-	uniform FVector3f ContactNormal;
-	uniform FVector3f ContactTangentU;
-	uniform FVector3f ContactTangentV;
+	FVector3f ContactNormal;
+	FVector3f ContactTangentU;
+	FVector3f ContactTangentV;
 
 	// Errors to correct along each of the contact axes
-	uniform float ContactDeltaNormal;
-	uniform float ContactDeltaTangentU;
-	uniform float ContactDeltaTangentV;
+	float ContactDeltaNormal;
+	float ContactDeltaTangentU;
+	float ContactDeltaTangentV;
 
 	// Target velocity along the normal direction
-	uniform float ContactTargetVelocityNormal;
+	float ContactTargetVelocityNormal;
 };
 
 export struct FPBDCollisionSolverManifoldPoint
 {
 	// World-space contact data
-	uniform FWorldContactPoint WorldContact;
+	FWorldContactPoint WorldContact;
 
 	// I^-1.(R x A) for each body where A is each axis (Normal, TangentU, TangentV)
-	uniform FVector3f WorldContactNormalAngular0;
-	uniform FVector3f WorldContactTangentUAngular0;
-	uniform FVector3f WorldContactTangentVAngular0;
-	uniform FVector3f WorldContactNormalAngular1;
-	uniform FVector3f WorldContactTangentUAngular1;
-	uniform FVector3f WorldContactTangentVAngular1;
+	FVector3f WorldContactNormalAngular0;
+	FVector3f WorldContactTangentUAngular0;
+	FVector3f WorldContactTangentVAngular0;
+	FVector3f WorldContactNormalAngular1;
+	FVector3f WorldContactTangentUAngular1;
+	FVector3f WorldContactTangentVAngular1;
 
 	// Contact mass (for non-friction)
-	uniform float ContactMassNormal;
-	uniform float ContactMassTangentU;
-	uniform float ContactMassTangentV;
+	float ContactMassNormal;
+	float ContactMassTangentU;
+	float ContactMassTangentV;
 
 	// Solver outputs
-	uniform float NetPushOutNormal;
-	uniform float NetPushOutTangentU;
-	uniform float NetPushOutTangentV;
-	uniform float NetImpulseNormal;
-	uniform float NetImpulseTangentU;
-	uniform float NetImpulseTangentV;
+	float NetPushOutNormal;
+	float NetPushOutTangentU;
+	float NetPushOutTangentV;
+	float NetImpulseNormal;
+	float NetImpulseTangentU;
+	float NetImpulseTangentV;
 
 	// A measure of how much we exceeded the static friction threshold.
 	// Equal to (NormalPushOut / TangentialPushOut) before clamping to the friction cone.
 	// Used to move the static friction anchors to the edge of the cone in Scatter.
-	uniform float StaticFrictionRatio;
+	float StaticFrictionRatio;
 };
 
 export struct FPBDCollisionSolver
 {
 	// Static Friction in the position-solve phase
-	uniform float StaticFriction;
+	float StaticFriction;
 
 	// Dynamic Friction in the position-solve phase
-	uniform float DynamicFriction;
+	float DynamicFriction;
 
 	// Dynamic Friction in the velocity-solve phase
-	uniform float VelocityFriction;
+	float VelocityFriction;
 
 	// Solver stiffness (scales all pushout and impulses)
-	uniform float Stiffness;
+	float Stiffness;
 
 	// Bodies and contacts
-	uniform FConstraintSolverBody SolverBodies[2];
-	uniform FPBDCollisionSolverManifoldPoint* uniform ManifoldPoints;
-	uniform int32 NumManifoldPoints;
-	uniform int32 MaxManifoldPoints;
+	FConstraintSolverBody SolverBodies[2];
+	FPBDCollisionSolverManifoldPoint* ManifoldPoints;
+	int32 NumManifoldPoints;
+	int32 MaxManifoldPoints;
 };
 
-static uniform float ManifoldPoint_CalculateContactPositionErrorNormal(
+
+//
+//
+// GENERAL FUNCTIONS
+//
+//
+
+
+inline uniform bool Body_IsDynamic(const uniform FSolverBody* uniform Body)
+{ 
+	return (Body->InvM > FLOAT_SMALL_NUMBER); 
+}
+
+inline void Body_ApplyPositionDelta(uniform FSolverBody* uniform Body, const uniform FVector3f DP)
+{
+	Body->DP = Body->DP + DP;
+}
+
+inline void Body_ApplyRotationDelta(uniform FSolverBody* uniform Body, const uniform FVector3f DQ)
+{
+	Body->DQ = Body->DQ + DQ;
+}
+
+static uniform float CalculatePositionErrorNormal(
 	uniform FPBDCollisionSolverManifoldPoint* uniform SolverManifoldPoint, 
 	const uniform FSolverBody* uniform Body0, 
 	const uniform FSolverBody* uniform Body1, 
@@ -205,7 +219,14 @@ static uniform float CalculatePositionCorrectionNormal(
 	}
 }
 
-static void CollisionSolver_ApplyPositionCorrectionNormal(
+
+//
+//
+// GAUSS SEIDEL SOLVER
+//
+//
+
+static void CollisionSolverApplyPositionCorrectionNormal(
 	uniform FPBDCollisionSolver* uniform Solver,
 	const uniform float Stiffness,
 	const uniform float ContactDeltaNormal,
@@ -234,7 +255,7 @@ static void CollisionSolver_ApplyPositionCorrectionNormal(
 	}
 }
 
-static void CollisionSolver_SolvePositionNoFriction(
+static void CollisionSolverSolvePositionNoFriction(
 	uniform FPBDCollisionSolver* uniform Solver, 
 	const uniform float Dt, 
 	const uniform float MaxPushOut)
@@ -243,12 +264,12 @@ static void CollisionSolver_SolvePositionNoFriction(
 	{
 		uniform FPBDCollisionSolverManifoldPoint* uniform SolverManifoldPoint = &Solver->ManifoldPoints[PointIndex];
 
-		const uniform float ContactDeltaNormal = ManifoldPoint_CalculateContactPositionErrorNormal(SolverManifoldPoint, Solver->SolverBodies[0].Body, Solver->SolverBodies[1].Body, MaxPushOut);
+		const uniform float ContactDeltaNormal = CalculatePositionErrorNormal(SolverManifoldPoint, Solver->SolverBodies[0].Body, Solver->SolverBodies[1].Body, MaxPushOut);
 
 		const uniform bool bProcessManifoldPoint = (ContactDeltaNormal < 0) || (SolverManifoldPoint->NetPushOutNormal > FLOAT_SMALL_NUMBER);
 		if (bProcessManifoldPoint)
 		{
-			CollisionSolver_ApplyPositionCorrectionNormal(Solver, Solver->Stiffness, ContactDeltaNormal, SolverManifoldPoint, Solver->SolverBodies[0].Body, Solver->SolverBodies[1].Body);
+			CollisionSolverApplyPositionCorrectionNormal(Solver, Solver->Stiffness, ContactDeltaNormal, SolverManifoldPoint, Solver->SolverBodies[0].Body, Solver->SolverBodies[1].Body);
 		}
 	}
 }
@@ -261,6 +282,85 @@ export void SolvePositionNoFriction(
 {
 	for (uniform int SolverIndex = 0; SolverIndex < NumSolvers; ++SolverIndex)
 	{
-		CollisionSolver_SolvePositionNoFriction(&Solvers[SolverIndex], Dt, MaxPushOut);
+		CollisionSolverSolvePositionNoFriction(&Solvers[SolverIndex], Dt, MaxPushOut);
+	}
+}
+
+
+//
+//
+// PARTIAL-JACOBI SOLVER
+//
+//
+
+static void CollisionSolver_SolvePositionNoFriction_Jacobi(
+	uniform FPBDCollisionSolver* uniform Solver, 
+	const uniform float Dt, 
+	const uniform float MaxPushOut,
+	const uniform float JacobiStiffness)
+{
+	uniform FSolverBody* uniform Body0 = Solver->SolverBodies[0].Body;
+	uniform FSolverBody* uniform Body1 = Solver->SolverBodies[1].Body;
+
+	uniform FVector3f DX0 = FloatZeroVector;
+	uniform FVector3f DX1 = FloatZeroVector;
+	uniform FVector3f DR0 = FloatZeroVector;
+	uniform FVector3f DR1 = FloatZeroVector;
+
+	const uniform float Stiffness = JacobiStiffness * Solver->Stiffness;
+
+	for (uniform int PointIndex = 0; PointIndex < Solver->NumManifoldPoints; ++PointIndex)
+	{
+		uniform FPBDCollisionSolverManifoldPoint* uniform ManifoldPoint = &Solver->ManifoldPoints[PointIndex];
+
+		const uniform float ContactDeltaNormal = CalculatePositionErrorNormal(ManifoldPoint, Solver->SolverBodies[0].Body, Solver->SolverBodies[1].Body, MaxPushOut);
+
+		const uniform bool bProcessManifoldPoint = (ContactDeltaNormal < 0) || (ManifoldPoint->NetPushOutNormal > FLOAT_SMALL_NUMBER);
+		if (bProcessManifoldPoint)
+		{
+			const uniform float PushOutNormal = CalculatePositionCorrectionNormal(Stiffness, ContactDeltaNormal, ManifoldPoint->ContactMassNormal, ManifoldPoint->NetPushOutNormal);
+
+			ManifoldPoint->NetPushOutNormal = ManifoldPoint->NetPushOutNormal + PushOutNormal;
+
+			if (Body_IsDynamic(Body0))
+			{
+				const uniform FVector3f DDX0 = ManifoldPoint->WorldContact.ContactNormal * (Body0->InvM * PushOutNormal);
+				const uniform FVector3f DDR0 = ManifoldPoint->WorldContactNormalAngular0 * PushOutNormal;
+				DX0 = DX0 + DDX0;
+				DR0 = DR0 + DDR0;
+			}
+			if (Body_IsDynamic(Body1))
+			{
+				const uniform FVector3f DDX1 = ManifoldPoint->WorldContact.ContactNormal * (Body1->InvM * PushOutNormal);
+				const uniform FVector3f DDR1 = ManifoldPoint->WorldContactNormalAngular1 * PushOutNormal;
+				DX1 = DX1 - DDX1;
+				DR1 = DR1 - DDR1;
+			}
+		}
+	}
+
+	if (Body_IsDynamic(Body0))
+	{
+		Body_ApplyPositionDelta(Body0, DX0);
+		Body_ApplyRotationDelta(Body0, DR0);
+	}
+	if (Body_IsDynamic(Body1))
+	{
+		Body_ApplyPositionDelta(Body1, DX1);
+		Body_ApplyRotationDelta(Body1, DR1);
+	}
+}
+
+
+export void SolvePositionNoFrictionJacobi(
+	uniform FPBDCollisionSolver* uniform Solvers, 
+	const uniform int NumSolvers, 
+	const uniform float Dt, 
+	const uniform float MaxPushOut,
+	const uniform float JacobiStiffness)
+{
+	for (uniform int SolverIndex = 0; SolverIndex < NumSolvers; ++SolverIndex)
+	{
+		CollisionSolver_SolvePositionNoFriction_Jacobi(&Solvers[SolverIndex], Dt, MaxPushOut, JacobiStiffness);
 	}
 }
\ No newline at end of file
diff --git a/Engine/Source/Runtime/Experimental/Chaos/Private/Chaos/Collision/SolverCollisionContainer.cpp b/Engine/Source/Runtime/Experimental/Chaos/Private/Chaos/Collision/SolverCollisionContainer.cpp
--- a/Engine/Source/Runtime/Experimental/Chaos/Private/Chaos/Collision/SolverCollisionContainer.cpp
+++ b/Engine/Source/Runtime/Experimental/Chaos/Private/Chaos/Collision/SolverCollisionContainer.cpp
@@ -34,10 +34,6 @@ namespace Chaos
 		FAutoConsoleVariableRef CVarChaosPBDCollisionSolverAutoStiffnessMassRatio1(TEXT("p.Chaos.PBDCollisionSolver.AutoStiffness.MassRatio1"), Chaos_PBDCollisionSolver_AutoStiffness_MassRatio1, TEXT(""));
 		FAutoConsoleVariableRef CVarChaosPBDCollisionSolverAutoStiffnessMassRatio2(TEXT("p.Chaos.PBDCollisionSolver.AutoStiffness.MassRatio2"), Chaos_PBDCollisionSolver_AutoStiffness_MassRatio2, TEXT(""));
 
-		// Whether to use the Jacobi collision pair solver
-		bool bChaos_PBDCollisionSolver_UseJacobiPairSolver = false;
-		FAutoConsoleVariableRef CVarChaosPBDCollisionSolverUseJacobiPairSolver(TEXT("p.Chaos.PBDCollisionSolver.UseJacobi"), bChaos_PBDCollisionSolver_UseJacobiPairSolver, TEXT(""));
-
 		// Jacobi solver stiffness
 		// @todo(chaos): to be tuned
 		FRealSingle Chaos_PBDCollisionSolver_JacobiStiffness = 0.5f;
diff --git a/Engine/Source/Runtime/Experimental/Chaos/Private/Chaos/Island/IslandGroupManager.cpp b/Engine/Source/Runtime/Experimental/Chaos/Private/Chaos/Island/IslandGroupManager.cpp
--- a/Engine/Source/Runtime/Experimental/Chaos/Private/Chaos/Island/IslandGroupManager.cpp
+++ b/Engine/Source/Runtime/Experimental/Chaos/Private/Chaos/Island/IslandGroupManager.cpp
@@ -8,13 +8,13 @@
 #include "ChaosStats.h"
 #include "ProfilingDebugging/ScopedTimers.h"
 
-DECLARE_CYCLE_STAT(TEXT("Chaos::SolveConstraints"), STAT_SolveConstraints, STATGROUP_Chaos);
+DECLARE_CYCLE_STAT(TEXT("Chaos::SolveConstraints"), STAT_SolveConstraints, STATGROUP_ChaosConstraintSolver);
 
-DECLARE_CYCLE_STAT(TEXT("FPBDRigidsEvolutionGBF::Gather"), STAT_Evolution_Gather, STATGROUP_Chaos);
-DECLARE_CYCLE_STAT(TEXT("FPBDRigidsEvolutionGBF::Scatter"), STAT_Evolution_Scatter, STATGROUP_Chaos);
-DECLARE_CYCLE_STAT(TEXT("FPBDRigidsEvolutionGBF::ApplyConstraintsPhase1"), STAT_Evolution_ApplyConstraintsPhase1, STATGROUP_Chaos);
-DECLARE_CYCLE_STAT(TEXT("FPBDRigidsEvolutionGBF::ApplyConstraintsPhase2"), STAT_Evolution_ApplyConstraintsPhase2, STATGROUP_Chaos);
-DECLARE_CYCLE_STAT(TEXT("FPBDRigidsEvolutionGBF::ApplyConstraintsPhase3"), STAT_Evolution_ApplyConstraintsPhase3, STATGROUP_Chaos);
+DECLARE_CYCLE_STAT(TEXT("FPBDRigidsEvolutionGBF::Gather"), STAT_Evolution_Gather, STATGROUP_ChaosConstraintSolver);
+DECLARE_CYCLE_STAT(TEXT("FPBDRigidsEvolutionGBF::Scatter"), STAT_Evolution_Scatter, STATGROUP_ChaosConstraintSolver);
+DECLARE_CYCLE_STAT(TEXT("FPBDRigidsEvolutionGBF::ApplyConstraintsPhase1"), STAT_Evolution_ApplyConstraintsPhase1, STATGROUP_ChaosConstraintSolver);
+DECLARE_CYCLE_STAT(TEXT("FPBDRigidsEvolutionGBF::ApplyConstraintsPhase2"), STAT_Evolution_ApplyConstraintsPhase2, STATGROUP_ChaosConstraintSolver);
+DECLARE_CYCLE_STAT(TEXT("FPBDRigidsEvolutionGBF::ApplyConstraintsPhase3"), STAT_Evolution_ApplyConstraintsPhase3, STATGROUP_ChaosConstraintSolver);
 
 DECLARE_CYCLE_STAT(TEXT("FPBDRigidsEvolutionGBF::Gather::Prepare"), STAT_Evolution_Gather_Prepare, STATGROUP_ChaosConstraintDetails);
 DECLARE_CYCLE_STAT(TEXT("FPBDRigidsEvolutionGBF::Gather::Bodies"), STAT_Evolution_Gather_Bodies, STATGROUP_ChaosConstraintDetails);
diff --git a/Engine/Source/Runtime/Experimental/Chaos/Private/Chaos/PBDCollisionConstraints.cpp b/Engine/Source/Runtime/Experimental/Chaos/Private/Chaos/PBDCollisionConstraints.cpp
--- a/Engine/Source/Runtime/Experimental/Chaos/Private/Chaos/PBDCollisionConstraints.cpp
+++ b/Engine/Source/Runtime/Experimental/Chaos/Private/Chaos/PBDCollisionConstraints.cpp
@@ -82,7 +82,6 @@ namespace Chaos
 #endif
 	
 	DECLARE_CYCLE_STAT(TEXT("Collisions::Reset"), STAT_Collisions_Reset, STATGROUP_ChaosCollision);
-	DECLARE_CYCLE_STAT(TEXT("Collisions::UpdatePointConstraints"), STAT_Collisions_UpdatePointConstraints, STATGROUP_ChaosCollision);
 	DECLARE_CYCLE_STAT(TEXT("Collisions::BeginDetect"), STAT_Collisions_BeginDetect, STATGROUP_ChaosCollision);
 	DECLARE_CYCLE_STAT(TEXT("Collisions::EndDetect"), STAT_Collisions_EndDetect, STATGROUP_ChaosCollision);
 	DECLARE_CYCLE_STAT(TEXT("Collisions::DetectProbeCollisions"), STAT_Collisions_DetectProbeCollisions, STATGROUP_ChaosCollision);
diff --git a/Engine/Source/Runtime/Experimental/Chaos/Public/ChaosStats.h b/Engine/Source/Runtime/Experimental/Chaos/Public/ChaosStats.h
--- a/Engine/Source/Runtime/Experimental/Chaos/Public/ChaosStats.h
+++ b/Engine/Source/Runtime/Experimental/Chaos/Public/ChaosStats.h
@@ -11,6 +11,7 @@ DECLARE_STATS_GROUP(TEXT("ChaosThread"), STATGROUP_ChaosThread, STATCAT_Advanced
 DECLARE_STATS_GROUP(TEXT("ChaosDedicated"), STATGROUP_ChaosDedicated, STATCAT_Advanced);
 DECLARE_STATS_GROUP(TEXT("ChaosEngine"), STATGROUP_ChaosEngine, STATCAT_Advanced);
 DECLARE_STATS_GROUP(TEXT("ChaosCollision"), STATGROUP_ChaosCollision, STATCAT_Advanced);
+DECLARE_STATS_GROUP(TEXT("ChaosConstraintSolver"), STATGROUP_ChaosConstraintSolver, STATCAT_Advanced);
 DECLARE_STATS_GROUP(TEXT("ChaosJoint"), STATGROUP_ChaosJoint, STATCAT_Advanced);
 DECLARE_STATS_GROUP(TEXT("ChaosMinEvolution"), STATGROUP_ChaosMinEvolution, STATCAT_Advanced);
 DECLARE_STATS_GROUP(TEXT("ChaosCounters"), STATGROUP_ChaosCounters, STATCAT_Advanced);
