Add "Disable" to midphase modification

diff --git a/Engine/Source/Runtime/Experimental/Chaos/Private/Chaos/Collision/ParticlePairMidPhase.cpp b/Engine/Source/Runtime/Experimental/Chaos/Private/Chaos/Collision/ParticlePairMidPhase.cpp
--- a/Engine/Source/Runtime/Experimental/Chaos/Private/Chaos/Collision/ParticlePairMidPhase.cpp
+++ b/Engine/Source/Runtime/Experimental/Chaos/Private/Chaos/Collision/ParticlePairMidPhase.cpp
@@ -842,6 +842,7 @@ namespace Chaos
 		ShapePairDetectors.Reset();
 		MultiShapePairDetectors.Reset();
 
+		Flags.bIsActive = true;
 		Flags.bIsCCD = false;
 		Flags.bIsCCDActive = false;
 		Flags.bIsSleeping = false;
@@ -852,6 +853,7 @@ namespace Chaos
 	{
 		if (Flags.bIsModified)
 		{
+			Flags.bIsActive = true;
 			Flags.bIsCCDActive = Flags.bIsCCD;
 			Flags.bIsModified = false;
 		}
@@ -997,37 +999,40 @@ namespace Chaos
 			return;
 		}
 
-		// CullDistance is scaled by the size of the dynamic objects.
-		FReal CullDistance = InCullDistance * CullDistanceScale;
-
-		// Extend cull distance for sweep-enabled CCD collision
-		const bool bUseSweep = Flags.bIsCCD && ShouldEnableCCD(Dt);
-		if (bUseSweep)
+		if (Flags.bIsActive)
 		{
-			const FReal VMax = (FConstGenericParticleHandle(GetParticle0())->V() - FConstGenericParticleHandle(GetParticle1())->V()).GetAbsMax();
-			CullDistance += VMax * Dt;
-		}
+			// CullDistance is scaled by the size of the dynamic objects.
+			FReal CullDistance = InCullDistance * CullDistanceScale;
 
-		// Run collision detection on all potentially colliding shape pairs
-		NumActiveConstraints = 0;
-		if (Flags.bIsCCD)
-		{
-			for (FSingleShapePairCollisionDetector& ShapePair : ShapePairDetectors)
+			// Extend cull distance for sweep-enabled CCD collision
+			const bool bUseSweep = Flags.bIsCCD && ShouldEnableCCD(Dt);
+			if (bUseSweep)
 			{
-				NumActiveConstraints += ShapePair.GenerateCollisionCCD(bUseSweep, CullDistance, Dt, Context);
+				const FReal VMax = (FConstGenericParticleHandle(GetParticle0())->V() - FConstGenericParticleHandle(GetParticle1())->V()).GetAbsMax();
+				CullDistance += VMax * Dt;
 			}
-		}
-		else
-		{
-			for (FSingleShapePairCollisionDetector& ShapePair : ShapePairDetectors)
+
+			// Run collision detection on all potentially colliding shape pairs
+			NumActiveConstraints = 0;
+			if (Flags.bIsCCD)
+			{
+				for (FSingleShapePairCollisionDetector& ShapePair : ShapePairDetectors)
+				{
+					NumActiveConstraints += ShapePair.GenerateCollisionCCD(bUseSweep, CullDistance, Dt, Context);
+				}
+			}
+			else
+			{
+				for (FSingleShapePairCollisionDetector& ShapePair : ShapePairDetectors)
+				{
+					NumActiveConstraints += ShapePair.GenerateCollision(CullDistance, Dt, Context);
+				}
+			}
+			for (FMultiShapePairCollisionDetector& MultiShapePair : MultiShapePairDetectors)
 			{
-				NumActiveConstraints += ShapePair.GenerateCollision(CullDistance, Dt, Context);
+				NumActiveConstraints += MultiShapePair.GenerateCollisions(CullDistance, Dt, Context);
 			}
 		}
-		for (FMultiShapePairCollisionDetector& MultiShapePair : MultiShapePairDetectors)
-		{
-			NumActiveConstraints += MultiShapePair.GenerateCollisions(CullDistance, Dt, Context);
-		}
 
 		// Reset any modifications applied by the MidPhaseModifier.
 		// In principle this belongs at the start of the frame but it is hard to avoid the cache miss there.
diff --git a/Engine/Source/Runtime/Experimental/Chaos/Private/Chaos/MidPhaseModification.cpp b/Engine/Source/Runtime/Experimental/Chaos/Private/Chaos/MidPhaseModification.cpp
--- a/Engine/Source/Runtime/Experimental/Chaos/Private/Chaos/MidPhaseModification.cpp
+++ b/Engine/Source/Runtime/Experimental/Chaos/Private/Chaos/MidPhaseModification.cpp
@@ -9,6 +9,11 @@
 
 namespace Chaos
 {
+	void FMidPhaseModifier::Disable()
+	{
+		MidPhase->SetIsActive(false);
+	}
+
 	void FMidPhaseModifier::DisableCCD()
 	{
 		MidPhase->SetCCDIsActive(false);
@@ -54,4 +59,29 @@ namespace Chaos
 	{
 		return FMidPhaseModifierParticleRange(this, Particle);
 	}
+
+	FMidPhaseModifier FMidPhaseModifierAccessor::GetMidPhase(
+		FGeometryParticleHandle* Particle0,
+		FGeometryParticleHandle* Particle1)
+	{
+		// Put the particle with fewer collisions in spot 0
+		int32 Num0 = Particle0->ParticleCollisions().Num();
+		int32 Num1 = Particle1->ParticleCollisions().Num();
+		if (Num1 < Num0)
+		{
+			Swap(Particle0, Particle1);
+		}
+
+		// Loop over the mid-phases of the particle that has fewer of them,
+		// find the one that involves the other particle, if any.
+		for (FMidPhaseModifier& Modifier : GetMidPhases(Particle0))
+		{
+			if (Modifier.GetOtherParticle(Particle0) == Particle1)
+			{
+				return Modifier;
+			}
+		}
+
+		return FMidPhaseModifier();
+	}
 }
diff --git a/Engine/Source/Runtime/Experimental/Chaos/Public/Chaos/Collision/ParticlePairMidPhase.h b/Engine/Source/Runtime/Experimental/Chaos/Public/Chaos/Collision/ParticlePairMidPhase.h
--- a/Engine/Source/Runtime/Experimental/Chaos/Public/Chaos/Collision/ParticlePairMidPhase.h
+++ b/Engine/Source/Runtime/Experimental/Chaos/Public/Chaos/Collision/ParticlePairMidPhase.h
@@ -439,7 +439,16 @@ namespace Chaos
 		}
 
 		/**
-		 * Override the CCD condition for this mid-phase. Used by the MidPhase modifier and gets reset every frame.
+		 * @brief By default midphases are active. If IsActive is set to false,
+		 * this midphase will not generate a narrow phase.
+		 */
+		void SetIsActive(bool bIsActive)
+		{
+			Flags.bIsActive = bIsActive;
+		}
+
+		/**
+		 * @brief Override the CCD condition for this mid-phase. Used by the MidPhase modifier and gets reset every frame.
 		 */
 		void SetCCDIsActive(bool bCCDIsActive)
 		{
@@ -480,6 +489,7 @@ namespace Chaos
 			FFlags() : Bits(0) {}
 			struct
 			{
+				uint32 bIsActive : 1;    // True if this midphase should generate a narrow phase at all
 				uint32 bIsCCD : 1;       // True if CCD is supported by either particle
 				uint32 bIsCCDActive : 1; // True if CCD is active for this midphase on this frame. This can be changes by modifiers and resets to bIsCCD each frame.
 				uint32 bIsSleeping : 1;
diff --git a/Engine/Source/Runtime/Experimental/Chaos/Public/Chaos/MidPhaseModification.h b/Engine/Source/Runtime/Experimental/Chaos/Public/Chaos/MidPhaseModification.h
--- a/Engine/Source/Runtime/Experimental/Chaos/Public/Chaos/MidPhaseModification.h
+++ b/Engine/Source/Runtime/Experimental/Chaos/Public/Chaos/MidPhaseModification.h
@@ -30,10 +30,22 @@ namespace Chaos
 			, Accessor(InAccessor)
 		{}
 
+		// Since a modifier can be invalid (ie, null midphase and/or accessor), make it
+		// castable to bool so that users can check validity.
+		bool IsValid() const
+		{
+			return
+				MidPhase != nullptr &&
+				Accessor != nullptr;
+		}
+
 		//
 		// Modifying functions
 		//
 
+		// Disable this midphase entirely
+		void Disable();
+
 		// Disable CCD for this pair
 		void DisableCCD();
 
@@ -141,6 +153,9 @@ namespace Chaos
 		// Get an object which allows for range iteration over the list of
 		// midphases for a particle
 		FMidPhaseModifierParticleRange GetMidPhases(FGeometryParticleHandle* Particle);
+
+		// Get a midphase modifier for a particular object pair
+		FMidPhaseModifier GetMidPhase(FGeometryParticleHandle* Particle0, FGeometryParticleHandle* Particle1);
 	};
 
 	using FMidPhaseModifierCallback = TFunction<void(const FCollisionContactModifier&)>;
