Add missing initialization for ParticlePairMidPhase::Flags::bIsActive

diff --git a/Engine/Source/Runtime/Experimental/Chaos/Private/Chaos/Collision/ParticlePairMidPhase.cpp b/Engine/Source/Runtime/Experimental/Chaos/Private/Chaos/Collision/ParticlePairMidPhase.cpp
--- a/Engine/Source/Runtime/Experimental/Chaos/Private/Chaos/Collision/ParticlePairMidPhase.cpp
+++ b/Engine/Source/Runtime/Experimental/Chaos/Private/Chaos/Collision/ParticlePairMidPhase.cpp
@@ -871,6 +871,8 @@ namespace Chaos
 		Particle1 = InParticle1;
 		Key = InKey;
 
+		Flags.bIsActive = true;
+
 		// If CCD is allowed in the current context and for at least one of
 		// the particles involved, enable it for this midphase.
 		//
