Optimization: Avoid to recreate FName from string when GeometryCollection GetRootIndex is called

diff --git a/Engine/Source/Runtime/Experimental/Chaos/Private/GeometryCollection/Facades/CollectionHierarchyFacade.cpp b/Engine/Source/Runtime/Experimental/Chaos/Private/GeometryCollection/Facades/CollectionHierarchyFacade.cpp
--- a/Engine/Source/Runtime/Experimental/Chaos/Private/GeometryCollection/Facades/CollectionHierarchyFacade.cpp
+++ b/Engine/Source/Runtime/Experimental/Chaos/Private/GeometryCollection/Facades/CollectionHierarchyFacade.cpp
@@ -6,15 +6,15 @@
 namespace Chaos::Facades
 {
 	FCollectionHierarchyFacade::FCollectionHierarchyFacade(FManagedArrayCollection& InCollection)
-		: ParentAttribute(InCollection, "Parent", FTransformCollection::TransformGroup)
-	    , ChildrenAttribute(InCollection, "Children", FTransformCollection::TransformGroup)
-		, LevelAttribute(InCollection, "Level", FTransformCollection::TransformGroup)
+		: ParentAttribute(InCollection, FTransformCollection::ParentAttribute, FTransformCollection::TransformGroup)
+	    , ChildrenAttribute(InCollection, FTransformCollection::ChildrenAttribute, FTransformCollection::TransformGroup)
+		, LevelAttribute(InCollection, FTransformCollection::LevelAttribute, FTransformCollection::TransformGroup)
 	{}
 
 	FCollectionHierarchyFacade::FCollectionHierarchyFacade(const FManagedArrayCollection& InCollection)
-		: ParentAttribute(InCollection, "Parent", FTransformCollection::TransformGroup)
-		, ChildrenAttribute(InCollection, "Children", FTransformCollection::TransformGroup)
-		, LevelAttribute(InCollection, "Level", FTransformCollection::TransformGroup)
+		: ParentAttribute(InCollection, FTransformCollection::ParentAttribute, FTransformCollection::TransformGroup)
+		, ChildrenAttribute(InCollection, FTransformCollection::ChildrenAttribute, FTransformCollection::TransformGroup)
+		, LevelAttribute(InCollection, FTransformCollection::LevelAttribute, FTransformCollection::TransformGroup)
 	{}
 
 	bool FCollectionHierarchyFacade::IsValid() const
diff --git a/Engine/Source/Runtime/Experimental/Chaos/Private/GeometryCollection/TransformCollection.cpp b/Engine/Source/Runtime/Experimental/Chaos/Private/GeometryCollection/TransformCollection.cpp
--- a/Engine/Source/Runtime/Experimental/Chaos/Private/GeometryCollection/TransformCollection.cpp
+++ b/Engine/Source/Runtime/Experimental/Chaos/Private/GeometryCollection/TransformCollection.cpp
@@ -12,6 +12,7 @@ const FName FTransformCollection::TransformAttribute = "Transform";
 const FName FTransformCollection::ParentAttribute = "Parent";
 const FName FTransformCollection::ChildrenAttribute = "Children";
 const FName FTransformCollection::ParticlesAttribute = "Particles";
+const FName FTransformCollection::LevelAttribute = "Level";
 
 FTransformCollection::FTransformCollection()
 	: FManagedArrayCollection()
diff --git a/Engine/Source/Runtime/Experimental/Chaos/Public/GeometryCollection/TransformCollection.h b/Engine/Source/Runtime/Experimental/Chaos/Public/GeometryCollection/TransformCollection.h
--- a/Engine/Source/Runtime/Experimental/Chaos/Public/GeometryCollection/TransformCollection.h
+++ b/Engine/Source/Runtime/Experimental/Chaos/Public/GeometryCollection/TransformCollection.h
@@ -50,6 +50,7 @@ public:
 	static const FName ParentAttribute;
 	static const FName ChildrenAttribute;
 	static const FName ParticlesAttribute;
+	static const FName LevelAttribute;
 
 
 	/** Schema */
