GeometryCollection : add extra parameters to SetNotofyCrumblings

diff --git a/Engine/Source/Runtime/Experimental/GeometryCollectionEngine/Private/GeometryCollection/GeometryCollectionComponent.cpp b/Engine/Source/Runtime/Experimental/GeometryCollectionEngine/Private/GeometryCollection/GeometryCollectionComponent.cpp
--- a/Engine/Source/Runtime/Experimental/GeometryCollectionEngine/Private/GeometryCollection/GeometryCollectionComponent.cpp
+++ b/Engine/Source/Runtime/Experimental/GeometryCollectionEngine/Private/GeometryCollection/GeometryCollectionComponent.cpp
@@ -746,12 +746,13 @@ void UGeometryCollectionComponent::SetNotifyRemovals(bool bNewNotifyRemovals)
 	}
 }
 
-void UGeometryCollectionComponent::SetNotifyCrumblings(bool bNewNotifyCrumblings)
+void UGeometryCollectionComponent::SetNotifyCrumblings(bool bNewNotifyCrumblings, bool bNewCrumblingEventIncludesChildren)
 {
 	if (bNotifyCrumblings != bNewNotifyCrumblings)
 	{
 		bNotifyCrumblings = bNewNotifyCrumblings;
-		UpdateBreakEventRegistration();
+		bCrumblingEventIncludesChildren = bNewCrumblingEventIncludesChildren;
+		UpdateCrumblingEventRegistration();
 	}
 }
 
diff --git a/Engine/Source/Runtime/Experimental/GeometryCollectionEngine/Public/GeometryCollection/GeometryCollectionComponent.h b/Engine/Source/Runtime/Experimental/GeometryCollectionEngine/Public/GeometryCollection/GeometryCollectionComponent.h
--- a/Engine/Source/Runtime/Experimental/GeometryCollectionEngine/Public/GeometryCollection/GeometryCollectionComponent.h
+++ b/Engine/Source/Runtime/Experimental/GeometryCollectionEngine/Public/GeometryCollection/GeometryCollectionComponent.h
@@ -906,7 +906,7 @@ public:
 
 	/** Changes whether or not this component will get future break notifications. */
 	UFUNCTION(BlueprintCallable, Category = "Physics")
-	void SetNotifyCrumblings(bool bNewNotifyCrumblings);
+	void SetNotifyCrumblings(bool bNewNotifyCrumblings, bool bNewCrumblingEventIncludesChildren = false);
 	
 	/** Overrideable native notification */
 	virtual void NotifyBreak(const FChaosBreakEvent& Event) {};
