Add option in GeometryCollectionComponent to skip UpdateNavigation in Tick and return simplifed collision for navigation data

diff --git a/Engine/Source/Runtime/Experimental/GeometryCollectionEngine/Private/GeometryCollection/GeometryCollectionComponent.cpp b/Engine/Source/Runtime/Experimental/GeometryCollectionEngine/Private/GeometryCollection/GeometryCollectionComponent.cpp
--- a/Engine/Source/Runtime/Experimental/GeometryCollectionEngine/Private/GeometryCollection/GeometryCollectionComponent.cpp
+++ b/Engine/Source/Runtime/Experimental/GeometryCollectionEngine/Private/GeometryCollection/GeometryCollectionComponent.cpp
@@ -2,6 +2,7 @@
 
 #include "GeometryCollection/GeometryCollectionComponent.h"
 
+#include "AI/Navigation/NavCollisionBase.h"
 #include "AI/NavigationSystemHelpers.h"
 #include "Async/ParallelFor.h"
 #include "Chaos/ChaosPhysicalMaterial.h"
@@ -307,6 +308,8 @@ UGeometryCollectionComponent::UGeometryCollectionComponent(const FObjectInitiali
 	, bNotifyCrumblings(false)
 	, bStoreVelocities(false)
 	, bShowBoneColors(false)
+	, bUseRootProxyForNavigation(false)
+	, bUpdateNavigationInTick(true) 
 #if WITH_EDITORONLY_DATA 
 	, bEnableRunTimeDataCollection(false)
 	, RunTimeDataCollectionGuid(FGuid::NewGuid())
@@ -926,6 +929,29 @@ bool UGeometryCollectionComponent::DoCustomNavigableGeometryExport(FNavigableGeo
 		return false;
 	}
 
+	if (bUseRootProxyForNavigation)
+	{
+		if (const UStaticMesh* ProxyMesh = Cast<UStaticMesh>(RestCollection->RootProxy.TryLoad()))
+		{
+			const FTransform& CompToWorld = GetComponentToWorld();
+			const FVector Scale3D = CompToWorld.GetScale3D();
+			if (!Scale3D.IsZero())
+			{
+				if (const UNavCollisionBase* NavCollision = ProxyMesh->GetNavCollision())
+				{
+					const bool bHasData = NavCollision->ExportGeometry(CompToWorld, GeomExport);
+					if (bHasData)
+					{
+						// skip default export
+						return false;
+					}
+				}
+			}
+		}
+
+		return true;
+	}
+
 	TArray<FVector> OutVertexBuffer;
 	TArray<int32> OutIndexBuffer;
 
@@ -2316,14 +2342,17 @@ void UGeometryCollectionComponent::TickComponent(float DeltaTime, enum ELevelTic
 				MarkRenderDynamicDataDirty();
 				bRenderStateDirty = false;
 
-				const UWorld* MyWorld = GetWorld();
-				if (MyWorld && MyWorld->IsGameWorld())
+				if (bUpdateNavigationInTick)
 				{
-					//cycle every 0xff frames
-					//@todo - Need way of seeing if the collection is actually changing
-					if (bNavigationRelevant && bRegistered && (((GFrameCounter + NavmeshInvalidationTimeSliceIndex) & 0xff) == 0))
+					const UWorld* MyWorld = GetWorld();
+					if (MyWorld && MyWorld->IsGameWorld())
 					{
-						UpdateNavigationData();
+						//cycle every 0xff frames
+						//@todo - Need way of seeing if the collection is actually changing
+						if (bNavigationRelevant && bRegistered && (((GFrameCounter + NavmeshInvalidationTimeSliceIndex) & 0xff) == 0))
+						{
+							UpdateNavigationData();
+						}
 					}
 				}
 			}
diff --git a/Engine/Source/Runtime/Experimental/GeometryCollectionEngine/Public/GeometryCollection/GeometryCollectionComponent.h b/Engine/Source/Runtime/Experimental/GeometryCollectionEngine/Public/GeometryCollection/GeometryCollectionComponent.h
--- a/Engine/Source/Runtime/Experimental/GeometryCollectionEngine/Public/GeometryCollection/GeometryCollectionComponent.h
+++ b/Engine/Source/Runtime/Experimental/GeometryCollectionEngine/Public/GeometryCollection/GeometryCollectionComponent.h
@@ -1016,6 +1016,12 @@ protected:
 	UPROPERTY(EditAnywhere, BlueprintReadWrite, Category = "ChaosPhysics|General")
 	bool bShowBoneColors;
 
+	UPROPERTY(EditAnywhere, BlueprintReadWrite, Category = Collision, AdvancedDisplay, config)
+	bool bUseRootProxyForNavigation;
+
+	UPROPERTY(EditAnywhere, BlueprintReadWrite, Category = Collision, AdvancedDisplay, config)
+	bool bUpdateNavigationInTick;
+
 #if WITH_EDITORONLY_DATA
 	UPROPERTY(EditAnywhere, Category = "ChaosPhysics|General")
 	bool bEnableRunTimeDataCollection;
