Geometry collection : Added max level support to anchoring by box BP functions

diff --git a/Engine/Source/Runtime/Experimental/Chaos/Private/GeometryCollection/Facades/CollectionHierarchyFacade.cpp b/Engine/Source/Runtime/Experimental/Chaos/Private/GeometryCollection/Facades/CollectionHierarchyFacade.cpp
--- a/Engine/Source/Runtime/Experimental/Chaos/Private/GeometryCollection/Facades/CollectionHierarchyFacade.cpp
+++ b/Engine/Source/Runtime/Experimental/Chaos/Private/GeometryCollection/Facades/CollectionHierarchyFacade.cpp
@@ -55,6 +55,15 @@ namespace Chaos::Facades
 		return GetRootIndices(ParentAttribute);
 	}
 
+	int32 FCollectionHierarchyFacade::GetInitialLevel(int32 TransformIndex) const
+	{
+		if (HasLevelAttribute())
+		{
+			return LevelAttribute.Get()[TransformIndex];
+		}
+		return INDEX_NONE;
+	}
+
 	void FCollectionHierarchyFacade::GenerateLevelAttribute()
 	{
 		check(!IsConst());
diff --git a/Engine/Source/Runtime/Experimental/Chaos/Private/PhysicsProxy/GeometryCollectionPhysicsProxy.cpp b/Engine/Source/Runtime/Experimental/Chaos/Private/PhysicsProxy/GeometryCollectionPhysicsProxy.cpp
--- a/Engine/Source/Runtime/Experimental/Chaos/Private/PhysicsProxy/GeometryCollectionPhysicsProxy.cpp
+++ b/Engine/Source/Runtime/Experimental/Chaos/Private/PhysicsProxy/GeometryCollectionPhysicsProxy.cpp
@@ -28,6 +28,7 @@
 #include "GeometryCollection/GeometryCollectionSizeSpecificUtility.h"
 #include "GeometryCollection/GeometryCollectionSimulationTypes.h"
 #include "GeometryCollection/ManagedArrayCollection.h"
+#include "GeometryCollection/Facades/CollectionHierarchyFacade.h"
 #include "Modules/ModuleManager.h"
 #include "Chaos/PullPhysicsDataImp.h"
 #include "Chaos/PBDRigidsEvolution.h"
@@ -1818,30 +1819,42 @@ void FGeometryCollectionPhysicsProxy::SetAnchoredByIndex_External(int32 Index, b
 	}
 }
 
-void FGeometryCollectionPhysicsProxy::SetAnchoredByTransformedBox_External(const FBox& Box, const FTransform& Transform, bool bAnchored)
+void FGeometryCollectionPhysicsProxy::SetAnchoredByTransformedBox_External(const FBox& Box, const FTransform& Transform, bool bAnchored, int32 MaxLevel)
 {
 	check(IsInGameThread());
 	if (Chaos::FPhysicsSolver* RBDSolver = GetSolver<Chaos::FPhysicsSolver>())
 	{
 		Chaos::FAABB3 BoundsToCheck(Box.Min, Box.Max);
 		Chaos::FRigidTransform3 BoxTransform(Transform);
-		RBDSolver->EnqueueCommandImmediate([this, BoundsToCheck, BoxTransform, bAnchored, RBDSolver]()
+		RBDSolver->EnqueueCommandImmediate([this, BoundsToCheck, BoxTransform, bAnchored, MaxLevel, RBDSolver]()
 			{
 				using namespace Chaos;
 				TSet<FPBDRigidClusteredParticleHandle*> TopParentHandles;
 
+				Facades::FCollectionHierarchyFacade HierarchyFacade(PhysicsThreadCollection);
+				const int32 MaxLevelToCheck = (HierarchyFacade.HasLevelAttribute()) ? MaxLevel : INDEX_NONE;
+
 				FPBDRigidsEvolution* Evolution = RBDSolver->GetEvolution();
-				for (FClusterHandle* ParticleHandle : GetSolverParticleHandles())
+				for (int32 TransformIndex = 0; TransformIndex < SolverParticleHandles.Num(); TransformIndex++)
 				{
-					if (ParticleHandle)
+					if (FClusterHandle * ParticleHandle = SolverParticleHandles[TransformIndex])
 					{
-						const FVec3 PositionInBoxSpace = BoxTransform.InverseTransformPosition(ParticleHandle->X());
-						if (BoundsToCheck.Contains(PositionInBoxSpace))
+						bool bCheckBounds = true;
+						if (MaxLevelToCheck > INDEX_NONE)
+						{
+							const int32 InitialLevel = HierarchyFacade.GetInitialLevel(TransformIndex);
+							bCheckBounds = (InitialLevel <= MaxLevelToCheck);
+						}
+						if (bCheckBounds)
 						{
-							SetParticleAnchored_Internal(Evolution, ParticleHandle, bAnchored);
-							if (Chaos::FPBDRigidClusteredParticleHandle* TopParentHandle = GetTopActiveClusteredParent_Internal(ParticleHandle))
+							const FVec3 PositionInBoxSpace = BoxTransform.InverseTransformPosition(ParticleHandle->X());
+							if (BoundsToCheck.Contains(PositionInBoxSpace))
 							{
-								TopParentHandles.Add(TopParentHandle);
+								SetParticleAnchored_Internal(Evolution, ParticleHandle, bAnchored);
+								if (Chaos::FPBDRigidClusteredParticleHandle* TopParentHandle = GetTopActiveClusteredParent_Internal(ParticleHandle))
+								{
+									TopParentHandles.Add(TopParentHandle);
+								}
 							}
 						}
 					}
diff --git a/Engine/Source/Runtime/Experimental/Chaos/Public/GeometryCollection/Facades/CollectionHierarchyFacade.h b/Engine/Source/Runtime/Experimental/Chaos/Public/GeometryCollection/Facades/CollectionHierarchyFacade.h
--- a/Engine/Source/Runtime/Experimental/Chaos/Public/GeometryCollection/Facades/CollectionHierarchyFacade.h
+++ b/Engine/Source/Runtime/Experimental/Chaos/Public/GeometryCollection/Facades/CollectionHierarchyFacade.h
@@ -44,6 +44,12 @@ namespace Chaos::Facades
 		/** Get the root indicies */
 		TArray<int32> GetRootIndices() const;
 
+		/** 
+		* Get initial level of a specific transform index 
+		* If the attribute is missing return INDEX_NONE
+		*/
+		int32 GetInitialLevel(int32 TransformIndex) const;
+
 		/** 
 		* Update level attribute for all elements (and create it if it is missing ) 
 		* @Param PersistencePolicy whether to make the attribute persistent or keep the existing state
diff --git a/Engine/Source/Runtime/Experimental/Chaos/Public/PhysicsProxy/GeometryCollectionPhysicsProxy.h b/Engine/Source/Runtime/Experimental/Chaos/Public/PhysicsProxy/GeometryCollectionPhysicsProxy.h
--- a/Engine/Source/Runtime/Experimental/Chaos/Public/PhysicsProxy/GeometryCollectionPhysicsProxy.h
+++ b/Engine/Source/Runtime/Experimental/Chaos/Public/PhysicsProxy/GeometryCollectionPhysicsProxy.h
@@ -392,7 +392,7 @@ public:
 	void BreakClusters_External(TArray<FGeometryCollectionItemIndex>&& ItemIndices);
 	void BreakActiveClusters_External();
 	void SetAnchoredByIndex_External(int32 Index, bool bAnchored);
-	void SetAnchoredByTransformedBox_External(const FBox& Box, const FTransform& Transform, bool bAnchored);
+	void SetAnchoredByTransformedBox_External(const FBox& Box, const FTransform& Transform, bool bAnchored, int32 MaxLevel = INDEX_NONE);
 	void RemoveAllAnchors_External();
 	void ApplyExternalStrain_External(FGeometryCollectionItemIndex ItemIndex, const FVector& WorldLocation, float Radius, int32 PropagationDepth, float PropagationFactor, float StrainValue);
 	void ApplyInternalStrain_External(FGeometryCollectionItemIndex ItemIndex, const FVector& WorldLocation, float Radius, int32 PropagationDepth, float PropagationFactor, float StrainValue);
diff --git a/Engine/Source/Runtime/Experimental/GeometryCollectionEngine/Private/GeometryCollection/GeometryCollectionComponent.cpp b/Engine/Source/Runtime/Experimental/GeometryCollectionEngine/Private/GeometryCollection/GeometryCollectionComponent.cpp
--- a/Engine/Source/Runtime/Experimental/GeometryCollectionEngine/Private/GeometryCollection/GeometryCollectionComponent.cpp
+++ b/Engine/Source/Runtime/Experimental/GeometryCollectionEngine/Private/GeometryCollection/GeometryCollectionComponent.cpp
@@ -4450,16 +4450,16 @@ void UGeometryCollectionComponent::SetAnchoredByIndex(int32 Index, bool bAnchore
 	}
 }
 
-void UGeometryCollectionComponent::SetAnchoredByBox(FBox WorldSpaceBox, bool bAnchored)
+void UGeometryCollectionComponent::SetAnchoredByBox(FBox WorldSpaceBox, bool bAnchored, int32 MaxLevel)
 {
-	SetAnchoredByTransformedBox(WorldSpaceBox, FTransform::Identity, bAnchored);
+	SetAnchoredByTransformedBox(WorldSpaceBox, FTransform::Identity, bAnchored, MaxLevel);
 }
 
-void UGeometryCollectionComponent::SetAnchoredByTransformedBox(FBox Box, FTransform Transform, bool bAnchored)
+void UGeometryCollectionComponent::SetAnchoredByTransformedBox(FBox Box, FTransform Transform, bool bAnchored, int32 MaxLevel)
 {
 	if (PhysicsProxy)
 	{
-		PhysicsProxy->SetAnchoredByTransformedBox_External(Box, Transform, bAnchored);
+		PhysicsProxy->SetAnchoredByTransformedBox_External(Box, Transform, bAnchored, MaxLevel);
 	}
 }
 void UGeometryCollectionComponent::RemoveAllAnchors()
diff --git a/Engine/Source/Runtime/Experimental/GeometryCollectionEngine/Public/GeometryCollection/GeometryCollectionComponent.h b/Engine/Source/Runtime/Experimental/GeometryCollectionEngine/Public/GeometryCollection/GeometryCollectionComponent.h
--- a/Engine/Source/Runtime/Experimental/GeometryCollectionEngine/Public/GeometryCollection/GeometryCollectionComponent.h
+++ b/Engine/Source/Runtime/Experimental/GeometryCollectionEngine/Public/GeometryCollection/GeometryCollectionComponent.h
@@ -560,13 +560,13 @@ public:
 	* Set all pieces within a world space bounding box to be anchored or not 
 	*/
 	UFUNCTION(BlueprintCallable, Category = "ChaosPhysics")
-	void SetAnchoredByBox(FBox WorldSpaceBox, bool bAnchored);
+	void SetAnchoredByBox(FBox WorldSpaceBox, bool bAnchored, int32 MaxLevel = -1);
 
 	/**
 	* Set all pieces within a world transformed bounding box to be anchored or not
 	*/
 	UFUNCTION(BlueprintCallable, Category = "ChaosPhysics")
-	void SetAnchoredByTransformedBox(FBox Box, FTransform Transform, bool bAnchored);
+	void SetAnchoredByTransformedBox(FBox Box, FTransform Transform, bool bAnchored, int32 MaxLevel = -1);
 
 	/**
 	* this will remove anchors on all the pieces ( including the static and kinematic initial states ones ) of the geometry colection
