Fix a bug in Tetrahedron::ProjectToSurface().

diff --git a/Engine/Source/Runtime/Experimental/Chaos/Public/Chaos/Tetrahedron.h b/Engine/Source/Runtime/Experimental/Chaos/Public/Chaos/Tetrahedron.h
--- a/Engine/Source/Runtime/Experimental/Chaos/Public/Chaos/Tetrahedron.h
+++ b/Engine/Source/Runtime/Experimental/Chaos/Public/Chaos/Tetrahedron.h
@@ -221,7 +221,7 @@ namespace Chaos {
 				TVec3<T> SurfacePoint(Location);
 				for (int32 i = 0; i < 4; i++)
 				{
-					if (Tris[i].GetPlane().SignedDistance(Location) > static_cast<T>(0.0))
+					if (Tris[i].GetPlane().SignedDistance(SurfacePoint) > static_cast<T>(0.0))
 					{
 						const TVec3<T> Norm = Tris[i].GetNormal();
 						SurfacePoint -= TVec3<T>::DotProduct(Location - Tris[i][0], Norm) * Norm;
