Chaos Cloth - Added a third type of property collection adapter class to make the default adapter non mutable, so that properties cannot be accidentally added/removed by code while the cloth is being simulated.

#rnx

diff --git a/Engine/Source/Runtime/Experimental/Chaos/Private/Chaos/PropertyCollectionAdapter.cpp b/Engine/Source/Runtime/Experimental/Chaos/Private/Chaos/PropertyCollectionAdapter.cpp
--- a/Engine/Source/Runtime/Experimental/Chaos/Private/Chaos/PropertyCollectionAdapter.cpp
+++ b/Engine/Source/Runtime/Experimental/Chaos/Private/Chaos/PropertyCollectionAdapter.cpp
@@ -82,18 +82,12 @@ namespace Chaos::Softs
 	FPropertyCollectionAdapter::FPropertyCollectionAdapter(const TSharedPtr<FManagedArrayCollection>& InManagedArrayCollection)
 		: FPropertyCollectionConstAdapter(InManagedArrayCollection, NoInit)
 	{
-		Construct();
 		Initialize();
 	}
-
-	void FPropertyCollectionAdapter::Construct()
+	
+	FPropertyCollectionAdapter::FPropertyCollectionAdapter(const TSharedPtr<FManagedArrayCollection>& InManagedArrayCollection, ENoInit)
+		: FPropertyCollectionConstAdapter(InManagedArrayCollection, NoInit)
 	{
-		// ClothProperty Group
-		GetManagedArrayCollection()->AddAttribute<FString>(KeyName, PropertyGroup);
-		GetManagedArrayCollection()->AddAttribute<FVector3f>(LowValueName, PropertyGroup);
-		GetManagedArrayCollection()->AddAttribute<FVector3f>(HighValueName, PropertyGroup);
-		GetManagedArrayCollection()->AddAttribute<FString>(StringValueName, PropertyGroup);
-		GetManagedArrayCollection()->AddAttribute<uint8>(FlagsName, PropertyGroup);
 	}
 
 	void FPropertyCollectionAdapter::EnableFlag(int32 KeyIndex, EPropertyFlag Flag, bool bEnable)
@@ -118,7 +112,24 @@ namespace Chaos::Softs
 		return KeyIndex;
 	}
 
-	int32 FPropertyCollectionAdapter::AddProperty(const FString& Key, bool bEnabled, bool bAnimatable)
+	FPropertyCollectionMutableAdapter::FPropertyCollectionMutableAdapter(const TSharedPtr<FManagedArrayCollection>& InManagedArrayCollection)
+		: FPropertyCollectionAdapter(InManagedArrayCollection, NoInit)
+	{
+		Construct();
+		Initialize();
+	}
+
+	void FPropertyCollectionMutableAdapter::Construct()
+	{
+		// ClothProperty Group
+		GetManagedArrayCollection()->AddAttribute<FString>(KeyName, PropertyGroup);
+		GetManagedArrayCollection()->AddAttribute<FVector3f>(LowValueName, PropertyGroup);
+		GetManagedArrayCollection()->AddAttribute<FVector3f>(HighValueName, PropertyGroup);
+		GetManagedArrayCollection()->AddAttribute<FString>(StringValueName, PropertyGroup);
+		GetManagedArrayCollection()->AddAttribute<uint8>(FlagsName, PropertyGroup);
+	}
+
+	int32 FPropertyCollectionMutableAdapter::AddProperty(const FString& Key, bool bEnabled, bool bAnimatable)
 	{
 		const int32 Index = GetManagedArrayCollection()->AddElements(1, PropertyGroup);
 		const uint8 Flags = (bEnabled ? (uint8)EPropertyFlag::Enabled : 0) | (bAnimatable ? (uint8)EPropertyFlag::Animatable : 0);
@@ -137,7 +148,7 @@ namespace Chaos::Softs
 		return Index;
 	}
 
-	int32 FPropertyCollectionAdapter::AddProperties(const TArray<FString>& Keys, bool bEnabled, bool bAnimatable)
+	int32 FPropertyCollectionMutableAdapter::AddProperties(const TArray<FString>& Keys, bool bEnabled, bool bAnimatable)
 	{
 		if (const int32 NumProperties = Keys.Num())
 		{
@@ -162,4 +173,19 @@ namespace Chaos::Softs
 		}
 		return INDEX_NONE;
 	}
+
+	void FPropertyCollectionMutableAdapter::Append(const TSharedPtr<const FManagedArrayCollection>& InManagedArrayCollection)
+	{
+		TArray<FName> GroupsToSkip = InManagedArrayCollection->GroupNames();
+		GroupsToSkip.RemoveSingleSwap(PropertyGroup);
+
+		InManagedArrayCollection->CopyTo(GetManagedArrayCollection().Get(), GroupsToSkip);
+	}
+
+	void FPropertyCollectionMutableAdapter::Copy(const TSharedPtr<const FManagedArrayCollection>& InManagedArrayCollection)
+	{
+		GetManagedArrayCollection()->Reset();
+		Append(InManagedArrayCollection);
+	}
+
 }  // End namespace Chaos
diff --git a/Engine/Source/Runtime/Experimental/Chaos/Public/Chaos/PropertyCollectionAdapter.h b/Engine/Source/Runtime/Experimental/Chaos/Public/Chaos/PropertyCollectionAdapter.h
--- a/Engine/Source/Runtime/Experimental/Chaos/Public/Chaos/PropertyCollectionAdapter.h
+++ b/Engine/Source/Runtime/Experimental/Chaos/Public/Chaos/PropertyCollectionAdapter.h
@@ -37,7 +37,7 @@
 		explicit FPropertyCollectionConstAdapter(const TSharedPtr<const FManagedArrayCollection>& InManagedArrayCollection);
 		virtual ~FPropertyCollectionConstAdapter() = default;
 
-		/** return the collection used by this adapter as a property collection. */
+		/** Return the collection used by this adapter as a property collection. */
 		const TSharedPtr<const FManagedArrayCollection>& GetManagedArrayCollection() const { return ManagedArrayCollection; }
 
 		/** Return the number of cloth properties in this collection. */
@@ -132,24 +132,17 @@
 	};
 
 	/**
-	 * Property collection adapter class to use a ManagedArrayCollection property collection and add simulation properties.
-	 * Note: Int properties are limited to 24 bits (-16777215 to 16777215).
+	 * Property collection adapter class to use a ManagedArrayCollection property collection and change simulation properties.
+	 * This adaptor is immutable and properties cannot be added from it.
+	 * Use a FPropertyCollectionMutableAdapter to initialize a new property colleciton or add new properties to the collection.
+	 * Note: Int property values are limited to 24 bits (-16777215 to 16777215).
 	 */
-	class FPropertyCollectionAdapter final : public FPropertyCollectionConstAdapter
+	class FPropertyCollectionAdapter : public FPropertyCollectionConstAdapter
 	{
 	public:
-		explicit FPropertyCollectionAdapter(const TSharedPtr<FManagedArrayCollection>& ManagedArrayCollection);
+		explicit FPropertyCollectionAdapter(const TSharedPtr<FManagedArrayCollection>& InManagedArrayCollection);
 		virtual ~FPropertyCollectionAdapter() = default;
 
-		/** Return the collection used by this adapter as a property collection. */
-		TSharedPtr<FManagedArrayCollection> GetManagedArrayCollection() { return ConstCastSharedPtr<FManagedArrayCollection>(ManagedArrayCollection); }
-
-		/** Add a single property, and return its index. */
-		int32 AddProperty(const FString& Key, bool bEnabled = true, bool bAnimatable = false);
-
-		/** Add new properties, and return the index of the first added property. */
-		int32 AddProperties(const TArray<FString>& Keys, bool bEnabled = true, bool bAnimatable = false);
-
 		//~ Values set per index
 		template<typename T, TEMPLATE_REQUIRES(TIsWeightedType<T>::Value)>
 		void SetLowValue(int32 KeyIndex, const T& Value) { GetLowValueArray()[KeyIndex] = FVector3f(Value); }
@@ -166,7 +159,7 @@
 		void SetValue(int32 KeyIndex, const T& Value) { SetWeightedValue(KeyIndex, Value, Value); }
 
 		void SetStringValue(int32 KeyIndex, const FString& Value) { GetStringValueArray()[KeyIndex] = Value; }
-		
+
 		void SetFlags(int32 KeyIndex, uint8 Flags) { GetFlagsArray()[KeyIndex] = Flags; }
 		void SetEnabled(int32 KeyIndex, bool bEnabled) { EnableFlag(KeyIndex, EPropertyFlag::Enabled, bEnabled); }
 		void SetAnimatable(int32 KeyIndex, bool bAnimatable) { EnableFlag(KeyIndex, EPropertyFlag::Animatable, bAnimatable); }
@@ -188,12 +181,66 @@
 		int32 SetValue(const FString& Key, const T& Value) { return SetWeightedValue(Key, Value, Value); }
 
 		int32 SetStringValue(const FString& Key, const FString& Value) { return SetValue(Key, GetStringValueArray(), Value); }
-		
+
 		int32 SetFlags(const FString& Key, uint8 Flags) { return SetValue(Key, GetFlagsArray(), Flags); }
 		int32 SetEnabled(const FString& Key, bool bEnabled) { return EnableFlag(Key, EPropertyFlag::Enabled, bEnabled); }
 		int32 SetAnimatable(const FString& Key, bool bAnimatable) { return EnableFlag(Key, EPropertyFlag::Animatable, bAnimatable); }
 		int32 SetDirty(const FString& Key, bool bDirty) { return EnableFlag(Key, EPropertyFlag::Animatable, bDirty); }
 
+	protected:
+		// No init constructor for FPropertyCollectionMutableAdapter
+		FPropertyCollectionAdapter(const TSharedPtr<FManagedArrayCollection>& InManagedArrayCollection, ENoInit);
+
+		// Access to a writeable ManagedArrayCollection is protected, use an FPropertyCollectionMutableAdapter if needed to get a non const pointer
+		TSharedPtr<FManagedArrayCollection> GetManagedArrayCollection() { return ConstCastSharedPtr<FManagedArrayCollection>(ManagedArrayCollection); }
+
+		const TArrayView<FString>& GetKeyArray() { return reinterpret_cast<TArrayView<FString>&>(KeyArray); }
+		const TArrayView<FVector3f>& GetLowValueArray() { return reinterpret_cast<TArrayView<FVector3f>&>(LowValueArray); }
+		const TArrayView<FVector3f>& GetHighValueArray() { return reinterpret_cast<TArrayView<FVector3f>&>(HighValueArray); }
+		const TArrayView<FString>& GetStringValueArray() { return reinterpret_cast<TArrayView<FString>&>(StringValueArray); }
+		const TArrayView<uint8>& GetFlagsArray() { return reinterpret_cast<const TArrayView<uint8>&>(FlagsArray); }
+
+	private:
+		template<typename T>
+		inline int32 SetValue(const FString& Key, const TArrayView<T>& ValueArray, const T& Value);
+
+		void EnableFlag(int32 KeyIndex, EPropertyFlag Flag, bool bEnable);
+		int32 EnableFlag(const FString& Key, EPropertyFlag Flag, bool bEnable);
+	};
+
+	/**
+	 * Property collection mutable adapter class to use a ManagedArrayCollection property collection and add/change simulation properties.
+	 * Note: Int property values are limited to 24 bits (-16777215 to 16777215).
+	 */
+	class FPropertyCollectionMutableAdapter final : public FPropertyCollectionAdapter
+	{
+	public:
+		explicit FPropertyCollectionMutableAdapter(const TSharedPtr<FManagedArrayCollection>& InManagedArrayCollection);
+		virtual ~FPropertyCollectionMutableAdapter() = default;
+
+		/** Return the collection used by this adapter as a property collection. */
+		using FPropertyCollectionAdapter::GetManagedArrayCollection;
+
+		/** Add a single property, and return its index. */
+		int32 AddProperty(const FString& Key, bool bEnabled = true, bool bAnimatable = false);
+
+		/** Add new properties, and return the index of the first added property. */
+		int32 AddProperties(const TArray<FString>& Keys, bool bEnabled = true, bool bAnimatable = false);
+
+		/**
+		 * Append all properties and values from an existing collection to this property collection.
+		 * This won't copy any other groups, only data from PropertyGroup.
+		 * Any pre-existing data will be preserved.
+		 */
+		void Append(const TSharedPtr<const FManagedArrayCollection>& InManagedArrayCollection);
+
+		/**
+		 * Copy all properties and values from an existing collection to this property collection.
+		 * This won't copy any other groups, only data from PropertyGroup.
+		 * Any pre-xisting data will be removed/replaced.
+		 */
+		void Copy(const TSharedPtr<const FManagedArrayCollection>& InManagedArrayCollection);
+
 		//~ Add values
 		template<typename T, TEMPLATE_REQUIRES(TIsWeightedType<T>::Value)>
 		inline int32 AddWeightedValue(const FString& Key, const T& LowValue, const T& HighValue, bool bEnabled = true, bool bAnimatable = false);
@@ -201,25 +248,13 @@
 		inline int32 AddWeightedFloatValue(const FString& Key, const FVector2f& Value, bool bEnabled, bool bAnimatable);
 
 		template<typename T, TEMPLATE_REQUIRES(TIsWeightedType<T>::Value)>
-		int32 AddValue(const FString& Key, const T& Value, bool bEnabled = true, bool bAnimatable = false) { return AddWeightedValue(Value, Value, bEnabled, bAnimatable); }
+		int32 AddValue(const FString& Key, const T& Value, bool bEnabled = true, bool bAnimatable = false) { return AddWeightedValue(Key, Value, Value, bEnabled, bAnimatable); }
 
 		inline int32 AddStringValue(const FString& Key, const FString& Value, bool bEnabled = true, bool bAnimatable = false);
-		
+
 	private:
 		// Add the property group to the target collection
 		void Construct();
-
-		template<typename T>
-		inline int32 SetValue(const FString& Key, const TArrayView<T>& ValueArray, const T& Value);
-
-		void EnableFlag(int32 KeyIndex, EPropertyFlag Flag, bool bEnable);
-		int32 EnableFlag(const FString& Key, EPropertyFlag Flag, bool bEnable);
-
-		const TArrayView<FString>& GetKeyArray() { return reinterpret_cast<TArrayView<FString>&>(KeyArray); }
-		const TArrayView<FVector3f>& GetLowValueArray() { return reinterpret_cast<TArrayView<FVector3f>&>(LowValueArray); }
-		const TArrayView<FVector3f>& GetHighValueArray() { return reinterpret_cast<TArrayView<FVector3f>&>(HighValueArray); }
-		const TArrayView<FString>& GetStringValueArray() { return reinterpret_cast<TArrayView<FString>&>(StringValueArray); }
-		const TArrayView<uint8>& GetFlagsArray() { return reinterpret_cast<const TArrayView<uint8>&>(FlagsArray); }
 	};
 
 	template<typename T, typename TEnableIf<TIsWeightedType<T>::Value, int>::type>
@@ -276,28 +311,6 @@
 		return KeyIndex;
 	}
 
-	template<typename T, typename TEnableIf<TIsWeightedType<T>::Value, int>::type>
-	inline int32 FPropertyCollectionAdapter::AddWeightedValue(const FString& Key, const T& LowValue, const T& HighValue, bool bEnabled, bool bAnimatable)
-	{
-		const int32 KeyIndex = AddProperty(Key, bEnabled, bAnimatable);
-		SetWeightedValue(KeyIndex, LowValue, HighValue);
-		return KeyIndex;
-	}
-
-	inline int32 FPropertyCollectionAdapter::AddWeightedFloatValue(const FString& Key, const FVector2f& Value, bool bEnabled, bool bAnimatable)
-	{
-		const int32 KeyIndex = AddProperty(Key, bEnabled, bAnimatable);
-		SetWeightedFloatValue(KeyIndex, Value);
-		return KeyIndex;
-	}
-
-	inline int32 FPropertyCollectionAdapter::AddStringValue(const FString& Key, const FString& Value, bool bEnabled, bool bAnimatable)
-	{
-		const int32 KeyIndex = AddProperty(Key, bEnabled, bAnimatable);
-		SetStringValue(Key, Value);
-		return KeyIndex;
-	}
-
 	template<typename T>
 	inline int32 FPropertyCollectionAdapter::SetValue(const FString& Key, const TArrayView<T>& ValueArray, const T& Value)
 	{
@@ -308,6 +321,28 @@
 		}
 		return KeyIndex;
 	}
+
+	template<typename T, typename TEnableIf<TIsWeightedType<T>::Value, int>::type>
+	inline int32 FPropertyCollectionMutableAdapter::AddWeightedValue(const FString& Key, const T& LowValue, const T& HighValue, bool bEnabled, bool bAnimatable)
+	{
+		const int32 KeyIndex = AddProperty(Key, bEnabled, bAnimatable);
+		SetWeightedValue(KeyIndex, LowValue, HighValue);
+		return KeyIndex;
+	}
+
+	inline int32 FPropertyCollectionMutableAdapter::AddWeightedFloatValue(const FString& Key, const FVector2f& Value, bool bEnabled, bool bAnimatable)
+	{
+		const int32 KeyIndex = AddProperty(Key, bEnabled, bAnimatable);
+		SetWeightedFloatValue(KeyIndex, Value);
+		return KeyIndex;
+	}
+
+	inline int32 FPropertyCollectionMutableAdapter::AddStringValue(const FString& Key, const FString& Value, bool bEnabled, bool bAnimatable)
+	{
+		const int32 KeyIndex = AddProperty(Key, bEnabled, bAnimatable);
+		SetStringValue(Key, Value);
+		return KeyIndex;
+	}
 }  // End namespace Chaos
 
 // Use this macro to add shorthands for property getters and direct access through the declared key index
