Chaos Cloth - Iterated on, and simplified the recently added spherical and backstop constraint scale API.

#rnx

diff --git a/Engine/Plugins/Experimental/ChaosCloth/Source/Chaos/Private/ChaosCloth/ChaosClothConstraints.cpp b/Engine/Plugins/Experimental/ChaosCloth/Source/Chaos/Private/ChaosCloth/ChaosClothConstraints.cpp
--- a/Engine/Plugins/Experimental/ChaosCloth/Source/Chaos/Private/ChaosCloth/ChaosClothConstraints.cpp
+++ b/Engine/Plugins/Experimental/ChaosCloth/Source/Chaos/Private/ChaosCloth/ChaosClothConstraints.cpp
@@ -730,11 +730,12 @@ void FClothConstraints::SetLongRangeAttachmentProperties(
 	}
 }
 
-void FClothConstraints::SetMaximumDistanceProperties(Softs::FSolverReal MaxDistancesMultiplier)
+void FClothConstraints::SetMaximumDistanceProperties(Softs::FSolverReal MeshScale)
 {
 	if (MaximumDistanceConstraints)
 	{
-		MaximumDistanceConstraints->SetTransformScale(MaxDistancesMultiplier);
+		constexpr Softs::FSolverReal MaxDistancesMultiplier = (Softs::FSolverReal)1.;
+		MaximumDistanceConstraints->SetScale(MaxDistancesMultiplier, MeshScale);
 	}
 }
 
@@ -760,12 +761,13 @@ void FClothConstraints::SetSelfCollisionProperties(Softs::FSolverReal SelfCollis
 	}
 }
 
-void FClothConstraints::SetBackstopProperties(bool bEnabled, Softs::FSolverReal BackstopDistancesMultiplier)
+void FClothConstraints::SetBackstopProperties(bool bEnabled, Softs::FSolverReal MeshScale)
 {
 	if (BackstopConstraints)
 	{
 		BackstopConstraints->SetEnabled(bEnabled);
-		BackstopConstraints->SetTransformScale(BackstopDistancesMultiplier);
+		constexpr Softs::FSolverReal BackstopDistancesMultiplier = (Softs::FSolverReal)1.;
+		BackstopConstraints->SetScale(BackstopDistancesMultiplier, MeshScale);
 	}
 }
 
diff --git a/Engine/Source/Runtime/Experimental/Chaos/Public/Chaos/PBDSphericalConstraint.h b/Engine/Source/Runtime/Experimental/Chaos/Public/Chaos/PBDSphericalConstraint.h
--- a/Engine/Source/Runtime/Experimental/Chaos/Public/Chaos/PBDSphericalConstraint.h
+++ b/Engine/Source/Runtime/Experimental/Chaos/Public/Chaos/PBDSphericalConstraint.h
@@ -39,13 +39,13 @@ public:
 		const uint32 InParticleCount,
 		const TArray<FSolverVec3>& InAnimationPositions,  // Use global indexation (will need adding ParticleOffset)
 		const TConstArrayView<FRealSingle>& InSphereRadii,  // Use local indexation
-		const FPropertyCollectionConstAdapter& PropertyCollection
+		const FPropertyCollectionConstAdapter& PropertyCollection,
+		FSolverReal MeshScale
 	)
 		: AnimationPositions(InAnimationPositions)
 		, SphereRadii(InSphereRadii)
 		, ParticleOffset(InParticleOffset)
-		, MaxDistanceScale((FSolverReal)FMath::Max(GetMaxDistanceScale(PropertyCollection, 1.f), 0.f))
-		, TransformScale((FSolverReal)1.)
+		, Scale((FSolverReal)FMath::Max(GetMaxDistanceScale(PropertyCollection, 1.f), 0.f) * MeshScale)
 	{
 		check(InSphereRadii.Num() == InParticleCount);
 	}
@@ -59,18 +59,17 @@ public:
 		: AnimationPositions(InAnimationPositions)
 		, SphereRadii(InSphereRadii)
 		, ParticleOffset(InParticleOffset)
-		, MaxDistanceScale((FSolverReal)1.)
-		, TransformScale((FSolverReal)1.)
+		, Scale((FSolverReal)1.)
 	{
 		check(InSphereRadii.Num() == InParticleCount);
 	}
 	~FPBDSphericalConstraint() {}
 
-	void SetProperties(const FPropertyCollectionConstAdapter& PropertyCollection)
+	void SetProperties(const FPropertyCollectionConstAdapter& PropertyCollection, FSolverReal MeshScale)
 	{
 		if (IsMaxDistanceScaleMutable(PropertyCollection))
 		{
-			MaxDistanceScale = (FSolverReal)FMath::Max(GetMaxDistanceScale(PropertyCollection), 0.f);
+			SetScale((FSolverReal)GetMaxDistanceScale(PropertyCollection), MeshScale);
 		}
 	}
 
@@ -88,25 +87,25 @@ public:
 		}
 	}
 
-	void SetTransformScale(const FSolverReal InTransformScale)
+	void SetScale(FSolverReal MaxDistanceScale, FSolverReal MeshScale)
 	{
-		TransformScale = FMath::Max(InTransformScale, (FSolverReal)0.);
+		Scale = FMath::Max(MaxDistanceScale, (FSolverReal)0.) * MeshScale;
 	}
 
-	FSolverReal GetTransformScale() const { return TransformScale; }
-	FSolverReal GetScale() const { return TransformScale * MaxDistanceScale; }
+	FSolverReal GetScale() const { return Scale; }
 
-	UE_DEPRECATED(5.2, "Use SetTransformScale instead.")
-	void SetSphereRadiiMultiplier(const FSolverReal InSphereRadiiMultiplier) { SetTransformScale(InSphereRadiiMultiplier); }
+	UE_DEPRECATED(5.2, "Use SetScale instead.")
+	void SetSphereRadiiMultiplier(FSolverReal InSphereRadiiMultiplier, FSolverReal MeshScale)
+	{
+		SetScale(InSphereRadiiMultiplier, MeshScale);
+	}
 
 private:
 	void ApplyHelper(FSolverParticles& Particles, const FSolverReal Dt) const
 	{
 		const int32 ParticleCount = SphereRadii.Num();
 
-		const FSolverReal Scale = GetScale();
-
-		PhysicsParallelFor(ParticleCount, [this, &Particles, Dt, Scale](int32 Index)  // TODO: profile need for parallel loop based on particle count
+		PhysicsParallelFor(ParticleCount, [this, &Particles, Dt](int32 Index)  // TODO: profile need for parallel loop based on particle count
 		{
 			const int32 ParticleIndex = ParticleOffset + Index;
 
@@ -137,12 +136,11 @@ protected:
 	const TArray<FSolverVec3>& AnimationPositions;  // Use global indexation (will need adding ParticleOffset)
 	const TConstArrayView<FRealSingle> SphereRadii;  // Use local indexation
 	const int32 ParticleOffset;
-	UE_DEPRECATED(5.2, "Use TransformScale and MaxDistanceScale instead.")
+	UE_DEPRECATED(5.2, "Use Scale instead.")
 	FSolverReal SphereRadiiMultiplier = 1.f;
 
 private:
-	FSolverReal MaxDistanceScale;
-	FSolverReal TransformScale;
+	FSolverReal Scale;
 
 	UE_CHAOS_DECLARE_PROPERTYCOLLECTION_NAME(MaxDistanceScale, float);
 };
@@ -163,15 +161,15 @@ public:
 		const TArray<FSolverVec3>& InAnimationNormals,  // Use global indexation (will need adding ParticleOffset)
 		const TConstArrayView<FRealSingle>& InSphereRadii,  // Use local indexation
 		const TConstArrayView<FRealSingle>& InSphereOffsetDistances,  // Use local indexation
-		const FPropertyCollectionConstAdapter& PropertyCollection
+		const FPropertyCollectionConstAdapter& PropertyCollection,
+		FSolverReal MeshScale
 	)
 		: AnimationPositions(InAnimationPositions)
 		, AnimationNormals(InAnimationNormals)
 		, SphereRadii(InSphereRadii)
 		, SphereOffsetDistances(InSphereOffsetDistances)
 		, ParticleOffset(InParticleOffset)
-		, BackstopScale((FSolverReal)FMath::Max(GetBackstopScale(PropertyCollection), 0.f))
-		, TransformScale((FSolverReal)1.)
+		, Scale((FSolverReal)FMath::Max(GetBackstopScale(PropertyCollection), 0.f) * MeshScale)
 		, bEnabled(true)
 		, bUseLegacyBackstop(false)
 	{
@@ -193,8 +191,7 @@ public:
 		, SphereRadii(InSphereRadii)
 		, SphereOffsetDistances(InSphereOffsetDistances)
 		, ParticleOffset(InParticleOffset)
-		, BackstopScale((FSolverReal)1.)
-		, TransformScale((FSolverReal)1.)
+		, Scale((FSolverReal)1.)
 		, bEnabled(true)
 		, bUseLegacyBackstop(bInUseLegacyBackstop)
 	{
@@ -203,7 +200,7 @@ public:
 	}
 	~FPBDSphericalBackstopConstraint() {}
 
-	void SetProperties(const FPropertyCollectionConstAdapter& PropertyCollection)
+	void SetProperties(const FPropertyCollectionConstAdapter& PropertyCollection, FSolverReal MeshScale)
 	{
 		if (IsBackstopScaleDirty(PropertyCollection) && IsBackstopScaleAnimatable(PropertyCollection))
 		{
@@ -212,7 +209,7 @@ public:
 
 		if (IsBackstopScaleMutable(PropertyCollection))
 		{
-			BackstopScale = (FSolverReal)FMath::Max(GetBackstopScale(PropertyCollection), 0.f);
+			SetScale((FSolverReal)GetBackstopScale(PropertyCollection), MeshScale);
 		}
 	}
 
@@ -254,17 +251,18 @@ public:
 		}
 	}
 
-	void SetTransformScale(const FSolverReal InTransformScale)
+	void SetScale(FSolverReal MaxDistanceScale, FSolverReal MeshScale)
 	{
-		TransformScale = FMath::Max(InTransformScale, (FSolverReal)0.);
+		Scale = FMath::Max(MaxDistanceScale, (FSolverReal)0.) * MeshScale;
 	}
-	
-	FSolverReal GetTransformScale() const { return TransformScale; }
 
-	FSolverReal GetScale() const { return TransformScale * BackstopScale; }
+	FSolverReal GetScale() const { return Scale; }
 
-	UE_DEPRECATED(5.2, "Use SetTransformScale instead.")
-	void SetSphereRadiiMultiplier(const FSolverReal InSphereRadiiMultiplier) { SetTransformScale(InSphereRadiiMultiplier); }
+	UE_DEPRECATED(5.2, "Use SetScale instead.")
+	void SetSphereRadiiMultiplier(FSolverReal InSphereRadiiMultiplier, FSolverReal MeshScale = (FSolverReal)1.)
+	{
+		SetScale(InSphereRadiiMultiplier, MeshScale);
+	}
 
 	UE_DEPRECATED(5.2, "Use GetScale() instead.")
 	FSolverReal GetSphereRadiiMultiplier() const { return GetScale(); }
@@ -278,9 +276,8 @@ private:
 	void ApplyHelper(FSolverParticles& Particles, const FSolverReal Dt) const
 	{
 		const int32 ParticleCount = SphereRadii.Num();
-		const FSolverReal Scale = GetScale();
 
-		PhysicsParallelFor(ParticleCount, [this, &Particles, Dt, Scale](int32 Index)  // TODO: profile need for parallel loop based on particle count
+		PhysicsParallelFor(ParticleCount, [this, &Particles, Dt](int32 Index)  // TODO: profile need for parallel loop based on particle count
 		{
 			const int32 ParticleIndex = ParticleOffset + Index;
 
@@ -316,9 +313,8 @@ private:
 	void ApplyLegacyHelper(FSolverParticles& Particles, const FSolverReal Dt) const
 	{
 		const int32 ParticleCount = SphereRadii.Num();
-		const FSolverReal Scale = GetScale();
 
-		PhysicsParallelFor(ParticleCount, [this, &Particles, Dt, Scale](int32 Index)  // TODO: profile need for parallel loop based on particle count
+		PhysicsParallelFor(ParticleCount, [this, &Particles, Dt](int32 Index)  // TODO: profile need for parallel loop based on particle count
 		{
 			const int32 ParticleIndex = ParticleOffset + Index;
 
@@ -360,8 +356,7 @@ private:
 	const TConstArrayView<FRealSingle> SphereRadii; // Start at index 0, use local indexation
 	const TConstArrayView<FRealSingle> SphereOffsetDistances;  // Sphere position offsets, use local indexation
 	const int32 ParticleOffset;
-	FSolverReal BackstopScale;
-	FSolverReal TransformScale;
+	FSolverReal Scale;
 	bool bEnabled;
 	bool bUseLegacyBackstop;
 
