GeometryCollection runtime optimizations
- make sure we do properly reset the dirty flag on the dynamic collection when scene proxy is not set
- fix condition that could lead to run part of the Tick function too many time when the collection is not dirty

diff --git a/Engine/Source/Runtime/Experimental/GeometryCollectionEngine/Private/GeometryCollection/GeometryCollectionComponent.cpp b/Engine/Source/Runtime/Experimental/GeometryCollectionEngine/Private/GeometryCollection/GeometryCollectionComponent.cpp
--- a/Engine/Source/Runtime/Experimental/GeometryCollectionEngine/Private/GeometryCollection/GeometryCollectionComponent.cpp
+++ b/Engine/Source/Runtime/Experimental/GeometryCollectionEngine/Private/GeometryCollection/GeometryCollectionComponent.cpp
@@ -2323,7 +2323,7 @@ void UGeometryCollectionComponent::TickComponent(float DeltaTime, enum ELevelTic
 			// 	}
 			// }
 			
-			if (RestCollection->HasVisibleGeometry() || DynamicCollection->IsDirty())
+			if (RestCollection->HasVisibleGeometry() && DynamicCollection->IsDirty())
 			{
 				// #todo review: When we've made changes to ISMC, we need to move this function call to SetRenderDynamicData_Concurrent
 				RefreshEmbeddedGeometry();
@@ -2810,13 +2810,13 @@ void UGeometryCollectionComponent::SendRenderDynamicData_Concurrent()
 					}
 				);
 			}
-		}
+		}		
+	}
 
-		// mark collection clean now that we have rendered
-		if (DynamicCollection)
-		{
-			DynamicCollection->MakeClean();
-		}			
+	// mark collection clean now that we have rendered
+	if (DynamicCollection)
+	{
+		DynamicCollection->MakeClean();
 	}
 }
 
