Geometry collection : optimize FindAttribute methods to only lookup in teh map once per call

diff --git a/Engine/Source/Runtime/Experimental/Chaos/Public/GeometryCollection/ManagedArrayCollection.h b/Engine/Source/Runtime/Experimental/Chaos/Public/GeometryCollection/ManagedArrayCollection.h
--- a/Engine/Source/Runtime/Experimental/Chaos/Public/GeometryCollection/ManagedArrayCollection.h
+++ b/Engine/Source/Runtime/Experimental/Chaos/Public/GeometryCollection/ManagedArrayCollection.h
@@ -267,10 +267,10 @@ public:
 	template<typename T>
 	TManagedArray<T>* FindAttribute(FName Name, FName Group)
 	{
-		if (HasAttribute(Name, Group))
+		const FKeyType Key = FManagedArrayCollection::MakeMapKey(Name, Group);
+		if (FValueType* FoundValue = Map.Find(Key))
 		{
-			FKeyType Key = FManagedArrayCollection::MakeMapKey(Name, Group);
-			return static_cast<TManagedArray<T>*>(Map[Key].Value);
+			return static_cast<TManagedArray<T>*>(FoundValue->Value);
 		}
 		return nullptr;
 	};
@@ -278,10 +278,10 @@ public:
 	template<typename T>
 	const TManagedArray<T>* FindAttribute(FName Name, FName Group) const
 	{
-		if (HasAttribute(Name, Group))
+		const FKeyType Key = FManagedArrayCollection::MakeMapKey(Name, Group);
+		if (const FValueType* FoundValue = Map.Find(Key))
 		{
-			FKeyType Key = FManagedArrayCollection::MakeMapKey(Name, Group);
-			return static_cast<TManagedArray<T>*>(Map[Key].Value);
+			return static_cast<const TManagedArray<T>*>(FoundValue->Value);
 		}
 		return nullptr;
 	};
@@ -295,14 +295,12 @@ public:
 	template<typename T>
 	TManagedArray<T>* FindAttributeTyped(FName Name, FName Group)
 	{
-		if(HasAttribute(Name, Group))
+		const FKeyType Key = FManagedArrayCollection::MakeMapKey(Name, Group);
+		if (FValueType* FoundValue = Map.Find(Key))
 		{
-			FKeyType Key = FManagedArrayCollection::MakeMapKey(Name, Group);
-			FValueType& FoundValue = Map[Key];
-
-			if(FoundValue.ArrayType == ManagedArrayType<T>())
+			if(FoundValue->ArrayType == ManagedArrayType<T>())
 			{
-				return static_cast<TManagedArray<T>*>(Map[Key].Value);
+				return static_cast<TManagedArray<T>*>(FoundValue->Value);
 			}
 		}
 		return nullptr;
@@ -311,14 +309,12 @@ public:
 	template<typename T>
 	const TManagedArray<T>* FindAttributeTyped(FName Name, FName Group) const
 	{
-		if(HasAttribute(Name, Group))
+		const FKeyType Key = FManagedArrayCollection::MakeMapKey(Name, Group);
+		if (const FValueType* FoundValue = Map.Find(Key))
 		{
-			FKeyType Key = FManagedArrayCollection::MakeMapKey(Name, Group);
-			const FValueType& FoundValue = Map[Key];
-
-			if(FoundValue.ArrayType == ManagedArrayType<T>())
+			if(FoundValue->ArrayType == ManagedArrayType<T>())
 			{
-				return static_cast<TManagedArray<T>*>(Map[Key].Value);
+				return static_cast<const TManagedArray<T>*>(FoundValue->Value);
 			}
 		}
 		return nullptr;
