Chaos Cloth - Iterated on the PropertyCollection changes for the solver constraint files: fixed wrong valence values, added missing mesh scales, ...etc.

#rnx

diff --git a/Engine/Source/Runtime/Experimental/Chaos/Private/Chaos/PBDLongRangeConstraintsBase.cpp b/Engine/Source/Runtime/Experimental/Chaos/Private/Chaos/PBDLongRangeConstraintsBase.cpp
--- a/Engine/Source/Runtime/Experimental/Chaos/Private/Chaos/PBDLongRangeConstraintsBase.cpp
+++ b/Engine/Source/Runtime/Experimental/Chaos/Private/Chaos/PBDLongRangeConstraintsBase.cpp
@@ -33,7 +33,8 @@ FPBDLongRangeConstraintsBase::FPBDLongRangeConstraintsBase(
 	const TConstArrayView<FRealSingle>& ScaleMultipliers,
 	const FSolverVec2& InStiffness,
 	const FSolverVec2& InScale,
-	FSolverReal MaxStiffness)
+	FSolverReal MaxStiffness,
+	FSolverReal MeshScale)
 	: Tethers(InTethers)
 	, Stiffness(
 		InStiffness,
@@ -43,7 +44,7 @@ FPBDLongRangeConstraintsBase::FPBDLongRangeConstraintsBase(
 		FPBDStiffness::DefaultParameterFitBase,
 		MaxStiffness)
 	, TetherScale(
-		InScale.ClampAxes(MinTetherScale, MaxTetherScale),
+		InScale.ClampAxes(MinTetherScale, MaxTetherScale) * MeshScale,
 		ScaleMultipliers,
 		InParticleCount)
 	, ParticleOffset(InParticleOffset)
diff --git a/Engine/Source/Runtime/Experimental/Chaos/Public/Chaos/PBDLongRangeConstraints.h b/Engine/Source/Runtime/Experimental/Chaos/Public/Chaos/PBDLongRangeConstraints.h
--- a/Engine/Source/Runtime/Experimental/Chaos/Public/Chaos/PBDLongRangeConstraints.h
+++ b/Engine/Source/Runtime/Experimental/Chaos/Public/Chaos/PBDLongRangeConstraints.h
@@ -28,7 +28,8 @@ public:
 		const TArray<TConstArrayView<TTuple<int32, int32, FRealSingle>>>& InTethers,
 		const TConstArrayView<FRealSingle>& StiffnessMultipliers,
 		const TConstArrayView<FRealSingle>& ScaleMultipliers,
-		const FPropertyCollectionConstAdapter& PropertyCollection)
+		const FPropertyCollectionConstAdapter& PropertyCollection,
+		FSolverReal MeshScale)
 		: FPBDLongRangeConstraintsBase(
 			Particles,
 			InParticleOffset,
@@ -37,7 +38,9 @@ public:
 			StiffnessMultipliers,
 			ScaleMultipliers,
 			FSolverVec2(GetWeightedFloatTetherStiffness(PropertyCollection)),
-			FSolverVec2(GetWeightedFloatTetherScale(PropertyCollection)))  // Scale clamping done in constructor
+			FSolverVec2(GetWeightedFloatTetherScale(PropertyCollection)),  // Scale clamping done in constructor
+			FPBDStiffness::DefaultPBDMaxStiffness,
+			MeshScale)
 	{}
 
 	FPBDLongRangeConstraints(
@@ -48,7 +51,8 @@ public:
 		const TConstArrayView<FRealSingle>& StiffnessMultipliers,
 		const TConstArrayView<FRealSingle>& ScaleMultipliers,
 		const FSolverVec2& InStiffness = FSolverVec2::UnitVector,
-		const FSolverVec2& InScale = FSolverVec2::UnitVector)
+		const FSolverVec2& InScale = FSolverVec2::UnitVector,
+		FSolverReal MeshScale = (FSolverReal)1.)
 		: FPBDLongRangeConstraintsBase(
 			Particles,
 			InParticleOffset,
@@ -57,14 +61,15 @@ public:
 			StiffnessMultipliers,
 			ScaleMultipliers,
 			InStiffness,
-			InScale)
+			InScale,
+			MeshScale)
 	{}
 
 	virtual ~FPBDLongRangeConstraints() override {}
 
 	using Base::SetProperties;
 
-	void SetProperties(const FPropertyCollectionConstAdapter& PropertyCollection)
+	void SetProperties(const FPropertyCollectionConstAdapter& PropertyCollection, FSolverReal MeshScale)
 	{
 		if (IsTetherStiffnessMutable(PropertyCollection))
 		{
@@ -72,7 +77,7 @@ public:
 		}
 		if (IsTetherScaleMutable(PropertyCollection))
 		{
-			TetherScale.SetWeightedValue(FSolverVec2(GetWeightedFloatTetherScale(PropertyCollection)).ClampAxes(MinTetherScale, MaxTetherScale));
+			TetherScale.SetWeightedValue(FSolverVec2(GetWeightedFloatTetherScale(PropertyCollection)).ClampAxes(MinTetherScale, MaxTetherScale) * MeshScale);
 		}
 	}
 
diff --git a/Engine/Source/Runtime/Experimental/Chaos/Public/Chaos/PBDLongRangeConstraintsBase.h b/Engine/Source/Runtime/Experimental/Chaos/Public/Chaos/PBDLongRangeConstraintsBase.h
--- a/Engine/Source/Runtime/Experimental/Chaos/Public/Chaos/PBDLongRangeConstraintsBase.h
+++ b/Engine/Source/Runtime/Experimental/Chaos/Public/Chaos/PBDLongRangeConstraintsBase.h
@@ -40,7 +40,8 @@ public:
 		const TConstArrayView<FRealSingle>& ScaleMultipliers,
 		const FSolverVec2& InStiffness = FSolverVec2::UnitVector,
 		const FSolverVec2& InScale = FSolverVec2::UnitVector,
-		FSolverReal MaxStiffness = FPBDStiffness::DefaultPBDMaxStiffness);
+		FSolverReal MaxStiffness = FPBDStiffness::DefaultPBDMaxStiffness,
+		FSolverReal MeshScale = (FSolverReal)1.);
 
 PRAGMA_DISABLE_DEPRECATION_WARNINGS  // For ScaleIndices and ScaleTable
 	virtual ~FPBDLongRangeConstraintsBase() {}
diff --git a/Engine/Source/Runtime/Experimental/Chaos/Public/Chaos/PBDSoftsEvolutionFwd.h b/Engine/Source/Runtime/Experimental/Chaos/Public/Chaos/PBDSoftsEvolutionFwd.h
--- a/Engine/Source/Runtime/Experimental/Chaos/Public/Chaos/PBDSoftsEvolutionFwd.h
+++ b/Engine/Source/Runtime/Experimental/Chaos/Public/Chaos/PBDSoftsEvolutionFwd.h
@@ -45,11 +45,17 @@ namespace Chaos
 
 		// Softs solver constraint classes
 		class FPBDSpringConstraints;
+		class FPBDEdgeSpringConstraints;
+		class FPBDBendingSpringConstraints;
 		class FXPBDSpringConstraints;
+		class FXPBDEdgeSpringConstraints;
+		class FXPBDBendingSpringConstraints;
 		class FPBDBendingConstraints;
 		class FXPBDBendingConstraints;
 		class FPBDAxialSpringConstraints;
 		class FXPBDAxialSpringConstraints;
+		class FPBDAreaSpringConstraints;
+		class FXPBDAreaSpringConstraints;
 		class FPBDVolumeConstraint;
 		class FPBDLongRangeConstraints;
 		class FXPBDLongRangeConstraints;
diff --git a/Engine/Source/Runtime/Experimental/Chaos/Public/Chaos/PBDSpringConstraints.h b/Engine/Source/Runtime/Experimental/Chaos/Public/Chaos/PBDSpringConstraints.h
--- a/Engine/Source/Runtime/Experimental/Chaos/Public/Chaos/PBDSpringConstraints.h
+++ b/Engine/Source/Runtime/Experimental/Chaos/Public/Chaos/PBDSpringConstraints.h
@@ -47,8 +47,6 @@ private:
 class CHAOS_API FPBDEdgeSpringConstraints final : public FPBDSpringConstraints
 {
 public:
-	static constexpr int32 Valence = 2;
-
 	static bool IsEnabled(const FPropertyCollectionConstAdapter& PropertyCollection)
 	{
 		return IsEdgeSpringStiffnessEnabled(PropertyCollection, false);
@@ -58,7 +56,7 @@ public:
 		const FSolverParticles& Particles,
 		int32 ParticleOffset,
 		int32 ParticleCount,
-		const TArray<TVector<int32, Valence>>& InConstraints,
+		const TArray<TVec3<int32>>& InConstraints,
 		const TConstArrayView<FRealSingle>& StiffnessMultipliers,
 		const FPropertyCollectionConstAdapter& PropertyCollection,
 		bool bTrimKinematicConstraints = false)
@@ -89,8 +87,6 @@ private:
 class CHAOS_API FPBDBendingSpringConstraints final : public FPBDSpringConstraints
 {
 public:
-	static constexpr int32 Valence = 4;
-
 	static bool IsEnabled(const FPropertyCollectionConstAdapter& PropertyCollection)
 	{
 		return IsBendingSpringStiffnessEnabled(PropertyCollection, false);
@@ -100,7 +96,7 @@ public:
 		const FSolverParticles& Particles,
 		int32 ParticleOffset,
 		int32 ParticleCount,
-		const TArray<TVector<int32, Valence>>& InConstraints,
+		const TArray<TVec2<int32>>& InConstraints,
 		const TConstArrayView<FRealSingle>& StiffnessMultipliers,
 		const FPropertyCollectionConstAdapter& PropertyCollection,
 		bool bTrimKinematicConstraints = false)
diff --git a/Engine/Source/Runtime/Experimental/Chaos/Public/Chaos/XPBDBendingConstraints.h b/Engine/Source/Runtime/Experimental/Chaos/Public/Chaos/XPBDBendingConstraints.h
--- a/Engine/Source/Runtime/Experimental/Chaos/Public/Chaos/XPBDBendingConstraints.h
+++ b/Engine/Source/Runtime/Experimental/Chaos/Public/Chaos/XPBDBendingConstraints.h
@@ -36,6 +36,7 @@ public:
 		TArray<TVec4<int32>>&& InConstraints,
 		const TConstArrayView<FRealSingle>& StiffnessMultipliers,
 		const TConstArrayView<FRealSingle>& BucklingStiffnessMultipliers,
+		const TConstArrayView<FRealSingle>& DampingMultipliers,
 		const FPropertyCollectionConstAdapter& PropertyCollection,
 		bool bTrimKinematicConstraints = false)
 		: Base(
@@ -50,7 +51,12 @@ public:
 			FSolverVec2(GetWeightedFloatXPBDBucklingStiffness(PropertyCollection, MaxStiffness)),
 			bTrimKinematicConstraints,
 			MaxStiffness)
-		, DampingRatio(FSolverVec2(GetWeightedFloatXPBDBendingElementDamping(PropertyCollection, 0.f)).ClampAxes(MinDamping, MaxDamping))
+		, DampingRatio(
+			FSolverVec2(GetWeightedFloatXPBDBendingElementDamping(PropertyCollection, 0.f)).ClampAxes(MinDamping, MaxDamping),
+			DampingMultipliers,
+			TConstArrayView<TVec2<int32>>(ConstraintSharedEdges),
+			ParticleOffset,
+			ParticleCount)
 	{
 		Lambdas.Init((FSolverReal)0., Constraints.Num());
 		InitColor(InParticles, ParticleOffset, ParticleCount);
diff --git a/Engine/Source/Runtime/Experimental/Chaos/Public/Chaos/XPBDSpringConstraints.h b/Engine/Source/Runtime/Experimental/Chaos/Public/Chaos/XPBDSpringConstraints.h
--- a/Engine/Source/Runtime/Experimental/Chaos/Public/Chaos/XPBDSpringConstraints.h
+++ b/Engine/Source/Runtime/Experimental/Chaos/Public/Chaos/XPBDSpringConstraints.h
@@ -158,8 +158,6 @@ private:
 class CHAOS_API FXPBDEdgeSpringConstraints final : public FXPBDSpringConstraints
 {
 public:
-	static constexpr int32 Valence = 2;
-
 	static bool IsEnabled(const FPropertyCollectionConstAdapter& PropertyCollection)
 	{
 		return IsXPBDEdgeSpringStiffnessEnabled(PropertyCollection, false);
@@ -169,11 +167,9 @@ public:
 		const FSolverParticles& Particles,
 		int32 ParticleOffset,
 		int32 ParticleCount,
-		const TArray<TVector<int32, 2>>& InConstraints,
+		const TArray<TVec3<int32>>& InConstraints,
 		const TConstArrayView<FRealSingle>& StiffnessMultipliers,
 		const TConstArrayView<FRealSingle>& DampingMultipliers,
-		const FSolverVec2& InStiffness,
-		const FSolverVec2& InDampingRatio,
 		const FPropertyCollectionConstAdapter& PropertyCollection,
 		bool bTrimKinematicConstraints = false)
 		: FXPBDSpringConstraints(
@@ -213,8 +209,6 @@ private:
 class CHAOS_API FXPBDBendingSpringConstraints : public FXPBDSpringConstraints
 {
 public:
-	static constexpr int32 Valence = 4;
-
 	static bool IsEnabled(const FPropertyCollectionConstAdapter& PropertyCollection)
 	{
 		return IsXPBDBendingSpringStiffnessEnabled(PropertyCollection, false);
@@ -224,11 +218,9 @@ public:
 		const FSolverParticles& Particles,
 		int32 ParticleOffset,
 		int32 ParticleCount,
-		const TArray<TVector<int32, 4>>& InConstraints,
+		const TArray<TVec2<int32>>& InConstraints,
 		const TConstArrayView<FRealSingle>& StiffnessMultipliers,
 		const TConstArrayView<FRealSingle>& DampingMultipliers,
-		const FSolverVec2& InStiffness,
-		const FSolverVec2& InDampingRatio,
 		const FPropertyCollectionConstAdapter& PropertyCollection,
 		bool bTrimKinematicConstraints = false)
 		: FXPBDSpringConstraints(
