Iris
- Fixed edge case in FReplicationWriter when only replicated destroy was written

diff --git a/Engine/Source/Runtime/Experimental/Iris/Core/Private/Iris/ReplicationSystem/ReplicationWriter.cpp b/Engine/Source/Runtime/Experimental/Iris/Core/Private/Iris/ReplicationSystem/ReplicationWriter.cpp
--- a/Engine/Source/Runtime/Experimental/Iris/Core/Private/Iris/ReplicationSystem/ReplicationWriter.cpp
+++ b/Engine/Source/Runtime/Experimental/Iris/Core/Private/Iris/ReplicationSystem/ReplicationWriter.cpp
@@ -2011,7 +2011,7 @@ FReplicationWriter::EWriteObjectStatus FReplicationWriter::WriteObjectInBatch(FN
 						UE_LOG(LogIris, Verbose, TEXT("Failed to write huge attachment for object %s ( InternalIndex: %u ), forcing fallback on hugeobject for attachments"), *NetRefHandle.ToString(), InternalIndex);
 						Writer.DoOverflow();
 					}
-					else if (!BatchEntry.bSentState)
+					else if (!(BatchEntry.bSentState | bSentTearOff | Info.SubObjectPendingDestroy))
 					{
 						// If we didn't send state and didn't send any attachments let's rollback
 						ObjectRollbackScope.Rollback();
