Chaos Cloth - Add missing CHAOS_API on the property adapter classes.

#rnx

diff --git a/Engine/Source/Runtime/Experimental/Chaos/Public/Chaos/PropertyCollectionAdapter.h b/Engine/Source/Runtime/Experimental/Chaos/Public/Chaos/PropertyCollectionAdapter.h
--- a/Engine/Source/Runtime/Experimental/Chaos/Public/Chaos/PropertyCollectionAdapter.h
+++ b/Engine/Source/Runtime/Experimental/Chaos/Public/Chaos/PropertyCollectionAdapter.h
@@ -31,7 +31,7 @@ namespace Chaos::Softs
 	/**
 	 * Property collection helper class to access simulation properties from a ManagedArrayCollection.
 	 */
-	class FPropertyCollectionConstAdapter
+	class CHAOS_API FPropertyCollectionConstAdapter
 	{
 	public:
 		explicit FPropertyCollectionConstAdapter(const TSharedPtr<const FManagedArrayCollection>& InManagedArrayCollection);
@@ -137,7 +137,7 @@ namespace Chaos::Softs
 	 * Use a FPropertyCollectionMutableAdapter to initialize a new property colleciton or add new properties to the collection.
 	 * Note: Int property values are limited to 24 bits (-16777215 to 16777215).
 	 */
-	class FPropertyCollectionAdapter : public FPropertyCollectionConstAdapter
+	class CHAOS_API FPropertyCollectionAdapter : public FPropertyCollectionConstAdapter
 	{
 	public:
 		explicit FPropertyCollectionAdapter(const TSharedPtr<FManagedArrayCollection>& InManagedArrayCollection);
@@ -212,7 +212,7 @@ namespace Chaos::Softs
 	 * Property collection mutable adapter class to use a ManagedArrayCollection property collection and add/change simulation properties.
 	 * Note: Int property values are limited to 24 bits (-16777215 to 16777215).
 	 */
-	class FPropertyCollectionMutableAdapter final : public FPropertyCollectionAdapter
+	class CHAOS_API FPropertyCollectionMutableAdapter final : public FPropertyCollectionAdapter
 	{
 	public:
 		explicit FPropertyCollectionMutableAdapter(const TSharedPtr<FManagedArrayCollection>& InManagedArrayCollection);
