Chaos Cloth - Added missing CHAOS_API on some property adapter template instantiations and specialiizations.

#rnx

diff --git a/Engine/Source/Runtime/Experimental/Chaos/Private/Chaos/PropertyCollectionAdapter.cpp b/Engine/Source/Runtime/Experimental/Chaos/Private/Chaos/PropertyCollectionAdapter.cpp
--- a/Engine/Source/Runtime/Experimental/Chaos/Private/Chaos/PropertyCollectionAdapter.cpp
+++ b/Engine/Source/Runtime/Experimental/Chaos/Private/Chaos/PropertyCollectionAdapter.cpp
@@ -51,22 +51,22 @@ namespace Chaos::Softs
 	{
 		return ValueArray[KeyIndex];
 	}
-	template FString FPropertyCollectionConstAdapter::GetValue<FString, FString>(int32 KeyIndex, const TConstArrayView<FString>& ValueArray) const;
-	template uint8 FPropertyCollectionConstAdapter::GetValue<uint8, uint8>(int32 KeyIndex, const TConstArrayView<uint8>& ValueArray) const;
+	template CHAOS_API FString FPropertyCollectionConstAdapter::GetValue<FString, FString>(int32 KeyIndex, const TConstArrayView<FString>& ValueArray) const;
+	template CHAOS_API uint8 FPropertyCollectionConstAdapter::GetValue<uint8, uint8>(int32 KeyIndex, const TConstArrayView<uint8>& ValueArray) const;
 
-	template<>
+	template<> CHAOS_API
 	bool FPropertyCollectionConstAdapter::GetValue<bool, FVector3f>(int32 KeyIndex, const TConstArrayView<FVector3f>& ValueArray) const
 	{
 		return (bool)ValueArray[KeyIndex].X;
 	}
 
-	template<>
+	template<> CHAOS_API
 	int32 FPropertyCollectionConstAdapter::GetValue<int32, FVector3f>(int32 KeyIndex, const TConstArrayView<FVector3f>& ValueArray) const
 	{
 		return (int32)ValueArray[KeyIndex].X;
 	}
 
-	template<>
+	template<> CHAOS_API
 	float FPropertyCollectionConstAdapter::GetValue<float, FVector3f>(int32 KeyIndex, const TConstArrayView<FVector3f>& ValueArray) const
 	{
 		return ValueArray[KeyIndex].X;
