Fix Clang missing external symbol
#rnx

diff --git a/Engine/Source/Runtime/Experimental/Chaos/Private/Chaos/PropertyCollectionAdapter.cpp b/Engine/Source/Runtime/Experimental/Chaos/Private/Chaos/PropertyCollectionAdapter.cpp
--- a/Engine/Source/Runtime/Experimental/Chaos/Private/Chaos/PropertyCollectionAdapter.cpp
+++ b/Engine/Source/Runtime/Experimental/Chaos/Private/Chaos/PropertyCollectionAdapter.cpp
@@ -52,6 +52,7 @@ namespace Chaos::Softs
 		return ValueArray[KeyIndex];
 	}
 	template CHAOS_API FString FPropertyCollectionConstAdapter::GetValue<FString, FString>(int32 KeyIndex, const TConstArrayView<FString>& ValueArray) const;
+	template CHAOS_API const FString& FPropertyCollectionConstAdapter::GetValue<const FString&, FString>(int32 KeyIndex, const TConstArrayView<FString>& ValueArray) const;
 	template CHAOS_API uint8 FPropertyCollectionConstAdapter::GetValue<uint8, uint8>(int32 KeyIndex, const TConstArrayView<uint8>& ValueArray) const;
 
 	template<> CHAOS_API
