Iris - Support overriding net cull distances.

Previously the NetObjectGridFilter would detect changes to the Actor's NetCullDistanceSquared property and act on that. Changing members on instances can cause unique memory growth on dedicated server instances though and this is a way to circumvent that to some degree.

diff --git a/Engine/Source/Runtime/Engine/Private/Net/Iris/ReplicationSystem/ReplicationSystemUtil.cpp b/Engine/Source/Runtime/Engine/Private/Net/Iris/ReplicationSystem/ReplicationSystemUtil.cpp
--- a/Engine/Source/Runtime/Engine/Private/Net/Iris/ReplicationSystem/ReplicationSystemUtil.cpp
+++ b/Engine/Source/Runtime/Engine/Private/Net/Iris/ReplicationSystem/ReplicationSystemUtil.cpp
@@ -583,7 +583,53 @@ void FReplicationSystemUtil::SetStaticPriority(const AActor* Actor, float Priori
 	});
 }
 
+void FReplicationSystemUtil::SetCullDistanceSqrOverride(const AActor* Actor, float CullDistSqr)
+{
+	FNetHandle ActorHandle = GetNetHandle(Actor);
+	if (!ActorHandle.IsValid())
+	{
+		return;
+	}
+	
+	ReplicationSystemUtil::ForEachReplicationSystem([ActorHandle, CullDistSqr](UReplicationSystem* ReplicationSystem)
+	{
+		if (ReplicationSystem->IsServer())
+		{
+			if (UObjectReplicationBridge* Bridge = ReplicationSystem->GetReplicationBridgeAs<UObjectReplicationBridge>())
+			{
+				FNetRefHandle RefHandle = Bridge->GetReplicatedRefHandle(ActorHandle);
+				if (RefHandle.IsValid())
+				{
+					ReplicationSystem->SetCullDistanceSqrOverride(RefHandle, CullDistSqr);
+				}
+			}
+		}
+	});
+}
 
+void FReplicationSystemUtil::ClearCullDistanceSqrOverride(const AActor* Actor)
+{
+	FNetHandle ActorHandle = GetNetHandle(Actor);
+	if (!ActorHandle.IsValid())
+	{
+		return;
+	}
+	
+	ReplicationSystemUtil::ForEachReplicationSystem([ActorHandle](UReplicationSystem* ReplicationSystem)
+	{
+		if (ReplicationSystem->IsServer())
+		{
+			if (UObjectReplicationBridge* Bridge = ReplicationSystem->GetReplicationBridgeAs<UObjectReplicationBridge>())
+			{
+				FNetRefHandle RefHandle = Bridge->GetReplicatedRefHandle(ActorHandle);
+				if (RefHandle.IsValid())
+				{
+					ReplicationSystem->ClearCullDistanceSqrOverride(RefHandle);
+				}
+			}
+		}
+	});
+}
 
 }
 
diff --git a/Engine/Source/Runtime/Engine/Public/Net/Iris/ReplicationSystem/ReplicationSystemUtil.h b/Engine/Source/Runtime/Engine/Public/Net/Iris/ReplicationSystem/ReplicationSystemUtil.h
--- a/Engine/Source/Runtime/Engine/Public/Net/Iris/ReplicationSystem/ReplicationSystemUtil.h
+++ b/Engine/Source/Runtime/Engine/Public/Net/Iris/ReplicationSystem/ReplicationSystemUtil.h
@@ -121,6 +121,12 @@ struct FReplicationSystemUtil
 	 * @see ReplicationSystem.h for details
 	 */
 	ENGINE_API static void SetStaticPriority(const AActor* Actor, float Priority);
+
+	/** Set the squared cull distance for an actor. This will cause affected code to ignore the NetCullDistanceSquared property. */
+	ENGINE_API static void SetCullDistanceSqrOverride(const AActor* Actor, float CullDistSqr);
+
+	/** Clears any previously set squared cull distance for an actor. This will cause affected code to respect the NetCullDistanceSquared property. */
+	ENGINE_API static void ClearCullDistanceSqrOverride(const AActor* Actor);
 };
 
 }
diff --git a/Engine/Source/Runtime/Experimental/Iris/Core/Private/Iris/ReplicationSystem/Filtering/NetObjectGridFilter.cpp b/Engine/Source/Runtime/Experimental/Iris/Core/Private/Iris/ReplicationSystem/Filtering/NetObjectGridFilter.cpp
--- a/Engine/Source/Runtime/Experimental/Iris/Core/Private/Iris/ReplicationSystem/Filtering/NetObjectGridFilter.cpp
+++ b/Engine/Source/Runtime/Experimental/Iris/Core/Private/Iris/ReplicationSystem/Filtering/NetObjectGridFilter.cpp
@@ -1,6 +1,7 @@
 // Copyright Epic Games, Inc. All Rights Reserved.
 
 #include "Iris/ReplicationSystem/Filtering/NetObjectGridFilter.h"
+#include "Iris/ReplicationSystem/NetCullDistanceOverrides.h"
 #include "Iris/ReplicationSystem/RepTag.h"
 #include "Iris/ReplicationSystem/ReplicationProtocol.h"
 #include "Iris/ReplicationSystem/ReplicationSystem.h"
@@ -18,6 +19,7 @@ void UNetObjectGridFilter::Init(FNetObjectFilterInitParams& Params)
 	PerConnectionInfos.SetNum(Params.MaxConnectionCount + 1);
 
 	WorldLocations = &Params.ReplicationSystem->GetWorldLocations();
+	NetCullDistanceOverrides = &Params.ReplicationSystem->GetNetCullDistanceOverrides();
 }
 
 void UNetObjectGridFilter::AddConnection(uint32 ConnectionId)
@@ -388,7 +390,12 @@ void UNetObjectGridFilter::UpdatePositionAndCullDistance(const UNetObjectGridFil
 	}
 
 	// Optionally update cull distance
-	if (PerObjectInfo.CullDistanceSqrStateOffset != InvalidStateOffset)
+	if (NetCullDistanceOverrides->HasCullDistanceOverride(PerObjectInfo.ObjectIndex))
+	{
+		const float CullDistanceSqr = NetCullDistanceOverrides->GetCullDistanceSqr(PerObjectInfo.ObjectIndex);
+		PerObjectInfo.CullDistance = FPlatformMath::Sqrt(CullDistanceSqr);
+	}
+	else if (PerObjectInfo.CullDistanceSqrStateOffset != InvalidStateOffset)
 	{
 		const UE::Net::FReplicationInstanceProtocol::FFragmentData& FragmentData = FragmentDatas[PerObjectInfo.CullDistanceSqrStateIndex];
 		const uint8* CullDistanceSqrAddress = FragmentData.ExternalSrcBuffer + PerObjectInfo.CullDistanceSqrStateOffset;
diff --git a/Engine/Source/Runtime/Experimental/Iris/Core/Private/Iris/ReplicationSystem/NetCullDistanceOverrides.cpp b/Engine/Source/Runtime/Experimental/Iris/Core/Private/Iris/ReplicationSystem/NetCullDistanceOverrides.cpp
new file mode 100644
--- /dev/null
+++ b/Engine/Source/Runtime/Experimental/Iris/Core/Private/Iris/ReplicationSystem/NetCullDistanceOverrides.cpp
@@ -0,0 +1,30 @@
+// Copyright Epic Games, Inc. All Rights Reserved.
+
+#include "Iris/ReplicationSystem/NetCullDistanceOverrides.h"
+#include "Iris/Core/IrisMemoryTracker.h"
+
+namespace UE::Net
+{
+
+void FNetCullDistanceOverrides::Init(const FNetCullDistanceOverridesInitParams& InitParams)
+{
+	ValidCullDistanceSqr.Init(InitParams.MaxObjectCount);
+}
+
+void FNetCullDistanceOverrides::ClearCullDistanceSqr(uint32 ObjectIndex)
+{
+	ValidCullDistanceSqr.ClearBit(ObjectIndex);
+}
+
+void FNetCullDistanceOverrides::SetCullDistanceSqr(uint32 ObjectIndex, float CullDistSqr)
+{
+	ValidCullDistanceSqr.SetBit(ObjectIndex);
+	if (ObjectIndex >= uint32(CullDistanceSqr.Num()))
+	{
+		LLM_SCOPE_BYTAG(Iris);
+		CullDistanceSqr.Add(ObjectIndex + 1U - CullDistanceSqr.Num());
+	}
+	CullDistanceSqr[ObjectIndex] = CullDistSqr;
+}
+
+}
diff --git a/Engine/Source/Runtime/Experimental/Iris/Core/Private/Iris/ReplicationSystem/ReplicationBridge.cpp b/Engine/Source/Runtime/Experimental/Iris/Core/Private/Iris/ReplicationSystem/ReplicationBridge.cpp
--- a/Engine/Source/Runtime/Experimental/Iris/Core/Private/Iris/ReplicationSystem/ReplicationBridge.cpp
+++ b/Engine/Source/Runtime/Experimental/Iris/Core/Private/Iris/ReplicationSystem/ReplicationBridge.cpp
@@ -368,6 +368,9 @@ void UReplicationBridge::InternalDestroyNetObject(FNetRefHandle Handle)
 
 	if (const FInternalNetRefIndex ObjectInternalIndex = NetRefHandleManager->GetInternalIndex(Handle))
 	{
+		FNetCullDistanceOverrides& NetCullDistanceOverrides = ReplicationSystem->GetReplicationSystemInternal()->GetNetCullDistanceOverrides();
+		NetCullDistanceOverrides.ClearCullDistanceSqr(ObjectInternalIndex);
+
 		const FNetRefHandleManager::FReplicatedObjectData& ObjectData = NetRefHandleManager->GetReplicatedObjectDataNoCheck(ObjectInternalIndex);
 
 		if (ObjectData.bTearOff && NetRefHandleManager->GetNetObjectRefCount(ObjectInternalIndex) > 0U)
diff --git a/Engine/Source/Runtime/Experimental/Iris/Core/Private/Iris/ReplicationSystem/ReplicationSystem.cpp b/Engine/Source/Runtime/Experimental/Iris/Core/Private/Iris/ReplicationSystem/ReplicationSystem.cpp
--- a/Engine/Source/Runtime/Experimental/Iris/Core/Private/Iris/ReplicationSystem/ReplicationSystem.cpp
+++ b/Engine/Source/Runtime/Experimental/Iris/Core/Private/Iris/ReplicationSystem/ReplicationSystem.cpp
@@ -79,6 +79,8 @@ public:
 
 		FNetRefHandleManager& NetRefHandleManager = ReplicationSystemInternal.GetNetRefHandleManager();
 
+		const uint32 MaxObjectCount =  NetRefHandleManager.GetMaxActiveObjectCount();
+
 		// $IRIS TODO: Need object ID range. Currently abusing hardcoded values from FNetRefHandleManager
 		FDirtyNetObjectTracker& DirtyNetObjectTracker = ReplicationSystemInternal.GetDirtyNetObjectTracker();
 		FDirtyNetObjectTrackerInitParams DirtyNetObjectTrackerInitParams;
@@ -86,7 +88,7 @@ public:
 			DirtyNetObjectTrackerInitParams.NetRefHandleManager = &NetRefHandleManager;
 			DirtyNetObjectTrackerInitParams.ReplicationSystemId = ReplicationSystemId;
 			DirtyNetObjectTrackerInitParams.NetObjectIndexRangeStart = 1;
-			DirtyNetObjectTrackerInitParams.NetObjectIndexRangeEnd = NetRefHandleManager.GetMaxActiveObjectCount() - 1U;
+			DirtyNetObjectTrackerInitParams.NetObjectIndexRangeEnd = MaxObjectCount - 1U;
 
 			DirtyNetObjectTracker.Init(DirtyNetObjectTrackerInitParams);
 		}
@@ -96,7 +98,7 @@ public:
 			FReplicationStateStorageInitParams InitParams;
 			InitParams.ReplicationSystem = ReplicationSystem;
 			InitParams.NetRefHandleManager = &NetRefHandleManager;
-			InitParams.MaxObjectCount = DirtyNetObjectTrackerInitParams.NetObjectIndexRangeEnd + 1U;
+			InitParams.MaxObjectCount = MaxObjectCount;
 			InitParams.MaxConnectionCount = ReplicationSystemInternal.GetConnections().GetMaxConnectionCount();
 			InitParams.MaxDeltaCompressedObjectCount = Params.MaxDeltaCompressedObjectCount;
 			StateStorage.Init(InitParams);
@@ -105,7 +107,7 @@ public:
 		FNetObjectGroups& Groups = ReplicationSystemInternal.GetGroups();
 		{
 			FNetObjectGroupInitParams InitParams = {};
-			InitParams.MaxObjectCount = DirtyNetObjectTrackerInitParams.NetObjectIndexRangeEnd + 1U;
+			InitParams.MaxObjectCount = MaxObjectCount;
 			InitParams.MaxGroupCount = Params.MaxNetObjectGroupCount;
 
 			Groups.Init(InitParams);
@@ -119,10 +121,17 @@ public:
 			Registry.Init(InitParams);
 		}		
 
+		FNetCullDistanceOverrides& NetCullDistanceOverrides = ReplicationSystemInternal.GetNetCullDistanceOverrides();
+		{
+			FNetCullDistanceOverridesInitParams InitParams;
+			InitParams.MaxObjectCount = MaxObjectCount;
+			NetCullDistanceOverrides.Init(InitParams);
+		}
+
 		FWorldLocations& WorldLocations = ReplicationSystemInternal.GetWorldLocations();
 		{
 			FWorldLocationsInitParams InitParams;
-			InitParams.MaxObjectCount = DirtyNetObjectTrackerInitParams.NetObjectIndexRangeEnd + 1U;
+			InitParams.MaxObjectCount = MaxObjectCount;
 			WorldLocations.Init(InitParams);
 		}
 	
@@ -131,7 +140,7 @@ public:
 		{
 			FDeltaCompressionBaselineInvalidationTrackerInitParams InitParams;
 			InitParams.BaselineManager = &DeltaCompressionBaselineManager;
-			InitParams.MaxObjectCount = DirtyNetObjectTrackerInitParams.NetObjectIndexRangeEnd + 1U;
+			InitParams.MaxObjectCount = MaxObjectCount;
 			DeltaCompressionBaselineInvalidationTracker.Init(InitParams);
 		}
 		{
@@ -140,7 +149,7 @@ public:
 			InitParams.Connections = &ReplicationSystemInternal.GetConnections();
 			InitParams.NetRefHandleManager = &NetRefHandleManager;
 			InitParams.ReplicationStateStorage = &StateStorage;
-			InitParams.MaxObjectCount = DirtyNetObjectTrackerInitParams.NetObjectIndexRangeEnd + 1U;
+			InitParams.MaxObjectCount = MaxObjectCount;
 			InitParams.MaxDeltaCompressedObjectCount = Params.MaxDeltaCompressedObjectCount;
 			InitParams.ReplicationSystem = ReplicationSystem;
 			DeltaCompressionBaselineManager.Init(InitParams);
@@ -154,7 +163,7 @@ public:
 			InitParams.NetRefHandleManager = &NetRefHandleManager;
 			InitParams.Groups = &Groups;
 			InitParams.BaselineInvalidationTracker = &ReplicationSystemInternal.GetDeltaCompressionBaselineInvalidationTracker();
-			InitParams.MaxObjectCount = DirtyNetObjectTrackerInitParams.NetObjectIndexRangeEnd + 1U;
+			InitParams.MaxObjectCount = MaxObjectCount;
 			InitParams.MaxGroupCount = Params.MaxNetObjectGroupCount;
 			ReplicationFiltering.Init(InitParams);
 		}
@@ -169,7 +178,7 @@ public:
 			InitParams.ReplicationFiltering = &ReplicationFiltering;
 			InitParams.NetObjectGroups = &Groups;
 			InitParams.BaselineInvalidationTracker = &ReplicationSystemInternal.GetDeltaCompressionBaselineInvalidationTracker();
-			InitParams.MaxObjectCount = DirtyNetObjectTrackerInitParams.NetObjectIndexRangeEnd + 1U;
+			InitParams.MaxObjectCount = MaxObjectCount;
 			InitParams.MaxConnectionCount = ReplicationSystemInternal.GetConnections().GetMaxConnectionCount();
 			ReplicationConditionals.Init(InitParams);
 		}
@@ -180,7 +189,7 @@ public:
 			InitParams.ReplicationSystem = ReplicationSystem;
 			InitParams.Connections = &ReplicationSystemInternal.GetConnections();
 			InitParams.NetRefHandleManager = &NetRefHandleManager;
-			InitParams.MaxObjectCount = DirtyNetObjectTrackerInitParams.NetObjectIndexRangeEnd + 1U;
+			InitParams.MaxObjectCount = MaxObjectCount;
 			ReplicationPrioritization.Init(InitParams);
 		}
 
@@ -1189,11 +1198,49 @@ void UReplicationSystem::AddReferencedObjects(UObject* InThis, FReferenceCollect
 	Super::AddReferencedObjects(InThis, Collector);
 }
 
+const UE::Net::FNetCullDistanceOverrides& UReplicationSystem::GetNetCullDistanceOverrides() const
+{
+	return Impl->ReplicationSystemInternal.GetNetCullDistanceOverrides();
+}
+
 const UE::Net::FWorldLocations& UReplicationSystem::GetWorldLocations() const
 {
 	return Impl->ReplicationSystemInternal.GetWorldLocations();
 }
 
+void UReplicationSystem::SetCullDistanceSqrOverride(FNetRefHandle Handle, float DistSqr)
+{
+	const UE::Net::Private::FInternalNetRefIndex ObjectInternalIndex = Impl->ReplicationSystemInternal.GetNetRefHandleManager().GetInternalIndex(Handle);
+	if (ObjectInternalIndex == UE::Net::Private::FNetRefHandleManager::InvalidInternalIndex)
+	{
+		return;
+	}
+
+	return Impl->ReplicationSystemInternal.GetNetCullDistanceOverrides().SetCullDistanceSqr(ObjectInternalIndex, DistSqr);
+}
+
+void UReplicationSystem::ClearCullDistanceSqrOverride(FNetRefHandle Handle)
+{
+	const UE::Net::Private::FInternalNetRefIndex ObjectInternalIndex = Impl->ReplicationSystemInternal.GetNetRefHandleManager().GetInternalIndex(Handle);
+	if (ObjectInternalIndex == UE::Net::Private::FNetRefHandleManager::InvalidInternalIndex)
+	{
+		return;
+	}
+
+	Impl->ReplicationSystemInternal.GetNetCullDistanceOverrides().ClearCullDistanceSqr(ObjectInternalIndex);
+}
+
+float UReplicationSystem::GetCullDistanceSqrOverride(FNetRefHandle Handle, float DefaultValue) const
+{
+	const UE::Net::Private::FInternalNetRefIndex ObjectInternalIndex = Impl->ReplicationSystemInternal.GetNetRefHandleManager().GetInternalIndex(Handle);
+	if (ObjectInternalIndex == UE::Net::Private::FNetRefHandleManager::InvalidInternalIndex)
+	{
+		return DefaultValue;
+	}
+
+	return Impl->ReplicationSystemInternal.GetNetCullDistanceOverrides().GetCullDistanceSqr(ObjectInternalIndex, DefaultValue);
+}
+
 namespace UE::Net
 {
 
diff --git a/Engine/Source/Runtime/Experimental/Iris/Core/Private/Iris/ReplicationSystem/ReplicationSystemInternal.h b/Engine/Source/Runtime/Experimental/Iris/Core/Private/Iris/ReplicationSystem/ReplicationSystemInternal.h
--- a/Engine/Source/Runtime/Experimental/Iris/Core/Private/Iris/ReplicationSystem/ReplicationSystemInternal.h
+++ b/Engine/Source/Runtime/Experimental/Iris/Core/Private/Iris/ReplicationSystem/ReplicationSystemInternal.h
@@ -6,6 +6,7 @@
 #include "Containers/Array.h"
 #include "Iris/ReplicationState/ReplicationStateStorage.h"
 #include "Iris/ReplicationSystem/DirtyNetObjectTracker.h"
+#include "Iris/ReplicationSystem/NetCullDistanceOverrides.h"
 #include "Iris/ReplicationSystem/NetRefHandleManager.h"
 #include "Iris/ReplicationSystem/ChangeMaskCache.h"
 #include "Iris/ReplicationSystem/Conditionals/ReplicationConditionals.h"
@@ -83,6 +84,8 @@ public:
 
 	FNetTokenStore& GetNetTokenStore() { return NetTokenStore; }
 
+	FNetCullDistanceOverrides& GetNetCullDistanceOverrides() { return NetCullDistanceOverrides; }
+
 	FWorldLocations& GetWorldLocations() { return WorldLocations; }
 
 	FDeltaCompressionBaselineManager& GetDeltaCompressionBaselineManager() { return DeltaCompressionBaselineManager; }
@@ -110,6 +113,7 @@ private:
 	FNetBlobManager NetBlobManager;
 	FNetTokenStore NetTokenStore;
 	FStringTokenStore StringTokenStore;
+	FNetCullDistanceOverrides NetCullDistanceOverrides;
 	FWorldLocations WorldLocations;
 	FDeltaCompressionBaselineManager DeltaCompressionBaselineManager;
 	FDeltaCompressionBaselineInvalidationTracker DeltaCompressionBaselineInvalidationTracker;
diff --git a/Engine/Source/Runtime/Experimental/Iris/Core/Public/Iris/ReplicationSystem/Filtering/NetObjectGridFilter.h b/Engine/Source/Runtime/Experimental/Iris/Core/Public/Iris/ReplicationSystem/Filtering/NetObjectGridFilter.h
--- a/Engine/Source/Runtime/Experimental/Iris/Core/Public/Iris/ReplicationSystem/Filtering/NetObjectGridFilter.h
+++ b/Engine/Source/Runtime/Experimental/Iris/Core/Public/Iris/ReplicationSystem/Filtering/NetObjectGridFilter.h
@@ -11,6 +11,7 @@
 
 namespace UE::Net
 {
+	class FNetCullDistanceOverrides;
 	class FWorldLocations;
 }
 
@@ -34,7 +35,7 @@ public:
 	UPROPERTY(Config)
 	float MaxCullDistance = 20000.0f;
 
-	/** Objects without a NetCullDistanceSquared property will assume to have this value but squared. */
+	/** Objects without a NetCullDistanceSquared property will assume to have this value but squared unless there's a cull distance override. */
 	UPROPERTY(Config)
 	float DefaultCullDistance = 15000.0f;
 
@@ -173,6 +174,7 @@ private:
 	uint32 FrameIndex = 0;
 
 	const UE::Net::FWorldLocations* WorldLocations = nullptr;
+	const UE::Net::FNetCullDistanceOverrides* NetCullDistanceOverrides = nullptr;
 };
 
 
diff --git a/Engine/Source/Runtime/Experimental/Iris/Core/Public/Iris/ReplicationSystem/NetCullDistanceOverrides.h b/Engine/Source/Runtime/Experimental/Iris/Core/Public/Iris/ReplicationSystem/NetCullDistanceOverrides.h
new file mode 100644
--- /dev/null
+++ b/Engine/Source/Runtime/Experimental/Iris/Core/Public/Iris/ReplicationSystem/NetCullDistanceOverrides.h
@@ -0,0 +1,72 @@
+// Copyright Epic Games, Inc. All Rights Reserved.
+
+#pragma once
+
+#include "HAL/Platform.h"
+#include "Containers/ChunkedArray.h"
+#include "Net/Core/NetBitArray.h"
+
+namespace UE::Net
+{
+
+struct FNetCullDistanceOverridesInitParams
+{
+	uint32 MaxObjectCount = 0;
+};
+
+class FNetCullDistanceOverrides
+{
+public:
+	void Init(const FNetCullDistanceOverridesInitParams& InitParams);
+
+	/** Returns whether the object has a valid world location or not. */
+	bool HasCullDistanceOverride(uint32 ObjectIndex) const;
+
+	/** Returns the object's squared cull distance override if it's valid or a negative value if it's not. */
+	float GetCullDistanceSqr(uint32 ObjectIndex) const;
+
+	/** Returns the object's squared cull distance override if it's valid or the provided DefaultValue if it's not. */
+	float GetCullDistanceSqr(uint32 ObjectIndex, float DefaultValue) const;
+
+public:
+	/** Remove cull distance override for object. */
+	void ClearCullDistanceSqr(uint32 ObjectIndex);
+
+	/** Set cull distance override for object. */
+	void SetCullDistanceSqr(uint32 ObjectIndex, float CullDistSqr);
+
+private:
+	enum : unsigned
+	{
+		// Arbitrary value. The chunked array will be allocated to fit to higest object index that overrides the cull distance.
+		BytesPerCullDistanceChunk = 4096U,
+	};
+
+	FNetBitArray ValidCullDistanceSqr;
+	TChunkedArray<float, BytesPerCullDistanceChunk> CullDistanceSqr;
+};
+
+inline bool FNetCullDistanceOverrides::HasCullDistanceOverride(uint32 ObjectIndex) const
+{
+	return ValidCullDistanceSqr.GetBit(ObjectIndex);
+}
+
+inline float FNetCullDistanceOverrides::GetCullDistanceSqr(uint32 ObjectIndex, float DefaultValue) const
+{
+	if (ValidCullDistanceSqr.GetBit(ObjectIndex))
+	{
+		return CullDistanceSqr[ObjectIndex];
+	}
+	else
+	{
+		return DefaultValue;
+	}
+}
+
+inline float FNetCullDistanceOverrides::GetCullDistanceSqr(uint32 ObjectIndex) const
+{
+	return GetCullDistanceSqr(ObjectIndex, -1.0f);
+}
+
+}
+
diff --git a/Engine/Source/Runtime/Experimental/Iris/Core/Public/Iris/ReplicationSystem/ReplicationSystem.h b/Engine/Source/Runtime/Experimental/Iris/Core/Public/Iris/ReplicationSystem/ReplicationSystem.h
--- a/Engine/Source/Runtime/Experimental/Iris/Core/Public/Iris/ReplicationSystem/ReplicationSystem.h
+++ b/Engine/Source/Runtime/Experimental/Iris/Core/Public/Iris/ReplicationSystem/ReplicationSystem.h
@@ -23,6 +23,7 @@ class UReplicationSystem;
 namespace UE::Net
 {
 	class FNetBitArray;
+	class FNetCullDistanceOverrides;
 	enum class ENetFilterStatus : uint32;
 	class FNetObjectAttachment;
 	enum class ENetObjectDeltaCompressionStatus : unsigned;
@@ -461,6 +462,13 @@ public:
 	 */
 	IRISCORE_API void MarkDirty(FNetRefHandle Handle);
 
+	/**
+	 * Retrieve the NetCullDistanceOverrides instance which holds cull distances for objects that have explicitly overriden it.
+	 * @return The NetCullDistanceOverrides instance.
+	 * @see UE::Net::FNetCullDistanceOverrides
+	 */
+	IRISCORE_API const UE::Net::FNetCullDistanceOverrides& GetNetCullDistanceOverrides() const;
+
 	/**
 	 * Retrieve the WorldLocations instance which holds world locations for all objects that support it. 
 	 * @return The WorldLocations instance.
@@ -490,6 +498,15 @@ public:
 
 	IRISCORE_API int32 GetPIEInstanceID() const { return PIEInstanceID; }
 
+	/** Set the squared cull distance for an object. This can be used by prioritizers and filters like UNetObjectGridFilter for example. For Actors this will override, not overwrite, the NetCullDistanceSquared property. */
+	IRISCORE_API void SetCullDistanceSqrOverride(FNetRefHandle Handle, float DistSqr);
+
+	/** Clears any previously set squared cull distance for an object. For Actors this will cause affected code to respect the NetCullDistanceSquared property. */
+	IRISCORE_API void ClearCullDistanceSqrOverride(FNetRefHandle Handle);
+
+	/** Returns the previously set squared cull distance for an object, or DefaultValue if it wasn't or the Handle isn't valid. */
+	IRISCORE_API float GetCullDistanceSqrOverride(FNetRefHandle Handle, float DefaultValue = -1.0f) const;
+
 public:
 	// For internal use and not exported.
 
