Fix Sphere Overlap

diff --git a/Engine/Source/Programs/HeadlessChaos/Private/HeadlessChaosTestOverlap.cpp b/Engine/Source/Programs/HeadlessChaos/Private/HeadlessChaosTestOverlap.cpp
--- a/Engine/Source/Programs/HeadlessChaos/Private/HeadlessChaosTestOverlap.cpp
+++ b/Engine/Source/Programs/HeadlessChaos/Private/HeadlessChaosTestOverlap.cpp
@@ -131,16 +131,7 @@ namespace ChaosTest
 					EXPECT_TRUE(bResult);
 				}
 			}
-			{
-				const FVec3 X1 = { 10.0, -1.0, 0.0 };
-				const FVec3 X2 = { 11.0, 1.0, 0.0 };
-				const FReal Radius = 0.1;
-				const FCapsule Capsule = FCapsule(X1, X2, Radius);
-				const FRigidTransform3 QueryTM(FVec3(0.0, 0.0, 0.0), FQuat::Identity);
-				bool bResult = TriangleMesh->OverlapGeom(Capsule, QueryTM, 0.0);
-				EXPECT_FALSE(bResult);
-			}
-			{
+{
 				const FVec3 X1 = { -11.0, 2.0, 0.0 };
 				const FVec3 X2 = { -9.0, -8.0, 0.0 };
 				const FReal Radius = 0.01;
@@ -154,6 +145,15 @@ namespace ChaosTest
 					EXPECT_EQ(bResult, bResultMTD);
 				}
 			}
+			{
+				const FVec3 X1 = { 10.0, -1.0, 0.0 };
+				const FVec3 X2 = { 11.0, 1.0, 0.0 };
+				const FReal Radius = 0.1;
+				const FCapsule Capsule = FCapsule(X1, X2, Radius);
+				const FRigidTransform3 QueryTM(FVec3(0.0, 0.0, 0.0), FQuat::Identity);
+				bool bResult = TriangleMesh->OverlapGeom(Capsule, QueryTM, 0.0);
+				EXPECT_FALSE(bResult);
+			}
 			{
 				// Sphere test
 				const FVec3 X = { 0.0, 0.0, 0.0 };
@@ -1212,7 +1212,7 @@ namespace ChaosTest
 				{
 					const FRigidTransform3 QueryTM(FVec3(0.0, 0.0, 0.0), FQuat::Identity);
 					bool bResult = TriangleMesh->OverlapGeom(Sphere, QueryTM, 0.0);
-					// EXPECT_TRUE(bResult); // TODO Enable test
+					EXPECT_TRUE(bResult);
 				}
 			}
 			{
@@ -1264,7 +1264,7 @@ namespace ChaosTest
 				{
 					const FRigidTransform3 QueryTM(FVec3(0.0, 0.0, 0.0), FQuat::Identity);
 					bool bResult = TriangleMesh->OverlapGeom(Sphere, QueryTM, 0.0);
-					// EXPECT_TRUE(bResult); // TODO Enable test
+					EXPECT_TRUE(bResult);
 				}
 			}
 			{
@@ -1276,7 +1276,7 @@ namespace ChaosTest
 				}
 			}
 		}
-		{
+{
 			FTriangleMeshImplicitObject::ParticlesType TrimeshParticles(
 				{
 					{0.0, 0.0, 0.0},
@@ -1343,6 +1343,16 @@ namespace ChaosTest
 						EXPECT_TRUE(bResult);
 					}
 				}
+				{
+					const FVec3 X = { -100, 0, 0 };
+					const FReal Radius = 99.0f;
+					const Chaos::FSphere Sphere = Chaos::FSphere(X, Radius);
+					{
+						const FRigidTransform3 QueryTM(FVec3(0.0, 0.0, 0.0), FQuat::Identity);
+						bool bResult = TriangleMesh->OverlapGeom(Sphere, QueryTM, 0.0);
+						EXPECT_TRUE(bResult);
+					}
+				}
 			}
 		}
 	}
diff --git a/Engine/Source/Runtime/Experimental/Chaos/Private/Chaos/Collision/TriangleOverlap.cpp b/Engine/Source/Runtime/Experimental/Chaos/Private/Chaos/Collision/TriangleOverlap.cpp
--- a/Engine/Source/Runtime/Experimental/Chaos/Private/Chaos/Collision/TriangleOverlap.cpp
+++ b/Engine/Source/Runtime/Experimental/Chaos/Private/Chaos/Collision/TriangleOverlap.cpp
@@ -31,10 +31,8 @@ namespace Chaos
 		const VectorRegister4Float BX1 = VectorSubtract(X1, B);
 		const VectorRegister4Float CX1 = VectorSubtract(X1, C);
 
-		const VectorRegister4Float RadiusSimd = VectorSetFloat1(Radius);
 		// Edge plane normals and signed distances to each segment point
 		// Mask to compare the distance of the triangle edge plane to the capsule radius.
-		constexpr static int32 AllEdges = 0b1111;
 		const VectorRegister4Float PlaneCA = VectorNormalize(VectorCross(CA, Normal));
 		{
 			FRealSingle EdgeX1 = VectorDot3Scalar(AX1, PlaneCA);
@@ -234,12 +232,12 @@ namespace Chaos
 	{
 		const VectorRegister4Float AB = VectorSubtract(B, A);
 		const VectorRegister4Float BC = VectorSubtract(C, B);
+		const VectorRegister4Float CA = VectorSubtract(A, C);
 
 		const VectorRegister4Float Normal = VectorNormalize(VectorCross(AB, BC));
 
 		// Plane Triangle
 		const VectorRegister4Float AX = VectorSubtract(X, A);
-
 		const FRealSingle AXDist = VectorDot3Scalar(AX, Normal);
 
 		if (FMath::Abs<FRealSingle>(AXDist) > Radius)
@@ -247,9 +245,9 @@ namespace Chaos
 			return false;
 		}
 
-		constexpr FRealSingle ThirdFloat = 1.0f / 3.0f;
-		constexpr VectorRegister4Float Third = MakeVectorRegisterFloatConstant(ThirdFloat, ThirdFloat, ThirdFloat, ThirdFloat);
-		const VectorRegister4Float Centroid = VectorMultiply(VectorAdd(VectorAdd(A, B), C), Third);
+		const VectorRegister4Float PlaneAB = VectorCross(AB, Normal);
+		const VectorRegister4Float PlaneBC = VectorCross(BC, Normal);
+		const VectorRegister4Float PlaneCA = VectorCross(CA, Normal);
 
 		const FRealSingle SqrRadius = Radius * Radius;
 		{
@@ -258,10 +256,9 @@ namespace Chaos
 			VectorRegister4Float Time = VectorClamp(VectorDivide(VectorDot3(AB, AX), SqrAB), VectorZeroFloat(), VectorOneFloat());
 			Time = VectorBitwiseNotAnd(ZeroMask, Time);
 			const VectorRegister4Float P = VectorMultiplyAdd(A, VectorSubtract(VectorOneFloat(), Time), VectorMultiply(B, Time));
-
 			const VectorRegister4Float PX = VectorSubtract(X, P);
-			const VectorRegister4Float CentroidP = VectorSubtract(P, Centroid);
-			const FRealSingle Dist = VectorDot3Scalar(PX, CentroidP);
+
+			const FRealSingle Dist = VectorDot3Scalar(PX, PlaneAB);
 			if (Dist > 0.0f)
 			{
 				if (VectorDot3Scalar(PX, PX) > SqrRadius)
@@ -277,10 +274,8 @@ namespace Chaos
 			VectorRegister4Float Time = VectorClamp(VectorDivide(VectorDot3(BC, BX), SqrBC), VectorZeroFloat(), VectorOneFloat());
 			Time = VectorBitwiseNotAnd(ZeroMask, Time);
 			const VectorRegister4Float P = VectorMultiplyAdd(B, VectorSubtract(VectorOneFloat(), Time), VectorMultiply(C, Time));
-
 			const VectorRegister4Float PX = VectorSubtract(X, P);
-			const VectorRegister4Float CentroidP = VectorSubtract(P, Centroid);
-			const FRealSingle Dist = VectorDot3Scalar(PX, CentroidP);
+			const FRealSingle Dist = VectorDot3Scalar(PX, PlaneBC);
 			if (Dist > 0.0f)
 			{
 				if (VectorDot3Scalar(PX, PX) > SqrRadius)
@@ -290,17 +285,14 @@ namespace Chaos
 			}
 		}
 		{
-			const VectorRegister4Float CA = VectorSubtract(A, C);
 			const VectorRegister4Float SqrCA = VectorDot3(CA, CA);
 			const VectorRegister4Float ZeroMask = VectorCompareEQ(VectorZeroFloat(), SqrCA);
 			const VectorRegister4Float CX = VectorSubtract(X, C);
 			VectorRegister4Float Time = VectorClamp(VectorDivide(VectorDot3(CA, CX), SqrCA), VectorZeroFloat(), VectorOneFloat());
 			Time = VectorBitwiseNotAnd(ZeroMask, Time);
 			const VectorRegister4Float P = VectorMultiplyAdd(C, VectorSubtract(VectorOneFloat(), Time), VectorMultiply(A, Time));
-
 			const VectorRegister4Float PX = VectorSubtract(X, P);
-			const VectorRegister4Float CentroidP = VectorSubtract(P, Centroid);
-			const FRealSingle Dist = VectorDot3Scalar(PX, CentroidP);
+			const FRealSingle Dist = VectorDot3Scalar(PX, PlaneCA);
 			if (Dist > 0.0f)
 			{
 				if (VectorDot3Scalar(PX, PX) > SqrRadius)
