add option to enable/disable adding strain to geometry collection clusters upon collisions

- also fixed issue where changing the SetNotifyX (break, crumbling, etc) on a geometry collection component wouldn't set the simulation parameter on the physics proxy.

diff --git a/Engine/Source/Runtime/Experimental/Chaos/Private/Chaos/PBDRigidClustering.cpp b/Engine/Source/Runtime/Experimental/Chaos/Private/Chaos/PBDRigidClustering.cpp
--- a/Engine/Source/Runtime/Experimental/Chaos/Private/Chaos/PBDRigidClustering.cpp
+++ b/Engine/Source/Runtime/Experimental/Chaos/Private/Chaos/PBDRigidClustering.cpp
@@ -1177,7 +1177,12 @@ namespace Chaos
 					bool bUseDamagePropagation = false;
 					if (const FGeometryCollectionPhysicsProxy* ConcreteProxy = GetConcreteProxy(Cluster))
 					{
-						bUseDamagePropagation = ConcreteProxy->GetSimParameters().bUseDamagePropagation;
+						const FSimulationParameters& SimParams = ConcreteProxy->GetSimParameters();
+						bUseDamagePropagation = SimParams.bUseDamagePropagation;
+						if (!SimParams.bEnableStrainOnCollision)
+						{
+							return;
+						}
 					}
 
 					if (bUseDamagePropagation)
diff --git a/Engine/Source/Runtime/Experimental/Chaos/Public/GeometryCollection/GeometryCollectionSimulationCoreTypes.h b/Engine/Source/Runtime/Experimental/Chaos/Public/GeometryCollection/GeometryCollectionSimulationCoreTypes.h
--- a/Engine/Source/Runtime/Experimental/Chaos/Public/GeometryCollection/GeometryCollectionSimulationCoreTypes.h
+++ b/Engine/Source/Runtime/Experimental/Chaos/Public/GeometryCollection/GeometryCollectionSimulationCoreTypes.h
@@ -210,6 +210,7 @@ struct FSimulationParameters
 		, SimulationFilterData()
 		, QueryFilterData()
 		, UserData(nullptr)
+		, bEnableStrainOnCollision(true)
 	{}
 
 	FSimulationParameters(const FSimulationParameters& Other)
@@ -258,6 +259,7 @@ struct FSimulationParameters
 		, SimulationFilterData(Other.SimulationFilterData)
 		, QueryFilterData(Other.QueryFilterData)
 		, UserData(Other.UserData)
+		, bEnableStrainOnCollision(Other.bEnableStrainOnCollision)
 	{
 	}
 
@@ -331,4 +333,5 @@ struct FSimulationParameters
 	FCollisionFilterData SimulationFilterData;
 	FCollisionFilterData QueryFilterData;
 	void* UserData;
+	bool bEnableStrainOnCollision;
 };
diff --git a/Engine/Source/Runtime/Experimental/GeometryCollectionEngine/Private/GeometryCollection/GeometryCollectionComponent.cpp b/Engine/Source/Runtime/Experimental/GeometryCollectionEngine/Private/GeometryCollection/GeometryCollectionComponent.cpp
--- a/Engine/Source/Runtime/Experimental/GeometryCollectionEngine/Private/GeometryCollection/GeometryCollectionComponent.cpp
+++ b/Engine/Source/Runtime/Experimental/GeometryCollectionEngine/Private/GeometryCollection/GeometryCollectionComponent.cpp
@@ -291,6 +291,7 @@ UGeometryCollectionComponent::UGeometryCollectionComponent(const FObjectInitiali
 	, MaxClusterLevel(100)
 	, DamageThreshold({ 500000.f, 50000.f, 5000.f })
 	, bUseSizeSpecificDamageThreshold(false)
+	, bEnableDamageFromCollision(true)
 	, bAllowRemovalOnSleep(true)
 	, bAllowRemovalOnBreak(true)
 	, ClusterConnectionType_DEPRECATED(EClusterConnectionTypeEnum::Chaos_MinimalSpanningSubsetDelaunayTriangulation)
@@ -735,6 +736,16 @@ void UGeometryCollectionComponent::SetNotifyBreaks(bool bNewNotifyBreaks)
 {
 	if (bNotifyBreaks != bNewNotifyBreaks)
 	{
+		if (Chaos::FPhysicsSolver* Solver = GetSolver(*this); Solver && PhysicsProxy)
+		{
+			Solver->EnqueueCommandImmediate(
+				[this, bNewNotifyBreaks]()
+				{
+					FSimulationParameters& SimParams = PhysicsProxy->GetSimParameters();
+					SimParams.bGenerateBreakingData = bNewNotifyBreaks;
+				}
+			);
+		}
 		bNotifyBreaks = bNewNotifyBreaks;
 		UpdateBreakEventRegistration();
 	}
@@ -744,6 +755,16 @@ void UGeometryCollectionComponent::SetNotifyRemovals(bool bNewNotifyRemovals)
 {
 	if (bNotifyRemovals != bNewNotifyRemovals)
 	{
+		if (Chaos::FPhysicsSolver* Solver = GetSolver(*this); Solver && PhysicsProxy)
+		{
+			Solver->EnqueueCommandImmediate(
+				[this, bNewNotifyRemovals]()
+				{
+					FSimulationParameters& SimParams = PhysicsProxy->GetSimParameters();
+					SimParams.bGenerateBreakingData = bNewNotifyRemovals;
+				}
+			);
+		}
 		bNotifyRemovals = bNewNotifyRemovals;
 		UpdateRemovalEventRegistration();
 	}
@@ -753,6 +774,17 @@ void UGeometryCollectionComponent::SetNotifyCrumblings(bool bNewNotifyCrumblings
 {
 	if (bNotifyCrumblings != bNewNotifyCrumblings)
 	{
+		if (Chaos::FPhysicsSolver* Solver = GetSolver(*this); Solver && PhysicsProxy)
+		{
+			Solver->EnqueueCommandImmediate(
+				[this, bNewNotifyCrumblings, bNewCrumblingEventIncludesChildren]()
+				{
+					FSimulationParameters& SimParams = PhysicsProxy->GetSimParameters();
+					SimParams.bGenerateCrumblingData = bNewNotifyCrumblings;
+					SimParams.bGenerateCrumblingChildrenData = bNewCrumblingEventIncludesChildren;
+				}
+			);
+		}
 		bNotifyCrumblings = bNewNotifyCrumblings;
 		bCrumblingEventIncludesChildren = bNewCrumblingEventIncludesChildren;
 		UpdateCrumblingEventRegistration();
@@ -2643,6 +2675,7 @@ void UGeometryCollectionComponent::RegisterAndInitializePhysicsProxy()
 		SimulationParameters.ShockDamagePropagationFactor = DamagePropagationData.ShockDamagePropagationFactor;
 		SimulationParameters.WorldTransform = GetComponentToWorld();
 		SimulationParameters.UserData = static_cast<void*>(&PhysicsUserData);
+		SimulationParameters.bEnableStrainOnCollision = bEnableDamageFromCollision;
 
 		UPhysicalMaterial* EnginePhysicalMaterial = GetPhysicalMaterial();
 		if (ensure(EnginePhysicalMaterial))
@@ -4722,4 +4755,20 @@ TArray<Chaos::FPhysicsObject*> UGeometryCollectionComponent::GetAllPhysicsObject
 		Objects.Add(GetPhysicsObjectById(Index));
 	}
 	return Objects;
+}
+
+void UGeometryCollectionComponent::SetEnableDamageFromCollision(bool bValue)
+{
+	if (Chaos::FPhysicsSolver* Solver = GetSolver(*this); Solver && PhysicsProxy)
+	{
+		Solver->EnqueueCommandImmediate(
+			[this, bValue]()
+			{
+				FSimulationParameters& SimParams = PhysicsProxy->GetSimParameters();
+				SimParams.bEnableStrainOnCollision = bValue;
+			}
+		);
+	}
+
+	bEnableDamageFromCollision = bValue;
 }
\ No newline at end of file
diff --git a/Engine/Source/Runtime/Experimental/GeometryCollectionEngine/Public/GeometryCollection/GeometryCollectionComponent.h b/Engine/Source/Runtime/Experimental/GeometryCollectionEngine/Public/GeometryCollection/GeometryCollectionComponent.h
--- a/Engine/Source/Runtime/Experimental/GeometryCollectionEngine/Public/GeometryCollection/GeometryCollectionComponent.h
+++ b/Engine/Source/Runtime/Experimental/GeometryCollectionEngine/Public/GeometryCollection/GeometryCollectionComponent.h
@@ -647,6 +647,9 @@ public:
 	void SelectEmbeddedGeometry();
 #endif
 
+	UFUNCTION(BlueprintCallable, Category = "ChaosPhysics")
+	void SetEnableDamageFromCollision(bool bValue);
+
 	/** API for getting at geometry collection data */
 	int32 GetNumElements(FName Group) const;
 
@@ -737,6 +740,10 @@ public:
 	UPROPERTY(EditAnywhere, BlueprintReadWrite, Category = "ChaosPhysics|Damage")
 	FGeometryCollectionDamagePropagationData DamagePropagationData;
 
+	/** Whether or not collisions against this geometry collection will apply strain which could cause the geometry collection to fracture. */
+	UPROPERTY(EditDefaultsOnly, BlueprintReadOnly, Category = "ChaosPhysics|Damage")
+	bool bEnableDamageFromCollision;
+
 	/** Allow removal on sleep for the instance if the rest collection has it enabled */
 	UPROPERTY(EditAnywhere, BlueprintReadWrite, Category = "ChaosPhysics|Removal")
 	bool bAllowRemovalOnSleep;
