Geometry collection : Do not create render state when ISMPool is in use and not in Editor

diff --git a/Engine/Source/Runtime/Experimental/GeometryCollectionEngine/Private/GeometryCollection/GeometryCollectionComponent.cpp b/Engine/Source/Runtime/Experimental/GeometryCollectionEngine/Private/GeometryCollection/GeometryCollectionComponent.cpp
--- a/Engine/Source/Runtime/Experimental/GeometryCollectionEngine/Private/GeometryCollection/GeometryCollectionComponent.cpp
+++ b/Engine/Source/Runtime/Experimental/GeometryCollectionEngine/Private/GeometryCollection/GeometryCollectionComponent.cpp
@@ -614,7 +614,8 @@ FPrimitiveSceneProxy* UGeometryCollectionComponent::CreateSceneProxy()
 
 	FPrimitiveSceneProxy* LocalSceneProxy = nullptr;
 
-	if (RestCollection)
+	const bool bUsesISMPool = this->CanUseISMPool();
+	if (RestCollection && !bUsesISMPool)
 	{
 		if (UseNanite(GetScene()->GetShaderPlatform()) &&
 			RestCollection->EnableNanite &&
@@ -4036,7 +4037,7 @@ void UGeometryCollectionComponent::InitializeEmbeddedGeometry()
 
 bool UGeometryCollectionComponent::CanUseISMPool() const 
 {
-	return bChaos_GC_UseISMPool && ISMPool;
+	return bChaos_GC_UseISMPool && ISMPool && GetWorld()->IsGameWorld();
 }
 
 void UGeometryCollectionComponent::RegisterToISMPool()
