Chaos - Collision ISPC perf improvement (minor tweaks)

diff --git a/Engine/Source/Runtime/Experimental/Chaos/Private/Chaos/Collision/PBDCollisionSolver.ispc b/Engine/Source/Runtime/Experimental/Chaos/Private/Chaos/Collision/PBDCollisionSolver.ispc
--- a/Engine/Source/Runtime/Experimental/Chaos/Private/Chaos/Collision/PBDCollisionSolver.ispc
+++ b/Engine/Source/Runtime/Experimental/Chaos/Private/Chaos/Collision/PBDCollisionSolver.ispc
@@ -1,7 +1,5 @@
 // Copyright Epic Games, Inc. All Rights Reserved.
 
-#define EXPLICIT_VECTOR4 1
-
 #include "Math/Vector.isph"
 #include "Math/Quat.isph"
 #include "Math/Transform.isph"
@@ -161,22 +159,22 @@ export struct FPBDCollisionSolver
 //
 
 
-inline uniform bool Body_IsDynamic(const uniform FSolverBody* uniform Body)
+static inline uniform bool Body_IsDynamic(const uniform FSolverBody* uniform Body)
 { 
 	return (Body->InvM > FLOAT_SMALL_NUMBER); 
 }
 
-inline void Body_ApplyPositionDelta(uniform FSolverBody* uniform Body, const uniform FVector3f DP)
+static inline void Body_ApplyPositionDelta(uniform FSolverBody* uniform Body, const uniform FVector3f DP)
 {
 	Body->DP = Body->DP + DP;
 }
 
-inline void Body_ApplyRotationDelta(uniform FSolverBody* uniform Body, const uniform FVector3f DQ)
+static inline void Body_ApplyRotationDelta(uniform FSolverBody* uniform Body, const uniform FVector3f DQ)
 {
 	Body->DQ = Body->DQ + DQ;
 }
 
-static uniform float CalculatePositionErrorNormal(
+static inline uniform float CalculatePositionErrorNormal(
 	uniform FPBDCollisionSolverManifoldPoint* uniform SolverManifoldPoint, 
 	const uniform FSolverBody* uniform Body0, 
 	const uniform FSolverBody* uniform Body1, 
@@ -199,7 +197,7 @@ static uniform float CalculatePositionErrorNormal(
 	return ContactDeltaNormal;
 }
 
-static uniform float CalculatePositionCorrectionNormal(
+static inline uniform float CalculatePositionCorrectionNormal(
 	const uniform float Stiffness,
 	const uniform float ContactDeltaNormal,
 	const uniform float ContactMassNormal,
@@ -209,14 +207,7 @@ static uniform float CalculatePositionCorrectionNormal(
 
 	// The total pushout so far this sub-step
 	// Unilateral constraint: Net-negative impulses not allowed (negative incremental impulses are allowed as long as the net is positive)
-	if ((NetPushOutNormal + PushOutNormal) > 0)
-	{
-		return PushOutNormal;
-	}
-	else
-	{
-		return -NetPushOutNormal;
-	}
+	return (NetPushOutNormal + PushOutNormal) > 0 ? PushOutNormal : -NetPushOutNormal;
 }
 
 
@@ -226,7 +217,7 @@ static uniform float CalculatePositionCorrectionNormal(
 //
 //
 
-static void CollisionSolverApplyPositionCorrectionNormal(
+static inline void CollisionSolverApplyPositionCorrectionNormal(
 	uniform FPBDCollisionSolver* uniform Solver,
 	const uniform float Stiffness,
 	const uniform float ContactDeltaNormal,
@@ -255,7 +246,7 @@ static void CollisionSolverApplyPositionCorrectionNormal(
 	}
 }
 
-static void CollisionSolverSolvePositionNoFriction(
+static inline void CollisionSolverSolvePositionNoFriction(
 	uniform FPBDCollisionSolver* uniform Solver, 
 	const uniform float Dt, 
 	const uniform float MaxPushOut)
@@ -363,4 +354,4 @@ export void SolvePositionNoFrictionJacobi(
 	{
 		CollisionSolver_SolvePositionNoFriction_Jacobi(&Solvers[SolverIndex], Dt, MaxPushOut, JacobiStiffness);
 	}
-}
\ No newline at end of file
+}
