Chaos - Separated Collision Jacobi Solver from the Gauss Seidel solver to prepare the way for custom data layout
- Add a mechanism to allow runtime switching of collision solver types
- Move FCollisionSolverSettings into its own file

diff --git a/Engine/Source/Runtime/Experimental/Chaos/Private/Chaos/Collision/PBDCollisionContainerSolverJacobi.cpp b/Engine/Source/Runtime/Experimental/Chaos/Private/Chaos/Collision/PBDCollisionContainerSolverJacobi.cpp
new file mode 100644
--- /dev/null
+++ b/Engine/Source/Runtime/Experimental/Chaos/Private/Chaos/Collision/PBDCollisionContainerSolverJacobi.cpp
@@ -0,0 +1,583 @@
+// Copyright Epic Games, Inc. All Rights Reserved.
+#include "Chaos/Collision/PBDCollisionContainerSolverJacobi.h"
+
+#include "Chaos/Collision/CollisionApplyType.h"
+#include "Chaos/Collision/PBDCollisionConstraint.h"
+#include "Chaos/CollisionResolution.h"
+#include "Chaos/DebugDrawQueue.h"
+#include "Chaos/Evolution/SolverBodyContainer.h"
+#include "Chaos/Island/IslandManager.h"
+#include "Chaos/Particle/ParticleUtilities.h"
+#include "Chaos/PBDCollisionConstraints.h"
+#include "Chaos/PBDCollisionConstraintsContact.h"
+#include "Chaos/Utilities.h"
+
+// Private includes
+
+#include "ChaosLog.h"
+
+//PRAGMA_DISABLE_OPTIMIZATION
+
+namespace Chaos
+{
+	namespace CVars
+	{
+		extern bool bChaos_PBDCollisionSolver_Position_SolveEnabled;
+		extern bool bChaos_PBDCollisionSolver_Velocity_SolveEnabled;
+
+		extern FRealSingle Chaos_PBDCollisionSolver_AutoStiffness_MassRatio1;
+		extern FRealSingle Chaos_PBDCollisionSolver_AutoStiffness_MassRatio2;
+
+		extern FRealSingle Chaos_PBDCollisionSolver_JacobiStiffness;
+	}
+
+	//////////////////////////////////////////////////////////////////////////////////////////////////
+	//////////////////////////////////////////////////////////////////////////////////////////////////
+	//////////////////////////////////////////////////////////////////////////////////////////////////
+	//////////////////////////////////////////////////////////////////////////////////////////////////
+
+	// Transform the Constraint's local-space data into world space for use by  the collision solver and also calculate tangents, errors, etc
+	FORCEINLINE_DEBUGGABLE void UpdateCollisionSolverContactPointFromConstraint(
+		Private::FPBDCollisionSolverJacobi& Solver, 
+		const int32 SolverPointIndex, 
+		const FPBDCollisionConstraint* Constraint, 
+		const int32 ConstraintPointIndex, 
+		const FRealSingle Dt, 
+		const FSolverBody& Body0, 
+		const FSolverBody& Body1)
+	{
+		const FManifoldPoint& ManifoldPoint = Constraint->GetManifoldPoint(ConstraintPointIndex);
+
+		const FRealSingle Restitution = FRealSingle(Constraint->GetRestitution());
+		const FRealSingle RestitutionVelocityThreshold = FRealSingle(Constraint->GetRestitutionThreshold()) * Dt;
+
+		// World-space shape transforms. Manifold data is currently relative to these spaces
+		const FRigidTransform3& ShapeWorldTransform0 = Constraint->GetShapeWorldTransform0();
+		const FRigidTransform3& ShapeWorldTransform1 = Constraint->GetShapeWorldTransform1();
+
+		// World-space contact points on each shape
+		const FVec3 WorldContact0 = ShapeWorldTransform0.TransformPositionNoScale(FVec3(ManifoldPoint.ContactPoint.ShapeContactPoints[0]));
+		const FVec3 WorldContact1 = ShapeWorldTransform1.TransformPositionNoScale(FVec3(ManifoldPoint.ContactPoint.ShapeContactPoints[1]));
+		const FVec3 WorldContact = FReal(0.5) * (WorldContact0 + WorldContact1);
+		const FVec3f WorldRelativeContact0 = FVec3f(WorldContact - Body0.P());
+		const FVec3f WorldRelativeContact1 = FVec3f(WorldContact - Body1.P());
+
+		// World-space normal
+		const FVec3f WorldContactNormal = FVec3f(ShapeWorldTransform1.TransformVectorNoScale(FVec3(ManifoldPoint.ContactPoint.ShapeContactNormal)));
+
+		// World-space tangents
+		FVec3f WorldContactTangentU = FVec3f::CrossProduct(FVec3f(0, 1, 0), WorldContactNormal);
+		if (!WorldContactTangentU.Normalize(FRealSingle(UE_KINDA_SMALL_NUMBER)))
+		{
+			WorldContactTangentU = FVec3f::CrossProduct(FVec3f(1, 0, 0), WorldContactNormal);
+			WorldContactTangentU = WorldContactTangentU.GetUnsafeNormal();
+		}
+		const FVec3f WorldContactTangentV = FVec3f::CrossProduct(WorldContactNormal, WorldContactTangentU);
+
+		// Calculate contact velocity if we will need it below (restitution and/or frist-contact friction)
+		const bool bNeedContactVelocity = (!ManifoldPoint.Flags.bHasStaticFrictionAnchor) || (Restitution > FRealSingle(0));
+		FVec3f ContactVel = FVec3(0);
+		if (bNeedContactVelocity)
+		{
+			const FVec3f ContactVel0 = Body0.V() + FVec3f::CrossProduct(Body0.W(), WorldRelativeContact0);
+			const FVec3f ContactVel1 = Body1.V() + FVec3f::CrossProduct(Body1.W(), WorldRelativeContact1);
+			ContactVel = ContactVel0 - ContactVel1;
+		}
+
+		// If we have contact data from a previous tick, use it to calculate the lateral position delta we need
+		// to apply to move the contacts back to their original relative locations (i.e., to enforce static friction).
+		// Otherwise, estimate the friction correction from the contact velocity.
+		// NOTE: quadratic shapes use the velocity-based path most of the time, unless the relative motion is very small.
+		FVec3f WorldFrictionDelta = FVec3f(0);
+		if (ManifoldPoint.Flags.bHasStaticFrictionAnchor)
+		{
+			const FVec3f FrictionDelta0 = ShapeWorldTransform0.TransformPositionNoScale(FVec3(ManifoldPoint.ShapeAnchorPoints[0])) - WorldContact0;
+			const FVec3f FrictionDelta1 = ShapeWorldTransform1.TransformPositionNoScale(FVec3(ManifoldPoint.ShapeAnchorPoints[1])) - WorldContact1;
+			WorldFrictionDelta = FrictionDelta0 - FrictionDelta1;
+		}
+		else
+		{
+			// @todo(chaos): consider adding a multiplier to the initial contact friction
+			WorldFrictionDelta = ContactVel * Dt;
+		}
+
+		// The contact point error we are trying to correct in this solver
+		const FRealSingle TargetPhi = ManifoldPoint.TargetPhi;
+		const FVec3f WorldContactDelta = FVec3f(WorldContact0 - WorldContact1);
+		const FRealSingle WorldContactDeltaNormal = FVec3f::DotProduct(WorldContactDelta, WorldContactNormal) - TargetPhi;
+		const FRealSingle WorldContactDeltaTangentU = FVec3f::DotProduct(WorldContactDelta + WorldFrictionDelta, WorldContactTangentU);
+		const FRealSingle WorldContactDeltaTangentV = FVec3f::DotProduct(WorldContactDelta + WorldFrictionDelta, WorldContactTangentV);
+
+		// The target contact velocity, taking restitution into account
+		FRealSingle WorldContactTargetVelocityNormal = FRealSingle(0);
+		if (Restitution > FRealSingle(0))
+		{
+			const FRealSingle ContactVelocityNormal = FVec3f::DotProduct(ContactVel, WorldContactNormal);
+			if (ContactVelocityNormal < -RestitutionVelocityThreshold)
+			{
+				WorldContactTargetVelocityNormal = -Restitution * ContactVelocityNormal;
+			}
+		}
+
+		Solver.InitManifoldPoint(
+			SolverPointIndex,
+			WorldRelativeContact0,
+			WorldRelativeContact1,
+			WorldContactNormal,
+			WorldContactTangentU,
+			WorldContactTangentV,
+			WorldContactDeltaNormal,
+			WorldContactDeltaTangentU,
+			WorldContactDeltaTangentV,
+			WorldContactTargetVelocityNormal);
+	}
+
+	FORCEINLINE_DEBUGGABLE void UpdateCollisionSolverManifoldFromConstraint(
+		Private::FPBDCollisionSolverJacobi& Solver, 
+		const FPBDCollisionConstraint* Constraint, 
+		const FSolverReal Dt, 
+		const int32 ConstraintPointBeginIndex, 
+		const int32 ConstraintPointEndIndex)
+	{
+		const FSolverBody& Body0 = Solver.SolverBody0().SolverBody();
+		const FSolverBody& Body1 = Solver.SolverBody1().SolverBody();
+
+		// Only calculate state for newly added contacts. Normally this is all of them, but maybe not if incremental collision is used by RBAN.
+		// Also we only add active points to the solver's manifold points list
+		for (int32 ConstraintManifoldPointIndex = ConstraintPointBeginIndex; ConstraintManifoldPointIndex < ConstraintPointEndIndex; ++ConstraintManifoldPointIndex)
+		{
+			if (!Constraint->GetManifoldPoint(ConstraintManifoldPointIndex).Flags.bDisabled)
+			{
+				const int32 SolverManifoldPointIndex = Solver.AddManifoldPoint();
+
+				// Transform the constraint contact data into world space for use by the solver
+				// We build this data directly into the solver's world-space contact data which looks a bit odd with "Init" called after but there you go
+				UpdateCollisionSolverContactPointFromConstraint(Solver, SolverManifoldPointIndex, Constraint, ConstraintManifoldPointIndex, Dt, Body0, Body1);
+
+				// Prepare the solver (set constraint-space mass etc)
+				Solver.FinalizeManifoldPoint(SolverManifoldPointIndex);
+			}
+		}
+	}
+
+	FORCEINLINE_DEBUGGABLE void UpdateCollisionSolverFromConstraint(
+		Private::FPBDCollisionSolverJacobi& Solver, 
+		const FPBDCollisionConstraint* Constraint, 
+		const FSolverReal Dt, 
+		const FPBDCollisionSolverSettings& SolverSettings, 
+		bool& bOutPerIterationCollision)
+	{
+		// Friction values. Static and Dynamic friction are applied in the position solve for most shapes.
+		// We can also run in a mode without static friction at all. This is faster but stacking is not possible.
+		const FSolverReal StaticFriction = FSolverReal(Constraint->GetStaticFriction());
+		const FSolverReal DynamicFriction = FSolverReal(Constraint->GetDynamicFriction());
+		FSolverReal PositionStaticFriction = FSolverReal(0);
+		FSolverReal PositionDynamicFriction = FSolverReal(0);
+		FSolverReal VelocityDynamicFriction = FSolverReal(0);
+		if (SolverSettings.NumPositionFrictionIterations > 0)
+		{
+			PositionStaticFriction = StaticFriction;
+			if (!Constraint->HasQuadraticShape())
+			{
+				PositionDynamicFriction = DynamicFriction;
+			}
+			else
+			{
+				// Quadratic shapes don't use PBD dynamic friction - it has issues at slow speeds where the WxR is
+				// less than the position tolerance for friction point matching
+				// @todo(chaos): fix PBD dynamic friction on quadratic shapes
+				VelocityDynamicFriction = DynamicFriction;
+			}
+		}
+		else
+		{
+			VelocityDynamicFriction = DynamicFriction;
+		}
+
+		Solver.SetFriction(PositionStaticFriction, PositionDynamicFriction, VelocityDynamicFriction);
+
+		const FReal SolverStiffness = Constraint->GetStiffness();
+		Solver.SetStiffness(FSolverReal(SolverStiffness));
+
+		Solver.SolverBody0().SetInvMScale(Constraint->GetInvMassScale0());
+		Solver.SolverBody0().SetInvIScale(Constraint->GetInvInertiaScale0());
+		Solver.SolverBody0().SetShockPropagationScale(FReal(1));
+		Solver.SolverBody1().SetInvMScale(Constraint->GetInvMassScale1());
+		Solver.SolverBody1().SetInvIScale(Constraint->GetInvInertiaScale1());
+		Solver.SolverBody1().SetShockPropagationScale(FReal(1));
+
+		bOutPerIterationCollision = (!Constraint->GetUseManifold() || Constraint->GetUseIncrementalCollisionDetection());
+
+		UpdateCollisionSolverManifoldFromConstraint(Solver, Constraint, Dt, 0, Constraint->NumManifoldPoints());
+	}
+
+	FORCEINLINE_DEBUGGABLE void UpdateCollisionConstraintFromSolver(
+		FPBDCollisionConstraint* Constraint, 
+		const Private::FPBDCollisionSolverJacobi& Solver, 
+		const FSolverReal Dt)
+	{
+		Constraint->ResetSolverResults();
+
+		// NOTE: We only put the non-pruned manifold points into the solver so the ManifoldPointIndex and
+		// SolverManifoldPointIndex do not necessarily match. See GatherManifoldPoints
+		int32 SolverManifoldPointIndex = 0;
+		for (int32 ManifoldPointIndex = 0; ManifoldPointIndex < Constraint->NumManifoldPoints(); ++ManifoldPointIndex)
+		{
+			FSolverVec3 NetPushOut = FSolverVec3(0);
+			FSolverVec3 NetImpulse = FSolverVec3(0);
+			FSolverReal StaticFrictionRatio = FSolverReal(0);
+
+			if (!Constraint->GetManifoldPoint(ManifoldPointIndex).Flags.bDisabled)
+			{
+				NetPushOut = Solver.GetNetPushOut(SolverManifoldPointIndex);
+
+				NetImpulse = Solver.GetNetImpulse(SolverManifoldPointIndex);
+
+				StaticFrictionRatio = Solver.GetStaticFrictionRatio(SolverManifoldPointIndex);
+
+				++SolverManifoldPointIndex;
+			}
+
+			// NOTE: We call this even for points we did not run the solver for (but with zero results)
+			Constraint->SetSolverResults(ManifoldPointIndex,
+				NetPushOut,
+				NetImpulse,
+				StaticFrictionRatio,
+				Dt);
+		}
+
+		Constraint->EndTick();
+	}
+
+	//////////////////////////////////////////////////////////////////////////////////////////////////
+	//////////////////////////////////////////////////////////////////////////////////////////////////
+	//////////////////////////////////////////////////////////////////////////////////////////////////
+	//////////////////////////////////////////////////////////////////////////////////////////////////
+
+
+	FPBDCollisionContainerSolverJacobi::FPBDCollisionContainerSolverJacobi(const FPBDCollisionConstraints& InConstraintContainer, const int32 InPriority)
+		: FConstraintContainerSolver(InPriority)
+		, ConstraintContainer(InConstraintContainer)
+		, bPerIterationCollisionDetection(false)
+	{
+#if !UE_BUILD_SHIPPING && !UE_BUILD_TEST && INTEL_ISPC
+		Private::FPBDCollisionSolverJacobiHelper::CheckISPC();
+#endif
+	}
+
+	FPBDCollisionContainerSolverJacobi::~FPBDCollisionContainerSolverJacobi()
+	{
+	}
+
+	void FPBDCollisionContainerSolverJacobi::Reset(const int32 MaxCollisions)
+	{
+		// A over-allocation policy to avoid reallocation every frame in the common case where a pile of objects is dropped
+		// and the number of contacts increases every tick.
+		if (MaxCollisions > CollisionSolvers.Max())
+		{
+			CollisionSolvers.Reset((5 * MaxCollisions) / 4);	// +25%
+		}
+
+		CollisionSolvers.Reset(MaxCollisions);
+		CollisionConstraints.Reset(CollisionSolvers.Max());
+		bCollisionConstraintPerIterationCollisionDetection.Reset(CollisionSolvers.Max());
+
+		// Just set the array size for these right away - all data will be initialized later and there is no constructor
+		// We use CollisionConstraints.Num() to track how many constraints we currently have
+		CollisionSolvers.SetNumUninitialized(MaxCollisions);
+		bCollisionConstraintPerIterationCollisionDetection.SetNumUninitialized(MaxCollisions);
+	}
+
+	void FPBDCollisionContainerSolverJacobi::AddConstraints()
+	{
+		Reset(ConstraintContainer.NumConstraints());
+
+		for (FPBDCollisionConstraintHandle* ConstraintHandle : ConstraintContainer.GetConstraintHandles())
+		{
+			FPBDCollisionConstraint& Constraint = ConstraintHandle->GetContact();
+
+			AddConstraint(Constraint);
+		}
+	}
+
+	void FPBDCollisionContainerSolverJacobi::AddConstraints(const TArrayView<Private::FPBDIslandConstraint>& IslandConstraints)
+	{
+		for (Private::FPBDIslandConstraint& IslandConstraint : IslandConstraints)
+		{
+			// We will only ever be given constraints from our container (asserts in non-shipping)
+			FPBDCollisionConstraint& Constraint = IslandConstraint.GetConstraint()->AsUnsafe<FPBDCollisionConstraintHandle>()->GetContact();
+
+			AddConstraint(Constraint);
+		}
+	}
+
+	void FPBDCollisionContainerSolverJacobi::AddConstraint(FPBDCollisionConstraint& Constraint)
+	{
+		// NOTE: No need to add to CollisionSolvers or bCollisionConstraintPerIterationCollisionDetection
+		const int32 Index = CollisionConstraints.Add(&Constraint);
+
+		CollisionSolvers[Index].Reset();
+	}
+
+
+	void FPBDCollisionContainerSolverJacobi::AddBodies(FSolverBodyContainer& SolverBodyContainer)
+	{
+		for (int32 SolverIndex = 0, SolverEndIndex = NumSolvers(); SolverIndex < SolverEndIndex; ++SolverIndex)
+		{
+			Private::FPBDCollisionSolverJacobi& CollisionSolver = CollisionSolvers[SolverIndex];
+			FPBDCollisionConstraint* Constraint = CollisionConstraints[SolverIndex];
+			check(Constraint != nullptr);
+
+			// Find the solver bodies for the particles we constrain. This will add them to the container
+			// if they aren't there already, and ensure that they are populated with the latest data.
+			FSolverBody* Body0 = SolverBodyContainer.FindOrAdd(Constraint->GetParticle0());
+			FSolverBody* Body1 = SolverBodyContainer.FindOrAdd(Constraint->GetParticle1());
+			check(Body0 != nullptr);
+			check(Body1 != nullptr);
+
+			CollisionSolver.SetSolverBodies(*Body0, *Body1);
+		}
+	}
+
+	void FPBDCollisionContainerSolverJacobi::GatherInput(const FReal Dt)
+	{
+		GatherInput(Dt, 0, NumSolvers());
+	}
+
+	void FPBDCollisionContainerSolverJacobi::CachePrefetchSolver(const int32 ConstraintIndex) const
+	{
+		if (ConstraintIndex < NumSolvers())
+		{
+			FPlatformMisc::PrefetchBlock(&CollisionSolvers[ConstraintIndex], sizeof(FPBDCollisionContainerSolverJacobi));
+		}
+	}
+
+	void FPBDCollisionContainerSolverJacobi::GatherInput(const FReal InDt, const int32 BeginIndex, const int32 EndIndex)
+	{
+		// NOTE: may be called in parallel. Should not change the container or any elements outside of [BeginIndex, EndIndex)
+
+		check(BeginIndex >= 0);
+		check(EndIndex <= NumSolvers());
+
+		const FSolverReal Dt = FSolverReal(InDt);
+
+		const int32 PrefetchCount = 4;
+		for (int32 PrefetchIndex = 0; PrefetchIndex < PrefetchCount; ++PrefetchIndex)
+		{
+			CachePrefetchSolver(PrefetchIndex);
+		}
+
+		bool bWantPerIterationCollisionDetection = false;
+		for (int32 ConstraintIndex = BeginIndex; ConstraintIndex < EndIndex; ++ConstraintIndex)
+		{
+			Private::FPBDCollisionSolverJacobi& CollisionSolver = CollisionSolvers[ConstraintIndex];
+			FPBDCollisionConstraint* Constraint = CollisionConstraints[ConstraintIndex];
+			bool& bPerIterationCollision = bCollisionConstraintPerIterationCollisionDetection[ConstraintIndex];
+
+			UpdateCollisionSolverFromConstraint(CollisionSolver, Constraint, Dt, ConstraintContainer.GetSolverSettings(), bPerIterationCollision);
+
+			// We need to run collision every iteration if we are not using manifolds, or are using incremental manifolds
+			bWantPerIterationCollisionDetection = bWantPerIterationCollisionDetection || bPerIterationCollision;
+
+			CachePrefetchSolver(ConstraintIndex + PrefetchCount);
+		}
+
+		if (bWantPerIterationCollisionDetection)
+		{
+			// We should lock here? We only ever set to true or do nothing so I think it doesn't matter if this happens on multiple threads...
+			bPerIterationCollisionDetection = true;
+		}
+	}
+
+	void FPBDCollisionContainerSolverJacobi::ScatterOutput(const FReal Dt)
+	{
+		ScatterOutput(Dt, 0, NumSolvers());
+	}
+
+	void FPBDCollisionContainerSolverJacobi::ScatterOutput(const FReal InDt, const int32 BeginIndex, const int32 EndIndex)
+	{
+		// NOTE: may be called in parallel. Should not change the container or any elements outside of [BeginIndex, EndIndex)
+
+		check(BeginIndex >= 0);
+		check(EndIndex <= NumSolvers());
+
+		const FSolverReal Dt = FSolverReal(InDt);
+
+		for (int32 ConstraintIndex = BeginIndex; ConstraintIndex < EndIndex; ++ConstraintIndex)
+		{
+			Private::FPBDCollisionSolverJacobi& CollisionSolver = CollisionSolvers[ConstraintIndex];
+			FPBDCollisionConstraint* Constraint = CollisionConstraints[ConstraintIndex];
+
+			UpdateCollisionConstraintFromSolver(Constraint, CollisionSolver, Dt);
+
+			// Reset the collision solver here as the pointers will be invalid on the next tick
+			// @todo(chaos): Pointers are always reinitalized before use next tick so this isn't strictly necessary
+			CollisionSolver.Reset();
+		}
+	}
+
+	void FPBDCollisionContainerSolverJacobi::ApplyPositionConstraints(const FReal Dt, const int32 It, const int32 NumIts)
+	{
+		SolvePositionImpl(Dt, It, NumIts, 0, NumSolvers(), ConstraintContainer.GetSolverSettings());
+	}
+
+	void FPBDCollisionContainerSolverJacobi::ApplyVelocityConstraints(const FReal Dt, const int32 It, const int32 NumIts)
+	{
+		SolveVelocityImpl(Dt, It, NumIts, 0, NumSolvers(), ConstraintContainer.GetSolverSettings());
+	}
+
+	void FPBDCollisionContainerSolverJacobi::ApplyProjectionConstraints(const FReal Dt, const int32 It, const int32 NumIts)
+	{
+		// Not supported for collisions
+	}
+
+	void FPBDCollisionContainerSolverJacobi::UpdatePositionShockPropagation(const FReal Dt, const int32 It, const int32 NumIts, const int32 BeginIndex, const int32 EndIndex, const FPBDCollisionSolverSettings& SolverSettings)
+	{
+		// If this is the first shock propagation iteration, enable it on each solver
+		const bool bEnableShockPropagation = (It == NumIts - SolverSettings.NumPositionShockPropagationIterations);
+		if (bEnableShockPropagation)
+		{
+			for (int32 SolverIndex = BeginIndex; SolverIndex < EndIndex; ++SolverIndex)
+			{
+				CollisionSolvers[SolverIndex].EnablePositionShockPropagation();
+			}
+		}
+	}
+
+	void FPBDCollisionContainerSolverJacobi::UpdateVelocityShockPropagation(const FReal Dt, const int32 It, const int32 NumIts, const int32 BeginIndex, const int32 EndIndex, const FPBDCollisionSolverSettings& SolverSettings)
+	{
+		// Set/reset the shock propagation based on current iteration. The position solve may
+		// have left the bodies with a mass scale and we want to change or reset it.
+		const bool bEnableShockPropagation = (It == NumIts - SolverSettings.NumVelocityShockPropagationIterations);
+		if (bEnableShockPropagation)
+		{
+			for (int32 SolverIndex = BeginIndex; SolverIndex < EndIndex; ++SolverIndex)
+			{
+				CollisionSolvers[SolverIndex].EnableVelocityShockPropagation();
+			}
+		}
+		else if (It == 0)
+		{
+			for (int32 SolverIndex = BeginIndex; SolverIndex < EndIndex; ++SolverIndex)
+			{
+				CollisionSolvers[SolverIndex].DisableShockPropagation();
+			}
+		}
+	}
+
+	bool FPBDCollisionContainerSolverJacobi::SolvePositionImpl(const FReal InDt, const int32 It, const int32 NumIts, const int32 BeginIndex, const int32 EndIndex, const FPBDCollisionSolverSettings& SolverSettings)
+	{
+		SCOPE_CYCLE_COUNTER(STAT_Collisions_Apply);
+		if (!CVars::bChaos_PBDCollisionSolver_Position_SolveEnabled)
+		{
+			return false;
+		}
+
+		if (EndIndex <= BeginIndex)
+		{
+			return false;
+		}
+
+		UpdatePositionShockPropagation(InDt, It, NumIts, BeginIndex, EndIndex, SolverSettings);
+
+		// We run collision detection here under two conditions (normally it is run after Integration and before the constraint solver phase):
+		// 1) When deferring collision detection until the solver phase for better joint-collision behaviour (RBAN). In this case, we only do this on the first iteration.
+		// 2) When using no manifolds or incremental manifolds, where we may add/replace manifold points every iteration.
+		const bool bRunDeferredCollisionDetection = (It == 0) && ConstraintContainer.GetDetectorSettings().bDeferNarrowPhase;
+		if (bRunDeferredCollisionDetection || bPerIterationCollisionDetection)
+		{
+			UpdateCollisions(InDt, BeginIndex, EndIndex);
+		}
+
+		// Only apply friction for the last few (tunable) iterations
+		// Adjust max pushout to attempt to make it iteration count independent
+		const FSolverReal Dt = FSolverReal(InDt);
+		const bool bApplyStaticFriction = (It >= (NumIts - SolverSettings.NumPositionFrictionIterations));
+		const FSolverReal MaxPushOut = (SolverSettings.MaxPushOutVelocity > 0) ? (FSolverReal(SolverSettings.MaxPushOutVelocity) * Dt) / FSolverReal(NumIts) : FSolverReal(0);
+
+		// Apply the position correction
+		if (bApplyStaticFriction)
+		{
+			Private::FPBDCollisionSolverJacobiHelper::SolvePositionWithFriction(MakeArrayView(&CollisionSolvers[BeginIndex], EndIndex - BeginIndex), Dt, MaxPushOut);
+		}
+		else
+		{
+			Private::FPBDCollisionSolverJacobiHelper::SolvePositionNoFriction(MakeArrayView(&CollisionSolvers[BeginIndex], EndIndex - BeginIndex), Dt, MaxPushOut);
+		}
+
+		return true;
+	}
+
+	bool FPBDCollisionContainerSolverJacobi::SolveVelocityImpl(const FReal InDt, const int32 It, const int32 NumIts, const int32 BeginIndex, const int32 EndIndex, const FPBDCollisionSolverSettings& SolverSettings)
+	{
+		SCOPE_CYCLE_COUNTER(STAT_Collisions_ApplyPushOut);
+		if (!CVars::bChaos_PBDCollisionSolver_Velocity_SolveEnabled)
+		{
+			return false;
+		}
+
+		if (EndIndex <= BeginIndex)
+		{
+			return false;
+		}
+
+		UpdateVelocityShockPropagation(InDt, It, NumIts, BeginIndex, EndIndex, SolverSettings);
+
+		const FSolverReal Dt = FSolverReal(InDt);
+		const bool bApplyDynamicFriction = (It >= NumIts - SolverSettings.NumVelocityFrictionIterations);
+
+		Private::FPBDCollisionSolverJacobiHelper::SolveVelocity(MakeArrayView(&CollisionSolvers[BeginIndex], EndIndex - BeginIndex), Dt, bApplyDynamicFriction);
+
+		return true;
+	}
+
+	void FPBDCollisionContainerSolverJacobi::UpdateCollisions(const FReal InDt, const int32 BeginIndex, const int32 EndIndex)
+	{
+		const FSolverReal Dt = FSolverReal(InDt);
+		const bool bDeferredCollisionDetection = ConstraintContainer.GetDetectorSettings().bDeferNarrowPhase;
+
+		bool bNeedsAnotherIteration = false;
+		for (int32 SolverIndex = BeginIndex; SolverIndex < EndIndex; ++SolverIndex)
+		{
+			if (bDeferredCollisionDetection || bCollisionConstraintPerIterationCollisionDetection[SolverIndex])
+			{
+				Private::FPBDCollisionSolverJacobi& CollisionSolver = CollisionSolvers[SolverIndex];
+				FPBDCollisionConstraint* Constraint = CollisionConstraints[SolverIndex];
+
+				// Run collision detection at the current transforms including any correction from previous iterations
+				const FSolverBody& Body0 = CollisionSolver.SolverBody0().SolverBody();
+				const FSolverBody& Body1 = CollisionSolver.SolverBody1().SolverBody();
+				const FRigidTransform3 CorrectedActorWorldTransform0 = FRigidTransform3(Body0.CorrectedActorP(), Body0.CorrectedActorQ());
+				const FRigidTransform3 CorrectedActorWorldTransform1 = FRigidTransform3(Body1.CorrectedActorP(), Body1.CorrectedActorQ());
+				const FRigidTransform3 CorrectedShapeWorldTransform0 = Constraint->GetShapeRelativeTransform0() * CorrectedActorWorldTransform0;
+				const FRigidTransform3 CorrectedShapeWorldTransform1 = Constraint->GetShapeRelativeTransform1() * CorrectedActorWorldTransform1;
+
+				// @todo(chaos): this is ugly - pass these to the required functions instead and remove from the constraint class
+				// This is now only needed for LevelSet collision (see UpdateLevelsetLevelsetConstraint)
+				Constraint->SetSolverBodies(&Body0, &Body1);
+
+				// Reset the manifold if we are not using manifolds (we just use the first manifold point)
+				if (!Constraint->GetUseManifold())
+				{
+					Constraint->ResetActiveManifoldContacts();
+					CollisionSolver.ResetManifold();
+				}
+
+				// We need to know how many points were added to the manifold
+				const int32 BeginPointIndex = Constraint->NumManifoldPoints();
+
+				// NOTE: We deliberately have not updated the ShapwWorldTranforms on the constraint. If we did that, we would calculate 
+				// errors incorrectly in UpdateManifoldPoints, because the solver assumes nothing has been moved as we iterate (we accumulate 
+				// corrections that will be applied later.)
+				Constraint->ResetPhi(Constraint->GetCullDistance());
+				Collisions::UpdateConstraint(*Constraint, CorrectedShapeWorldTransform0, CorrectedShapeWorldTransform1, Dt);
+
+				// Update the manifold based on the new or updated contacts
+				UpdateCollisionSolverManifoldFromConstraint(CollisionSolver, Constraint, Dt, BeginPointIndex, Constraint->NumManifoldPoints());
+
+				Constraint->SetSolverBodies(nullptr, nullptr);
+			}
+		}
+	}
+
+}
diff --git a/Engine/Source/Runtime/Experimental/Chaos/Private/Chaos/Collision/PBDCollisionSolver.cpp b/Engine/Source/Runtime/Experimental/Chaos/Private/Chaos/Collision/PBDCollisionSolver.cpp
--- a/Engine/Source/Runtime/Experimental/Chaos/Private/Chaos/Collision/PBDCollisionSolver.cpp
+++ b/Engine/Source/Runtime/Experimental/Chaos/Private/Chaos/Collision/PBDCollisionSolver.cpp
@@ -19,8 +19,6 @@
 
 //PRAGMA_DISABLE_OPTIMIZATION
 
-DEFINE_LOG_CATEGORY(LogChaosCollision);
-
 // TEMP: to be removed
 DECLARE_CYCLE_STAT(TEXT("SolvePositionNoFriction"), STAT_SolvePositionNoFriction, STATGROUP_ChaosConstraintSolver);
 
@@ -38,7 +36,7 @@ namespace Chaos
 		bool bChaos_PBDCollisionSolver_ISPC = false;	// NOTE: WORK IN PROGRESS. Not ready to be enabled.
 		bool bChaos_PBDCollisionSolver_UseJacobiPairSolver = false;
 		FAutoConsoleVariableRef CVarChaos_PBDCollisionSolver_ISPC(TEXT("p.Chaos.PBDCollisionSolver.ISPC"), bChaos_PBDCollisionSolver_ISPC, TEXT("Use ISPC collision solver (WIP)"));
-		FAutoConsoleVariableRef CVarChaosPBDCollisionSolverUseJacobiPairSolver(TEXT("p.Chaos.PBDCollisionSolver.UseJacobi"), bChaos_PBDCollisionSolver_UseJacobiPairSolver, TEXT("Whether to use the Jacobi collision pair solver"));
+		FAutoConsoleVariableRef CVarChaosPBDCollisionSolverUseJacobiPairSolver(TEXT("p.Chaos.PBDCollisionSolver.UseJacobi"), bChaos_PBDCollisionSolver_UseJacobiPairSolver, TEXT("Whether to use the Jacobi collision pair solver (as option in main collision solver)"));
 
 		//
 		// Position Solver Settings
diff --git a/Engine/Source/Runtime/Experimental/Chaos/Private/Chaos/Collision/PBDCollisionSolverJacobi.cpp b/Engine/Source/Runtime/Experimental/Chaos/Private/Chaos/Collision/PBDCollisionSolverJacobi.cpp
new file mode 100644
--- /dev/null
+++ b/Engine/Source/Runtime/Experimental/Chaos/Private/Chaos/Collision/PBDCollisionSolverJacobi.cpp
@@ -0,0 +1,128 @@
+// Copyright Epic Games, Inc. All Rights Reserved.
+#include "Chaos/Collision/PBDCollisionSolverJacobi.h"
+
+#include "Chaos/Collision/CollisionApplyType.h"
+#include "Chaos/Collision/PBDCollisionConstraint.h"
+#include "Chaos/Collision/SolverCollisionContainer.h"
+#include "Chaos/CollisionResolution.h"
+#include "Chaos/CollisionResolutionTypes.h"
+#include "Chaos/CollisionResolutionUtil.h"
+#include "Chaos/Evolution/SolverBodyContainer.h"
+#include "Chaos/Particle/ParticleUtilities.h"
+#include "Chaos/PBDCollisionConstraints.h"
+#include "Chaos/PBDCollisionConstraintsContact.h"
+#include "Chaos/Utilities.h"
+
+//PRAGMA_DISABLE_OPTIMIZATION
+
+namespace Chaos
+{
+	namespace CVars
+	{
+		extern float Chaos_PBDCollisionSolver_Position_MinInvMassScale;
+		extern float Chaos_PBDCollisionSolver_Velocity_MinInvMassScale;
+	}
+
+	namespace Private
+	{
+		//////////////////////////////////////////////////////////////////////////////////////////////////
+		//////////////////////////////////////////////////////////////////////////////////////////////////
+		//////////////////////////////////////////////////////////////////////////////////////////////////
+		//////////////////////////////////////////////////////////////////////////////////////////////////
+
+
+		void FPBDCollisionSolverJacobi::EnablePositionShockPropagation()
+		{
+			SetShockPropagationInvMassScale(CVars::Chaos_PBDCollisionSolver_Position_MinInvMassScale);
+		}
+
+		void FPBDCollisionSolverJacobi::EnableVelocityShockPropagation()
+		{
+			SetShockPropagationInvMassScale(CVars::Chaos_PBDCollisionSolver_Velocity_MinInvMassScale);
+		}
+
+		void FPBDCollisionSolverJacobi::DisableShockPropagation()
+		{
+			SetShockPropagationInvMassScale(FReal(1));
+		}
+
+		void FPBDCollisionSolverJacobi::SetShockPropagationInvMassScale(const FSolverReal InvMassScale)
+		{
+			FConstraintSolverBody& Body0 = SolverBody0();
+			FConstraintSolverBody& Body1 = SolverBody1();
+
+			// Shock propagation decreases the inverse mass of bodies that are lower in the pile
+			// of objects. This significantly improves stability of heaps and stacks. Height in the pile is indictaed by the "level". 
+			// No need to set an inverse mass scale if the other body is kinematic (with inv mass of 0).
+			// Bodies at the same level do not take part in shock propagation.
+			if (Body0.IsDynamic() && Body1.IsDynamic() && (Body0.Level() != Body1.Level()))
+			{
+				// Set the inv mass scale of the "lower" body to make it heavier
+				bool bInvMassUpdated = false;
+				if (Body0.Level() < Body1.Level())
+				{
+					if (Body0.ShockPropagationScale() != InvMassScale)
+					{
+						Body0.SetShockPropagationScale(InvMassScale);
+						bInvMassUpdated = true;
+					}
+				}
+				else
+				{
+					if (Body1.ShockPropagationScale() != InvMassScale)
+					{
+						Body1.SetShockPropagationScale(InvMassScale);
+						bInvMassUpdated = true;
+					}
+				}
+
+				// If the masses changed, we need to rebuild the contact mass for each manifold point
+				if (bInvMassUpdated)
+				{
+					for (int32 PointIndex = 0; PointIndex < NumManifoldPoints(); ++PointIndex)
+					{
+						State.ManifoldPoints.UpdateMassNormal(PointIndex, Body0, Body1);
+					}
+				}
+			}
+		}
+
+		//////////////////////////////////////////////////////////////////////////////////////////////////
+		//////////////////////////////////////////////////////////////////////////////////////////////////
+		//////////////////////////////////////////////////////////////////////////////////////////////////
+		//////////////////////////////////////////////////////////////////////////////////////////////////
+
+		void FPBDCollisionSolverJacobiHelper::SolvePositionNoFriction(const TArrayView<FPBDCollisionSolverJacobi>& CollisionSolvers, const FSolverReal Dt, const FSolverReal MaxPushOut)
+		{
+			//SCOPE_CYCLE_COUNTER(STAT_SolvePositionNoFriction);
+
+			for (FPBDCollisionSolverJacobi& CollisionSolver : CollisionSolvers)
+			{
+				CollisionSolver.SolvePositionNoFriction(Dt, MaxPushOut);
+			}
+		}
+
+		void FPBDCollisionSolverJacobiHelper::SolvePositionWithFriction(const TArrayView<FPBDCollisionSolverJacobi>& CollisionSolvers, const FSolverReal Dt, const FSolverReal MaxPushOut)
+		{
+			for (FPBDCollisionSolverJacobi& CollisionSolver : CollisionSolvers)
+			{
+				CollisionSolver.SolvePositionWithFriction(Dt, MaxPushOut);
+			}
+		}
+
+		void FPBDCollisionSolverJacobiHelper::SolveVelocity(const TArrayView<FPBDCollisionSolverJacobi>& CollisionSolvers, const FSolverReal Dt, const bool bApplyDynamicFriction)
+		{
+			for (FPBDCollisionSolverJacobi& CollisionSolver : CollisionSolvers)
+			{
+				CollisionSolver.SolveVelocity(Dt, bApplyDynamicFriction);
+			}
+		}
+
+		void FPBDCollisionSolverJacobiHelper::CheckISPC()
+		{
+#if INTEL_ISPC
+#endif
+		}
+
+	}	// namespace Private
+}	// namespace Chaos
\ No newline at end of file
diff --git a/Engine/Source/Runtime/Experimental/Chaos/Private/Chaos/Collision/SolverCollisionContainer.cpp b/Engine/Source/Runtime/Experimental/Chaos/Private/Chaos/Collision/SolverCollisionContainer.cpp
--- a/Engine/Source/Runtime/Experimental/Chaos/Private/Chaos/Collision/SolverCollisionContainer.cpp
+++ b/Engine/Source/Runtime/Experimental/Chaos/Private/Chaos/Collision/SolverCollisionContainer.cpp
@@ -3,7 +3,6 @@
 
 #include "Chaos/Collision/CollisionApplyType.h"
 #include "Chaos/Collision/PBDCollisionConstraint.h"
-#include "Chaos/Collision/PBDCollisionSolver.h"
 #include "Chaos/CollisionResolution.h"
 #include "Chaos/DebugDrawQueue.h"
 #include "Chaos/Evolution/SolverBodyContainer.h"
@@ -50,24 +49,6 @@ namespace Chaos
 		FAutoConsoleVariableRef CVarChaosPBDCollisionSolverJacobiRotationTolerance(TEXT("p.Chaos.PBDCollisionSolver.JacobiRotationTolerance"), Chaos_PBDCollisionSolver_JacobiRotationTolerance, TEXT(""));
 	}
 
-	using namespace CVars;
-	using namespace Private;
-
-	//////////////////////////////////////////////////////////////////////////////////////////////////
-	//////////////////////////////////////////////////////////////////////////////////////////////////
-	//////////////////////////////////////////////////////////////////////////////////////////////////
-	//////////////////////////////////////////////////////////////////////////////////////////////////
-
-
-	FPBDCollisionSolverSettings::FPBDCollisionSolverSettings()
-		: MaxPushOutVelocity(0)
-		, NumPositionFrictionIterations(4)
-		, NumVelocityFrictionIterations(1)
-		, NumPositionShockPropagationIterations(3)
-		, NumVelocityShockPropagationIterations(1)
-	{
-	}
-
 
 	//////////////////////////////////////////////////////////////////////////////////////////////////
 	//////////////////////////////////////////////////////////////////////////////////////////////////
@@ -75,7 +56,7 @@ namespace Chaos
 	//////////////////////////////////////////////////////////////////////////////////////////////////
 
 	// Transform the Constraint's local-space data into world space for use by  the collision solver and also calculate tangents, errors, etc
-	FORCEINLINE_DEBUGGABLE void UpdateCollisionSolverContactPointFromConstraint(FPBDCollisionSolver& Solver, const int32 SolverPointIndex, const FPBDCollisionConstraint* Constraint, const int32 ConstraintPointIndex, const FRealSingle Dt, const FSolverBody& Body0, const FSolverBody& Body1)
+	FORCEINLINE_DEBUGGABLE void UpdateCollisionSolverContactPointFromConstraint(Private::FPBDCollisionSolver& Solver, const int32 SolverPointIndex, const FPBDCollisionConstraint* Constraint, const int32 ConstraintPointIndex, const FRealSingle Dt, const FSolverBody& Body0, const FSolverBody& Body1)
 	{
 		const FManifoldPoint& ManifoldPoint = Constraint->GetManifoldPoint(ConstraintPointIndex);
 
@@ -162,7 +143,7 @@ namespace Chaos
 		WorldContactPoint.ContactTargetVelocityNormal = WorldContactTargetVelocityNormal;
 	}
 
-	FORCEINLINE_DEBUGGABLE void UpdateCollisionSolverManifoldFromConstraint(FPBDCollisionSolver& Solver, const FPBDCollisionConstraint* Constraint, const FSolverReal Dt, const int32 ConstraintPointBeginIndex, const int32 ConstraintPointEndIndex)
+	FORCEINLINE_DEBUGGABLE void UpdateCollisionSolverManifoldFromConstraint(Private::FPBDCollisionSolver& Solver, const FPBDCollisionConstraint* Constraint, const FSolverReal Dt, const int32 ConstraintPointBeginIndex, const int32 ConstraintPointEndIndex)
 	{
 		const FSolverBody& Body0 = Solver.SolverBody0().SolverBody();
 		const FSolverBody& Body1 = Solver.SolverBody1().SolverBody();
@@ -185,7 +166,7 @@ namespace Chaos
 		}
 	}
 
-	FORCEINLINE_DEBUGGABLE void UpdateCollisionSolverFromConstraint(FPBDCollisionSolver& Solver, const FPBDCollisionConstraint* Constraint, const FSolverReal Dt, const FPBDCollisionSolverSettings& SolverSettings, bool& bOutPerIterationCollision)
+	FORCEINLINE_DEBUGGABLE void UpdateCollisionSolverFromConstraint(Private::FPBDCollisionSolver& Solver, const FPBDCollisionConstraint* Constraint, const FSolverReal Dt, const FPBDCollisionSolverSettings& SolverSettings, bool& bOutPerIterationCollision)
 	{
 		// Friction values. Static and Dynamic friction are applied in the position solve for most shapes.
 		// We can also run in a mode without static friction at all. This is faster but stacking is not possible.
@@ -231,7 +212,7 @@ namespace Chaos
 		UpdateCollisionSolverManifoldFromConstraint(Solver, Constraint, Dt, 0, Constraint->NumManifoldPoints());
 	}
 
-	FORCEINLINE_DEBUGGABLE void UpdateCollisionConstraintFromSolver(FPBDCollisionConstraint* Constraint, const FPBDCollisionSolver& Solver, const FSolverReal Dt)
+	FORCEINLINE_DEBUGGABLE void UpdateCollisionConstraintFromSolver(FPBDCollisionConstraint* Constraint, const Private::FPBDCollisionSolver& Solver, const FSolverReal Dt)
 	{
 		Constraint->ResetSolverResults();
 
@@ -244,7 +225,7 @@ namespace Chaos
 			FSolverVec3 NetImpulse = FSolverVec3(0);
 			FSolverReal StaticFrictionRatio = FSolverReal(0);
 
-			const FPBDCollisionSolverManifoldPoint& SolverManifoldPoint = Solver.GetManifoldPoint(SolverManifoldPointIndex);
+			const Private::FPBDCollisionSolverManifoldPoint& SolverManifoldPoint = Solver.GetManifoldPoint(SolverManifoldPointIndex);
 			if (!Constraint->GetManifoldPoint(ManifoldPointIndex).Flags.bDisabled)
 			{
 				NetPushOut =
@@ -305,7 +286,7 @@ namespace Chaos
 		CollisionSolvers.Reset(MaxCollisions);
 		CollisionConstraints.Reset(CollisionSolvers.Max());
 		bCollisionConstraintPerIterationCollisionDetection.Reset(CollisionSolvers.Max());
-		CollisionSolverManifoldPoints.Reset(CollisionSolvers.Max() * FPBDCollisionSolver::MaxPointsPerConstraint);
+		CollisionSolverManifoldPoints.Reset(CollisionSolvers.Max() * Private::FPBDCollisionSolver::MaxPointsPerConstraint);
 
 		// Just set the array size for these right away - all data will be initialized later and there is no constructor
 		// We use CollisionConstraints.Num() to track how many constraints we currently have
@@ -345,7 +326,7 @@ namespace Chaos
 		// NOTE: we don't know how many points we will create if collision detection is deferred or incrememntal, so just allocate space for the max allowed
 		// @todo(chaos): see if we can do better - we don't want to pay for this logic when it is only used rarely
 		const bool bDeferredCollisionDetection = ConstraintContainer.GetDetectorSettings().bDeferNarrowPhase;
-		const int32 ManifoldPointMax = (bDeferredCollisionDetection || Constraint.GetUseIncrementalManifold()) ? FPBDCollisionSolver::MaxPointsPerConstraint : Constraint.NumManifoldPoints();
+		const int32 ManifoldPointMax = (bDeferredCollisionDetection || Constraint.GetUseIncrementalManifold()) ? Private::FPBDCollisionSolver::MaxPointsPerConstraint : Constraint.NumManifoldPoints();
 		if (ManifoldPointMax > 0)
 		{
 			const int32 ManifoldPointStart = CollisionSolverManifoldPoints.Num();
@@ -367,7 +348,7 @@ namespace Chaos
 	{
 		for (int32 SolverIndex = 0, SolverEndIndex = NumSolvers(); SolverIndex < SolverEndIndex; ++SolverIndex)
 		{
-			FPBDCollisionSolver& CollisionSolver = CollisionSolvers[SolverIndex];
+			Private::FPBDCollisionSolver& CollisionSolver = CollisionSolvers[SolverIndex];
 			FPBDCollisionConstraint* Constraint = CollisionConstraints[SolverIndex];
 			check(Constraint != nullptr);
 
@@ -394,7 +375,7 @@ namespace Chaos
 			const int32 NumManifoldPoints = CollisionSolvers[ConstraintIndex].NumManifoldPoints();
 			if (NumManifoldPoints > 0)
 			{
-				FPlatformMisc::PrefetchBlock(&CollisionSolvers[ConstraintIndex].GetManifoldPoint(0), NumManifoldPoints * sizeof(FPBDCollisionSolverManifoldPoint));
+				FPlatformMisc::PrefetchBlock(&CollisionSolvers[ConstraintIndex].GetManifoldPoint(0), NumManifoldPoints * sizeof(Private::FPBDCollisionSolverManifoldPoint));
 			}
 			FPlatformMisc::PrefetchBlock(&CollisionSolvers[ConstraintIndex], sizeof(FPBDCollisionContainerSolver));
 		}
@@ -418,7 +399,7 @@ namespace Chaos
 		bool bWantPerIterationCollisionDetection = false;
 		for (int32 ConstraintIndex = BeginIndex; ConstraintIndex < EndIndex; ++ConstraintIndex)
 		{
-			FPBDCollisionSolver& CollisionSolver = CollisionSolvers[ConstraintIndex];
+			Private::FPBDCollisionSolver& CollisionSolver = CollisionSolvers[ConstraintIndex];
 			FPBDCollisionConstraint* Constraint = CollisionConstraints[ConstraintIndex];
 			bool& bPerIterationCollision = bCollisionConstraintPerIterationCollisionDetection[ConstraintIndex];
 
@@ -453,7 +434,7 @@ namespace Chaos
 
 		for (int32 ConstraintIndex = BeginIndex; ConstraintIndex < EndIndex; ++ConstraintIndex)
 		{
-			FPBDCollisionSolver& CollisionSolver = CollisionSolvers[ConstraintIndex];
+			Private::FPBDCollisionSolver& CollisionSolver = CollisionSolvers[ConstraintIndex];
 			FPBDCollisionConstraint* Constraint = CollisionConstraints[ConstraintIndex];
 
 			UpdateCollisionConstraintFromSolver(Constraint, CollisionSolver, Dt);
@@ -516,7 +497,7 @@ namespace Chaos
 	bool FPBDCollisionContainerSolver::SolvePositionImpl(const FReal InDt, const int32 It, const int32 NumIts, const int32 BeginIndex, const int32 EndIndex, const FPBDCollisionSolverSettings& SolverSettings)
 	{
 		SCOPE_CYCLE_COUNTER(STAT_Collisions_Apply);
-		if (!bChaos_PBDCollisionSolver_Position_SolveEnabled)
+		if (!CVars::bChaos_PBDCollisionSolver_Position_SolveEnabled)
 		{
 			return false;
 		}
@@ -559,7 +540,7 @@ namespace Chaos
 	bool FPBDCollisionContainerSolver::SolveVelocityImpl(const FReal InDt, const int32 It, const int32 NumIts, const int32 BeginIndex, const int32 EndIndex, const FPBDCollisionSolverSettings& SolverSettings)
 	{
 		SCOPE_CYCLE_COUNTER(STAT_Collisions_ApplyPushOut);
-		if (!bChaos_PBDCollisionSolver_Velocity_SolveEnabled)
+		if (!CVars::bChaos_PBDCollisionSolver_Velocity_SolveEnabled)
 		{
 			return false;
 		}
@@ -589,7 +570,7 @@ namespace Chaos
 		{
 			if (bDeferredCollisionDetection || bCollisionConstraintPerIterationCollisionDetection[SolverIndex])
 			{
-				FPBDCollisionSolver& CollisionSolver = CollisionSolvers[SolverIndex];
+				Private::FPBDCollisionSolver& CollisionSolver = CollisionSolvers[SolverIndex];
 				FPBDCollisionConstraint* Constraint = CollisionConstraints[SolverIndex];
 
 				// Run collision detection at the current transforms including any correction from previous iterations
diff --git a/Engine/Source/Runtime/Experimental/Chaos/Private/Chaos/Evolution/ConstraintGroupSolver.cpp b/Engine/Source/Runtime/Experimental/Chaos/Private/Chaos/Evolution/ConstraintGroupSolver.cpp
--- a/Engine/Source/Runtime/Experimental/Chaos/Private/Chaos/Evolution/ConstraintGroupSolver.cpp
+++ b/Engine/Source/Runtime/Experimental/Chaos/Private/Chaos/Evolution/ConstraintGroupSolver.cpp
@@ -48,7 +48,10 @@ namespace Chaos
 			PrioritizedConstraintContainerSolvers.Reset(ConstraintContainerSolvers.Num());
 			for (TUniquePtr<FConstraintContainerSolver>& SolverContainer : ConstraintContainerSolvers)
 			{
-				PrioritizedConstraintContainerSolvers.Add(SolverContainer.Get());
+				if (SolverContainer != nullptr)
+				{
+					PrioritizedConstraintContainerSolvers.Add(SolverContainer.Get());
+				}
 			}
 
 			PrioritizedConstraintContainerSolvers.StableSort(
diff --git a/Engine/Source/Runtime/Experimental/Chaos/Private/Chaos/Island/IslandGroupManager.cpp b/Engine/Source/Runtime/Experimental/Chaos/Private/Chaos/Island/IslandGroupManager.cpp
--- a/Engine/Source/Runtime/Experimental/Chaos/Private/Chaos/Island/IslandGroupManager.cpp
+++ b/Engine/Source/Runtime/Experimental/Chaos/Private/Chaos/Island/IslandGroupManager.cpp
@@ -93,6 +93,15 @@ namespace Chaos
 			}
 		}
 
+		void FPBDIslandGroupManager::RemoveConstraintContainer(FPBDConstraintContainer& ConstraintContainer)
+		{
+			const int32 ContainerId = ConstraintContainer.GetContainerId();
+			for (int32 GroupIndex = 0; GroupIndex < IslandGroups.Num(); ++GroupIndex)
+			{
+				IslandGroups[GroupIndex]->SetConstraintSolver(ContainerId, nullptr);
+			}
+		}
+
 		void FPBDIslandGroupManager::SetConstraintContainerPriority(const int32 ContainerId, const int32 Priority)
 		{
 			for (int32 GroupIndex = 0; GroupIndex < IslandGroups.Num(); ++GroupIndex)
diff --git a/Engine/Source/Runtime/Experimental/Chaos/Private/Chaos/PBDCollisionConstraints.cpp b/Engine/Source/Runtime/Experimental/Chaos/Private/Chaos/PBDCollisionConstraints.cpp
--- a/Engine/Source/Runtime/Experimental/Chaos/Private/Chaos/PBDCollisionConstraints.cpp
+++ b/Engine/Source/Runtime/Experimental/Chaos/Private/Chaos/PBDCollisionConstraints.cpp
@@ -9,6 +9,7 @@
 #include "Chaos/CollisionResolution.h"
 #include "Chaos/Collision/CollisionPruning.h"
 #include "Chaos/Collision/PBDCollisionSolver.h"
+#include "Chaos/Collision/PBDCollisionContainerSolverJacobi.h"
 #include "Chaos/Collision/SolverCollisionContainer.h"
 #include "Chaos/Defines.h"
 #include "Chaos/Evolution/SolverBodyContainer.h"
@@ -27,6 +28,11 @@
 
 namespace Chaos
 {
+	namespace CVars
+	{
+		extern bool bChaos_PBDCollisionSolver_UseJacobiPairSolver2;
+	}
+
 	int32 CollisionParticlesBVHDepth = 4;
 	FAutoConsoleVariableRef CVarCollisionParticlesBVHDepth(TEXT("p.CollisionParticlesBVHDepth"), CollisionParticlesBVHDepth, TEXT("The maximum depth for collision particles bvh"));
 
@@ -113,6 +119,7 @@ namespace Chaos
 		, bEnableEdgePruning(true)
 		, bIsDeterministic(false)
 		, bCanDisableContacts(true)
+		, CollisionSolverType(Private::ECollisionSolverType::GaussSeidel)
 		, GravityDirection(FVec3(0,0,-1))
 		, GravitySize(980)
 		, SolverSettings()
@@ -126,6 +133,23 @@ namespace Chaos
 	{
 	}
 
+	TUniquePtr<FConstraintContainerSolver> FPBDCollisionConstraints::CreateSceneSolver(const int32 Priority)
+	{
+		if (CollisionSolverType == Private::ECollisionSolverType::PartialJacobi)
+		{
+			return MakeUnique<FPBDCollisionContainerSolverJacobi>(*this, Priority);
+		}
+		else
+		{
+			return MakeUnique<FPBDCollisionContainerSolver>(*this, Priority);
+		}
+	}
+
+	TUniquePtr<FConstraintContainerSolver> FPBDCollisionConstraints::CreateGroupSolver(const int32 Priority)
+	{
+		return CreateSceneSolver(Priority);
+	}
+
 	void FPBDCollisionConstraints::DisableHandles()
 	{
 		check(NumConstraints() == 0);
diff --git a/Engine/Source/Runtime/Experimental/Chaos/Private/Chaos/PBDRigidsEvolutionGBF.cpp b/Engine/Source/Runtime/Experimental/Chaos/Private/Chaos/PBDRigidsEvolutionGBF.cpp
--- a/Engine/Source/Runtime/Experimental/Chaos/Private/Chaos/PBDRigidsEvolutionGBF.cpp
+++ b/Engine/Source/Runtime/Experimental/Chaos/Private/Chaos/PBDRigidsEvolutionGBF.cpp
@@ -59,6 +59,10 @@ namespace Chaos
 		bool bChaosCollisionCCDUseTightBoundingBox = true;
 		FAutoConsoleVariableRef  CVarChaosCollisionCCDUseTightBoundingBox(TEXT("p.Chaos.Collision.CCD.UseTightBoundingBox"), bChaosCollisionCCDUseTightBoundingBox , TEXT(""));
 
+		// Collision Solver Type: Full Gaus Seidel or Partial Jacobi (see ECollisionSolverType)
+		int32 ChaosSolverCollisionSolverType = -1;
+		FAutoConsoleVariableRef CVarChaosSolverCollisionSolverType(TEXT("p.Chaos.Solver.Collision.SolverType"), ChaosSolverCollisionSolverType, TEXT("-1: Use default (Full Gauss Seidel); 0: Full Gauss Seidel; 1: Partial Jacobi"));
+
 		int32 ChaosSolverCollisionPriority = 0;
 		FAutoConsoleVariableRef CVarChaosSolverCollisionPriority(TEXT("p.Chaos.Solver.Collision.Priority"), ChaosSolverCollisionPriority, TEXT("Set constraint priority. Larger values are evaluated later [def:0]"));
 
@@ -300,6 +304,22 @@ void FPBDRigidsEvolutionGBF::BuildDisabledParticles(const int32 Island, TArray<T
 	SleepedIslands[Island] = GetIslandManager().SleepInactive(Island, PhysicsMaterials, SolverPhysicsMaterials);
 }
 
+void FPBDRigidsEvolutionGBF::UpdateCollisionSolverType()
+{
+	// If we have changed the collision solver type we must destroy any existing solvers so that they get recreated.
+	// This is not intended to be performant - it is to allow switching for behaviour and performance comparisons only, without restarting the game.
+	if (ChaosSolverCollisionSolverType >= 0)
+	{
+		const Private::ECollisionSolverType CollisionSolverType = (ChaosSolverCollisionSolverType != 0) ? Private::ECollisionSolverType::PartialJacobi : Private::ECollisionSolverType::GaussSeidel;
+		if (CollisionSolverType != CollisionConstraints.GetSolverType())
+		{
+			IslandGroupManager.RemoveConstraintContainer(CollisionConstraints);
+			CollisionConstraints.SetSolverType(CollisionSolverType);
+			IslandGroupManager.AddConstraintContainer(CollisionConstraints, ChaosSolverCollisionPriority);
+		}
+	}
+}
+
 int32 DrawAwake = 0;
 FAutoConsoleVariableRef CVarDrawAwake(TEXT("p.chaos.DebugDrawAwake"),DrawAwake,TEXT("Draw particles that are awake"));
 
@@ -324,6 +344,9 @@ void FPBDRigidsEvolutionGBF::AdvanceOneTimeStepImpl(const FReal Dt, const FSubSt
 	}
 #endif
 
+	// Update the collision solver type (used to support runtime comparisons of solver types for debugging/testing)
+	UpdateCollisionSolverType();
+
 	{
 		SCOPE_CYCLE_COUNTER(STAT_Evolution_UnclusterUnions);
 		Clustering.UnionClusterGroups();
diff --git a/Engine/Source/Runtime/Experimental/Chaos/Public/Chaos/Collision/PBDCollisionContainerSolverJacobi.h b/Engine/Source/Runtime/Experimental/Chaos/Public/Chaos/Collision/PBDCollisionContainerSolverJacobi.h
new file mode 100644
--- /dev/null
+++ b/Engine/Source/Runtime/Experimental/Chaos/Public/Chaos/Collision/PBDCollisionContainerSolverJacobi.h
@@ -0,0 +1,71 @@
+// Copyright Epic Games, Inc. All Rights Reserved.
+#pragma once
+
+#include "Chaos/CollisionResolutionTypes.h"
+#include "Chaos/Collision/CollisionApplyType.h"
+#include "Chaos/Collision/PBDCollisionSolverJacobi.h"
+#include "Chaos/Collision/PBDCollisionSolverSettings.h"
+#include "Chaos/Evolution/SolverConstraintContainer.h"
+
+namespace Chaos
+{
+	namespace Private
+	{
+		class FPBDIsland;
+		class FPBDIslandConstraint;
+	}
+	class FPBDCollisionConstraint;
+	class FPBDCollisionConstraints;
+	class FSolverBody;
+	class FSolverBodyContainer;
+
+	/**
+	 * The solver for a set of collision constraints. This collects all the data required to solve a set of collision
+	 * constraints into a contiguous, ordered buffer.
+	 * 
+	 * This version runs a Gauss-Seidel outer loop over manifolds, and a Jacobi loop over contacts in each manifold.
+	*/
+	class FPBDCollisionContainerSolverJacobi : public FConstraintContainerSolver
+	{
+	public:
+		UE_NONCOPYABLE(FPBDCollisionContainerSolverJacobi);
+
+		FPBDCollisionContainerSolverJacobi(const FPBDCollisionConstraints& InConstraintContainer, const int32 InPriority);
+		~FPBDCollisionContainerSolverJacobi();
+
+		int32 NumSolvers() const { return CollisionConstraints.Num(); }
+
+		virtual void Reset(const int32 InMaxCollisions) override final;
+
+		virtual int32 GetNumConstraints() const override final { return NumSolvers(); }
+
+		//
+		// IslandGroup API
+		//
+		virtual void AddConstraints() override final;
+		virtual void AddConstraints(const TArrayView<Private::FPBDIslandConstraint>& ConstraintHandles) override final;
+		virtual void AddBodies(FSolverBodyContainer& SolverBodyContainer) override final;
+		virtual void GatherInput(const FReal Dt) override final;
+		virtual void GatherInput(const FReal Dt, const int32 BeginIndex, const int32 EndIndex) override final;
+		virtual void ScatterOutput(const FReal Dt) override final;
+		virtual void ScatterOutput(const FReal Dt, const int32 BeginIndex, const int32 EndIndex) override final;
+		virtual void ApplyPositionConstraints(const FReal Dt, const int32 It, const int32 NumIts) override final;
+		virtual void ApplyVelocityConstraints(const FReal Dt, const int32 It, const int32 NumIts) override final;
+		virtual void ApplyProjectionConstraints(const FReal Dt, const int32 It, const int32 NumIts) override final;
+
+	private:
+		void CachePrefetchSolver(const int32 ConstraintIndex) const;
+		void AddConstraint(FPBDCollisionConstraint& Constraint);
+		void UpdatePositionShockPropagation(const FReal Dt, const int32 It, const int32 NumIts, const int32 BeginIndex, const int32 EndIndex, const FPBDCollisionSolverSettings& SolverSettings);
+		void UpdateVelocityShockPropagation(const FReal Dt, const int32 It, const int32 NumIts, const int32 BeginIndex, const int32 EndIndex, const FPBDCollisionSolverSettings& SolverSettings);
+		bool SolvePositionImpl(const FReal Dt, const int32 It, const int32 NumIts, const int32 BeginIndex, const int32 EndIndex, const FPBDCollisionSolverSettings& SolverSettings);
+		bool SolveVelocityImpl(const FReal Dt, const int32 It, const int32 NumIts, const int32 BeginIndex, const int32 EndIndex, const FPBDCollisionSolverSettings& SolverSettings);
+		void UpdateCollisions(const FReal InDt, const int32 BeginIndex, const int32 EndIndex);
+
+		const FPBDCollisionConstraints& ConstraintContainer;
+		TArray<Private::FPBDCollisionSolverJacobi> CollisionSolvers;
+		TArray<FPBDCollisionConstraint*> CollisionConstraints;
+		TArray<bool> bCollisionConstraintPerIterationCollisionDetection;
+		bool bPerIterationCollisionDetection;
+	};
+}
diff --git a/Engine/Source/Runtime/Experimental/Chaos/Public/Chaos/Collision/PBDCollisionSolver.h b/Engine/Source/Runtime/Experimental/Chaos/Public/Chaos/Collision/PBDCollisionSolver.h
--- a/Engine/Source/Runtime/Experimental/Chaos/Public/Chaos/Collision/PBDCollisionSolver.h
+++ b/Engine/Source/Runtime/Experimental/Chaos/Public/Chaos/Collision/PBDCollisionSolver.h
@@ -6,8 +6,6 @@
 #include "Chaos/Defines.h"
 #include "Chaos/Evolution/SolverBody.h"
 
-CHAOS_API DECLARE_LOG_CATEGORY_EXTERN(LogChaosCollision, Log, All);
-
 // Set to 0 to use a linearized error calculation, and set to 1 to use a non-linear error calculation in collision detection. 
 // In principle nonlinear is more accurate when large rotation corrections occur, but this is not too important for collisions because 
 // when the bodies settle the corrections are small. The linearized version is significantly faster than the non-linear version because 
@@ -77,7 +75,6 @@ namespace Chaos
 			*/
 			void CalculateContactPositionErrorNormal(const FConstraintSolverBody& Body0, const FConstraintSolverBody& Body1, const FSolverReal MaxPushOut, FSolverReal& OutContactDeltaNormal) const;
 			void CalculateContactPositionErrorTangential(const FConstraintSolverBody& Body0, const FConstraintSolverBody& Body1, FSolverReal& OutContactDeltaTanget0, FSolverReal& OutContactDeltaTangent1) const;
-			void CalculateContactPositionError(const FConstraintSolverBody& Body0, const FConstraintSolverBody& Body1, const FSolverReal MaxPushOut, FSolverReal& OutContactDeltaNormal, FSolverReal& OutContactDeltaTanget0, FSolverReal& OutContactDeltaTangent1) const;
 
 			/**
 			 * @brief Calculate the velocity error at the current transforms
@@ -89,12 +86,6 @@ namespace Chaos
 		public:
 			friend class FPBDCollisionSolver;
 
-			/**
-			 * @brief Calculate the relative velocity at the contact point
-			 * @note InitContact must be called before calling this function
-			*/
-			FSolverVec3 CalculateContactVelocity(const FConstraintSolverBody& Body0, const FConstraintSolverBody& Body1) const;
-
 			/**
 			 * @brief Whether we need to solve velocity for this manifold point (only if we were penetrating or applied a pushout)
 			*/
@@ -757,36 +748,6 @@ namespace Chaos
 			OutContactDeltaTangentV = WorldContact.ContactDeltaTangentV + FSolverVec3::DotProduct(ContactDelta, WorldContact.ContactTangentV);
 		}
 
-		FORCEINLINE_DEBUGGABLE void FPBDCollisionSolverManifoldPoint::CalculateContactPositionError(const FConstraintSolverBody& Body0, const FConstraintSolverBody& Body1, const FSolverReal MaxPushOut, FSolverReal& OutContactDeltaNormal, FSolverReal& OutContactDeltaTangentU, FSolverReal& OutContactDeltaTangentV) const
-		{
-	#if CHAOS_NONLINEAR_COLLISIONS_ENABLED
-	#error "Non-linear collision solver probably no longer functional - fix it!"
-			// Non-linear version: calculate the contact delta after we have converted the current positional impulses into position and rotation corrections.
-			// We could precalculate and store the LocalContactPositions if we really want to use this nonlinear version
-			const FSolverVec3 LocalContactPosition0 = Body0.Q().Inverse() * RelativeContactPosition0;
-			const FSolverVec3 LocalContactPosition1 = Body1.Q().Inverse() * RelativeContactPosition1;
-			const FSolverVec3 ContactDelta = (Body0.CorrectedP() + Body0.CorrectedQ() * LocalContactPosition0) - (Body1.CorrectedP() + Body1.CorrectedQ() * LocalContactPosition1);
-			OutContactDeltaNormal = FSolverVec3::DotProduct(ContactDelta, WorldContactNormal);
-			OutContactDeltaTangentU = FSolverVec3::DotProduct(ContactDelta, WorldContactTangentU);
-			OutContactDeltaTangentV = FSolverVec3::DotProduct(ContactDelta, WorldContactTangentV);
-	#else
-			// Linear version: calculate the contact delta assuming linear motion after applying a positional impulse at the contact point. There will be an error that depends on the size of the rotation.
-			const FSolverVec3 ContactDelta0 = Body0.DP() + FSolverVec3::CrossProduct(Body0.DQ(), WorldContact.RelativeContactPoints[0]);
-			const FSolverVec3 ContactDelta1 = Body1.DP() + FSolverVec3::CrossProduct(Body1.DQ(), WorldContact.RelativeContactPoints[1]);
-			const FSolverVec3 ContactDelta = ContactDelta0 - ContactDelta1;
-			OutContactDeltaNormal = WorldContact.ContactDeltaNormal + FSolverVec3::DotProduct(ContactDelta, WorldContact.ContactNormal);
-			OutContactDeltaTangentU = WorldContact.ContactDeltaTangentU + FSolverVec3::DotProduct(ContactDelta, WorldContact.ContactTangentU);
-			OutContactDeltaTangentV = WorldContact.ContactDeltaTangentV + FSolverVec3::DotProduct(ContactDelta, WorldContact.ContactTangentV);
-	#endif
-
-			// NOTE: OutContactDeltaNormal is negative for penetration
-			// NOTE: MaxPushOut == 0 disables the pushout limits
-			if ((MaxPushOut > 0) && (OutContactDeltaNormal < -MaxPushOut))
-			{
-				OutContactDeltaNormal = -MaxPushOut;
-			}
-		}
-
 		FORCEINLINE_DEBUGGABLE void FPBDCollisionSolverManifoldPoint::CalculateContactVelocityError(const FConstraintSolverBody& Body0, const FConstraintSolverBody& Body1, const FSolverReal DynamicFriction, const FSolverReal Dt, FSolverReal& OutContactVelocityDeltaNormal, FSolverReal& OutContactVelocityDeltaTangent0, FSolverReal& OutContactVelocityDeltaTangent1) const
 		{
 			const FSolverVec3 ContactVelocity0 = Body0.V() + FSolverVec3::CrossProduct(Body0.W(), WorldContact.RelativeContactPoints[0]);
@@ -812,13 +773,6 @@ namespace Chaos
 			OutContactVelocityDeltaNormal = (ContactVelocityNormal - WorldContact.ContactTargetVelocityNormal);
 		}
 
-		FORCEINLINE_DEBUGGABLE FSolverVec3 FPBDCollisionSolverManifoldPoint::CalculateContactVelocity(const FConstraintSolverBody& Body0, const FConstraintSolverBody& Body1) const
-		{
-			const FSolverVec3 ContactVelocity0 = Body0.V() + FVec3::CrossProduct(Body0.W(), WorldContact.RelativeContactPoints[0]);
-			const FSolverVec3 ContactVelocity1 = Body1.V() + FVec3::CrossProduct(Body1.W(), WorldContact.RelativeContactPoints[1]);
-			return ContactVelocity0 - ContactVelocity1;
-		}
-
 		FORCEINLINE_DEBUGGABLE bool FPBDCollisionSolverManifoldPoint::ShouldSolveVelocity() const
 		{
 			// We ensure positive separating velocity for close contacts even if they didn't receive a pushout
diff --git a/Engine/Source/Runtime/Experimental/Chaos/Public/Chaos/Collision/PBDCollisionSolverJacobi.h b/Engine/Source/Runtime/Experimental/Chaos/Public/Chaos/Collision/PBDCollisionSolverJacobi.h
new file mode 100644
--- /dev/null
+++ b/Engine/Source/Runtime/Experimental/Chaos/Public/Chaos/Collision/PBDCollisionSolverJacobi.h
@@ -0,0 +1,976 @@
+// Copyright Epic Games, Inc. All Rights Reserved.
+#pragma once
+
+#include "Chaos/Core.h"
+#include "Chaos/Collision/ContactPoint.h"
+#include "Chaos/Defines.h"
+#include "Chaos/Evolution/SolverBody.h"
+#include "Chaos/Framework/UncheckedArray.h"
+
+// Set to 0 to use a linearized error calculation, and set to 1 to use a non-linear error calculation in collision detection. 
+// In principle nonlinear is more accurate when large rotation corrections occur, but this is not too important for collisions because 
+// when the bodies settle the corrections are small. The linearized version is significantly faster than the non-linear version because 
+// the non-linear version requires a quaternion multiply and renormalization whereas the linear version is just a cross product.
+#define CHAOS_NONLINEAR_COLLISIONS_ENABLED 0
+
+namespace Chaos
+{
+	class FManifoldPoint;
+	class FPBDCollisionConstraint;
+
+	namespace Private
+	{
+		class FPBDCollisionSolverJacobi;
+	}
+
+	namespace CVars
+	{
+		extern bool bChaos_PBDCollisionSolver_Velocity_FrictionEnabled;
+		extern float Chaos_PBDCollisionSolver_Position_StaticFrictionStiffness;
+		extern float Chaos_PBDCollisionSolver_JacobiStiffness;
+		extern float Chaos_PBDCollisionSolver_JacobiPositionTolerance;
+		extern float Chaos_PBDCollisionSolver_JacobiRotationTolerance;
+	}
+
+	namespace Private
+	{
+		class FSolverVec3SOA
+		{
+			AlignedFloat4 VX;
+			AlignedFloat4 VY;
+			AlignedFloat4 VZ;
+		};
+
+		class FSolverRealSOA
+		{
+			AlignedFloat4 V;
+		};
+
+		/**
+		 * @brief A set of 4 manifold points in a FPBDCollisionSolver
+		*/
+		class FPBDCollisionSolverJacobiManifoldPoints
+		{
+		public:
+			static const int32 MaxManifoldPoints = 4;
+
+			void Reset()
+			{
+				ManifoldPoints.Reset();
+			}
+
+			FORCEINLINE_DEBUGGABLE int32 NumManifoldPoints() const
+			{
+				return ManifoldPoints.Num();
+			}
+
+			FORCEINLINE_DEBUGGABLE int32 AddManifoldPoint()
+			{
+				return ManifoldPoints.Add();
+			}
+
+			FORCEINLINE_DEBUGGABLE void SetWorldContact(
+				const int32 ManifoldPointIndex,
+				const FSolverVec3& InRelativeContactPosition0,
+				const FSolverVec3& InRelativeContactPosition1,
+				const FSolverVec3& InWorldContactNormal,
+				const FSolverVec3& InWorldContactTangentU,
+				const FSolverVec3& InWorldContactTangentV,
+				const FSolverReal InWorldContactDeltaNormal,
+				const FSolverReal InWorldContactDeltaTangentU,
+				const FSolverReal InWorldContactDeltaTangentV,
+				const FSolverReal InWorldContactVelocityTargetNormal)
+			{
+				ManifoldPoints[ManifoldPointIndex].RelativeContactPosition0 = InRelativeContactPosition0;
+				ManifoldPoints[ManifoldPointIndex].RelativeContactPosition1 = InRelativeContactPosition1;
+				ManifoldPoints[ManifoldPointIndex].ContactNormal = InWorldContactNormal;
+				ManifoldPoints[ManifoldPointIndex].ContactTangentU = InWorldContactTangentU;
+				ManifoldPoints[ManifoldPointIndex].ContactTangentV = InWorldContactTangentV;
+				ManifoldPoints[ManifoldPointIndex].ContactDeltaNormal = InWorldContactDeltaNormal;
+				ManifoldPoints[ManifoldPointIndex].ContactDeltaTangentU = InWorldContactDeltaTangentU;
+				ManifoldPoints[ManifoldPointIndex].ContactDeltaTangentV = InWorldContactDeltaTangentV;
+				ManifoldPoints[ManifoldPointIndex].ContactTargetVelocityNormal = InWorldContactVelocityTargetNormal;
+			}
+
+			/**
+			 * @brief Initialize the geometric data for the contact
+			*/
+			FORCEINLINE_DEBUGGABLE void InitContact(
+				const int32 ManifoldPointIndex,
+				const FConstraintSolverBody& Body0,
+				const FConstraintSolverBody& Body1)
+			{
+				ManifoldPoints[ManifoldPointIndex].NetPushOutNormal = FSolverReal(0);
+				ManifoldPoints[ManifoldPointIndex].NetPushOutTangentU = FSolverReal(0);
+				ManifoldPoints[ManifoldPointIndex].NetPushOutTangentV = FSolverReal(0);
+				ManifoldPoints[ManifoldPointIndex].NetImpulseNormal = FSolverReal(0);
+				ManifoldPoints[ManifoldPointIndex].NetImpulseTangentU = FSolverReal(0);
+				ManifoldPoints[ManifoldPointIndex].NetImpulseTangentV = FSolverReal(0);
+				ManifoldPoints[ManifoldPointIndex].StaticFrictionRatio = FSolverReal(0);
+
+				UpdateMass(ManifoldPointIndex, Body0, Body1);
+			}
+
+			/**
+			 * @brief Update the cached mass properties based on the current body transforms
+			*/
+			void UpdateMass(
+				const int32 ManifoldPointIndex, 
+				const FConstraintSolverBody& Body0, 
+				const FConstraintSolverBody& Body1)
+			{
+				FSolverReal ContactMassInvNormal = FSolverReal(0);
+				FSolverReal ContactMassInvTangentU = FSolverReal(0);
+				FSolverReal ContactMassInvTangentV = FSolverReal(0);
+
+				if (Body0.IsDynamic())
+				{
+					const FSolverVec3 R0xN = FSolverVec3::CrossProduct(ManifoldPoints[ManifoldPointIndex].RelativeContactPosition0, ManifoldPoints[ManifoldPointIndex].ContactNormal);
+					const FSolverVec3 R0xU = FSolverVec3::CrossProduct(ManifoldPoints[ManifoldPointIndex].RelativeContactPosition0, ManifoldPoints[ManifoldPointIndex].ContactTangentU);
+					const FSolverVec3 R0xV = FSolverVec3::CrossProduct(ManifoldPoints[ManifoldPointIndex].RelativeContactPosition0, ManifoldPoints[ManifoldPointIndex].ContactTangentV);
+
+					const FSolverMatrix33 InvI0 = Body0.InvI();
+
+					ManifoldPoints[ManifoldPointIndex].ContactNormalAngular0 = InvI0 * R0xN;
+					ManifoldPoints[ManifoldPointIndex].ContactTangentUAngular0 = InvI0 * R0xU;
+					ManifoldPoints[ManifoldPointIndex].ContactTangentVAngular0 = InvI0 * R0xV;
+
+					ContactMassInvNormal += FSolverVec3::DotProduct(R0xN, ManifoldPoints[ManifoldPointIndex].ContactNormalAngular0) + Body0.InvM();
+					ContactMassInvTangentU += FSolverVec3::DotProduct(R0xU, ManifoldPoints[ManifoldPointIndex].ContactTangentUAngular0) + Body0.InvM();
+					ContactMassInvTangentV += FSolverVec3::DotProduct(R0xV, ManifoldPoints[ManifoldPointIndex].ContactTangentVAngular0) + Body0.InvM();
+				}
+				if (Body1.IsDynamic())
+				{
+					const FSolverVec3 R1xN = FSolverVec3::CrossProduct(ManifoldPoints[ManifoldPointIndex].RelativeContactPosition1, ManifoldPoints[ManifoldPointIndex].ContactNormal);
+					const FSolverVec3 R1xU = FSolverVec3::CrossProduct(ManifoldPoints[ManifoldPointIndex].RelativeContactPosition1, ManifoldPoints[ManifoldPointIndex].ContactTangentU);
+					const FSolverVec3 R1xV = FSolverVec3::CrossProduct(ManifoldPoints[ManifoldPointIndex].RelativeContactPosition1, ManifoldPoints[ManifoldPointIndex].ContactTangentV);
+
+					const FSolverMatrix33 InvI1 = Body1.InvI();
+
+					ManifoldPoints[ManifoldPointIndex].ContactNormalAngular1 = InvI1 * R1xN;
+					ManifoldPoints[ManifoldPointIndex].ContactTangentUAngular1 = InvI1 * R1xU;
+					ManifoldPoints[ManifoldPointIndex].ContactTangentVAngular1 = InvI1 * R1xV;
+
+					ContactMassInvNormal += FSolverVec3::DotProduct(R1xN, ManifoldPoints[ManifoldPointIndex].ContactNormalAngular1) + Body1.InvM();
+					ContactMassInvTangentU += FSolverVec3::DotProduct(R1xU, ManifoldPoints[ManifoldPointIndex].ContactTangentUAngular1) + Body1.InvM();
+					ContactMassInvTangentV += FSolverVec3::DotProduct(R1xV, ManifoldPoints[ManifoldPointIndex].ContactTangentVAngular1) + Body1.InvM();
+				}
+
+				ManifoldPoints[ManifoldPointIndex].ContactMassNormal = (ContactMassInvNormal > FSolverReal(UE_SMALL_NUMBER)) ? FSolverReal(1) / ContactMassInvNormal : FSolverReal(0);
+				ManifoldPoints[ManifoldPointIndex].ContactMassTangentU = (ContactMassInvTangentU > FSolverReal(UE_SMALL_NUMBER)) ? FSolverReal(1) / ContactMassInvTangentU : FSolverReal(0);
+				ManifoldPoints[ManifoldPointIndex].ContactMassTangentV = (ContactMassInvTangentV > FSolverReal(UE_SMALL_NUMBER)) ? FSolverReal(1) / ContactMassInvTangentV : FSolverReal(0);
+			}
+
+			/**
+			 * @brief Update the contact mass for the normal correction
+			 * This is used by shock propagation.
+			*/
+			void UpdateMassNormal(
+				const int32 ManifoldPointIndex, 
+				const FConstraintSolverBody& Body0, 
+				const FConstraintSolverBody& Body1)
+			{
+				FSolverReal ContactMassInvNormal = FSolverReal(0);
+				if (Body0.IsDynamic())
+				{
+					const FSolverVec3 R0xN = FSolverVec3::CrossProduct(ManifoldPoints[ManifoldPointIndex].RelativeContactPosition0, ManifoldPoints[ManifoldPointIndex].ContactNormal);
+					const FSolverMatrix33 InvI0 = Body0.InvI();
+					ManifoldPoints[ManifoldPointIndex].ContactNormalAngular0 = InvI0 * R0xN;
+					ContactMassInvNormal += FSolverVec3::DotProduct(R0xN, ManifoldPoints[ManifoldPointIndex].ContactNormalAngular0) + Body0.InvM();
+				}
+				if (Body1.IsDynamic())
+				{
+					const FSolverVec3 R1xN = FSolverVec3::CrossProduct(ManifoldPoints[ManifoldPointIndex].RelativeContactPosition1, ManifoldPoints[ManifoldPointIndex].ContactNormal);
+					const FSolverMatrix33 InvI1 = Body1.InvI();
+					ManifoldPoints[ManifoldPointIndex].ContactNormalAngular1 = InvI1 * R1xN;
+					ContactMassInvNormal += FSolverVec3::DotProduct(R1xN, ManifoldPoints[ManifoldPointIndex].ContactNormalAngular1) + Body1.InvM();
+				}
+				ManifoldPoints[ManifoldPointIndex].ContactMassNormal = (ContactMassInvNormal > FSolverReal(UE_SMALL_NUMBER)) ? FSolverReal(1) / ContactMassInvNormal : FSolverReal(0);
+			}
+
+			FORCEINLINE_DEBUGGABLE FSolverVec3 GetNetPushOut(const int32 ManifoldPointIndex) const
+			{
+				return ManifoldPoints[ManifoldPointIndex].NetPushOutNormal* ManifoldPoints[ManifoldPointIndex].ContactNormal +
+					ManifoldPoints[ManifoldPointIndex].NetPushOutTangentU * ManifoldPoints[ManifoldPointIndex].ContactTangentU +
+					ManifoldPoints[ManifoldPointIndex].NetPushOutTangentV * ManifoldPoints[ManifoldPointIndex].ContactTangentV;
+			}
+
+			FORCEINLINE_DEBUGGABLE FSolverVec3 GetNetImpulse(const int32 ManifoldPointIndex) const
+			{
+				return ManifoldPoints[ManifoldPointIndex].NetImpulseNormal* ManifoldPoints[ManifoldPointIndex].ContactNormal +
+					ManifoldPoints[ManifoldPointIndex].NetImpulseTangentU * ManifoldPoints[ManifoldPointIndex].ContactTangentU +
+					ManifoldPoints[ManifoldPointIndex].NetImpulseTangentV * ManifoldPoints[ManifoldPointIndex].ContactTangentV;
+			}
+
+			FORCEINLINE_DEBUGGABLE FSolverReal GetStaticFrictionRatio(const int32 ManifoldPointIndex) const
+			{
+				return ManifoldPoints[ManifoldPointIndex].StaticFrictionRatio;
+			}
+
+			/**
+			 * @brief Calculate the position error at the current transforms
+			 * @param MaxPushOut a limit on the position error for this iteration to prevent initial-penetration explosion (a common PBD problem)
+			*/
+			FORCEINLINE_DEBUGGABLE void CalculateContactPositionErrorNormal(
+				const int32 ManifoldPointIndex, 
+				const FConstraintSolverBody& Body0, 
+				const FConstraintSolverBody& Body1, 
+				const FSolverReal MaxPushOut, 
+				FSolverReal& OutContactDeltaNormal) const
+			{
+				// Linear version: calculate the contact delta assuming linear motion after applying a positional impulse at the contact point. There will be an error that depends on the size of the rotation.
+				const FSolverVec3 ContactDelta0 = Body0.DP() + FSolverVec3::CrossProduct(Body0.DQ(), ManifoldPoints[ManifoldPointIndex].RelativeContactPosition0);
+				const FSolverVec3 ContactDelta1 = Body1.DP() + FSolverVec3::CrossProduct(Body1.DQ(), ManifoldPoints[ManifoldPointIndex].RelativeContactPosition1);
+				const FSolverVec3 ContactDelta = ContactDelta0 - ContactDelta1;
+				OutContactDeltaNormal = ManifoldPoints[ManifoldPointIndex].ContactDeltaNormal + FSolverVec3::DotProduct(ContactDelta, ManifoldPoints[ManifoldPointIndex].ContactNormal);
+
+				// NOTE: OutContactDeltaNormal is negative for penetration
+				// NOTE: MaxPushOut == 0 disables the pushout limits
+				if ((MaxPushOut > 0) && (OutContactDeltaNormal < -MaxPushOut))
+				{
+					OutContactDeltaNormal = -MaxPushOut;
+				}
+			}
+
+			FORCEINLINE_DEBUGGABLE void CalculateContactPositionErrorTangential(
+				const int32 ManifoldPointIndex, 
+				const FConstraintSolverBody& Body0, 
+				const FConstraintSolverBody& Body1, 
+				FSolverReal& OutContactDeltaTangentU, 
+				FSolverReal& OutContactDeltaTangentV) const
+			{
+				// Linear version: calculate the contact delta assuming linear motion after applying a positional impulse at the contact point. There will be an error that depends on the size of the rotation.
+				const FSolverVec3 ContactDelta0 = Body0.DP() + FSolverVec3::CrossProduct(Body0.DQ(), ManifoldPoints[ManifoldPointIndex].RelativeContactPosition0);
+				const FSolverVec3 ContactDelta1 = Body1.DP() + FSolverVec3::CrossProduct(Body1.DQ(), ManifoldPoints[ManifoldPointIndex].RelativeContactPosition1);
+				const FSolverVec3 ContactDelta = ContactDelta0 - ContactDelta1;
+				OutContactDeltaTangentU = ManifoldPoints[ManifoldPointIndex].ContactDeltaTangentU + FSolverVec3::DotProduct(ContactDelta, ManifoldPoints[ManifoldPointIndex].ContactTangentU);
+				OutContactDeltaTangentV = ManifoldPoints[ManifoldPointIndex].ContactDeltaTangentV + FSolverVec3::DotProduct(ContactDelta, ManifoldPoints[ManifoldPointIndex].ContactTangentV);
+			}
+
+			FORCEINLINE_DEBUGGABLE void CalculatePositionCorrectionNormal(
+				const int32 ManifoldPointIndex,
+				const FSolverReal Stiffness,
+				FSolverReal ContactDeltaNormal,
+				FConstraintSolverBody& Body0,
+				FConstraintSolverBody& Body1,
+				FSolverVec3& InOutDX0,
+				FSolverVec3& InOutDR0,
+				FSolverVec3& InOutDX1,
+				FSolverVec3& InOutDR1)
+			{
+				FSolverReal PushOutNormal = -Stiffness * ContactDeltaNormal * ManifoldPoints[ManifoldPointIndex].ContactMassNormal;
+
+				// The total pushout so far this sub-step
+				// Unilateral constraint: Net-negative impulses not allowed (negative incremental impulses are allowed as long as the net is positive)
+				if ((ManifoldPoints[ManifoldPointIndex].NetPushOutNormal + PushOutNormal) > FSolverReal(0))
+				{
+					ManifoldPoints[ManifoldPointIndex].NetPushOutNormal += PushOutNormal;
+				}
+				else
+				{
+					PushOutNormal = -ManifoldPoints[ManifoldPointIndex].NetPushOutNormal;
+					ManifoldPoints[ManifoldPointIndex].NetPushOutNormal = 0;
+				}
+
+				// Update the particle state based on the pushout
+				if (Body0.IsDynamic())
+				{
+					InOutDX0 += (Body0.InvM() * PushOutNormal) * ManifoldPoints[ManifoldPointIndex].ContactNormal;
+					InOutDR0 += ManifoldPoints[ManifoldPointIndex].ContactNormalAngular0 * PushOutNormal;
+				}
+				if (Body1.IsDynamic())
+				{
+					InOutDX1 -= (Body1.InvM() * PushOutNormal) * ManifoldPoints[ManifoldPointIndex].ContactNormal;
+					InOutDR1 -= ManifoldPoints[ManifoldPointIndex].ContactNormalAngular1 * PushOutNormal;
+				}
+			}
+
+			FORCEINLINE_DEBUGGABLE void ApplyFrictionCone(
+				const int32 ManifoldPointIndex,
+				const FSolverReal StaticFriction,
+				const FSolverReal DynamicFriction,
+				const FSolverReal MaxFrictionPushOut,
+				FSolverReal& InOutPushOutTangentU,
+				FSolverReal& InOutPushOutTangentV)
+			{
+				// Assume we stay in the friction cone...
+				ManifoldPoints[ManifoldPointIndex].StaticFrictionRatio = FSolverReal(1);
+
+				if (MaxFrictionPushOut < FSolverReal(UE_KINDA_SMALL_NUMBER))
+				{
+					// Note: we have already added the current iteration's PushOut to the NetPushOut but it has not been applied to the body
+					// so we must subtract it again to calculate the actual pushout we want to undo (i.e., the net pushout that has been applied 
+					// to the body so far from previous iterations)
+					InOutPushOutTangentU = -(ManifoldPoints[ManifoldPointIndex].NetPushOutTangentU - InOutPushOutTangentU);
+					InOutPushOutTangentV = -(ManifoldPoints[ManifoldPointIndex].NetPushOutTangentV - InOutPushOutTangentV);
+					ManifoldPoints[ManifoldPointIndex].NetPushOutTangentU = FSolverReal(0);
+					ManifoldPoints[ManifoldPointIndex].NetPushOutTangentV = FSolverReal(0);
+					ManifoldPoints[ManifoldPointIndex].StaticFrictionRatio = FSolverReal(0);
+				}
+				else
+				{
+					// If we exceed the static friction cone, clip to the dynamic friction cone
+					const FSolverReal MaxStaticPushOutTangentSq = FMath::Square(StaticFriction * MaxFrictionPushOut);
+					const FSolverReal NetPushOutTangentSq = FMath::Square(ManifoldPoints[ManifoldPointIndex].NetPushOutTangentU) + FMath::Square(ManifoldPoints[ManifoldPointIndex].NetPushOutTangentV);
+					if (NetPushOutTangentSq > MaxStaticPushOutTangentSq)
+					{
+						const FSolverReal MaxDynamicPushOutTangent = DynamicFriction * MaxFrictionPushOut;
+						const FSolverReal FrictionMultiplier = MaxDynamicPushOutTangent * FMath::InvSqrt(NetPushOutTangentSq);
+						const FSolverReal NetPushOutTangentU = FrictionMultiplier * ManifoldPoints[ManifoldPointIndex].NetPushOutTangentU;
+						const FSolverReal NetPushOutTangentV = FrictionMultiplier * ManifoldPoints[ManifoldPointIndex].NetPushOutTangentV;
+						InOutPushOutTangentU = NetPushOutTangentU - (ManifoldPoints[ManifoldPointIndex].NetPushOutTangentU - InOutPushOutTangentU);
+						InOutPushOutTangentV = NetPushOutTangentV - (ManifoldPoints[ManifoldPointIndex].NetPushOutTangentV - InOutPushOutTangentV);
+						ManifoldPoints[ManifoldPointIndex].NetPushOutTangentU = NetPushOutTangentU;
+						ManifoldPoints[ManifoldPointIndex].NetPushOutTangentV = NetPushOutTangentV;
+						ManifoldPoints[ManifoldPointIndex].StaticFrictionRatio = FrictionMultiplier;
+					}
+				}
+			}
+
+			FORCEINLINE_DEBUGGABLE void SolvePositionNoFriction(
+				const FSolverReal Stiffness,
+				const FSolverReal MaxPushOut,
+				FConstraintSolverBody& Body0,
+				FConstraintSolverBody& Body1)
+			{
+				// Apply the position correction so that all contacts have zero separation
+				FSolverVec3 DX0 = FSolverVec3(0);
+				FSolverVec3 DX1 = FSolverVec3(0);
+				FSolverVec3 DR0 = FSolverVec3(0);
+				FSolverVec3 DR1 = FSolverVec3(0);
+
+				for (int32 ManifoldPointIndex = 0; ManifoldPointIndex < NumManifoldPoints(); ++ManifoldPointIndex)
+				{
+					FSolverReal ContactDeltaNormal;
+					CalculateContactPositionErrorNormal(ManifoldPointIndex, Body0.SolverBody(), Body1.SolverBody(), MaxPushOut, ContactDeltaNormal);
+
+					const bool bProcessManifoldPoint = (ContactDeltaNormal < FSolverReal(0)) || (ManifoldPoints[ManifoldPointIndex].NetPushOutNormal > FSolverReal(UE_SMALL_NUMBER));
+					if (bProcessManifoldPoint)
+					{
+						CalculatePositionCorrectionNormal(
+							ManifoldPointIndex,
+							Stiffness,
+							ContactDeltaNormal,
+							Body0,
+							Body1,
+							DX0, DR0, DX1, DR1);
+					}
+				}
+
+				if (Body0.IsDynamic())
+				{
+					Body0.ApplyPositionDelta(DX0);
+					Body0.ApplyRotationDelta(DR0);
+				}
+				if (Body1.IsDynamic())
+				{
+					Body1.ApplyPositionDelta(DX1);
+					Body1.ApplyRotationDelta(DR1);
+				}
+			}
+
+			FORCEINLINE_DEBUGGABLE bool SolvePositionWithFriction(
+				const FSolverReal Stiffness,
+				const FSolverReal Dt,
+				const FSolverReal StaticFriction,
+				const FSolverReal DynamicFriction,
+				const FSolverReal MaxPushOut,
+				FConstraintSolverBody& Body0,
+				FConstraintSolverBody& Body1)
+			{
+				// Accumulate net pushout for friction limits below
+				bool bApplyFriction[MaxManifoldPoints] = { false, };
+				int32 NumFrictionContacts = 0;
+				FSolverReal TotalPushOutNormal = FSolverReal(0);
+
+				FSolverVec3 DX0 = FSolverVec3(0);
+				FSolverVec3 DX1 = FSolverVec3(0);
+				FSolverVec3 DR0 = FSolverVec3(0);
+				FSolverVec3 DR1 = FSolverVec3(0);
+
+				// Apply the position correction along the normal and determine if we want to run friction on each point
+				for (int32 ManifoldPointIndex = 0; ManifoldPointIndex < NumManifoldPoints(); ++ManifoldPointIndex)
+				{
+					FSolverReal ContactDeltaNormal;
+					CalculateContactPositionErrorNormal(ManifoldPointIndex, Body0.SolverBody(), Body1.SolverBody(), MaxPushOut, ContactDeltaNormal);
+
+					// Apply a normal correction if we still have penetration or if we are now separated but have previously applied a correction that we may want to undo
+					const bool bProcessManifoldPoint = (ContactDeltaNormal < FSolverReal(0)) || (ManifoldPoints[ManifoldPointIndex].NetPushOutNormal > FSolverReal(UE_SMALL_NUMBER));
+					if (bProcessManifoldPoint)
+					{
+						CalculatePositionCorrectionNormal(
+							ManifoldPointIndex,
+							Stiffness,
+							ContactDeltaNormal,
+							Body0,
+							Body1,
+							DX0, DR0, DX1, DR1);
+
+						TotalPushOutNormal += ManifoldPoints[ManifoldPointIndex].NetPushOutNormal;
+					}
+
+					// Friction gets updated for any point with a net normal correction or where we have previously had a normal correction and 
+					// already applied friction (in which case we may need to zero it)
+					if ((ManifoldPoints[ManifoldPointIndex].NetPushOutNormal != 0) || (ManifoldPoints[ManifoldPointIndex].NetPushOutTangentU != 0) || (ManifoldPoints[ManifoldPointIndex].NetPushOutTangentV != 0))
+					{
+						bApplyFriction[ManifoldPointIndex] = true;
+						++NumFrictionContacts;
+					}
+				}
+
+				// Apply the tangential position correction if required
+				if (NumFrictionContacts > 0)
+				{
+					// We clip the tangential correction at each contact to the friction cone, but we use to average impulse
+					// among all contacts as the clipping limit. This is not really correct but it is much more stable to 
+					// differences in contacts from tick to tick
+					const FSolverReal FrictionMaxPushOut = TotalPushOutNormal / FSolverReal(NumFrictionContacts);
+					const FSolverReal FrictionStiffness = Stiffness * CVars::Chaos_PBDCollisionSolver_Position_StaticFrictionStiffness;
+
+					for (int32 ManifoldPointIndex = 0; ManifoldPointIndex < NumManifoldPoints(); ++ManifoldPointIndex)
+					{
+						if (bApplyFriction[ManifoldPointIndex])
+						{
+							FSolverReal ContactDeltaTangentU, ContactDeltaTangentV;
+							CalculateContactPositionErrorTangential(ManifoldPointIndex, Body0.SolverBody(), Body1.SolverBody(), ContactDeltaTangentU, ContactDeltaTangentV);
+
+							// Bilateral constraint - negative values allowed (unlike the normal correction)
+							FSolverReal PushOutTangentU = -Stiffness * ManifoldPoints[ManifoldPointIndex].ContactMassTangentU * ContactDeltaTangentU;
+							FSolverReal PushOutTangentV = -Stiffness * ManifoldPoints[ManifoldPointIndex].ContactMassTangentV * ContactDeltaTangentV;
+
+							ManifoldPoints[ManifoldPointIndex].NetPushOutTangentU += PushOutTangentU;
+							ManifoldPoints[ManifoldPointIndex].NetPushOutTangentV += PushOutTangentV;
+
+							ApplyFrictionCone(
+								ManifoldPointIndex,
+								StaticFriction,
+								DynamicFriction,
+								FrictionMaxPushOut,
+								PushOutTangentU,			// InOut
+								PushOutTangentV);			// InOut
+
+							// Update the particle state based on the pushout
+							const FSolverVec3 PushOut = PushOutTangentU * ManifoldPoints[ManifoldPointIndex].ContactTangentU + PushOutTangentV * ManifoldPoints[ManifoldPointIndex].ContactTangentV;
+							if (Body0.IsDynamic())
+							{
+								DX0 += Body0.InvM() * PushOut;
+								DR0 += ManifoldPoints[ManifoldPointIndex].ContactTangentUAngular0 * PushOutTangentU + ManifoldPoints[ManifoldPointIndex].ContactTangentVAngular0 * PushOutTangentV;
+							}
+							if (Body1.IsDynamic())
+							{
+								DX1 -= Body1.InvM() * PushOut;
+								DR1 -= ManifoldPoints[ManifoldPointIndex].ContactTangentUAngular1 * PushOutTangentU + ManifoldPoints[ManifoldPointIndex].ContactTangentVAngular1 * PushOutTangentV;
+							}
+						}
+					}
+
+					if (Body0.IsDynamic())
+					{
+						Body0.ApplyPositionDelta(DX0);
+						Body0.ApplyRotationDelta(DR0);
+					}
+					if (Body1.IsDynamic())
+					{
+						Body1.ApplyPositionDelta(DX1);
+						Body1.ApplyRotationDelta(DR1);
+					}
+				}
+
+				return false;
+			}
+
+			/**
+			 * @brief Calculate the velocity error at the current transforms
+			*/
+			FORCEINLINE_DEBUGGABLE void CalculateContactVelocityErrorNormal(
+				const int32 ManifoldPointIndex, 
+				const FConstraintSolverBody& Body0, 
+				const FConstraintSolverBody& Body1, 
+				FSolverReal& OutContactVelocityDeltaNormal) const
+			{
+				const FSolverVec3 ContactVelocity0 = Body0.V() + FSolverVec3::CrossProduct(Body0.W(), ManifoldPoints[ManifoldPointIndex].RelativeContactPosition0);
+				const FSolverVec3 ContactVelocity1 = Body1.V() + FSolverVec3::CrossProduct(Body1.W(), ManifoldPoints[ManifoldPointIndex].RelativeContactPosition1);
+				const FSolverVec3 ContactVelocity = ContactVelocity0 - ContactVelocity1;
+				const FSolverReal ContactVelocityNormal = FSolverVec3::DotProduct(ContactVelocity, ManifoldPoints[ManifoldPointIndex].ContactNormal);
+
+				// Add up the errors in the velocity (current velocity - desired velocity)
+				OutContactVelocityDeltaNormal = (ContactVelocityNormal - ManifoldPoints[ManifoldPointIndex].ContactTargetVelocityNormal);
+			}
+
+			FORCEINLINE_DEBUGGABLE void CalculateContactVelocityError(
+				const int32 ManifoldPointIndex, 
+				const FConstraintSolverBody& Body0, 
+				const FConstraintSolverBody& Body1, 
+				const FSolverReal DynamicFriction, 
+				const FSolverReal Dt, 
+				FSolverReal& OutContactVelocityDeltaNormal, 
+				FSolverReal& OutContactVelocityDeltaTangent0, 
+				FSolverReal& OutContactVelocityDeltaTangent1) const
+			{
+				const FSolverVec3 ContactVelocity0 = Body0.V() + FSolverVec3::CrossProduct(Body0.W(), ManifoldPoints[ManifoldPointIndex].RelativeContactPosition0);
+				const FSolverVec3 ContactVelocity1 = Body1.V() + FSolverVec3::CrossProduct(Body1.W(), ManifoldPoints[ManifoldPointIndex].RelativeContactPosition1);
+				const FSolverVec3 ContactVelocity = ContactVelocity0 - ContactVelocity1;
+				const FSolverReal ContactVelocityNormal = FSolverVec3::DotProduct(ContactVelocity, ManifoldPoints[ManifoldPointIndex].ContactNormal);
+				const FSolverReal ContactVelocityTangent0 = FSolverVec3::DotProduct(ContactVelocity, ManifoldPoints[ManifoldPointIndex].ContactTangentU);
+				const FSolverReal ContactVelocityTangent1 = FSolverVec3::DotProduct(ContactVelocity, ManifoldPoints[ManifoldPointIndex].ContactTangentV);
+
+				OutContactVelocityDeltaNormal = (ContactVelocityNormal - ManifoldPoints[ManifoldPointIndex].ContactTargetVelocityNormal);
+				OutContactVelocityDeltaTangent0 = ContactVelocityTangent0;
+				OutContactVelocityDeltaTangent1 = ContactVelocityTangent1;
+			}
+
+			// @todo(chaos): dynamic friction
+			FORCEINLINE_DEBUGGABLE void SolveVelocity(
+				const FSolverReal Stiffness,
+				const FSolverReal Dt,
+				const bool bApplyDynamicFriction,
+				FConstraintSolverBody& Body0,
+				FConstraintSolverBody& Body1)
+			{
+				FSolverVec3 DV0 = FSolverVec3(0);
+				FSolverVec3 DV1 = FSolverVec3(0);
+				FSolverVec3 DW0 = FSolverVec3(0);
+				FSolverVec3 DW1 = FSolverVec3(0);
+
+				for (int32 ManifoldPointIndex = 0; ManifoldPointIndex < NumManifoldPoints(); ++ManifoldPointIndex)
+				{
+					if (ShouldSolveVelocity(ManifoldPointIndex))
+					{
+						const FSolverReal MinImpulseNormal = FMath::Min(FSolverReal(0), -ManifoldPoints[ManifoldPointIndex].NetPushOutNormal / Dt);
+
+						FSolverReal ContactVelocityDeltaNormal;
+						CalculateContactVelocityErrorNormal(ManifoldPointIndex, Body0, Body1, ContactVelocityDeltaNormal);
+
+						FSolverReal ImpulseNormal = -(Stiffness * ManifoldPoints[ManifoldPointIndex].ContactMassNormal) * ContactVelocityDeltaNormal;
+
+						const FSolverReal NetImpulseNormal = ManifoldPoints[ManifoldPointIndex].NetImpulseNormal + ImpulseNormal;
+						if (NetImpulseNormal < MinImpulseNormal)
+						{
+							ImpulseNormal = ImpulseNormal - NetImpulseNormal + MinImpulseNormal;
+						}
+
+						ManifoldPoints[ManifoldPointIndex].NetImpulseNormal += ImpulseNormal;
+
+						FSolverVec3 Impulse = ImpulseNormal * ManifoldPoints[ManifoldPointIndex].ContactNormal;
+						if (Body0.IsDynamic())
+						{
+							DV0 += Body0.InvM() * Impulse;
+							DW0 += ManifoldPoints[ManifoldPointIndex].ContactNormalAngular0 * ImpulseNormal;
+						}
+						if (Body1.IsDynamic())
+						{
+							DV1 -= Body1.InvM() * Impulse;
+							DW1 -= ManifoldPoints[ManifoldPointIndex].ContactNormalAngular1 * ImpulseNormal;
+						}
+					}
+				}
+
+				if (Body0.IsDynamic())
+				{
+					Body0.ApplyVelocityDelta(DV0, DW0);
+				}
+				if (Body1.IsDynamic())
+				{
+					Body1.ApplyVelocityDelta(DV1, DW1);
+				}
+			}
+
+			// @todo(chaos): make private
+		public:
+			friend class FPBDCollisionSolverJacobi;
+
+			/**
+			 * @brief Whether we need to solve velocity for this manifold point (only if we were penetrating or applied a pushout)
+			*/
+			FORCEINLINE_DEBUGGABLE bool ShouldSolveVelocity(const int32 ManifoldPointIndex) const
+			{
+				// We ensure positive separating velocity for close contacts even if they didn't receive a pushout
+				return (ManifoldPoints[ManifoldPointIndex].NetPushOutNormal > FSolverReal(0)) || (ManifoldPoints[ManifoldPointIndex].ContactDeltaNormal < FSolverReal(0));
+			}
+
+			struct FManifoldPoint
+			{
+				// World-space contact point relative to each particle's center of mass
+				FSolverVec3 RelativeContactPosition0;
+				FSolverVec3 RelativeContactPosition1;
+
+				// World-space contact normal and tangents
+				FSolverVec3 ContactNormal;
+				FSolverVec3 ContactTangentU;
+				FSolverVec3 ContactTangentV;
+
+				// Errors to correct along each of the contact axes
+				FSolverReal ContactDeltaNormal;
+				FSolverReal ContactDeltaTangentU;
+				FSolverReal ContactDeltaTangentV;
+
+				// Target velocity along the normal direction
+				FSolverReal ContactTargetVelocityNormal;
+
+				// I^-1.(R x A) for each body where A is each axis (Normal, TangentU, TangentV)
+				FSolverVec3 ContactNormalAngular0;
+				FSolverVec3 ContactTangentUAngular0;
+				FSolverVec3 ContactTangentVAngular0;
+				FSolverVec3 ContactNormalAngular1;
+				FSolverVec3 ContactTangentUAngular1;
+				FSolverVec3 ContactTangentVAngular1;
+
+				// Contact mass (for non-friction)
+				FSolverReal ContactMassNormal;
+				FSolverReal ContactMassTangentU;
+				FSolverReal ContactMassTangentV;
+
+				// Solver outputs
+				FSolverReal NetPushOutNormal;
+				FSolverReal NetPushOutTangentU;
+				FSolverReal NetPushOutTangentV;
+				FSolverReal NetImpulseNormal;
+				FSolverReal NetImpulseTangentU;
+				FSolverReal NetImpulseTangentV;
+
+				// A measure of how much we exceeded the static friction threshold.
+				// Equal to (NormalPushOut / TangentialPushOut) before clamping to the friction cone.
+				// Used to move the static friction anchors to the edge of the cone in Scatter.
+				FSolverReal StaticFrictionRatio;
+			};
+
+			TCArray<FManifoldPoint, MaxManifoldPoints> ManifoldPoints;
+		};
+
+		/**
+		 * A Jocobi solver for the contact manifold between two shapes. This solves each
+		 * manifold point independently and sums and dampens the net impulses before
+		 * applying to the bodies.
+		 * 
+		 * Pros (vs Gauss Seidel):
+		 *	-	If the contacts are symmetric, the impusles will be too. This fixed spurious rotations when 
+		 *		a box with high restitution bounces exactly flat on the ground.
+		 *	-	We can run the points in the manifold in parallel
+		 *	-	Stacking is much more stable (but see below)
+		 * Cons
+		 *	-	We overestimate the impulse and must apply a damping factor to prevent energy gain. It is not
+		 *		clear how we select the value of this damping - it is empirically tuned.
+		 *	-	The damping factor leads to spongey stacks/piles
+		 * 
+		*/
+		class FPBDCollisionSolverJacobi
+		{
+		public:
+			static const int32 MaxConstrainedBodies = 2;
+			static const int32 MaxPointsPerConstraint = 4;
+
+			// Create a solver that is initialized to safe defaults
+			static FPBDCollisionSolverJacobi MakeInitialized()
+			{
+				FPBDCollisionSolverJacobi Solver;
+				Solver.State.Init();
+				return Solver;
+			}
+
+			// Create a solver with no initialization
+			static FPBDCollisionSolverJacobi MakeUninitialized()
+			{
+				return FPBDCollisionSolverJacobi();
+			}
+
+			// NOTE: Does not initialize any properties. See MakeInitialized
+			FPBDCollisionSolverJacobi() {}
+
+			/** Reset the state of the collision solver */
+			FORCEINLINE_DEBUGGABLE void Reset()
+			{
+				State.SolverBodies[0].Reset();
+				State.SolverBodies[1].Reset();
+				ResetManifold();
+			}
+
+			FORCEINLINE_DEBUGGABLE void ResetManifold()
+			{
+				State.ManifoldPoints.Reset();
+			}
+
+			FORCEINLINE_DEBUGGABLE int32 AddManifoldPoint()
+			{
+				return State.ManifoldPoints.AddManifoldPoint();
+			}
+
+			FORCEINLINE_DEBUGGABLE FSolverReal StaticFriction() const
+			{ 
+				return State.StaticFriction;
+			}
+
+			FORCEINLINE_DEBUGGABLE FSolverReal DynamicFriction() const
+			{ 
+				return State.DynamicFriction; 
+			}
+
+			FORCEINLINE_DEBUGGABLE FSolverReal VelocityFriction() const
+			{ 
+				return State.VelocityFriction;
+			}
+
+			FORCEINLINE_DEBUGGABLE void SetFriction(const FSolverReal InStaticFriction, const FSolverReal InDynamicFriction, const FSolverReal InVelocityFriction)
+			{
+				State.StaticFriction = InStaticFriction;
+				State.DynamicFriction = InDynamicFriction;
+				State.VelocityFriction = InVelocityFriction;
+			}
+
+			FORCEINLINE_DEBUGGABLE void SetStiffness(const FSolverReal InStiffness)
+			{
+				State.Stiffness = InStiffness;
+			}
+
+			FORCEINLINE_DEBUGGABLE void SetSolverBodies(FSolverBody& SolverBody0, FSolverBody& SolverBody1)
+			{
+				State.SolverBodies[0].SetSolverBody(SolverBody0);
+				State.SolverBodies[1].SetSolverBody(SolverBody1);
+			}
+
+			FORCEINLINE_DEBUGGABLE int32 NumManifoldPoints() const
+			{
+				return State.ManifoldPoints.NumManifoldPoints();
+			}
+
+			FORCEINLINE_DEBUGGABLE FSolverVec3 GetNetPushOut(const int32 ManifoldPointIndex) const
+			{
+				return State.ManifoldPoints.GetNetPushOut(ManifoldPointIndex);
+			}
+
+			FORCEINLINE_DEBUGGABLE FSolverVec3 GetNetImpulse(const int32 ManifoldPointIndex) const
+			{
+				return State.ManifoldPoints.GetNetImpulse(ManifoldPointIndex);
+			}
+
+			FORCEINLINE_DEBUGGABLE FSolverReal GetStaticFrictionRatio(const int32 ManifoldPointIndex) const
+			{
+				return State.ManifoldPoints.GetStaticFrictionRatio(ManifoldPointIndex);
+			}
+
+			/**
+			 * Set up a manifold point (also calls FinalizeManifoldPoint)
+			*/
+			FORCEINLINE_DEBUGGABLE void SetManifoldPoint(
+				const int32 ManifoldPointIndex,
+				const FSolverVec3& InRelativeContactPosition0,
+				const FSolverVec3& InRelativeContactPosition1,
+				const FSolverVec3& InWorldContactNormal,
+				const FSolverVec3& InWorldContactTangentU,
+				const FSolverVec3& InWorldContactTangentV,
+				const FSolverReal InWorldContactDeltaNormal,
+				const FSolverReal InWorldContactDeltaTangentU,
+				const FSolverReal InWorldContactDeltaTangentV,
+				const FSolverReal InWorldContactVelocityTargetNormal)
+			{
+				InitManifoldPoint(
+					ManifoldPointIndex,
+					InRelativeContactPosition0,
+					InRelativeContactPosition1,
+					InWorldContactNormal,
+					InWorldContactTangentU,
+					InWorldContactTangentV,
+					InWorldContactDeltaNormal,
+					InWorldContactDeltaTangentU,
+					InWorldContactDeltaTangentV,
+					InWorldContactVelocityTargetNormal);
+
+				FinalizeManifoldPoint(ManifoldPointIndex);
+			}
+
+			FORCEINLINE_DEBUGGABLE void InitManifoldPoint(
+				const int32 ManifoldPointIndex,
+				const FSolverVec3& InRelativeContactPosition0,
+				const FSolverVec3& InRelativeContactPosition1,
+				const FSolverVec3& InWorldContactNormal,
+				const FSolverVec3& InWorldContactTangentU,
+				const FSolverVec3& InWorldContactTangentV,
+				const FSolverReal InWorldContactDeltaNormal,
+				const FSolverReal InWorldContactDeltaTangentU,
+				const FSolverReal InWorldContactDeltaTangentV,
+				const FSolverReal InWorldContactVelocityTargetNormal)
+			{
+				State.ManifoldPoints.SetWorldContact(
+					ManifoldPointIndex,
+					InRelativeContactPosition0,
+					InRelativeContactPosition1,
+					InWorldContactNormal,
+					InWorldContactTangentU,
+					InWorldContactTangentV,
+					InWorldContactDeltaNormal,
+					InWorldContactDeltaTangentU,
+					InWorldContactDeltaTangentV,
+					InWorldContactVelocityTargetNormal
+				);
+			}
+
+			/** 
+			 * Finish manifold point setup.
+			 * NOTE: Can only be called after the InitManifoldPoint has been called
+			*/
+			FORCEINLINE_DEBUGGABLE void FinalizeManifoldPoint(const int32 ManifoldPointIndex)
+			{
+				State.ManifoldPoints.InitContact(
+					ManifoldPointIndex,
+					State.SolverBodies[0],
+					State.SolverBodies[1]);
+			}
+
+			/**
+			 * @brief Get the first (decorated) solver body
+			 * The decorator add a possible mass scale
+			*/
+			FORCEINLINE_DEBUGGABLE FConstraintSolverBody& SolverBody0() { return State.SolverBodies[0]; }
+			FORCEINLINE_DEBUGGABLE const FConstraintSolverBody& SolverBody0() const { return State.SolverBodies[0]; }
+
+			/**
+			 * @brief Get the second (decorated) solver body
+			 * The decorator add a possible mass scale
+			*/
+			FORCEINLINE_DEBUGGABLE FConstraintSolverBody& SolverBody1() { return State.SolverBodies[1]; }
+			FORCEINLINE_DEBUGGABLE const FConstraintSolverBody& SolverBody1() const { return State.SolverBodies[1]; }
+
+			/**
+			 * @brief Set up the mass scaling for shock propagation, using the position-phase mass scale
+			*/
+			void EnablePositionShockPropagation();
+
+			/**
+			 * @brief Set up the mass scaling for shock propagation, using the velocity-phase mass scale
+			*/
+			void EnableVelocityShockPropagation();
+
+			/**
+			 * @brief Disable mass scaling
+			*/
+			void DisableShockPropagation();
+
+			/**
+			 * @brief Calculate and apply the position correction for this iteration
+			 * @return true if we need to run more iterations, false if we did not apply any correction
+			*/
+			FORCEINLINE_DEBUGGABLE bool SolvePositionNoFriction(const FSolverReal Dt, const FSolverReal MaxPushOut)
+			{
+				// SolverBody decorator used to add mass scaling
+				FConstraintSolverBody& Body0 = SolverBody0();
+				FConstraintSolverBody& Body1 = SolverBody1();
+
+				const FSolverReal Stiffness = CVars::Chaos_PBDCollisionSolver_JacobiStiffness * State.Stiffness;
+
+				State.ManifoldPoints.SolvePositionNoFriction(
+					Stiffness,
+					MaxPushOut,
+					Body0, 
+					Body1);
+
+				return false;
+			}
+
+			FORCEINLINE_DEBUGGABLE bool SolvePositionWithFriction(const FSolverReal Dt, const FSolverReal MaxPushOut)
+			{
+				// SolverBody decorator used to add mass scaling
+				FConstraintSolverBody& Body0 = SolverBody0();
+				FConstraintSolverBody& Body1 = SolverBody1();
+
+				const FSolverReal Stiffness = CVars::Chaos_PBDCollisionSolver_JacobiStiffness * State.Stiffness;
+
+				State.ManifoldPoints.SolvePositionWithFriction(
+					Stiffness,
+					Dt,
+					State.StaticFriction,
+					State.DynamicFriction,
+					MaxPushOut,
+					Body0,
+					Body1);
+
+				return false;
+			}
+
+
+			/**
+			 * @brief Calculate and apply the velocity correction for this iteration
+			 * @return true if we need to run more iterations, false if we did not apply any correction
+			*/
+			FORCEINLINE_DEBUGGABLE bool SolveVelocity(const FSolverReal Dt, const bool bApplyDynamicFriction)
+			{
+				FConstraintSolverBody& Body0 = SolverBody0();
+				FConstraintSolverBody& Body1 = SolverBody1();
+
+				const FSolverReal Stiffness = CVars::Chaos_PBDCollisionSolver_JacobiStiffness * State.Stiffness;
+
+				State.ManifoldPoints.SolveVelocity(
+					Stiffness,
+					Dt, 
+					bApplyDynamicFriction,
+					Body0,
+					Body1);
+
+				return false;
+			}
+
+		private:
+			/**
+			 * @brief Apply the inverse mass scale the body with the lower level
+			 * @param InvMassScale
+			*/
+			void SetShockPropagationInvMassScale(const FSolverReal InvMassScale);
+
+			struct FState
+			{
+				FState()
+				{
+				}
+
+				void Init()
+				{
+					SolverBodies[0].Init();
+					SolverBodies[1].Init();
+					StaticFriction = 0;
+					DynamicFriction = 0;
+					VelocityFriction = 0;
+					Stiffness = 1;
+					ManifoldPoints.Reset();
+				}
+
+				// Static Friction in the position-solve phase
+				FSolverReal StaticFriction;
+
+				// Dynamic Friction in the position-solve phase
+				FSolverReal DynamicFriction;
+
+				// Dynamic Friction in the velocity-solve phase
+				FSolverReal VelocityFriction;
+
+				// Solver stiffness (scales all pushout and impulses)
+				FSolverReal Stiffness;
+
+				// Bodies and contacts
+				FConstraintSolverBody SolverBodies[MaxConstrainedBodies];
+				FPBDCollisionSolverJacobiManifoldPoints ManifoldPoints;
+			};
+
+			FState State;
+		};
+
+
+		//////////////////////////////////////////////////////////////////////////////////////////////////
+		//////////////////////////////////////////////////////////////////////////////////////////////////
+		//////////////////////////////////////////////////////////////////////////////////////////////////
+		//////////////////////////////////////////////////////////////////////////////////////////////////
+
+
+		/**
+		 * A helper for solving arrays of constraints
+		 */
+		class FPBDCollisionSolverJacobiHelper
+		{
+		public:
+			static void SolvePositionNoFriction(const TArrayView<FPBDCollisionSolverJacobi>& CollisionSolvers, const FSolverReal Dt, const FSolverReal MaxPushOut);
+			static void SolvePositionWithFriction(const TArrayView<FPBDCollisionSolverJacobi>& CollisionSolvers, const FSolverReal Dt, const FSolverReal MaxPushOut);
+			static void SolveVelocity(const TArrayView<FPBDCollisionSolverJacobi>& CollisionSolvers, const FSolverReal Dt, const bool bApplyDynamicFriction);
+
+			static void CheckISPC();
+		};
+
+	}	// namespace Private
+}	// namespace Chaos
+
diff --git a/Engine/Source/Runtime/Experimental/Chaos/Public/Chaos/Collision/PBDCollisionSolverSettings.h b/Engine/Source/Runtime/Experimental/Chaos/Public/Chaos/Collision/PBDCollisionSolverSettings.h
new file mode 100644
--- /dev/null
+++ b/Engine/Source/Runtime/Experimental/Chaos/Public/Chaos/Collision/PBDCollisionSolverSettings.h
@@ -0,0 +1,44 @@
+// Copyright Epic Games, Inc. All Rights Reserved.
+#pragma once
+
+#include "Chaos/Core.h"
+
+namespace Chaos
+{
+	/**
+		* @brief Settings to control the low-level collision solver behaviour
+	*/
+	class FPBDCollisionSolverSettings
+	{
+	public:
+		UE_NONCOPYABLE(FPBDCollisionSolverSettings);
+
+		FPBDCollisionSolverSettings()
+			: MaxPushOutVelocity(0)
+			, NumPositionFrictionIterations(4)
+			, NumVelocityFrictionIterations(1)
+			, NumPositionShockPropagationIterations(3)
+			, NumVelocityShockPropagationIterations(1)
+		{
+		}
+
+		// Maximum speed at which two objects can depenetrate (actually, how much relative velocity can be added
+		// to a contact per frame when depentrating. Stacks and deep penetrations can lead to larger velocities)
+		FReal MaxPushOutVelocity;
+
+		// How many of the position iterations should run static/dynamic friction
+		int32 NumPositionFrictionIterations;
+
+		// How many of the velocity iterations should run dynamic friction
+		// @todo(chaos): if NumVelocityFrictionIterations > 1, then dynamic friction in the velocity phase will be iteration 
+		// count dependent (velocity-solve friction is currentlyused by quadratic shapes and RBAN)
+		int32 NumVelocityFrictionIterations;
+
+		// How many position iterations should have shock propagation enabled
+		int32 NumPositionShockPropagationIterations;
+
+		// How many velocity iterations should have shock propagation enabled
+		int32 NumVelocityShockPropagationIterations;
+	};
+
+}	// namespace Chaos
diff --git a/Engine/Source/Runtime/Experimental/Chaos/Public/Chaos/Collision/SolverCollisionContainer.h b/Engine/Source/Runtime/Experimental/Chaos/Public/Chaos/Collision/SolverCollisionContainer.h
--- a/Engine/Source/Runtime/Experimental/Chaos/Public/Chaos/Collision/SolverCollisionContainer.h
+++ b/Engine/Source/Runtime/Experimental/Chaos/Public/Chaos/Collision/SolverCollisionContainer.h
@@ -1,10 +1,10 @@
 // Copyright Epic Games, Inc. All Rights Reserved.
 #pragma once
 
-#include "CoreMinimal.h"
 #include "Chaos/CollisionResolutionTypes.h"
 #include "Chaos/Collision/CollisionApplyType.h"
 #include "Chaos/Collision/PBDCollisionSolver.h"
+#include "Chaos/Collision/PBDCollisionSolverSettings.h"
 #include "Chaos/Evolution/SolverConstraintContainer.h"
 
 namespace Chaos
@@ -19,35 +19,6 @@ namespace Chaos
 	class FSolverBody;
 	class FSolverBodyContainer;
 
-	/**
-	 * @brief Settings to control the low-level collision solver behaviour
-	*/
-	class FPBDCollisionSolverSettings
-	{
-	public:
-		UE_NONCOPYABLE(FPBDCollisionSolverSettings);
-
-		FPBDCollisionSolverSettings();
-
-		// Maximum speed at which two objects can depenetrate (actually, how much relative velocity can be added
-		// to a contact per frame when depentrating. Stacks and deep penetrations can lead to larger velocities)
-		FReal MaxPushOutVelocity;
-
-		// How many of the position iterations should run static/dynamic friction
-		int32 NumPositionFrictionIterations;
-
-		// How many of the velocity iterations should run dynamic friction
-		// @todo(chaos): if NumVelocityFrictionIterations > 1, then dynamic friction in the velocity phase will be iteration 
-		// count dependent (velocity-solve friction is currentlyused by quadratic shapes and RBAN)
-		int32 NumVelocityFrictionIterations;
-
-		// How many position iterations should have shock propagation enabled
-		int32 NumPositionShockPropagationIterations;
-
-		// How many velocity iterations should have shock propagation enabled
-		int32 NumVelocityShockPropagationIterations;
-	};
-
 	/**
 	 * The solver for a set of collision constraints. This collects all the data required to solve a set of collision
 	 * constraints into a contiguous, ordered buffer.
diff --git a/Engine/Source/Runtime/Experimental/Chaos/Public/Chaos/Island/IslandGroupManager.h b/Engine/Source/Runtime/Experimental/Chaos/Public/Chaos/Island/IslandGroupManager.h
--- a/Engine/Source/Runtime/Experimental/Chaos/Public/Chaos/Island/IslandGroupManager.h
+++ b/Engine/Source/Runtime/Experimental/Chaos/Public/Chaos/Island/IslandGroupManager.h
@@ -92,6 +92,11 @@ namespace Chaos
 			*/
 			void AddConstraintContainer(FPBDConstraintContainer& ConstraintContainer, const int32 Priority = 0);
 
+			/**
+			 * Remove a previously-added container. Generall this is only needed for debuggin as cleanup is automatic on destruction.
+			 */
+			void RemoveConstraintContainer(FPBDConstraintContainer& ConstraintContainer);
+
 			/**
 			 * Set the priority for the specified container (that must have been pre-registered with AddConstraintContainer).
 			*/
diff --git a/Engine/Source/Runtime/Experimental/Chaos/Public/Chaos/PBDCollisionConstraints.h b/Engine/Source/Runtime/Experimental/Chaos/Public/Chaos/PBDCollisionConstraints.h
--- a/Engine/Source/Runtime/Experimental/Chaos/Public/Chaos/PBDCollisionConstraints.h
+++ b/Engine/Source/Runtime/Experimental/Chaos/Public/Chaos/PBDCollisionConstraints.h
@@ -7,7 +7,7 @@
 #include "Chaos/Collision/CollisionContext.h"
 #include "Chaos/Collision/PBDCollisionConstraint.h"
 #include "Chaos/Collision/PBDCollisionConstraintHandle.h"
-#include "Chaos/Collision/SolverCollisionContainer.h"
+#include "Chaos/Collision/PBDCollisionSolverSettings.h"
 #include "Chaos/PBDConstraintContainer.h"
 #include "Framework/BufferedData.h"
 
@@ -28,6 +28,17 @@ using FRigidBodyContactConstraintsPostComputeCallback = TFunction<void()>;
 using FRigidBodyContactConstraintsPostApplyCallback = TFunction<void(const FReal Dt, const TArray<FPBDCollisionConstraintHandle*>&)>;
 using FRigidBodyContactConstraintsPostApplyPushOutCallback = TFunction<void(const FReal Dt, const TArray<FPBDCollisionConstraintHandle*>&, bool)>;
 
+
+namespace Private
+{
+	// The type of solver to use for collisions
+	enum class ECollisionSolverType
+	{
+		GaussSeidel,
+		PartialJacobi,
+	};
+}
+
 /**
  * A container and solver for collision constraints.
  * 
@@ -141,17 +152,22 @@ public:
 	virtual void PrepareTick() override final {}
 	virtual void UnprepareTick() override final {}
 
-	virtual TUniquePtr<FConstraintContainerSolver> CreateSceneSolver(const int32 Priority) override final
+	virtual TUniquePtr<FConstraintContainerSolver> CreateSceneSolver(const int32 Priority) override final;
+
+	virtual TUniquePtr<FConstraintContainerSolver> CreateGroupSolver(const int32 Priority) override final;
+
+	// The type of solver we are creating
+	Private::ECollisionSolverType GetSolverType() const
 	{
-		return MakeUnique<FPBDCollisionContainerSolver>(*this, Priority);
+		return CollisionSolverType;
 	}
 
-	virtual TUniquePtr<FConstraintContainerSolver> CreateGroupSolver(const int32 Priority) override final
+	// Set the solver type. NOTE: Any previously created solvers will not be recreated at this level. (See FPBDRigidsEvolutionGBF::UpdateCollisionSolverType)
+	void SetSolverType(const Private::ECollisionSolverType InSolverType)
 	{
-		return MakeUnique<FPBDCollisionContainerSolver>(*this, Priority);
+		CollisionSolverType = InSolverType;
 	}
 
-
 	//
 	// Member Access
 	//
@@ -277,7 +293,6 @@ protected:
 	void PruneEdgeCollisions();
 
 private:
-
 	const FPBDRigidsSOAs& Particles;
 
 	Private::FCollisionConstraintAllocator ConstraintAllocator;
@@ -301,6 +316,8 @@ private:
 	// This improves performance, but can decrease stability if contacts are culled prematurely.
 	bool bCanDisableContacts;
 
+	Private::ECollisionSolverType CollisionSolverType;
+
 	// Used to determine constraint directions
 	FVec3 GravityDirection;
 	FReal GravitySize;
diff --git a/Engine/Source/Runtime/Experimental/Chaos/Public/Chaos/PBDRigidsEvolutionGBF.h b/Engine/Source/Runtime/Experimental/Chaos/Public/Chaos/PBDRigidsEvolutionGBF.h
--- a/Engine/Source/Runtime/Experimental/Chaos/Public/Chaos/PBDRigidsEvolutionGBF.h
+++ b/Engine/Source/Runtime/Experimental/Chaos/Public/Chaos/PBDRigidsEvolutionGBF.h
@@ -347,6 +347,8 @@ namespace Chaos
 			return CurrentStepResimCacheImp;
 		}
 
+		void UpdateCollisionSolverType();
+
 		FRigidClustering Clustering;
 
 		FPBDJointConstraints JointConstraints;
